Copying //UE4/Dev-Core to Dev-Main (//UE4/Dev-Main)

diff --git a/Engine/Plugins/FX/Niagara/Source/Niagara/Private/NiagaraDataInterfaceCollisionQuery.cpp b/Engine/Plugins/FX/Niagara/Source/Niagara/Private/NiagaraDataInterfaceCollisionQuery.cpp
--- a/Engine/Plugins/FX/Niagara/Source/Niagara/Private/NiagaraDataInterfaceCollisionQuery.cpp
+++ b/Engine/Plugins/FX/Niagara/Source/Niagara/Private/NiagaraDataInterfaceCollisionQuery.cpp
@@ -17,9 +17,9 @@ FCriticalSection UNiagaraDataInterfaceCollisionQuery::CriticalSection;
 UNiagaraDataInterfaceCollisionQuery::UNiagaraDataInterfaceCollisionQuery(FObjectInitializer const& ObjectInitializer)
 	: Super(ObjectInitializer)
 {
-	TraceChannelEnum = FindObject<UEnum>(ANY_PACKAGE, TEXT("ECollisionChannel"), true);
+	TraceChannelEnum = StaticEnum<ECollisionChannel>();
 
-	Proxy = MakeShared<FNiagaraDataIntefaceProxyCollisionQuery, ESPMode::ThreadSafe>();
+    Proxy = MakeShared<FNiagaraDataIntefaceProxyCollisionQuery, ESPMode::ThreadSafe>();
 }
 
 bool UNiagaraDataInterfaceCollisionQuery::InitPerInstanceData(void* PerInstanceData, FNiagaraSystemInstance* InSystemInstance)
diff --git a/Engine/Source/Developer/TraceAnalysis/Private/Analysis/Context.cpp b/Engine/Source/Developer/TraceAnalysis/Private/Analysis/Context.cpp
new file mode 100644
--- /dev/null
+++ b/Engine/Source/Developer/TraceAnalysis/Private/Analysis/Context.cpp
@@ -0,0 +1,28 @@
+// Copyright 1998-2019 Epic Games, Inc. All Rights Reserved.
+
+#include "Engine.h"
+#include "Trace/Analysis.h"
+
+namespace Trace
+{
+
+////////////////////////////////////////////////////////////////////////////////
+void FAnalysisContext::AddAnalyzer(IAnalyzer& Analyzer)
+{
+	Analyzers.Add(&Analyzer);
+}
+
+////////////////////////////////////////////////////////////////////////////////
+FAnalysisProcessor FAnalysisContext::Process()
+{
+	FAnalysisProcessor Processor;
+
+	if (Analyzers.Num())
+	{
+		Processor.Impl = new FAnalysisEngine(MoveTemp(Analyzers));
+	}
+
+	return MoveTemp(Processor);
+}
+
+} // namespace Trace
diff --git a/Engine/Source/Developer/TraceAnalysis/Private/Analysis/Engine.cpp b/Engine/Source/Developer/TraceAnalysis/Private/Analysis/Engine.cpp
new file mode 100644
--- /dev/null
+++ b/Engine/Source/Developer/TraceAnalysis/Private/Analysis/Engine.cpp
@@ -0,0 +1,612 @@
+// Copyright 1998-2019 Epic Games, Inc. All Rights Reserved.
+
+#include "Engine.h"
+#include "Containers/ArrayView.h"
+#include "CoreGlobals.h"
+#include "HAL/UnrealMemory.h"
+#include "Trace/Analysis.h"
+#include "Trace/Analyzer.h"
+#if 0
+#include "Trace/Private/DataDecoder.h"
+#endif // 0
+#include "Trace/Private/Event.h"
+
+namespace Trace
+{
+
+////////////////////////////////////////////////////////////////////////////////
+class FFnv1aHash
+{
+public:
+					FFnv1aHash() = default;
+					FFnv1aHash(uint32 PrevResult)		{ Result = PrevResult; }
+	void			Add(const ANSICHAR* String)			{ for (; *String; ++String) { Add(*String); } }
+	const uint8*	Add(const uint8* Data, uint32 Size)	{ for (uint32 i = 0; i < Size; ++Data, ++i) { Add(*Data); } return Data; }
+	void			Add(uint8 Value)					{ Result ^= Value; Result *= 0x01000193; }
+	uint32			Get() const							{ return Result; }
+
+private:
+	uint32			Result = 0x811c9dc5;
+	// uint32: bias=0x811c9dc5			prime=0x01000193
+	// uint64: bias=0xcbf29ce484222325	prime=0x00000100000001b3;
+};
+
+
+
+////////////////////////////////////////////////////////////////////////////////
+class FTransportReader
+{
+public:
+	void					SetSource(FStreamReader::FData& InSource);
+	template <typename RetType>
+	RetType const*			GetPointer();
+	template <typename RetType>
+	RetType const*			GetPointer(uint32 BlockSize);
+	virtual void			Advance(uint32 BlockSize);
+
+protected:
+	virtual const uint8*	GetPointerImpl(uint32 BlockSize);
+	FStreamReader::FData*	Source;
+};
+
+////////////////////////////////////////////////////////////////////////////////
+void FTransportReader::SetSource(FStreamReader::FData& InSource)
+{
+	Source = &InSource;
+}
+
+////////////////////////////////////////////////////////////////////////////////
+template <typename RetType>
+RetType const* FTransportReader::GetPointer()
+{
+	return GetPointer<RetType>(sizeof(RetType));
+}
+
+////////////////////////////////////////////////////////////////////////////////
+template <typename RetType>
+RetType const* FTransportReader::GetPointer(uint32 BlockSize)
+{
+	return (RetType const*)GetPointerImpl(BlockSize);
+}
+
+////////////////////////////////////////////////////////////////////////////////
+void FTransportReader::Advance(uint32 BlockSize)
+{
+	Source->Advance(BlockSize);
+}
+
+////////////////////////////////////////////////////////////////////////////////
+const uint8* FTransportReader::GetPointerImpl(uint32 BlockSize)
+{
+	return Source->GetPointer(BlockSize);
+}
+
+
+
+////////////////////////////////////////////////////////////////////////////////
+#if 0
+class FLz4TransportReader
+	: public FTransportReader
+{
+	enum : uint32 { MaxChunks = 4 };
+
+public:
+							FLz4TransportReader(uint32 InChunkPow2);
+	virtual					~FLz4TransportReader();
+
+private:
+	virtual void			Advance(uint32 BlockSize) override;
+	virtual const uint8*	GetPointerImpl(uint32 BlockSize) override;
+	bool					NextChunk();
+    FDataDecoder			Decoder;
+	uint8*					Buffer;
+	const uint8*			Cursor;
+	const uint8*			End;
+	uint32					Index = 0;
+	const uint32			ChunkPow2;
+};
+
+////////////////////////////////////////////////////////////////////////////////
+FLz4TransportReader::FLz4TransportReader(uint32 InChunkPow2)
+: ChunkPow2(InChunkPow2)
+{
+	uint32 ChunkSize = 1 << ChunkPow2;
+	Buffer = new uint8[ChunkSize * (MaxChunks + 1)];
+	Cursor = End = Buffer + (ChunkSize * 2);
+}
+
+////////////////////////////////////////////////////////////////////////////////
+FLz4TransportReader::~FLz4TransportReader()
+{
+	delete[] Buffer;
+}
+
+////////////////////////////////////////////////////////////////////////////////
+void FLz4TransportReader::Advance(uint32 BlockSize)
+{
+	Cursor += BlockSize;
+}
+
+////////////////////////////////////////////////////////////////////////////////
+const uint8* FLz4TransportReader::GetPointerImpl(uint32 BlockSize)
+{
+	while (true)
+	{
+		uint32 Remaining = uint32(UPTRINT(End - Cursor));
+		if (Remaining >= BlockSize)
+		{
+			break;
+		}
+
+		if (!NextChunk())
+		{
+			return nullptr;
+		}
+	}
+
+	return Cursor;
+}
+
+////////////////////////////////////////////////////////////////////////////////
+bool FLz4TransportReader::NextChunk()
+{
+	const struct {
+		uint32	Size;
+		uint8	Data[];
+	}* Block;
+
+	Block = decltype(Block)(Source->GetPointer(sizeof(*Block)));
+	if (Block == nullptr)
+	{
+		return false;
+	}
+
+	int32 BlockSize = Block->Size;
+	if (BlockSize < 0)
+	{
+		BlockSize = ~BlockSize;
+	}
+
+	Block = decltype(Block)(Source->GetPointer(BlockSize));
+	if (Block == nullptr)
+	{
+		return false;
+	}
+
+	Index = (Index + 1) & (MaxChunks - 1);
+	if (!Index)
+	{
+		uint32 Remaining = uint32(UPTRINT(End - Cursor));
+		uint32 OverflowChunks = Remaining >> ChunkPow2;
+		if (OverflowChunks >= MaxChunks - 1)
+		{
+			return false;
+		}
+
+		char* Dest = (char*)(Buffer + (1ull << ChunkPow2));
+		Dest -= Remaining - OverflowChunks;
+		memcpy(Dest, Cursor, Remaining);
+		Cursor = (uint8*)Dest;
+		End = Cursor + Remaining;
+	}
+
+	uint32 SrcSize = BlockSize - sizeof(*Block);
+	uint32 DestSize = 1 << ChunkPow2;
+	int DecodedSize = Decoder.Decode(Block->Data, End, SrcSize, DestSize);
+	if (DecodedSize < 0)
+	{
+		return false;
+	}
+
+	End += DecodedSize;
+
+	if (BlockSize != Block->Size)
+	{
+		Decoder.Reset();
+	}
+
+	Source->Advance(BlockSize);
+	return true;
+}
+#endif // 0
+
+
+
+////////////////////////////////////////////////////////////////////////////////
+struct FAnalysisEngine::FEventDataImpl
+	: public FEventData
+{
+	virtual					~FEventDataImpl() = default;
+	virtual const FValue&	GetValue(const ANSICHAR* FieldName) const override;
+	virtual const FArray&	GetArray(const ANSICHAR* FieldName) const override;
+	virtual const uint8*	GetData() const override;
+	virtual const uint8*	GetAttachment() const override;
+	virtual uint16			GetAttachmentSize() const override;
+	virtual uint16			GetTotalSize() const override;
+	const FDispatch*		Dispatch;
+	const uint8*			Ptr;
+	uint16					Size;
+};
+
+////////////////////////////////////////////////////////////////////////////////
+const IAnalyzer::FValue& FAnalysisEngine::FEventDataImpl::GetValue(const ANSICHAR* FieldName) const
+{
+	UPTRINT Ret = ~0ull;
+
+	FFnv1aHash Hash;
+	Hash.Add(FieldName);
+	uint32 NameHash = Hash.Get();
+
+	for (int i = 0, n = Dispatch->FieldCount; i < n; ++i)
+	{
+		const auto& Field = Dispatch->Fields[i];
+		if (Field.Hash == NameHash)
+		{
+			Ret = UPTRINT(Ptr + Field.Offset);
+			Ret |= UPTRINT(Field.TypeInfo) << 48;
+			break;
+		}
+	}
+
+	return *(FValue*)Ret;
+}
+
+////////////////////////////////////////////////////////////////////////////////
+const IAnalyzer::FArray& FAnalysisEngine::FEventDataImpl::GetArray(const ANSICHAR* FieldName) const
+{
+	return *(FArray*)0x493;
+}
+
+////////////////////////////////////////////////////////////////////////////////
+const uint8* FAnalysisEngine::FEventDataImpl::GetData() const
+{
+	return Ptr;
+}
+
+////////////////////////////////////////////////////////////////////////////////
+const uint8* FAnalysisEngine::FEventDataImpl::GetAttachment() const
+{
+	return Ptr + Dispatch->EventSize;
+}
+
+////////////////////////////////////////////////////////////////////////////////
+uint16 FAnalysisEngine::FEventDataImpl::GetAttachmentSize() const
+{
+	return Size - Dispatch->EventSize;
+}
+
+////////////////////////////////////////////////////////////////////////////////
+uint16 FAnalysisEngine::FEventDataImpl::GetTotalSize() const
+{
+	return Size;
+}
+
+
+
+////////////////////////////////////////////////////////////////////////////////
+enum ERouteId : uint16
+{
+	RouteId_NewEvent,
+	RouteId_NewTrace,
+	RouteId_Timing,
+};
+
+////////////////////////////////////////////////////////////////////////////////
+FAnalysisEngine::FAnalysisEngine(TArray<IAnalyzer*>&& InAnalyzers)
+: Analyzers(MoveTemp(InAnalyzers))
+{
+	EventDataImpl = new FEventDataImpl();
+
+	uint16 SelfIndex = Analyzers.Num();
+	Analyzers.Add(this);
+
+	// Manually add event routing for known events, and those we don't quite know
+	// yet but are expecting.
+	FDispatch& NewEventDispatch = AddDispatch(uint16(FNewEventEvent::Uid), 0);
+	NewEventDispatch.FirstRoute = 0;
+	AddRoute(SelfIndex, RouteId_NewEvent, 0);
+	AddRoute(SelfIndex, RouteId_NewTrace, "$Trace", "NewTrace");
+	AddRoute(SelfIndex, RouteId_Timing, "$Trace", "Timing");
+}
+
+////////////////////////////////////////////////////////////////////////////////
+FAnalysisEngine::~FAnalysisEngine()
+{
+	for (IAnalyzer* Analyzer : Analyzers)
+	{
+		Analyzer->OnAnalysisEnd();
+	}
+
+	for (FDispatch* Dispatch : Dispatches)
+	{
+		FMemory::Free(Dispatch);
+	}
+
+	delete EventDataImpl;
+}
+
+////////////////////////////////////////////////////////////////////////////////
+void FAnalysisEngine::OnAnalysisBegin(const FOnAnalysisContext& Context)
+{
+}
+
+////////////////////////////////////////////////////////////////////////////////
+void FAnalysisEngine::OnAnalysisEnd()
+{
+}
+
+////////////////////////////////////////////////////////////////////////////////
+void FAnalysisEngine::OnEvent(uint16 RouteId, const FOnEventContext& Context)
+{
+	switch (RouteId)
+	{
+	case RouteId_NewEvent:	return OnNewEvent(Context);
+	case RouteId_NewTrace:	return OnNewTrace(Context);
+	case RouteId_Timing:	return OnTiming(Context);
+	}
+}
+
+////////////////////////////////////////////////////////////////////////////////
+void FAnalysisEngine::AddRoute(
+	uint16 AnalyzerIndex,
+	uint16 Id,
+	const ANSICHAR* Logger,
+	const ANSICHAR* Event)
+{
+	FFnv1aHash Hash;
+	Hash.Add(Logger);
+	Hash.Add(Event);
+	AddRoute(AnalyzerIndex, Id, Hash.Get());
+}
+
+////////////////////////////////////////////////////////////////////////////////
+void FAnalysisEngine::AddRoute(
+	uint16 AnalyzerIndex,
+	uint16 Id,
+	uint32 Hash)
+{
+	check(AnalyzerIndex < Analyzers.Num());
+
+	uint16 HashIndex = 0;
+	for (uint16 n = Hashes.Num(); HashIndex < n; ++HashIndex)
+	{
+		if (Hashes[HashIndex] == Hash)
+		{
+			break;
+		}
+	}
+
+	if (HashIndex == Hashes.Num())
+	{
+		Hashes.Add(Hash);
+	}
+
+	FRoute& Route = Routes.Emplace_GetRef();
+	Route.Id = Id;
+	Route.HashIndex = HashIndex;
+	Route.Count = 1;
+	Route.AnalyzerIndex = AnalyzerIndex;
+}
+
+////////////////////////////////////////////////////////////////////////////////
+void FAnalysisEngine::OnNewTrace(const FOnEventContext& Context)
+{
+	struct : IAnalyzer::FInterfaceBuilder
+	{
+		virtual void RouteEvent(uint16 RouteId, const ANSICHAR* Logger, const ANSICHAR* Event) override
+		{
+			Self->AddRoute(AnalyzerIndex, RouteId, Logger, Event);
+		}
+
+		FAnalysisEngine* Self;
+		uint16 AnalyzerIndex;
+	} Builder;
+	Builder.Self = this;
+
+	FOnAnalysisContext OnAnalysisContext = { { SessionContext }, Builder };
+	for (uint16 i = 0, n = Analyzers.Num(); i < n; ++i)
+	{
+		Builder.AnalyzerIndex = i;
+		Analyzers[i]->OnAnalysisBegin(OnAnalysisContext);
+	}
+
+	Algo::SortBy(Routes, [] (const FRoute& Route) { return Route.HashIndex; });
+
+	FRoute* Cursor = Routes.GetData();
+	Cursor->Count = 1;
+
+	for (uint16 i = 1, n = Routes.Num(); i < n; ++i)
+	{
+		if (Routes[i].HashIndex == Cursor->HashIndex)
+		{
+			Cursor->Count++;
+		}
+		else
+		{
+			Cursor = Routes.GetData() + i;
+			Cursor->Count = 1;
+		}
+	}
+
+	// Add a terminal route for events that aren't subscribed to
+	FRoute& Route = Routes.Emplace_GetRef();
+	Route.HashIndex = Hashes.Num();
+	Route.Count = 0;
+}
+
+////////////////////////////////////////////////////////////////////////////////
+void FAnalysisEngine::OnTiming(const FOnEventContext& Context)
+{
+	SessionContext.StartCycle = Context.EventData.GetValue("StartCycle").As<uint64>();
+	SessionContext.CycleFrequency = Context.EventData.GetValue("CycleFrequency").As<uint64>();
+}
+
+////////////////////////////////////////////////////////////////////////////////
+FAnalysisEngine::FDispatch& FAnalysisEngine::AddDispatch(uint16 Uid, uint16 FieldCount)
+{
+	// Allocate a block of memory to hold the dispatch
+	uint32 Size = sizeof(FDispatch) + (sizeof(FDispatch::FField) * FieldCount);
+	auto* Dispatch = (FDispatch*)FMemory::Malloc(Size);
+	Dispatch->FieldCount = FieldCount;
+	Dispatch->EventSize = 0;
+	Dispatch->FirstRoute = -1;
+
+	// Add the new dispatch in the dispatch table
+	if (Uid >= Dispatches.Num())
+	{
+		Dispatches.SetNum(Uid + 1);
+	}
+	check(Dispatches[Uid] == nullptr);
+	Dispatches[Uid] = Dispatch;
+
+	return *Dispatch;
+}
+
+////////////////////////////////////////////////////////////////////////////////
+void FAnalysisEngine::OnNewEvent(const FOnEventContext& Context)
+{
+	const auto& NewEvent = *(FNewEventEvent*)Context.EventData.GetData();
+
+	FDispatch& Dispatch = AddDispatch(NewEvent.EventUid, NewEvent.FieldCount);
+
+	if (NewEvent.FieldCount)
+	{
+		auto& LastField = NewEvent.Fields[NewEvent.FieldCount - 1];
+		Dispatch.EventSize = LastField.Offset + LastField.Size;
+	}
+
+	const uint8* NameCursor = (uint8*)(NewEvent.Fields + NewEvent.FieldCount);
+
+	// Calculate this dispatch's hash.
+	FFnv1aHash DispatchHash;
+	NameCursor = DispatchHash.Add(NameCursor, NewEvent.LoggerNameSize);
+	NameCursor = DispatchHash.Add(NameCursor, NewEvent.EventNameSize);
+
+	// Fill out the fields
+	for (int i = 0, n = NewEvent.FieldCount; i < n; ++i)
+	{
+		const auto& In = NewEvent.Fields[i];
+		auto& Out = Dispatch.Fields[i];
+
+		FFnv1aHash FieldHash;
+		NameCursor = FieldHash.Add(NameCursor, In.NameSize);
+
+		Out.Hash = FieldHash.Get();
+		Out.Offset = In.Offset;
+		Out.Size = In.Size;
+		Out.TypeInfo = In.TypeInfo;
+	}
+
+	TArrayView<FDispatch::FField> Fields(Dispatch.Fields, Dispatch.FieldCount);
+	Algo::SortBy(Fields, [] (const auto& Field) { return Field.Hash; });
+
+	// Find routes that have subscribed to this event.
+	uint16 HashIndex = 0;
+	for (uint16 n = Hashes.Num(); HashIndex < n; ++HashIndex)
+	{
+		if (Hashes[HashIndex] == DispatchHash.Get())
+		{
+			break;
+		}
+	}
+
+	uint16 RouteIndex = 0;
+	for (uint16 n = Routes.Num(); RouteIndex < n; ++RouteIndex)
+	{
+		if (Routes[RouteIndex].HashIndex == HashIndex)
+		{
+			break;
+		}
+	}
+
+	check(RouteIndex < Routes.Num());
+	Dispatch.FirstRoute = RouteIndex;
+}
+
+////////////////////////////////////////////////////////////////////////////////
+bool FAnalysisEngine::EstablishTransport(FStreamReader::FData& Data)
+{
+	const struct {
+		uint8 Format;
+		uint8 Parameter;
+	}* Header = decltype(Header)(Data.GetPointer(sizeof(*Header)));
+	if (Header == nullptr)
+	{
+		return false;
+	}
+
+	Data.Advance(sizeof(*Header));
+
+	switch (Header->Format)
+	{
+	case 1:		Transport = new FTransportReader();						break;
+#if 0
+	case 4:		Transport = new FLz4TransportReader(Header->Parameter); break;
+#endif // 0
+	default:	return false;
+	}
+
+	return true;
+}
+
+////////////////////////////////////////////////////////////////////////////////
+bool FAnalysisEngine::OnData(FStreamReader::FData& Data)
+{
+	if (Transport == nullptr)
+	{
+		if (!EstablishTransport(Data))
+		{
+			return false;
+		}
+	}
+
+	struct FEventHeader
+	{
+		uint16	Uid;
+		uint16	Size;
+		uint8	EventData[];
+	};
+
+	Transport->SetSource(Data);
+
+	while (true)
+	{
+		const auto* Header = Transport->GetPointer<FEventHeader>();
+		if (Header == nullptr)
+		{
+			break;
+		}
+
+		uint32 BlockSize = Header->Size + sizeof(FEventHeader);
+		Header = Transport->GetPointer<FEventHeader>(BlockSize);
+		if (Header == nullptr)
+		{
+			break;
+		}
+
+		uint16 Uid = uint16(Header->Uid & ((1 << 14) - 1));
+		if (Uid >= Dispatches.Num())
+		{
+			return false;
+		}
+
+		Transport->Advance(BlockSize);
+
+		const FDispatch* Dispatch = Dispatches[Uid];
+
+		EventDataImpl->Dispatch = Dispatch;
+		EventDataImpl->Ptr = Header->EventData;
+		EventDataImpl->Size = Header->Size;
+
+		const FRoute* Route = Routes.GetData() + Dispatch->FirstRoute;
+		for (uint32 n = Route->Count; n--; ++Route)
+		{
+			IAnalyzer* Analyzer = Analyzers[Route->AnalyzerIndex];
+			Analyzer->OnEvent(Route->Id, { SessionContext, *EventDataImpl });
+		}
+	}
+
+	return true;
+}
+
+} // namespace Trace
diff --git a/Engine/Source/Developer/TraceAnalysis/Private/Analysis/Engine.h b/Engine/Source/Developer/TraceAnalysis/Private/Analysis/Engine.h
new file mode 100644
--- /dev/null
+++ b/Engine/Source/Developer/TraceAnalysis/Private/Analysis/Engine.h
@@ -0,0 +1,78 @@
+// Copyright 1998-2019 Epic Games, Inc. All Rights Reserved.
+
+#pragma once
+
+#include "Containers/Array.h"
+#include "DataStream.h"
+#include "Trace/Analysis.h"
+#include "Trace/Analyzer.h"
+
+class IFileHandle;
+
+namespace Trace
+{
+
+////////////////////////////////////////////////////////////////////////////////
+struct FNewEventEvent;
+class FTransportReader;
+
+////////////////////////////////////////////////////////////////////////////////
+class FAnalysisEngine
+	: public IAnalyzer
+{
+public:
+						FAnalysisEngine(TArray<IAnalyzer*>&& InAnalyzers);
+						~FAnalysisEngine();
+	bool				OnData(FStreamReader::FData& Data);
+
+private:
+	struct FDispatch
+	{
+		struct FField
+		{
+			uint32		Hash;
+			uint16		Offset;
+			uint16		Size;
+			uint16		_Unused0;
+			uint8		TypeInfo;
+			uint8		_Unused1;
+		};
+
+		uint16			FirstRoute;
+		uint16			FieldCount;
+		uint16			EventSize;
+		uint16			_Unused0;
+		FField			Fields[];
+	};
+
+	struct FRoute
+	{
+		uint16			HashIndex;
+		int16			Count;
+		uint16			Id;
+		uint16			AnalyzerIndex;
+	};
+
+	struct				FEventDataImpl;
+
+	virtual void		OnAnalysisBegin(const FOnAnalysisContext& Context) override;
+	virtual void		OnAnalysisEnd() override;
+	virtual void		OnEvent(uint16 RouteId, const FOnEventContext& Context) override;
+
+	bool				EstablishTransport(FStreamReader::FData& Data);
+	FDispatch&			AddDispatch(uint16 Uid, uint16 FieldCount);
+	void				AddRoute(uint16 AnalyzerIndex, uint16 Id, const ANSICHAR* Logger, const ANSICHAR* Event);
+	void				AddRoute(uint16 AnalyzerIndex, uint16 Id, uint32 Hash);
+	void				OnNewTrace(const FOnEventContext& Context);
+	void				OnTiming(const FOnEventContext& Context);
+	void				OnNewEvent(const FOnEventContext& Context);
+	FSessionContext		SessionContext;
+	TArray<uint32>		Hashes;
+	TArray<FRoute>		Routes;
+	TArray<IAnalyzer*>	Analyzers;
+	TArray<FDispatch*>	Dispatches;
+	FTransportReader*	Transport = nullptr;
+	FEventDataImpl*		EventDataImpl;
+};
+
+} // namespace Trace
diff --git a/Engine/Source/Developer/TraceAnalysis/Private/Analysis/Processor.cpp b/Engine/Source/Developer/TraceAnalysis/Private/Analysis/Processor.cpp
new file mode 100644
--- /dev/null
+++ b/Engine/Source/Developer/TraceAnalysis/Private/Analysis/Processor.cpp
@@ -0,0 +1,30 @@
+// Copyright 1998-2019 Epic Games, Inc. All Rights Reserved.
+
+#include "Engine.h"
+#include "Trace/Analysis.h"
+
+namespace Trace
+{
+
+////////////////////////////////////////////////////////////////////////////////
+void FAnalysisProcessor::Start(IInDataStream& DataStream)
+{
+	FAnalysisEngine* Engine = (FAnalysisEngine*)Impl;
+	if (Engine == nullptr)
+	{
+		return;
+	}
+
+	FStreamReader Reader(DataStream);
+	while (FStreamReader::FData* Data = Reader.Read())
+	{
+		if (!Engine->OnData(*Data))
+		{
+			break;
+		}
+	}
+	delete Engine;
+	Engine = nullptr;
+}
+
+} // namespace Trace
diff --git a/Engine/Source/Developer/TraceAnalysis/Private/DataStream.cpp b/Engine/Source/Developer/TraceAnalysis/Private/DataStream.cpp
new file mode 100644
--- /dev/null
+++ b/Engine/Source/Developer/TraceAnalysis/Private/DataStream.cpp
@@ -0,0 +1,85 @@
+// Copyright 1998-2019 Epic Games, Inc. All Rights Reserved.
+
+#include "DataStream.h"
+#include "GenericPlatform/GenericPlatformFile.h"
+
+namespace Trace
+{
+
+FFileStream::FFileStream(const TCHAR* InFilePath)
+	: FilePath(InFilePath)
+{
+	OpenFileInternal();
+}
+
+FFileStream::~FFileStream()
+{
+	delete Inner;
+}
+
+int32 FFileStream::Read(void* Data, uint32 Size)
+{
+	if (!Inner)
+	{
+		return 0;
+	}
+	uint64 Remaining = End - Cursor;
+	if (Remaining == 0)
+	{
+		return 0;
+	}
+
+	uint64 Size64 = Size;
+	Size64 = (Remaining < Size64) ? Remaining : Size64;
+
+	if (!Inner->Read((uint8*)Data, Size64))
+	{
+		return 0;
+	}
+
+	Cursor += Size64;
+	return int32(Size64);
+}
+
+void FFileStream::UpdateFileSize()
+{
+	delete Inner;
+	OpenFileInternal();
+	if (Inner)
+	{
+		Inner->Seek(Cursor);
+	}
+}
+
+void FFileStream::OpenFileInternal()
+{
+	IPlatformFile& FileSystem = IPlatformFile::GetPlatformPhysical();
+	Inner = FileSystem.OpenRead(*FilePath, true);
+	if (Inner)
+	{
+		End = Inner->Size();
+	}
+}
+
+FStreamReader::FStreamReader(IInDataStream& InDataStream)
+	: DataStream(InDataStream)
+{
+	Buffer = new uint8[BufferSize];
+}
+
+FStreamReader::~FStreamReader()
+{
+	delete[] Buffer;
+}
+
+IInDataStream* DataStream_ReadFile(const TCHAR* FilePath)
+{
+	IPlatformFile& FileSystem = IPlatformFile::GetPlatformPhysical();
+	if (!FileSystem.FileExists(FilePath))
+	{
+		return nullptr;
+	}
+	return new FFileStream(FilePath);
+}
+
+}
diff --git a/Engine/Source/Developer/TraceAnalysis/Private/DataStream.h b/Engine/Source/Developer/TraceAnalysis/Private/DataStream.h
new file mode 100644
--- /dev/null
+++ b/Engine/Source/Developer/TraceAnalysis/Private/DataStream.h
@@ -0,0 +1,99 @@
+// Copyright 1998-2019 Epic Games, Inc. All Rights Reserved.
+
+#pragma once
+
+#include "CoreTypes.h"
+#include "Logging/LogMacros.h"
+#include "Containers/UnrealString.h"
+#include "Trace/DataStream.h"
+
+#include <memory.h>
+
+namespace Trace
+{
+
+////////////////////////////////////////////////////////////////////////////////
+class FFileStream
+	: public IInDataStream
+{
+public:
+					FFileStream(const TCHAR* FilePath);
+	virtual			~FFileStream();
+	virtual int32	Read(void* Data, uint32 Size) override;
+	void			UpdateFileSize();
+
+private:
+	void			OpenFileInternal();
+
+	FString			FilePath;
+	IFileHandle*	Inner = nullptr;
+	uint64			Cursor = 0;
+	uint64			End = 0;
+};
+
+////////////////////////////////////////////////////////////////////////////////
+class FStreamReader
+{
+public:
+	class FData
+	{
+		friend			FStreamReader;
+		const uint8*	Cursor = nullptr;
+		const uint8*	End = nullptr;
+
+	public:
+		const uint8*	GetPointer(uint32 Size) const;
+		void			Advance(uint32 Size);
+	};
+
+						FStreamReader(IInDataStream& InDataStream);
+						~FStreamReader();
+	FData*				Read();
+
+private:
+	static const uint32	BufferSize = 1 << 18;
+	FData				Data;
+	IInDataStream&		DataStream;
+	uint8*				Buffer;
+};
+
+////////////////////////////////////////////////////////////////////////////////
+inline FStreamReader::FData* FStreamReader::Read()
+{
+	int32 Remaining = int32(UPTRINT(Data.End - Data.Cursor));
+	if (Remaining != 0)
+	{
+		memcpy(Buffer, Data.Cursor, Remaining);
+	}
+
+	uint8* Dest = Buffer + Remaining;
+	int32 ReadSize = DataStream.Read(Dest, BufferSize - Remaining);
+	if (ReadSize <= 0)
+	{
+		return nullptr;
+	}
+
+	Data.Cursor = Buffer;
+	Data.End = Dest + ReadSize;
+	return &Data;
+}
+
+////////////////////////////////////////////////////////////////////////////////
+inline const uint8* FStreamReader::FData::GetPointer(uint32 Size) const
+{
+	if (Cursor + Size > End)
+	{
+		return nullptr;
+	}
+
+	return Cursor;
+}
+
+////////////////////////////////////////////////////////////////////////////////
+inline void FStreamReader::FData::Advance(uint32 Size)
+{
+	Cursor += Size;
+	check(Cursor <= End);
+}
+
+} // namespace Trace
diff --git a/Engine/Source/Developer/TraceAnalysis/Private/TraceAnalysisModule.cpp b/Engine/Source/Developer/TraceAnalysis/Private/TraceAnalysisModule.cpp
new file mode 100644
--- /dev/null
+++ b/Engine/Source/Developer/TraceAnalysis/Private/TraceAnalysisModule.cpp
@@ -0,0 +1,5 @@
+// Copyright 1998-2019 Epic Games, Inc. All Rights Reserved.
+
+#include "Modules/ModuleManager.h"
+
+IMPLEMENT_MODULE(FDefaultModuleImpl, TraceAnalysis);
diff --git a/Engine/Source/Developer/TraceAnalysis/Public/Trace/Analysis.h b/Engine/Source/Developer/TraceAnalysis/Public/Trace/Analysis.h
new file mode 100644
--- /dev/null
+++ b/Engine/Source/Developer/TraceAnalysis/Public/Trace/Analysis.h
@@ -0,0 +1,41 @@
+// Copyright 1998-2019 Epic Games, Inc. All Rights Reserved.
+
+#pragma once
+
+#include "Containers/Array.h"
+
+namespace Trace
+{
+
+class IAnalyzer;
+class IInDataStream;
+
+////////////////////////////////////////////////////////////////////////////////
+class TRACEANALYSIS_API FAnalysisProcessor
+{
+public:
+	explicit	operator bool () const;
+	void		Start(IInDataStream& DataStream);
+	void		Stop();
+	void		Wait();
+	void		Pause(bool State);
+
+private:
+	friend		class FAnalysisContext;
+	void*		Impl = nullptr;
+};
+
+
+
+////////////////////////////////////////////////////////////////////////////////
+class TRACEANALYSIS_API FAnalysisContext
+{
+public:
+	void				AddAnalyzer(IAnalyzer& Analyzer);
+	FAnalysisProcessor	Process();
+
+private:
+	TArray<IAnalyzer*>	Analyzers;
+};
+
+} // namespace Trace
diff --git a/Engine/Source/Developer/TraceAnalysis/Public/Trace/Analyzer.h b/Engine/Source/Developer/TraceAnalysis/Public/Trace/Analyzer.h
new file mode 100644
--- /dev/null
+++ b/Engine/Source/Developer/TraceAnalysis/Public/Trace/Analyzer.h
@@ -0,0 +1,97 @@
+// Copyright 1998-2019 Epic Games, Inc. All Rights Reserved.
+
+#pragma once
+
+#include "CoreTypes.h"
+#include "Logging/LogMacros.h"
+#include "Trace/Private/Field.h"
+
+#include <memory.h>
+
+namespace Trace
+{
+
+////////////////////////////////////////////////////////////////////////////////
+class TRACEANALYSIS_API IAnalyzer
+{
+public:
+	struct FSessionContext
+	{
+		uint64 StartCycle;
+		uint64 CycleFrequency;
+
+		double TimestampFromCycle(uint64 Cycle) const
+		{
+			return double(Cycle - StartCycle) / double(CycleFrequency);
+		}
+
+		double DurationFromCycleCount(uint64 CycleCount) const
+		{
+			return double(CycleCount) / double(CycleFrequency);
+		}
+	};
+
+	struct FInterfaceBuilder
+	{
+		virtual void			RouteEvent(uint16 RouteId, const ANSICHAR* Logger, const ANSICHAR* Event) = 0;
+	};
+
+	struct FOnAnalysisContext
+	{
+		const FSessionContext&	SessionContext;
+		FInterfaceBuilder&		InterfaceBuilder;
+	};
+
+	struct FValue
+	{
+		template <typename AsType>
+		AsType As() const;
+	};
+
+	struct FArray
+	{
+	};
+
+	struct FEventData
+	{
+		virtual const FValue&	GetValue(const ANSICHAR* FieldName) const = 0;
+		virtual const FArray&	GetArray(const ANSICHAR* FieldName) const = 0;
+		virtual const uint8*	GetData() const = 0;
+		virtual const uint8*	GetAttachment() const = 0;
+		virtual uint16			GetAttachmentSize() const = 0;
+		virtual uint16			GetTotalSize() const = 0;
+	};
+
+	struct FOnEventContext
+	{
+		const FSessionContext&	SessionContext;
+		const FEventData&		EventData;
+	};
+
+	virtual				~IAnalyzer() = default;
+	virtual void		OnAnalysisBegin(const FOnAnalysisContext& Context) = 0;
+	virtual void		OnAnalysisEnd() = 0;
+	virtual void		OnEvent(uint16 RouteId, const FOnEventContext& Context) = 0;
+};
+
+////////////////////////////////////////////////////////////////////////////////
+template <typename AsType>
+AsType IAnalyzer::FValue::As() const
+{
+	uint16 FieldType = uint16(UPTRINT(this) >> 48);
+	if (FieldType == 0xffff)
+	{
+		return AsType(0);
+	}
+
+	check((FieldType & _Field_Float) == 0);
+
+	const void* Addr = (void*)(UPTRINT(this) & ((1ull << 48) - 1));
+	UPTRINT Value = 0;
+	uint32 Size = 1 << (FieldType & _Field_Pow2SizeMask);
+	memcpy(&Value, Addr, Size);
+
+	return AsType(Value);
+}
+
+} // namespace Trace
diff --git a/Engine/Source/Developer/TraceAnalysis/Public/Trace/DataStream.h b/Engine/Source/Developer/TraceAnalysis/Public/Trace/DataStream.h
new file mode 100644
--- /dev/null
+++ b/Engine/Source/Developer/TraceAnalysis/Public/Trace/DataStream.h
@@ -0,0 +1,26 @@
+// Copyright 1998-2019 Epic Games, Inc. All Rights Reserved.
+
+#pragma once
+
+#include "CoreTypes.h"
+
+namespace Trace
+{
+
+class IInDataStream
+{
+public:
+	virtual			~IInDataStream() = default;
+	virtual int32	Read(void* Data, uint32 Size) = 0;
+};
+
+class IOutDataStream
+{
+public:
+	virtual			~IOutDataStream() = default;
+	virtual bool	Write(const void* Data, uint32 Size) = 0;
+};
+
+TRACEANALYSIS_API IInDataStream* DataStream_ReadFile(const TCHAR* FilePath);
+
+}
diff --git a/Engine/Source/Developer/TraceAnalysis/Public/Trace/Store.h b/Engine/Source/Developer/TraceAnalysis/Public/Trace/Store.h
new file mode 100644
--- /dev/null
+++ b/Engine/Source/Developer/TraceAnalysis/Public/Trace/Store.h
@@ -0,0 +1,35 @@
+// Copyright 1998-2019 Epic Games, Inc. All Rights Reserved.
+
+#pragma once
+
+#include "CoreTypes.h"
+#include "Containers/UnrealString.h"
+#include "Templates/SharedPointer.h"
+
+namespace Trace
+{
+
+class IInDataStream;
+class IOutDataStream;
+typedef uint64 FStoreSessionHandle;
+
+struct FStoreSessionInfo
+{
+	FStoreSessionHandle Handle;
+	const TCHAR* Uri;
+	const TCHAR* Name;
+	bool bIsLive;
+};
+
+class IStore
+{
+public:
+	virtual ~IStore() = default;
+	virtual void GetAvailableSessions(TArray<FStoreSessionInfo>& OutSessions) const = 0;
+	virtual TTuple<FStoreSessionHandle, IOutDataStream*> CreateNewSession() = 0;
+	virtual IInDataStream* OpenSessionStream(FStoreSessionHandle Handle) = 0;
+};
+
+TRACEANALYSIS_API TSharedPtr<IStore> Store_Create(const TCHAR* StoreDir);
+
+}
diff --git a/Engine/Source/Developer/TraceAnalysis/TraceAnalysis.Build.cs b/Engine/Source/Developer/TraceAnalysis/TraceAnalysis.Build.cs
new file mode 100644
--- /dev/null
+++ b/Engine/Source/Developer/TraceAnalysis/TraceAnalysis.Build.cs
@@ -0,0 +1,13 @@
+// Copyright 1998-2019 Epic Games, Inc. All Rights Reserved.
+
+using UnrealBuildTool;
+
+public class TraceAnalysis : ModuleRules
+{
+	public TraceAnalysis(ReadOnlyTargetRules Target) : base(Target)
+	{
+		PrivateDependencyModuleNames.Add("Core");
+		PrivateDependencyModuleNames.Add("Sockets");
+		PrivateDependencyModuleNames.Add("DirectoryWatcher");
+	}
+}
diff --git a/Engine/Source/Developer/TraceInsights/Private/Insights/Common/TimeUtils.cpp b/Engine/Source/Developer/TraceInsights/Private/Insights/Common/TimeUtils.cpp
new file mode 100644
--- /dev/null
+++ b/Engine/Source/Developer/TraceInsights/Private/Insights/Common/TimeUtils.cpp
@@ -0,0 +1,778 @@
+// Copyright 1998-2019 Epic Games, Inc. All Rights Reserved.
+
+#include "TimeUtils.h"
+
+#include <limits>
+//#include <math.h>
+
+//#include "TimingProfilerCommon.h" // for UE_LOG(TimingProfiler, ...
+
+namespace TimeUtils {
+
+////////////////////////////////////////////////////////////////////////////////////////////////////
+
+FString FormatTimeAuto(const double InDuration)
+{
+	//TestTimeAutoFormatting();
+
+	double Duration = InDuration;
+
+	if (Duration == 0.0)
+	{
+		return TEXT("0");
+	}
+
+	FString Str;
+
+	if (FMath::IsNegativeDouble(Duration))
+	{
+		Str = TEXT("-");
+		Duration = -Duration;
+	}
+
+	if (Duration == DBL_MAX || Duration == std::numeric_limits<double>::infinity())
+	{
+		Str += TEXT("∞");
+	}
+	else if (Duration < TimeUtils::Picosecond)
+	{
+		// (0 .. 1ps)
+		Str = TEXT("~0");
+	}
+	else if (Duration < TimeUtils::Nanosecond)
+	{
+		// [1ps .. 1ns)
+		if (Duration >= 999.5 * TimeUtils::Picosecond)
+		{
+			Str += TEXT("1 ns");
+		}
+		else
+		{
+			const double Picoseconds = Duration * 1000000000000.0 + 0.5;
+			const int32 IntPicoseconds = static_cast<int32>(Picoseconds);
+			ensure(IntPicoseconds <= 999);
+			//if (IntPicoseconds > 999)
+			//{
+			//	Str += TEXT("1 ns");
+			//}
+			//else
+			{
+				Str += FString::Printf(TEXT("%d ps"), IntPicoseconds);
+			}
+		}
+	}
+	else if (Duration < TimeUtils::Microsecond)
+	{
+		// [1ns .. 1µs)
+		if (Duration >= 999.5 * TimeUtils::Nanosecond)
+		{
+			Str += TEXT("1 µs");
+		}
+		else
+		{
+			const double Nanoseconds = Duration * 1000000000.0 + 0.5;
+			const int32 IntNanoseconds = static_cast<int32>(Nanoseconds);
+			ensure(IntNanoseconds <= 999);
+			//if (IntNanoseconds > 999)
+			//{
+			//	Str += TEXT("1 µs");
+			//}
+			//else
+			{
+				Str += FString::Printf(TEXT("%d ns"), IntNanoseconds);
+			}
+		}
+	}
+	else if (Duration < TimeUtils::Milisecond)
+	{
+		// [1µs .. 1ms)
+		const double Microseconds = Duration * 1000000.0;
+		if (Microseconds >= 999.95)
+		{
+			Str += TEXT("1 ms");
+		}
+		else
+		{
+			FString S = FString::Printf(TEXT("%.1f"), Microseconds);
+			S.RemoveFromEnd(TEXT(".0"));
+			Str += S;
+			Str += TEXT(" µs");
+		}
+	}
+	else if (Duration < TimeUtils::Second)
+	{
+		// [1ms .. 1s)
+		const double Miliseconds = Duration * 1000.0;
+		if (Miliseconds >= 999.95)
+		{
+			Str += TEXT("1s");
+		}
+		else
+		{
+			FString S = FString::Printf(TEXT("%.1f"), Miliseconds);
+			S.RemoveFromEnd(TEXT(".0"));
+			Str += S;
+			Str += TEXT(" ms");
+		}
+	}
+	else if (Duration < TimeUtils::Minute)
+	{
+		// [1s .. 1m)
+		if (Duration >= 59.95)
+		{
+			Str += TEXT("1m");
+		}
+		else
+		{
+			FString S = FString::Printf(TEXT("%.1f"), Duration);
+			S.RemoveFromEnd(TEXT(".0"));
+			Str += S;
+			Str += TEXT("s");
+		}
+	}
+	else if (Duration < TimeUtils::Hour)
+	{
+		// [1m .. 1h)
+		const double Minutes = FMath::FloorToDouble(Duration / TimeUtils::Minute);
+		Str += FString::Printf(TEXT("%dm"), static_cast<int32>(Minutes));
+		Duration -= Minutes * TimeUtils::Minute;
+		const double Seconds = FMath::FloorToDouble(Duration / TimeUtils::Second);
+		if (Seconds > 0.5)
+		{
+			Str += FString::Printf(TEXT(" %ds"), static_cast<int32>(Seconds));
+		}
+	}
+	else if (Duration < TimeUtils::Day)
+	{
+		// [1h .. 1d)
+		const double Hours = FMath::FloorToDouble(Duration / TimeUtils::Hour);
+		Str += FString::Printf(TEXT("%dh"), static_cast<int32>(Hours));
+		Duration -= Hours * TimeUtils::Hour;
+		const double Minutes = FMath::FloorToDouble(Duration / TimeUtils::Minute);
+		if (Minutes > 0.5)
+		{
+			Str += FString::Printf(TEXT(" %dm"), static_cast<int32>(Minutes));
+		}
+	}
+	else
+	{
+		// [1d .. ∞)
+		const double Days = FMath::FloorToDouble(Duration / TimeUtils::Day);
+		Str += FString::Printf(TEXT("%dd"), static_cast<int32>(Days));
+		Duration -= Days * TimeUtils::Day;
+		const double Hours = FMath::FloorToDouble(Duration / TimeUtils::Hour);
+		if (Hours > 0.5)
+		{
+			Str += FString::Printf(TEXT(" %dh"), static_cast<int32>(Hours));
+		}
+	}
+
+	return Str;
+}
+
+////////////////////////////////////////////////////////////////////////////////////////////////////
+
+FString FormatTimeMs(const double InDuration, const int32 NumDigits, bool bAddTimeUnit)
+{
+	if (FMath::IsNearlyZero(InDuration, TimeUtils::Picosecond))
+	{
+		return TEXT("0");
+	}
+
+	double Duration = InDuration;
+	FString Str;
+
+	if (FMath::IsNegativeDouble(Duration))
+	{
+		Duration = -Duration;
+		Str = TEXT("-");
+	}
+
+	if (Duration == DBL_MAX || Duration == std::numeric_limits<double>::infinity())
+	{
+		Str += TEXT("∞");
+	}
+	else
+	{
+		Str += FString::Printf(TEXT("%.*f"), NumDigits, Duration * 1000.0);
+
+		if (bAddTimeUnit)
+		{
+			Str += TEXT(" ms");
+		}
+	}
+
+	return Str;
+}
+
+////////////////////////////////////////////////////////////////////////////////////////////////////
+
+FString FormatTime(const double InTime, const double Precision)
+{
+	if (FMath::IsNearlyZero(InTime, FMath::Max(TimeUtils::Picosecond, Precision / 10.0f)))
+	{
+		return TEXT("0");
+	}
+
+	double Time = InTime;
+	FString Str;
+
+	if (FMath::IsNegativeDouble(Time))
+	{
+		Time = -Time;
+		Str = TEXT("-");
+	}
+
+	if (Time == DBL_MAX || Time == std::numeric_limits<double>::infinity())
+	{
+		Str += TEXT("∞");
+		return Str;
+	}
+
+	bool bIsSpaceNeeded = false;
+
+	int32 Days = static_cast<int32>(Time / TimeUtils::Day);
+	if (Days > 0)
+	{
+		Str += FString::Printf(TEXT("%dd"), Days);
+		bIsSpaceNeeded = true;
+		Time -= static_cast<double>(Days) * TimeUtils::Day;
+	}
+
+	if (Precision >= TimeUtils::Day)
+	{
+		if (!bIsSpaceNeeded)
+		{
+			Str = TEXT("~0");
+		}
+		return Str;
+	}
+
+	int32 Hours = static_cast<int32>(Time / TimeUtils::Hour);
+	if (Hours > 0)
+	{
+		if (bIsSpaceNeeded)
+		{
+			Str += TEXT(" ");
+		}
+		Str += FString::Printf(TEXT("%dh"), Hours);
+		bIsSpaceNeeded = true;
+		Time -= static_cast<double>(Hours) * TimeUtils::Hour;
+	}
+
+	if (Precision >= TimeUtils::Hour)
+	{
+		if (!bIsSpaceNeeded)
+		{
+			Str = TEXT("~0");
+		}
+		return Str;
+	}
+
+	int32 Minutes = static_cast<int32>(Time / TimeUtils::Minute);
+	if (Minutes > 0)
+	{
+		if (bIsSpaceNeeded)
+		{
+			Str += TEXT(" ");
+		}
+		Str += FString::Printf(TEXT("%dm"), Minutes);
+		bIsSpaceNeeded = true;
+		Time -= static_cast<double>(Minutes) * TimeUtils::Minute;
+	}
+
+	if (Precision >= TimeUtils::Minute)
+	{
+		if (!bIsSpaceNeeded)
+		{
+			Str = TEXT("~0");
+		}
+		return Str;
+	}
+
+	//TestOptimizationIssue();
+
+	////float Log10 = -FMath::LogX(10.0f, static_cast<float>(Precision));
+	//double Log10 = -log10(Precision);
+	//int32 Digits = (Log10 > 0) ? FMath::CeilToInt(Log10) : 0;
+
+	static const double DigitThresholds[] =
+	{
+		TimeUtils::Second,              // 0 digits
+		TimeUtils::Milisecond * 100,    // 1 digit
+		TimeUtils::Milisecond * 10,     // 2 digits
+		TimeUtils::Milisecond,          // 3 digits
+		TimeUtils::Microsecond * 100,   // 4 digits
+		TimeUtils::Microsecond * 10,    // 5 digits
+		TimeUtils::Microsecond,         // 6 digits
+		TimeUtils::Nanosecond * 100,    // 7 digits
+		TimeUtils::Nanosecond * 10,     // 8 digits
+		TimeUtils::Nanosecond,          // 9 digits
+		TimeUtils::Picosecond * 100,    // 10 digits
+		TimeUtils::Picosecond * 10,     // 11 digits
+		TimeUtils::Picosecond,          // 12 digits
+		0                               // 13 digits
+	};
+	int32 Digits = 0;
+	while (Precision < DigitThresholds[Digits])
+	{
+		Digits++;
+	}
+
+	if (Digits == 0)
+	{
+		int32 Seconds = static_cast<int32>(Time / TimeUtils::Second + 0.5);
+		if (Seconds > 0)
+		{
+			if (bIsSpaceNeeded)
+			{
+				Str += TEXT(" ");
+			}
+			Str += FString::Printf(TEXT("%ds"), Seconds);
+		}
+		else if (!bIsSpaceNeeded)
+		{
+			Str += TEXT("~0");
+		}
+	}
+	//else if (Digits <= 9)
+	//{
+	//	if (bIsSpaceNeeded)
+	//	{
+	//		Str += TEXT(" ");
+	//	}
+	//	int32 Seconds = static_cast<int32>(Time / TimeUtils::Second);
+	//	Time -= static_cast<double>(Seconds) * TimeUtils::Second;
+	//	int64 SubSeconds = static_cast<int64>(Time * FMath::Pow(10.0f, Digits) + 0.5);
+	//	Str += FString::Printf(TEXT("%d.%0*llds"), Seconds, Digits, SubSeconds);
+	//}
+	else
+	{
+		if (bIsSpaceNeeded)
+		{
+			Str += TEXT(" ");
+		}
+		Str += FString::Printf(TEXT("%.*fs"), Digits, Time);
+	}
+
+	return Str;
+}
+
+////////////////////////////////////////////////////////////////////////////////////////////////////
+
+FString FormatTimeHMS(const double Time, const double Precision)
+{
+	//TODO: DD:HH:MM:SS.mmm.uuu.nnn.ppp
+	return FormatTime(Time, Precision);
+}
+
+////////////////////////////////////////////////////////////////////////////////////////////////////
+
+void SplitTime(const double InTime, FTimeSplit& OutTimeSplit)
+{
+	double Time = InTime;
+
+	if (FMath::IsNearlyZero(Time, TimeUtils::Picosecond))
+	{
+		OutTimeSplit.bIsZero = true;
+		OutTimeSplit.bIsNegative = false;
+		OutTimeSplit.bIsInfinite = false;
+		return;
+	}
+
+	if (FMath::IsNegativeDouble(Time))
+	{
+		Time = -Time;
+		OutTimeSplit.bIsNegative = true;
+	}
+	else
+	{
+		OutTimeSplit.bIsNegative = false;
+	}
+
+	if (Time == DBL_MAX || Time == std::numeric_limits<double>::infinity())
+	{
+		OutTimeSplit.bIsZero = false;
+		OutTimeSplit.bIsInfinite = true;
+		for (int32 Index = 0; Index < 8; ++Index)
+		{
+			OutTimeSplit.Units[Index] = 0;
+		}
+		return;
+	}
+
+	OutTimeSplit.bIsInfinite = false;
+	bool bIsZero = true; // assume true and see if any split unit is != 0
+
+	const double Days = FMath::FloorToDouble(Time / TimeUtils::Day);
+	OutTimeSplit.Days = static_cast<int32>(Days);
+	if (OutTimeSplit.Days > 0)
+	{
+		Time -= Days * TimeUtils::Day;
+		bIsZero = false;
+	}
+
+	const double Hours = FMath::FloorToDouble(Time / TimeUtils::Hour);
+	OutTimeSplit.Hours = static_cast<int32>(Hours);
+	if (OutTimeSplit.Hours > 0)
+	{
+		Time -= Hours * TimeUtils::Hour;
+	}
+
+	const double Minutes = FMath::FloorToDouble(Time / TimeUtils::Minute);
+	OutTimeSplit.Minutes = static_cast<int32>(Minutes);
+	if (OutTimeSplit.Minutes > 0)
+	{
+		Time -= Minutes * TimeUtils::Minute;
+		bIsZero = false;
+	}
+
+	const double Seconds = FMath::FloorToDouble(Time / TimeUtils::Second);
+	OutTimeSplit.Seconds = static_cast<int32>(Seconds);
+	if (OutTimeSplit.Seconds > 0)
+	{
+		Time -= Seconds * TimeUtils::Second;
+		bIsZero = false;
+	}
+
+	const double Miliseconds = FMath::FloorToDouble(Time / TimeUtils::Milisecond);
+	OutTimeSplit.Miliseconds = static_cast<int32>(Miliseconds);
+	if (OutTimeSplit.Miliseconds > 0)
+	{
+		Time -= Miliseconds * TimeUtils::Milisecond;
+		bIsZero = false;
+	}
+
+	const double Microseconds = FMath::FloorToDouble(Time / TimeUtils::Microsecond);
+	OutTimeSplit.Microseconds = static_cast<int32>(Microseconds);
+	if (OutTimeSplit.Microseconds > 0)
+	{
+		Time -= Microseconds * TimeUtils::Microsecond;
+		bIsZero = false;
+	}
+
+	const double Nanoseconds = FMath::FloorToDouble(Time / TimeUtils::Nanosecond);
+	OutTimeSplit.Nanoseconds = static_cast<int32>(Nanoseconds);
+	if (OutTimeSplit.Nanoseconds > 0)
+	{
+		Time -= Nanoseconds * TimeUtils::Nanosecond;
+		bIsZero = false;
+	}
+
+	const double Picoseconds = FMath::FloorToDouble(Time / TimeUtils::Picosecond);
+	OutTimeSplit.Picoseconds = static_cast<int32>(Picoseconds);
+	if (OutTimeSplit.Picoseconds > 0)
+	{
+		Time -= Picoseconds * TimeUtils::Picosecond;
+		bIsZero = false;
+	}
+
+	OutTimeSplit.bIsZero = bIsZero;
+}
+
+////////////////////////////////////////////////////////////////////////////////////////////////////
+
+FString FormatTimeSplit(const FTimeSplit& InTimeSplit, const double Precision)
+{
+	FString Str;
+
+	if (InTimeSplit.bIsZero)
+	{
+		Str = TEXT("0");
+		return Str;
+	}
+
+	if (InTimeSplit.bIsNegative)
+	{
+		Str = TEXT("-");
+	}
+
+	if (InTimeSplit.bIsInfinite)
+	{
+		Str += TEXT("∞");
+		return Str;
+	}
+
+	bool bIsSpaceNeeded = false;
+
+	if (InTimeSplit.Days > 0)
+	{
+		Str += FString::Printf(TEXT("%dd"), InTimeSplit.Days);
+		bIsSpaceNeeded = true;
+	}
+
+	if (Precision >= TimeUtils::Day)
+	{
+		if (!bIsSpaceNeeded)
+		{
+			Str = TEXT("~0");
+		}
+		return Str;
+	}
+
+	if (InTimeSplit.Hours > 0)
+	{
+		if (bIsSpaceNeeded)
+		{
+			Str += TEXT(" ");
+		}
+		Str += FString::Printf(TEXT("%dh"), InTimeSplit.Hours);
+		bIsSpaceNeeded = true;
+	}
+
+	if (Precision >= TimeUtils::Hour)
+	{
+		if (!bIsSpaceNeeded)
+		{
+			Str = TEXT("~0");
+		}
+		return Str;
+	}
+
+	if (InTimeSplit.Minutes > 0)
+	{
+		if (bIsSpaceNeeded)
+		{
+			Str += TEXT(" ");
+		}
+		Str += FString::Printf(TEXT("%dm"), InTimeSplit.Minutes);
+		bIsSpaceNeeded = true;
+	}
+
+	if (Precision >= TimeUtils::Minute)
+	{
+		if (!bIsSpaceNeeded)
+		{
+			Str = TEXT("~0");
+		}
+		return Str;
+	}
+
+	if (InTimeSplit.Seconds > 0)
+	{
+		if (bIsSpaceNeeded)
+		{
+			Str += TEXT(" ");
+		}
+		Str += FString::Printf(TEXT("%ds"), InTimeSplit.Seconds);
+		bIsSpaceNeeded = true;
+	}
+
+	if (Precision >= TimeUtils::Second)
+	{
+		if (!bIsSpaceNeeded)
+		{
+			Str = TEXT("~0");
+		}
+		return Str;
+	}
+
+	if (InTimeSplit.Miliseconds > 0)
+	{
+		if (bIsSpaceNeeded)
+		{
+			Str += TEXT(" ");
+		}
+		Str += FString::Printf(TEXT("%dms"), InTimeSplit.Miliseconds);
+		bIsSpaceNeeded = true;
+	}
+
+	if (Precision >= TimeUtils::Milisecond)
+	{
+		if (!bIsSpaceNeeded)
+		{
+			Str = TEXT("~0");
+		}
+		return Str;
+	}
+
+	if (InTimeSplit.Microseconds > 0)
+	{
+		if (bIsSpaceNeeded)
+		{
+			Str += TEXT(" ");
+		}
+		Str += FString::Printf(TEXT("%dµs"), InTimeSplit.Microseconds);
+		bIsSpaceNeeded = true;
+	}
+
+	if (Precision >= TimeUtils::Microsecond)
+	{
+		if (!bIsSpaceNeeded)
+		{
+			Str = TEXT("~0");
+		}
+		return Str;
+	}
+
+	if (InTimeSplit.Nanoseconds > 0)
+	{
+		if (bIsSpaceNeeded)
+		{
+			Str += TEXT(" ");
+		}
+		Str += FString::Printf(TEXT("%dns"), InTimeSplit.Nanoseconds);
+		bIsSpaceNeeded = true;
+	}
+
+	if (Precision >= TimeUtils::Nanosecond)
+	{
+		if (!bIsSpaceNeeded)
+		{
+			Str = TEXT("~0");
+		}
+		return Str;
+	}
+
+	if (InTimeSplit.Picoseconds > 0)
+	{
+		if (bIsSpaceNeeded)
+		{
+			Str += TEXT(" ");
+		}
+		Str += FString::Printf(TEXT("%dps"), InTimeSplit.Picoseconds);
+		bIsSpaceNeeded = true;
+	}
+
+	if (!bIsSpaceNeeded)
+	{
+		Str = TEXT("~0");
+	}
+
+	return Str;
+}
+
+////////////////////////////////////////////////////////////////////////////////////////////////////
+
+FString FormatTimeSplit(const double Time, const double Precision)
+{
+	FTimeSplit TimeSplit;
+	SplitTime(Time, TimeSplit);
+	return FormatTimeSplit(TimeSplit, Precision);
+}
+
+////////////////////////////////////////////////////////////////////////////////////////////////////
+
+void TestTimeFormatting()
+{
+	double T1 = 1 * TimeUtils::Day + 2 * TimeUtils::Hour + 3 * TimeUtils::Minute + 4.567890123456789;
+	FString S1 = FormatTime(T1, TimeUtils::Day);
+	//UE_LOG(TimingProfiler, Log, TEXT("D-T: %s"), *S1);
+	FString S2 = FormatTime(T1, TimeUtils::Hour);
+	//UE_LOG(TimingProfiler, Log, TEXT("H-T: %s"), *S2);
+	FString S3 = FormatTime(T1, TimeUtils::Minute);
+	//UE_LOG(TimingProfiler, Log, TEXT("M-T: %s"), *S3);
+	for (double P = 10.0; P >= TimeUtils::Picosecond; P /= 10.0)
+	{
+		FString SP = FormatTime(T1, P);
+		//UE_LOG(TimingProfiler, Log, TEXT("P:%g T: %s"), P, *SP);
+	}
+}
+
+////////////////////////////////////////////////////////////////////////////////////////////////////
+
+void TestTimeAutoFormatting()
+{
+	static bool Once = false;
+	if (Once) return;
+	Once = true;
+
+	struct FTestTimeAutoData
+	{
+		double T;
+		const TCHAR* Msg;
+	};
+
+	static const FTestTimeAutoData Data[] =
+	{
+		{ TimeUtils::Minute,                TEXT("1m") },
+		{ TimeUtils::Second * 59.99,        TEXT("59.99s") },
+		{ TimeUtils::Second * 59.95,        TEXT("59.95s") },
+		{ 0, TEXT("[threshold 1m / 59.9s]") },
+		{ TimeUtils::Second * 59.94,        TEXT("59.94s") },
+		{ TimeUtils::Second * 59.9,         TEXT("59.9s") },
+		{ TimeUtils::Second * 10,           TEXT("10s") },
+		{ TimeUtils::Second,                TEXT("1s") },
+		{ TimeUtils::Milisecond * 999.99,   TEXT("999.99ms") },
+		{ TimeUtils::Milisecond * 999.95,   TEXT("999.95ms") },
+		{ 0, TEXT("[threshold 1s / 999.9 ms]") },
+		{ TimeUtils::Milisecond * 999.94,   TEXT("999.94ms") },
+		{ TimeUtils::Milisecond * 999.9,    TEXT("999.9ms") },
+		{ TimeUtils::Milisecond * 999,      TEXT("999ms") },
+		{ TimeUtils::Milisecond * 100,      TEXT("100ms") },
+		{ TimeUtils::Milisecond * 10,       TEXT("10ms") },
+		{ TimeUtils::Milisecond * 1.55,     TEXT("1.55ms") },
+		{ TimeUtils::Milisecond * 1.5,      TEXT("1.5ms") },
+		{ TimeUtils::Milisecond * 1.05,     TEXT("1.05ms") },
+		{ TimeUtils::Milisecond,            TEXT("1ms") },
+		{ TimeUtils::Microsecond * 999.99,  TEXT("999.99µs") },
+		{ TimeUtils::Microsecond * 999.95,  TEXT("999.95µs") },
+		{ 0, TEXT("[threshold 1 ms / 999.9 µs]") },
+		{ TimeUtils::Microsecond * 999.94,  TEXT("999.94µs") },
+		{ TimeUtils::Microsecond * 999.9,   TEXT("999.9µs") },
+		{ TimeUtils::Microsecond * 999,     TEXT("999µs") },
+		{ TimeUtils::Microsecond * 100,     TEXT("100µs") },
+		{ TimeUtils::Microsecond * 10,      TEXT("10µs") },
+		{ TimeUtils::Microsecond,           TEXT("1µs") },
+		{ TimeUtils::Nanosecond * 999.9,    TEXT("999.9ns") },
+		{ TimeUtils::Nanosecond * 999.5,    TEXT("999.5ns") },
+		{ 0, TEXT("[threshold 1 µs / 999 ns]") },
+		{ TimeUtils::Nanosecond * 999.4,    TEXT("999.4ns") },
+		{ TimeUtils::Nanosecond * 999,      TEXT("999ns") },
+		{ TimeUtils::Nanosecond * 100,      TEXT("100ns") },
+		{ TimeUtils::Nanosecond * 10,       TEXT("10ns") },
+		{ TimeUtils::Nanosecond,            TEXT("1ns") },
+		{ TimeUtils::Picosecond * 999.9,    TEXT("999.9ps") },
+		{ TimeUtils::Picosecond * 999.5,    TEXT("999.5ps") },
+		{ 0, TEXT("[threshold 1 ns / 999 ps]") },
+		{ TimeUtils::Picosecond * 999.4,    TEXT("999.4ps") },
+		{ TimeUtils::Picosecond * 999,      TEXT("999ps") },
+		{ TimeUtils::Picosecond * 100,      TEXT("100ps") },
+		{ TimeUtils::Picosecond * 10,       TEXT("10ps") },
+		{ TimeUtils::Picosecond,            TEXT("1ps") },
+		{ TimeUtils::Picosecond * 0.1,      TEXT("0.1ps") },
+	};
+	const int32 DataCount = sizeof(Data) / sizeof(FTestTimeAutoData);
+
+	for (int32 Index = 0; Index < DataCount; ++Index)
+	{
+		FString Str = FormatTimeAuto(Data[Index].T);
+		//UE_LOG(TimingProfiler, Log, TEXT("%s : %s"), Data[Index].Msg, *Str);
+	}
+}
+
+////////////////////////////////////////////////////////////////////////////////////////////////////
+
+int32 GetNumDigits(const double Precision)
+{
+	//float Log10 = -log10(Precision);
+	float Log10 = -FMath::LogX(10.0f, static_cast<float>(Precision));
+	int32 D = (Log10 > 0) ? FMath::CeilToInt(Log10) : 0;
+	return D;
+}
+
+PRAGMA_DISABLE_OPTIMIZATION
+int32 GetNumDigitsOptDisabled(const double Precision)
+{
+	//double Log10 = -log10(Precision);
+	float Log10 = -FMath::LogX(10.0f, static_cast<float>(Precision));
+	int32 D = (Log10 > 0) ? FMath::CeilToInt(Log10) : 0;
+	return D;
+}
+PRAGMA_ENABLE_OPTIMIZATION
+
+PRAGMA_DISABLE_OPTIMIZATION
+void TestOptimizationIssue()
+{
+	constexpr double Ns = 0.000000001;
+	int32 D1 = GetNumDigits(Ns);
+	int32 D2 = GetNumDigitsOptDisabled(Ns);
+	//UE_LOG(TimingProfiler, Log, TEXT("D1 = %d, D2 = %d"), D1, D2);
+	ensure(D1 == 9); // 10 ?
+	ensure(D1 == D2);
+}
+PRAGMA_ENABLE_OPTIMIZATION
+
+////////////////////////////////////////////////////////////////////////////////////////////////////
+
+} // namespace TimeUtils
diff --git a/Engine/Source/Developer/TraceInsights/Private/Insights/Common/TimeUtils.h b/Engine/Source/Developer/TraceInsights/Private/Insights/Common/TimeUtils.h
new file mode 100644
--- /dev/null
+++ b/Engine/Source/Developer/TraceInsights/Private/Insights/Common/TimeUtils.h
@@ -0,0 +1,68 @@
+// Copyright 1998-2019 Epic Games, Inc. All Rights Reserved.
+
+#pragma once
+
+#include "CoreMinimal.h"
+#include "Containers/UnrealString.h"
+
+/*
+double is a 64 bit IEEE 754 double precision Floating Point Number
+1 bit for the sign, 11 bits for the exponent, and 52* bits for the value
+15.5 decimal digits of precision (max value*: ‭4`503`599`627`370`495)
+|------------------------------------------------------------------------------
+| 900`000.000`000`001     --> up to 10 days   with    1 ns precision
+|  90`000.000`000`000`1   --> up to 25 hours  with  0.1 ns precision (100 ps)
+|   9`000.000`000`000`01  --> up to 2.5 hours with 0.01 ns precision (10 ps)
+|     900.000`000`000`001 --> up to 15 min    with    1 ps precision
+|------------------------------------------------------------------------------
+|   3`600.000`000`000`001 --> 1 ps precision at 1 hour : NOT OK (16 digits)!!!
+|  86`400.000`000`000`001 --> 1 ps precision at 1 day  : NOT OK (17 digits)!!!
+|------------------------------------------------------------------------------
+*/
+
+namespace TimeUtils
+{
+	static constexpr double Picosecond = 0.000000000001;
+	static constexpr double Nanosecond = 0.000000001;
+	static constexpr double Microsecond = 0.000001;
+	static constexpr double Milisecond = 0.001;
+	static constexpr double Second = 1.0;
+	static constexpr double Minute = 60.0;
+	static constexpr double Hour = 3600.0;
+	static constexpr double Day = 86400.0;
+	static constexpr double Week = 604800.0;
+
+	struct FTimeSplit
+	{
+		union
+		{
+			int32 Units[8]; // Units[0] = Days, Units[1] = Hours, ..., Units[7] = Picoseconds
+			struct
+			{
+				int32 Days;
+				int32 Hours;
+				int32 Minutes;
+				int32 Seconds;
+				int32 Miliseconds;
+				int32 Microseconds;
+				int32 Nanoseconds;
+				int32 Picoseconds;
+			};
+		};
+		bool bIsZero;
+		bool bIsNegative;
+		bool bIsInfinite;
+	};
+
+	FString FormatTimeAuto(const double Duration);
+	FString FormatTimeMs(const double Duration, const int32 NumDigits = 2, bool bAddTimeUnit = false);
+	FString FormatTime(const double Time, const double Precision = 0.0);
+	FString FormatTimeHMS(const double Time, const double Precision = 0.0);
+	void SplitTime(const double Time, FTimeSplit& OutTimeSplit);
+	FString FormatTimeSplit(const FTimeSplit& TimeSplit, const double Precision = 0.0);
+	FString FormatTimeSplit(const double Time, const double Precision = 0.0);
+
+	void TestTimeFormatting();
+	void TestTimeAutoFormatting();
+	void TestOptimizationIssue();
+}
diff --git a/Engine/Source/Developer/TraceInsights/Private/Insights/InsightsCommands.cpp b/Engine/Source/Developer/TraceInsights/Private/Insights/InsightsCommands.cpp
new file mode 100644
--- /dev/null
+++ b/Engine/Source/Developer/TraceInsights/Private/Insights/InsightsCommands.cpp
@@ -0,0 +1,205 @@
+// Copyright 1998-2019 Epic Games, Inc. All Rights Reserved.
+
+#include "Insights/InsightsCommands.h"
+
+#include "DesktopPlatformModule.h"
+#include "EditorStyleSet.h"
+#include "Framework/Application/SlateApplication.h"
+//#include "Framework/MultiBox/MultiBoxBuilder.h"
+#include "HAL/FileManagerGeneric.h"
+#include "Misc/Paths.h"
+
+// Insights
+#include "Insights/InsightsManager.h"
+
+////////////////////////////////////////////////////////////////////////////////////////////////////
+
+#define LOCTEXT_NAMESPACE "FInsightsCommands"
+
+////////////////////////////////////////////////////////////////////////////////////////////////////
+// FInsightsCommands
+////////////////////////////////////////////////////////////////////////////////////////////////////
+
+FInsightsCommands::FInsightsCommands()
+	: TCommands<FInsightsCommands>(
+		TEXT("InsightsCommands"), // Context name for fast lookup
+		NSLOCTEXT("Contexts", "InsightsCommand", "Insights Command"), // Localized context name for displaying
+		NAME_None, // Parent
+		FEditorStyle::GetStyleSetName() // Icon Style Set
+	)
+{
+}
+
+////////////////////////////////////////////////////////////////////////////////////////////////////
+
+// UI_COMMAND takes long for the compile to optimize
+PRAGMA_DISABLE_OPTIMIZATION
+void FInsightsCommands::RegisterCommands()
+{
+	UI_COMMAND(InsightsManager_Live, "Live", "Loads profiler data from a live or last trace session", EUserInterfaceActionType::Button, FInputChord());
+	UI_COMMAND(InsightsManager_Load, "Load...", "Loads profiler data from a trace file", EUserInterfaceActionType::Button, FInputChord(EModifierKey::Control, EKeys::L));
+	UI_COMMAND(ToggleDebugInfo, "Debug", "Toggles the display of debug info", EUserInterfaceActionType::ToggleButton, FInputChord(EModifierKey::Control, EKeys::D));
+	UI_COMMAND(OpenSettings, "Settings", "Opens the Unreal Insights settings", EUserInterfaceActionType::Button, FInputChord(EModifierKey::Control, EKeys::O));
+}
+PRAGMA_ENABLE_OPTIMIZATION
+
+////////////////////////////////////////////////////////////////////////////////////////////////////
+// ToggleDebugInfo
+////////////////////////////////////////////////////////////////////////////////////////////////////
+
+void FInsightsActionManager::Map_ToggleDebugInfo_Global()
+{
+	This->CommandList->MapAction(This->GetCommands().ToggleDebugInfo, ToggleDebugInfo_Custom());
+}
+
+const FUIAction FInsightsActionManager::ToggleDebugInfo_Custom()
+{
+	FUIAction UIAction;
+	UIAction.ExecuteAction = FExecuteAction::CreateRaw(this, &FInsightsActionManager::ToggleDebugInfo_Execute);
+	UIAction.CanExecuteAction = FCanExecuteAction::CreateRaw(this, &FInsightsActionManager::ToggleDebugInfo_CanExecute);
+	UIAction.GetActionCheckState = FGetActionCheckState::CreateRaw(this, &FInsightsActionManager::ToggleDebugInfo_GetCheckState);
+	return UIAction;
+}
+
+void FInsightsActionManager::ToggleDebugInfo_Execute()
+{
+	const bool bIsDebugInfoEnabled = !This->IsDebugInfoEnabled();
+	This->SetDebugInfo(bIsDebugInfoEnabled);
+}
+
+bool FInsightsActionManager::ToggleDebugInfo_CanExecute() const
+{
+	return FInsightsManager::Get()->GetSession().IsValid();
+}
+
+ECheckBoxState FInsightsActionManager::ToggleDebugInfo_GetCheckState() const
+{
+	const bool bIsDebugInfoEnabled = This->IsDebugInfoEnabled();
+	return bIsDebugInfoEnabled ? ECheckBoxState::Checked : ECheckBoxState::Unchecked;
+}
+
+////////////////////////////////////////////////////////////////////////////////////////////////////
+// InsightsManager_Live
+////////////////////////////////////////////////////////////////////////////////////////////////////
+
+void FInsightsActionManager::Map_InsightsManager_Live()
+{
+	FUIAction UIAction;
+	UIAction.ExecuteAction = FExecuteAction::CreateRaw(this, &FInsightsActionManager::InsightsManager_Live_Execute);
+	UIAction.CanExecuteAction = FCanExecuteAction::CreateRaw(this, &FInsightsActionManager::InsightsManager_Live_CanExecute);
+
+	This->CommandList->MapAction(This->GetCommands().InsightsManager_Live, UIAction);
+}
+
+////////////////////////////////////////////////////////////////////////////////////////////////////
+
+void FInsightsActionManager::InsightsManager_Live_Execute()
+{
+	This->LoadLastLiveSession();
+	if (!This->GetSession().IsValid())
+	{
+		This->LoadLastSession();
+	}
+}
+
+////////////////////////////////////////////////////////////////////////////////////////////////////
+
+bool FInsightsActionManager::InsightsManager_Live_CanExecute() const
+{
+	return true;
+}
+
+////////////////////////////////////////////////////////////////////////////////////////////////////
+// InsightsManager_Load
+////////////////////////////////////////////////////////////////////////////////////////////////////
+
+void FInsightsActionManager::Map_InsightsManager_Load()
+{
+	FUIAction UIAction;
+	UIAction.ExecuteAction = FExecuteAction::CreateRaw(this, &FInsightsActionManager::InsightsManager_Load_Execute);
+	UIAction.CanExecuteAction = FCanExecuteAction::CreateRaw(this, &FInsightsActionManager::InsightsManager_Load_CanExecute);
+
+	This->CommandList->MapAction(This->GetCommands().InsightsManager_Load, UIAction);
+}
+
+////////////////////////////////////////////////////////////////////////////////////////////////////
+
+void FInsightsActionManager::InsightsManager_Load_Execute()
+{
+	//const FString ProfilingDirectory(FPaths::ConvertRelativePathToFull(*FPaths::ProfilingDir()));
+	TSharedRef<Trace::ISessionService> SessionService = This->GetSessionService();
+	const FString ProfilingDirectory(FPaths::ConvertRelativePathToFull(SessionService->GetLocalSessionDirectory()));
+
+	TArray<FString> OutFiles;
+
+	bool bOpened = false;
+
+	IDesktopPlatform* DesktopPlatform = FDesktopPlatformModule::Get();
+	if (DesktopPlatform != nullptr)
+	{
+		FSlateApplication::Get().CloseToolTip();
+
+		bOpened = DesktopPlatform->OpenFileDialog
+		(
+			FSlateApplication::Get().FindBestParentWindowHandleForDialogs(nullptr),
+			LOCTEXT("LoadTrace_FileDesc", "Open trace file...").ToString(),
+			ProfilingDirectory,
+			TEXT(""),
+			LOCTEXT("LoadTrace_FileFilter", "Trace files (*.utrace)|*.utrace|All files (*.*)|*.*").ToString(),
+			EFileDialogFlags::None,
+			OutFiles
+		);
+	}
+
+	if (bOpened == true)
+	{
+		if (OutFiles.Num() == 1)
+		{
+			This->LoadTraceFile(OutFiles[0]);
+		}
+	}
+}
+
+////////////////////////////////////////////////////////////////////////////////////////////////////
+
+bool FInsightsActionManager::InsightsManager_Load_CanExecute() const
+{
+	return true;
+}
+
+////////////////////////////////////////////////////////////////////////////////////////////////////
+// OpenSettings
+////////////////////////////////////////////////////////////////////////////////////////////////////
+
+void FInsightsActionManager::Map_OpenSettings_Global()
+{
+	This->CommandList->MapAction(This->GetCommands().OpenSettings, OpenSettings_Custom());
+}
+
+////////////////////////////////////////////////////////////////////////////////////////////////////
+
+const FUIAction FInsightsActionManager::OpenSettings_Custom()
+{
+	FUIAction UIAction;
+	UIAction.ExecuteAction = FExecuteAction::CreateRaw(this, &FInsightsActionManager::OpenSettings_Execute);
+	UIAction.CanExecuteAction = FCanExecuteAction::CreateRaw(this, &FInsightsActionManager::OpenSettings_CanExecute);
+	return UIAction;
+}
+
+////////////////////////////////////////////////////////////////////////////////////////////////////
+
+void FInsightsActionManager::OpenSettings_Execute()
+{
+	This->OpenSettings();
+}
+
+////////////////////////////////////////////////////////////////////////////////////////////////////
+
+bool FInsightsActionManager::OpenSettings_CanExecute() const
+{
+	return !This->Settings.IsEditing();
+}
+
+////////////////////////////////////////////////////////////////////////////////////////////////////
+
+#undef LOCTEXT_NAMESPACE
diff --git a/Engine/Source/Developer/TraceInsights/Private/Insights/InsightsCommands.h b/Engine/Source/Developer/TraceInsights/Private/Insights/InsightsCommands.h
new file mode 100644
--- /dev/null
+++ b/Engine/Source/Developer/TraceInsights/Private/Insights/InsightsCommands.h
@@ -0,0 +1,102 @@
+// Copyright 1998-2019 Epic Games, Inc. All Rights Reserved.
+
+#pragma once
+
+#include "CoreMinimal.h"
+#include "Framework/Commands/UIAction.h"
+#include "Framework/Commands/Commands.h"
+
+////////////////////////////////////////////////////////////////////////////////////////////////////
+/**
+ * Class that holds all profiler commands.
+ */
+class FInsightsCommands : public TCommands<FInsightsCommands>
+{
+public:
+	/** Default constructor. */
+	FInsightsCommands();
+
+	/** Initialize commands. */
+	virtual void RegisterCommands() override;
+
+public:
+	//////////////////////////////////////////////////
+	// Global commands need to implement following method:
+	//     void Map_<CommandName>_Global();
+	// Custom commands needs to implement also the following method:
+	//     const FUIAction <CommandName>_Custom(...) const;
+	//////////////////////////////////////////////////
+
+	/** Load profiler data from a live trace session. Global version. */
+	TSharedPtr<FUICommandInfo> InsightsManager_Live;
+
+	/** Load profiler data from a trace file. Global version. */
+	TSharedPtr<FUICommandInfo> InsightsManager_Load;
+
+	/** Toggles the debug info. Global and custom command. */
+	TSharedPtr<FUICommandInfo> ToggleDebugInfo;
+
+	/** Open settings for the profiler manager. */
+	TSharedPtr<FUICommandInfo> OpenSettings;
+};
+
+////////////////////////////////////////////////////////////////////////////////////////////////////
+/**
+ * Class that provides helper functions for the commands to avoid cluttering manager with many small functions.
+ * Can't contain any variables. Directly operates on the manager instance.
+ */
+class FInsightsActionManager
+{
+	friend class FInsightsManager;
+
+private:
+	/** Private constructor. */
+	FInsightsActionManager(class FInsightsManager* Instance)
+		: This(Instance)
+	{}
+
+	//////////////////////////////////////////////////
+	// InsightsManager_Live
+
+public:
+	void Map_InsightsManager_Live(); /**< Maps UI command info InsightsManager_Live with the specified UI command list. */
+protected:
+	void InsightsManager_Live_Execute(); /**< Handles FExecuteAction for InsightsManager_Live. */
+	bool InsightsManager_Live_CanExecute() const; /**< Handles FCanExecuteAction for InsightsManager_Live. */
+
+	//////////////////////////////////////////////////
+	// InsightsManager_Load
+
+public:
+	void Map_InsightsManager_Load(); /**< Maps UI command info InsightsManager_Load with the specified UI command list. */
+protected:
+	void InsightsManager_Load_Execute(); /**< Handles FExecuteAction for InsightsManager_Load. */
+	bool InsightsManager_Load_CanExecute() const; /**< Handles FCanExecuteAction for InsightsManager_Load. */
+
+	//////////////////////////////////////////////////
+	// ToggleDebugInfo
+
+public:
+	void Map_ToggleDebugInfo_Global(); /**< Maps UI command info ToggleDebugInfo with the specified UI command list. */
+	const FUIAction ToggleDebugInfo_Custom(); /**< UI action for ToggleDebugInfo command. */
+protected:
+	void ToggleDebugInfo_Execute(); /**< Handles FExecuteAction for ToggleDebugInfo. */
+	bool ToggleDebugInfo_CanExecute() const; /**< Handles FCanExecuteAction for ToggleDebugInfo. */
+	ECheckBoxState ToggleDebugInfo_GetCheckState() const; /**< Handles FGetActionCheckState for ToggleDebugInfo. */
+
+	//////////////////////////////////////////////////
+	// OpenSettings
+
+public:
+	void Map_OpenSettings_Global(); /**< Maps UI command info OpenSettings with the specified UI command list. */
+	const FUIAction OpenSettings_Custom(); /**< UI action for OpenSettings command. */
+protected:
+	void OpenSettings_Execute(); /**< Handles FExecuteAction for OpenSettings. */
+	bool OpenSettings_CanExecute() const; /**< Handles FCanExecuteAction for OpenSettings. */
+
+	//////////////////////////////////////////////////
+
+protected:
+	/** Reference to the global instance of the Insights manager. */
+	class FInsightsManager* This;
+};
diff --git a/Engine/Source/Developer/TraceInsights/Private/Insights/InsightsManager.cpp b/Engine/Source/Developer/TraceInsights/Private/Insights/InsightsManager.cpp
new file mode 100644
--- /dev/null
+++ b/Engine/Source/Developer/TraceInsights/Private/Insights/InsightsManager.cpp
@@ -0,0 +1,299 @@
+// Copyright 1998-2019 Epic Games, Inc. All Rights Reserved.
+
+#include "InsightsManager.h"
+
+#include "Framework/Application/SlateApplication.h"
+#include "Modules/ModuleManager.h"
+#include "Templates/ScopedPointer.h"
+#include "Templates/UniquePtr.h"
+
+// Insights
+#include "Insights/TimingProfilerManager.h"
+#include "Insights/Widgets/SStartPageWindow.h"
+#include "Insights/Widgets/STimingProfilerWindow.h"
+
+////////////////////////////////////////////////////////////////////////////////////////////////////
+
+#define LOCTEXT_NAMESPACE "InsightsManager"
+
+const FName FInsightsManagerTabs::StartPageTabId(TEXT("StartPage"));
+const FName FInsightsManagerTabs::TimingProfilerTabId(TEXT("TimingProfiler"));
+const FName FInsightsManagerTabs::IoProfilerTabId(TEXT("IoProfiler"));
+
+TSharedPtr<FInsightsManager> FInsightsManager::Instance = nullptr;
+
+////////////////////////////////////////////////////////////////////////////////////////////////////
+
+FInsightsManager::FInsightsManager(TSharedRef<Trace::IAnalysisService> InTraceAnalysisService,
+								   TSharedRef<Trace::ISessionService> InTraceSessionService,
+								   TSharedRef<Trace::IModuleService> InTraceModuleService)
+	: AnalysisService(InTraceAnalysisService)
+	, SessionService(InTraceSessionService)
+	, ModuleService(InTraceModuleService)
+	, CommandList(new FUICommandList())
+	, ActionManager(this)
+	, Settings()
+	, bIsDebugInfoEnabled(false)
+{
+}
+
+////////////////////////////////////////////////////////////////////////////////////////////////////
+
+void FInsightsManager::PostConstructor()
+{
+	// Register tick functions.
+	//OnTick = FTickerDelegate::CreateSP(this, &FInsightsManager::Tick);
+	//OnTickHandle = FTicker::GetCoreTicker().AddTicker(OnTick, 1.0f);
+
+	FInsightsCommands::Register();
+	BindCommands();
+}
+
+////////////////////////////////////////////////////////////////////////////////////////////////////
+
+void FInsightsManager::BindCommands()
+{
+	ActionManager.Map_InsightsManager_Live();
+	ActionManager.Map_InsightsManager_Load();
+	ActionManager.Map_ToggleDebugInfo_Global();
+	ActionManager.Map_OpenSettings_Global();
+}
+
+////////////////////////////////////////////////////////////////////////////////////////////////////
+
+FInsightsManager::~FInsightsManager()
+{
+	FInsightsCommands::Unregister();
+
+	// Unregister tick function.
+	//FTicker::GetCoreTicker().RemoveTicker(OnTickHandle);
+}
+
+////////////////////////////////////////////////////////////////////////////////////////////////////
+
+TSharedPtr<FInsightsManager> FInsightsManager::Get()
+{
+	return FInsightsManager::Instance;
+}
+
+////////////////////////////////////////////////////////////////////////////////////////////////////
+
+TSharedPtr<const Trace::IAnalysisSession> FInsightsManager::GetSession() const
+{
+	return Session;
+}
+
+////////////////////////////////////////////////////////////////////////////////////////////////////
+
+const TSharedRef<FUICommandList> FInsightsManager::GetCommandList() const
+{
+	return CommandList;
+}
+
+////////////////////////////////////////////////////////////////////////////////////////////////////
+
+const FInsightsCommands& FInsightsManager::GetCommands()
+{
+	return FInsightsCommands::Get();
+}
+
+////////////////////////////////////////////////////////////////////////////////////////////////////
+
+FInsightsActionManager& FInsightsManager::GetActionManager()
+{
+	return FInsightsManager::Instance->ActionManager;
+}
+
+////////////////////////////////////////////////////////////////////////////////////////////////////
+
+FInsightsSettings& FInsightsManager::GetSettings()
+{
+	return FInsightsManager::Instance->Settings;
+}
+
+////////////////////////////////////////////////////////////////////////////////////////////////////
+
+bool FInsightsManager::Tick(float DeltaTime)
+{
+	//SCOPE_CYCLE_COUNTER(STAT_IM_Tick);
+
+	return true;
+}
+
+////////////////////////////////////////////////////////////////////////////////////////////////////
+
+void FInsightsManager::ResetSession()
+{
+	if (Session.IsValid())
+	{
+		Session.Reset();
+		OnSessionChanged();
+	}
+}
+
+////////////////////////////////////////////////////////////////////////////////////////////////////
+
+void FInsightsManager::OnSessionChanged()
+{
+	if (TSharedPtr<FTimingProfilerManager> TimingProfilerManager = FTimingProfilerManager::Get())
+	{
+		// FIXME: make TimingProfilerManager to register to SessionChangedEvent instead
+		TimingProfilerManager->OnSessionChanged();
+	}
+
+	SessionChangedEvent.Broadcast();
+}
+
+////////////////////////////////////////////////////////////////////////////////////////////////////
+
+void FInsightsManager::SpawnAndActivateTabs()
+{
+	// Open tabs for profilers.
+	if (FGlobalTabmanager::Get()->CanSpawnTab(FInsightsManagerTabs::TimingProfilerTabId))
+	{
+		FGlobalTabmanager::Get()->InvokeTab(FInsightsManagerTabs::TimingProfilerTabId);
+	}
+	if (FGlobalTabmanager::Get()->CanSpawnTab(FInsightsManagerTabs::IoProfilerTabId))
+	{
+		FGlobalTabmanager::Get()->InvokeTab(FInsightsManagerTabs::IoProfilerTabId);
+	}
+
+	// Ensure Timing Insights / Timing View is the active tab / view.
+	if (TSharedPtr<SDockTab> TimingInsightsTab = FGlobalTabmanager::Get()->FindExistingLiveTab(FInsightsManagerTabs::TimingProfilerTabId))
+	{
+		TimingInsightsTab->ActivateInParent(ETabActivationCause::SetDirectly);
+
+		//TOOD: FTimingProfilerManager::Get()->SpawnAndActivateTabs();
+		TSharedPtr<class STimingProfilerWindow> Wnd = FTimingProfilerManager::Get()->GetProfilerWindow();
+		if (Wnd)
+		{
+			TSharedPtr<FTabManager> TabManager = Wnd->GetTabManager();
+
+			if (TSharedPtr<SDockTab> TimingViewTab = TabManager->FindExistingLiveTab(FTimingProfilerTabs::TimingViewID))
+			{
+				TimingViewTab->ActivateInParent(ETabActivationCause::SetDirectly);
+				FSlateApplication::Get().SetKeyboardFocus(TimingViewTab->GetContent());
+			}
+		}
+	}
+}
+
+////////////////////////////////////////////////////////////////////////////////////////////////////
+
+bool FInsightsManager::IsAnyLiveSessionAvailable(Trace::FSessionHandle& OutLastLiveSessionHandle) const
+{
+	TArray<Trace::FSessionHandle> LiveSessions;
+	SessionService->GetLiveSessions(LiveSessions);
+
+	if (LiveSessions.Num() > 0)
+	{
+		OutLastLiveSessionHandle = LiveSessions.Last();
+		return true;
+	}
+
+	return false;
+}
+
+////////////////////////////////////////////////////////////////////////////////////////////////////
+
+bool FInsightsManager::IsAnySessionAvailable(Trace::FSessionHandle& OutLastSessionHandle) const
+{
+	TArray<Trace::FSessionHandle> AvailableSessions;
+	SessionService->GetAvailableSessions(AvailableSessions);
+
+	if (AvailableSessions.Num() > 0)
+	{
+		OutLastSessionHandle = AvailableSessions.Last();
+		return true;
+	}
+
+	return false;
+}
+
+////////////////////////////////////////////////////////////////////////////////////////////////////
+
+void FInsightsManager::LoadLastLiveSession()
+{
+	ResetSession();
+
+	TArray<Trace::FSessionHandle> AvailableSessions;
+	SessionService->GetAvailableSessions(AvailableSessions);
+
+	// Iterate in reverse order as we want the most recent live session first.
+	for (int32 SessionIndex = AvailableSessions.Num() - 1; SessionIndex >= 0; --SessionIndex)
+	{
+		Trace::FSessionHandle SessionHandle = AvailableSessions[SessionIndex];
+
+		Trace::FSessionInfo SessionInfo;
+		SessionService->GetSessionInfo(SessionHandle, SessionInfo);
+
+		if (SessionInfo.bIsLive)
+		{
+			LoadSession(SessionHandle);
+			break;
+		}
+	}
+}
+
+////////////////////////////////////////////////////////////////////////////////////////////////////
+
+void FInsightsManager::LoadLastSession()
+{
+	ResetSession();
+
+	TArray<Trace::FSessionHandle> AvailableSessions;
+	SessionService->GetAvailableSessions(AvailableSessions);
+
+	if (AvailableSessions.Num() > 0)
+	{
+		LoadSession(AvailableSessions.Last());
+	}
+}
+
+////////////////////////////////////////////////////////////////////////////////////////////////////
+
+void FInsightsManager::LoadSession(Trace::FSessionHandle SessionHandle)
+{
+	ResetSession();
+
+	Trace::FSessionInfo SessionInfo;
+	SessionService->GetSessionInfo(SessionHandle, SessionInfo);
+
+	TUniquePtr<Trace::IInDataStream> DataStream(SessionService->OpenSessionStream(SessionHandle));
+	if (DataStream)
+	{
+		Session = AnalysisService->StartAnalysis(SessionInfo.Name, MoveTemp(DataStream));
+		SpawnAndActivateTabs();
+		OnSessionChanged();
+	}
+}
+
+////////////////////////////////////////////////////////////////////////////////////////////////////
+
+void FInsightsManager::LoadTraceFile(const FString& TraceFilepath)
+{
+	ResetSession();
+
+	TUniquePtr<Trace::IInDataStream> DataStream(SessionService->OpenSessionFromFile(*TraceFilepath));
+	if (DataStream)
+	{
+		Session = AnalysisService->StartAnalysis(*TraceFilepath, MoveTemp(DataStream));
+		SpawnAndActivateTabs();
+		OnSessionChanged();
+	}
+}
+
+////////////////////////////////////////////////////////////////////////////////////////////////////
+
+void FInsightsManager::OpenSettings()
+{
+	TSharedPtr<SStartPageWindow> Wnd = GetStartPageWindow();
+	if (Wnd.IsValid())
+	{
+		Wnd->OpenSettings();
+	}
+}
+
+////////////////////////////////////////////////////////////////////////////////////////////////////
+
+#undef LOCTEXT_NAMESPACE
diff --git a/Engine/Source/Developer/TraceInsights/Private/Insights/InsightsManager.h b/Engine/Source/Developer/TraceInsights/Private/Insights/InsightsManager.h
new file mode 100644
--- /dev/null
+++ b/Engine/Source/Developer/TraceInsights/Private/Insights/InsightsManager.h
@@ -0,0 +1,232 @@
+// Copyright 1998-2019 Epic Games, Inc. All Rights Reserved.
+
+#pragma once
+
+#include "Containers/Ticker.h"
+#include "CoreMinimal.h"
+#include "Framework/Commands/UICommandList.h"
+#include "TraceServices/AnalysisService.h"
+#include "TraceServices/ModuleService.h"
+#include "TraceServices/SessionService.h"
+
+// Insights
+#include "Insights/InsightsCommands.h"
+#include "Insights/InsightsSettings.h"
+
+////////////////////////////////////////////////////////////////////////////////////////////////////
+
+class SStartPageWindow;
+
+////////////////////////////////////////////////////////////////////////////////////////////////////
+
+struct FInsightsManagerTabs
+{
+	static const FName StartPageTabId;
+	static const FName TimingProfilerTabId;
+	static const FName IoProfilerTabId;
+};
+
+////////////////////////////////////////////////////////////////////////////////////////////////////
+/** Enumerates loading a trace file progress states. */
+
+enum class ELoadingProgressState
+{
+	Started,
+	InProgress,
+	Loaded,
+	Failed,
+	Cancelled,
+
+	/** Invalid enum type, may be used as a number of enumerations. */
+	InvalidOrMax,
+};
+
+////////////////////////////////////////////////////////////////////////////////////////////////////
+
+enum class EInsightsNotificationType
+{
+	LoadingTraceFile,
+
+	/** Invalid enum type, may be used as a number of enumerations. */
+	InvalidOrMax,
+};
+
+////////////////////////////////////////////////////////////////////////////////////////////////////
+/**
+ * This class manages following areas:
+ *     Connecting/disconnecting to source trace
+ *     Lifetime of all other managers (specific profilers)
+ *     Global Unreal Insights application state and settings
+ */
+class FInsightsManager
+	: public TSharedFromThis<FInsightsManager>
+{
+	friend class FInsightsActionManager;
+
+public:
+	/** Creates the main manager, only one instance can exist. */
+	FInsightsManager(TSharedRef<Trace::IAnalysisService> TraceAnalysisService,
+					 TSharedRef<Trace::ISessionService> SessionService,
+					 TSharedRef<Trace::IModuleService> TraceModuleService);
+
+	/** Virtual destructor. */
+	virtual ~FInsightsManager();
+
+	/**
+	 * Creates an instance of the main manager and initializes global instance with the previously created instance of the manager.
+	 * @param TraceAnalysisService The trace analysis service
+	 * @param TraceSessionService  The trace session service
+	 * @param TraceModuleService   The trace module service
+	 */
+	static TSharedPtr<FInsightsManager> Initialize(TSharedRef<Trace::IAnalysisService> TraceAnalysisService,
+												   TSharedRef<Trace::ISessionService> TraceSessionService,
+												   TSharedRef<Trace::IModuleService> TraceModuleService)
+	{
+		if (FInsightsManager::Instance.IsValid())
+		{
+			FInsightsManager::Instance.Reset();
+		}
+
+		FInsightsManager::Instance = MakeShareable(new FInsightsManager(TraceAnalysisService, TraceSessionService, TraceModuleService));
+		FInsightsManager::Instance->PostConstructor();
+
+		return FInsightsManager::Instance;
+	}
+
+	/** Shutdowns the main manager. */
+	void Shutdown()
+	{
+		FInsightsManager::Instance.Reset();
+	}
+
+	/** @return the global instance of the main manager (FInsightsManager). */
+	static TSharedPtr<FInsightsManager> Get();
+
+	TSharedRef<Trace::IAnalysisService> GetAnalysisService() const { return AnalysisService; }
+	TSharedRef<Trace::ISessionService> GetSessionService() const { return SessionService; }
+	TSharedRef<Trace::IModuleService> GetModuleService() const { return ModuleService; }
+
+	/** @return an instance of the trace analysis session. */
+	TSharedPtr<const Trace::IAnalysisSession> GetSession() const;
+
+	/** @returns UI command list for the main manager. */
+	const TSharedRef<FUICommandList> GetCommandList() const;
+
+	/** @return an instance of the main commands. */
+	static const FInsightsCommands& GetCommands();
+
+	/** @return an instance of the main action manager. */
+	static FInsightsActionManager& GetActionManager();
+
+	/** @return an instance of the main settings. */
+	static FInsightsSettings& GetSettings();
+
+	void AssignStartPageWindow(const TSharedRef<SStartPageWindow>& InStartPageWindow)
+	{
+		StartPageWindow = InStartPageWindow;
+	}
+
+	/**
+	 * Converts profiler window weak pointer to a shared pointer and returns it.
+	 * Make sure the returned pointer is valid before trying to dereference it.
+	 */
+	TSharedPtr<class SStartPageWindow> GetStartPageWindow() const
+	{
+		return StartPageWindow.Pin();
+	}
+
+	////////////////////////////////////////////////////////////////////////////////////////////////////
+	// Getters and setters used by Toggle Commands.
+
+	/** @return true, if UI is allowed to display debug info. */
+	const bool IsDebugInfoEnabled() const { return bIsDebugInfoEnabled; }
+	void SetDebugInfo(const bool bDebugInfoEnabledState) { bIsDebugInfoEnabled = bDebugInfoEnabledState; }
+
+	////////////////////////////////////////////////////////////////////////////////////////////////////
+
+	bool IsAnyLiveSessionAvailable(Trace::FSessionHandle& OutLastLiveSessionHandle) const;
+	bool IsAnySessionAvailable(Trace::FSessionHandle& OutLastSessionHandle) const;
+
+	/** Creates a new analysis session instance and loads the latest available trace session that is live. */
+	void LoadLastLiveSession();
+
+	/** Creates a new analysis session instance and loads the latest available trace session. */
+	void LoadLastSession();
+
+	/**
+	 * Creates a new analysis session instance using specified session handle.
+	 * @param SessionHandle - The handle for session to analyze
+	 */
+	void LoadSession(Trace::FSessionHandle SessionHandle);
+
+	/**
+	 * Creates a new analysis session instance and loads a trace file from the specified location.
+	 * @param TraceFilepath - The path to the trace file
+	 */
+	void LoadTraceFile(const FString& TraceFilepath);
+
+	/** Opens the Settings dialog. */
+	void OpenSettings();
+
+	////////////////////////////////////////////////////////////////////////////////////////////////////
+	// SessionChangedEvent
+
+public:
+	/** The event to execute when the session has changed. */
+	DECLARE_EVENT(FTimingProfilerManager, FSessionChangedEvent);
+	FSessionChangedEvent& GetSessionChangedEvent() { return SessionChangedEvent; }
+private:
+	/** The event to execute when the session has changed. */
+	FSessionChangedEvent SessionChangedEvent;
+
+	////////////////////////////////////////////////////////////////////////////////////////////////////
+
+private:
+	/** Finishes initialization of the profiler manager. */
+	void PostConstructor();
+
+	/** Binds our UI commands to delegates. */
+	void BindCommands();
+
+	/** Updates this manager, done through FCoreTicker. */
+	bool Tick(float DeltaTime);
+
+	/** Resets (closes) current session instance. */
+	void ResetSession();
+
+	void OnSessionChanged();
+
+	void SpawnAndActivateTabs();
+
+private:
+	/** The delegate to be invoked when this manager ticks. */
+	FTickerDelegate OnTick;
+
+	/** Handle to the registered OnTick. */
+	FDelegateHandle OnTickHandle;
+
+	TSharedRef<Trace::IAnalysisService> AnalysisService;
+	TSharedRef<Trace::ISessionService> SessionService;
+	TSharedRef<Trace::IModuleService> ModuleService;
+
+	/** The trace analysis session. */
+	TSharedPtr<const Trace::IAnalysisSession> Session;
+
+	/** List of UI commands for this manager. This will be filled by this and corresponding classes. */
+	TSharedRef<FUICommandList> CommandList;
+
+	/** An instance of the main action manager. */
+	FInsightsActionManager ActionManager;
+
+	/** An instance of the main settings. */
+	FInsightsSettings Settings;
+
+	/** A weak pointer to the Start Page window. */
+	TWeakPtr<class SStartPageWindow> StartPageWindow;
+
+	/** If enabled, UI can display additional info for debugging purposes. */
+	bool bIsDebugInfoEnabled;
+
+	/** A shared pointer to the global instance of the main manager. */
+	static TSharedPtr<FInsightsManager> Instance;
+};
diff --git a/Engine/Source/Developer/TraceInsights/Private/Insights/InsightsSettings.cpp b/Engine/Source/Developer/TraceInsights/Private/Insights/InsightsSettings.cpp
new file mode 100644
--- /dev/null
+++ b/Engine/Source/Developer/TraceInsights/Private/Insights/InsightsSettings.cpp
@@ -0,0 +1,5 @@
+// Copyright 1998-2019 Epic Games, Inc. All Rights Reserved.
+
+#include "InsightsSettings.h"
+
+FInsightsSettings FInsightsSettings::Defaults(true);
diff --git a/Engine/Source/Developer/TraceInsights/Private/Insights/InsightsSettings.h b/Engine/Source/Developer/TraceInsights/Private/Insights/InsightsSettings.h
new file mode 100644
--- /dev/null
+++ b/Engine/Source/Developer/TraceInsights/Private/Insights/InsightsSettings.h
@@ -0,0 +1,89 @@
+// Copyright 1998-2019 Epic Games, Inc. All Rights Reserved.
+
+#pragma once
+
+#include "CoreMinimal.h"
+#include "Misc/ConfigCacheIni.h"
+
+/** Contains all settings for the Unreal Insights, accessible through the main manager. */
+class FInsightsSettings
+{
+public:
+	FInsightsSettings(bool bInIsDefault = false)
+		: bIsEditing(false)
+		, bIsDefault(bInIsDefault)
+		, bShowEmptyTracksByDefault(false)
+	{
+		if (!bIsDefault)
+		{
+			LoadFromConfig();
+		}
+	}
+
+	~FInsightsSettings()
+	{
+		if (!bIsDefault)
+		{
+			SaveToConfig();
+		}
+	}
+
+	void LoadFromConfig()
+	{
+		FConfigCacheIni::LoadGlobalIniFile(SettingsIni, TEXT("UnrealInsightsSettings"));
+
+		GConfig->GetBool(TEXT("Insights.TimingProfiler"), TEXT("bShowEmptyTracksByDefault"), bShowEmptyTracksByDefault, SettingsIni);
+	}
+
+	void SaveToConfig()
+	{
+		GConfig->SetBool(TEXT("Insights.TimingProfiler"), TEXT("bShowEmptyTracksByDefault"), bShowEmptyTracksByDefault, SettingsIni);
+
+		GConfig->Flush(false, SettingsIni);
+	}
+
+	void EnterEditMode()
+	{
+		bIsEditing = true;
+	}
+
+	void ExitEditMode()
+	{
+		bIsEditing = false;
+	}
+
+	const bool IsEditing() const
+	{
+		return bIsEditing;
+	}
+
+	const FInsightsSettings& GetDefaults() const
+	{
+		return Defaults;
+	}
+
+	void ResetToDefaults()
+	{
+		bShowEmptyTracksByDefault = Defaults.bShowEmptyTracksByDefault;
+	}
+
+//protected:
+public:
+	/** Contains default settings. */
+	static FInsightsSettings Defaults;
+
+	/** Setting filename ini. */
+	FString SettingsIni;
+
+	/** Whether profiler settings is in edit mode. */
+	bool bIsEditing;
+
+	/** Whether this instance contains defaults. */
+	bool bIsDefault;
+
+	//////////////////////////////////////////////////
+	// Actual settings.
+
+	/** If True, the Timing View will show empty tracks by default. */
+	bool bShowEmptyTracksByDefault;
+};
diff --git a/Engine/Source/Developer/TraceInsights/Private/Insights/InsightsStyle.cpp b/Engine/Source/Developer/TraceInsights/Private/Insights/InsightsStyle.cpp
new file mode 100644
--- /dev/null
+++ b/Engine/Source/Developer/TraceInsights/Private/Insights/InsightsStyle.cpp
@@ -0,0 +1,99 @@
+// Copyright 1998-2019 Epic Games, Inc. All Rights Reserved.
+
+#include "InsightsStyle.h"
+#include "Styling/SlateStyleRegistry.h"
+
+#define IMAGE_BRUSH(RelativePath, ...)  FSlateImageBrush (FPaths::EngineContentDir() / "Editor/Slate"/ RelativePath + TEXT(".png"), __VA_ARGS__)
+#define BOX_BRUSH(RelativePath, ...)    FSlateBoxBrush   (FPaths::EngineContentDir() / "Editor/Slate"/ RelativePath + TEXT(".png"), __VA_ARGS__)
+#define BORDER_BRUSH(RelativePath, ...) FSlateBorderBrush(FPaths::EngineContentDir() / "Editor/Slate"/ RelativePath + TEXT(".png"), __VA_ARGS__)
+
+TSharedPtr<FSlateStyleSet> FInsightsStyle::StyleInstance = nullptr;
+
+void FInsightsStyle::Initialize()
+{
+	if (!StyleInstance.IsValid())
+	{
+		StyleInstance = Create();
+		FSlateStyleRegistry::RegisterSlateStyle(*StyleInstance);
+	}
+}
+
+void FInsightsStyle::Shutdown()
+{
+	FSlateStyleRegistry::UnRegisterSlateStyle(*StyleInstance);
+	ensure(StyleInstance.IsUnique());
+	StyleInstance.Reset();
+}
+
+FName FInsightsStyle::GetStyleSetName()
+{
+	static FName StyleSetName(TEXT("InsightsStyle"));
+	return StyleSetName;
+}
+
+TSharedRef<FSlateStyleSet> FInsightsStyle::Create()
+{
+	TSharedRef<FSlateStyleSet> StyleRef = MakeShareable(new FSlateStyleSet(FInsightsStyle::GetStyleSetName()));
+	StyleRef->SetContentRoot(FPaths::EngineContentDir() / TEXT("Editor/Slate"));
+	StyleRef->SetCoreContentRoot(FPaths::EngineContentDir() / TEXT("Slate"));
+
+	FSlateStyleSet& Style = StyleRef.Get();
+
+	//////////////////////////////////////////////////
+	// Icons for major components
+
+	Style.Set("StartPage.Icon.Large", new IMAGE_BRUSH("/Icons/icon_tab_Tools_16x", FVector2D(32.0f, 32.0f)));
+	Style.Set("StartPage.Icon.Small", new IMAGE_BRUSH("/Icons/icon_tab_Tools_16x", FVector2D(16.0f, 16.0f)));
+
+	Style.Set("TimingInsights.Icon.Large", new IMAGE_BRUSH("Icons/Profiler/profiler_stats_40x", FVector2D(32.0f, 32.0f)));
+	Style.Set("TimingInsights.Icon.Small", new IMAGE_BRUSH("Icons/Profiler/profiler_stats_40x", FVector2D(16.0f, 16.0f)));
+
+	Style.Set("AssetLoadingInsights.Icon.Large", new IMAGE_BRUSH("Icons/Profiler/profiler_stats_40x", FVector2D(32.0f, 32.0f)));
+	Style.Set("AssetLoadingInsights.Icon.Small", new IMAGE_BRUSH("Icons/Profiler/profiler_stats_40x", FVector2D(16.0f, 16.0f)));
+
+	Style.Set("Toolbar.Icon.Large", new IMAGE_BRUSH("/Icons/icon_tab_Tools_16x", FVector2D(32.0f, 32.0f)));
+	Style.Set("Toolbar.Icon.Small", new IMAGE_BRUSH("/Icons/icon_tab_Tools_16x", FVector2D(16.0f, 16.0f)));
+
+	//////////////////////////////////////////////////
+	// Timing Insights
+
+	Style.Set("FramesTrack.Icon.Large", new IMAGE_BRUSH("Icons/Profiler/profiler_stats_40x", FVector2D(32.0f, 32.0f)));
+	Style.Set("FramesTrack.Icon.Small", new IMAGE_BRUSH("Icons/Profiler/profiler_stats_40x", FVector2D(16.0f, 16.0f)));
+
+	Style.Set("GraphTrack.Icon.Large", new IMAGE_BRUSH("Icons/Profiler/profiler_stats_40x", FVector2D(32.0f, 32.0f)));
+	Style.Set("GraphTrack.Icon.Small", new IMAGE_BRUSH("Icons/Profiler/profiler_stats_40x", FVector2D(16.0f, 16.0f)));
+
+	Style.Set("TimingView.Icon.Large", new IMAGE_BRUSH("Icons/Profiler/profiler_stats_40x", FVector2D(32.0f, 32.0f)));
+	Style.Set("TimingView.Icon.Small", new IMAGE_BRUSH("Icons/Profiler/profiler_stats_40x", FVector2D(16.0f, 16.0f)));
+
+	Style.Set("TimersView.Icon.Large", new IMAGE_BRUSH("Icons/Profiler/Profiler_Data_Capture_40x", FVector2D(32.0f, 32.0f)));
+	Style.Set("TimersView.Icon.Small", new IMAGE_BRUSH("Icons/Profiler/Profiler_Data_Capture_40x", FVector2D(16.0f, 16.0f)));
+
+	Style.Set("StatsCountersView.Icon.Large", new IMAGE_BRUSH("Icons/Profiler/Profiler_Data_Capture_40x", FVector2D(32.0f, 32.0f)));
+	Style.Set("StatsCountersView.Icon.Small", new IMAGE_BRUSH("Icons/Profiler/Profiler_Data_Capture_40x", FVector2D(16.0f, 16.0f)));
+
+	Style.Set("LogView.Icon.Large", new IMAGE_BRUSH("Icons/Profiler/profiler_CopyToClipboard_32x", FVector2D(32.0f, 32.0f)));
+	Style.Set("LogView.Icon.Small", new IMAGE_BRUSH("Icons/Profiler/profiler_CopyToClipboard_32x", FVector2D(16.0f, 16.0f)));
+
+	//////////////////////////////////////////////////
+	// Start Page buttons
+
+	Style.Set("Live.Icon.Large", new IMAGE_BRUSH("Icons/Profiler/Profiler_LoadMultiple_Profiler_40x", FVector2D(32.0f, 32.0f)));
+	Style.Set("Live.Icon.Small", new IMAGE_BRUSH("Icons/Profiler/Profiler_Load_Profiler_40x", FVector2D(16.0f, 16.0f)));
+
+	Style.Set("Load.Icon.Large", new IMAGE_BRUSH("Icons/LV_Load", FVector2D(32.0f, 32.0f)));
+	Style.Set("Load.Icon.Small", new IMAGE_BRUSH("Icons/LV_Load", FVector2D(16.0f, 16.0f)));
+
+	//////////////////////////////////////////////////
+
+	return StyleRef;
+}
+
+#undef IMAGE_BRUSH
+#undef BOX_BRUSH
+#undef BORDER_BRUSH
+
+const ISlateStyle& FInsightsStyle::Get()
+{
+	return *StyleInstance;
+}
diff --git a/Engine/Source/Developer/TraceInsights/Private/Insights/IoProfilerCommands.cpp b/Engine/Source/Developer/TraceInsights/Private/Insights/IoProfilerCommands.cpp
new file mode 100644
--- /dev/null
+++ b/Engine/Source/Developer/TraceInsights/Private/Insights/IoProfilerCommands.cpp
@@ -0,0 +1,102 @@
+// Copyright 1998-2019 Epic Games, Inc. All Rights Reserved.
+
+#include "Insights/IoProfilerCommands.h"
+
+#include "DesktopPlatformModule.h"
+#include "EditorStyleSet.h"
+#include "Framework/Application/SlateApplication.h"
+#include "Framework/MultiBox/MultiBoxBuilder.h"
+
+// Insights
+#include "Insights/InsightsManager.h"
+#include "Insights/IoProfilerManager.h"
+
+////////////////////////////////////////////////////////////////////////////////////////////////////
+
+#define LOCTEXT_NAMESPACE "FIoProfilerCommands"
+
+////////////////////////////////////////////////////////////////////////////////////////////////////
+// FIoProfilerMenuBuilder
+////////////////////////////////////////////////////////////////////////////////////////////////////
+
+void FIoProfilerMenuBuilder::AddMenuEntry(FMenuBuilder& MenuBuilder, const TSharedPtr< FUICommandInfo >& UICommandInfo, const FUIAction& UIAction)
+{
+	MenuBuilder.AddMenuEntry
+	(
+		UICommandInfo->GetLabel(),
+		UICommandInfo->GetDescription(),
+		UICommandInfo->GetIcon(),
+		UIAction,
+		NAME_None,
+		UICommandInfo->GetUserInterfaceType()
+	);
+}
+
+////////////////////////////////////////////////////////////////////////////////////////////////////
+// FIoProfilerCommands
+////////////////////////////////////////////////////////////////////////////////////////////////////
+
+FIoProfilerCommands::FIoProfilerCommands()
+	: TCommands<FIoProfilerCommands>(
+		TEXT("IoProfilerCommand"), // Context name for fast lookup
+		NSLOCTEXT("Contexts", "IoProfilerCommand", "Asset Loading Insights Command"), // Localized context name for displaying
+		NAME_None, // Parent
+		FEditorStyle::GetStyleSetName() // Icon Style Set
+	)
+{
+}
+
+////////////////////////////////////////////////////////////////////////////////////////////////////
+
+// UI_COMMAND takes long for the compiler to optimize
+PRAGMA_DISABLE_OPTIMIZATION
+void FIoProfilerCommands::RegisterCommands()
+{
+	UI_COMMAND(ToggleTimingViewVisibility, "Timing", "Toggles the visibility of the main Timing view", EUserInterfaceActionType::ToggleButton, FInputChord(EModifierKey::Control, EKeys::T));
+}
+PRAGMA_ENABLE_OPTIMIZATION
+
+////////////////////////////////////////////////////////////////////////////////////////////////////
+// Toggle Commands
+////////////////////////////////////////////////////////////////////////////////////////////////////
+
+#define IMPLEMENT_TOGGLE_COMMAND(CmdName, IsEnabled, SetIsEnabled) \
+	\
+	void FIoProfilerActionManager::Map_##CmdName##_Global()\
+	{\
+		This->CommandList->MapAction(This->GetCommands().CmdName, CmdName##_Custom());\
+	}\
+	\
+	const FUIAction FIoProfilerActionManager::CmdName##_Custom() \
+	{\
+		FUIAction UIAction;\
+		UIAction.ExecuteAction = FExecuteAction::CreateRaw(this, &FIoProfilerActionManager::CmdName##_Execute);\
+		UIAction.CanExecuteAction = FCanExecuteAction::CreateRaw(this, &FIoProfilerActionManager::CmdName##_CanExecute);\
+		UIAction.GetActionCheckState = FGetActionCheckState::CreateRaw(this, &FIoProfilerActionManager::CmdName##_GetCheckState);\
+		return UIAction;\
+	}\
+	\
+	void FIoProfilerActionManager::CmdName##_Execute()\
+	{\
+		const bool b##IsEnabled = !This->IsEnabled();\
+		This->SetIsEnabled(b##IsEnabled);\
+	}\
+	\
+	bool FIoProfilerActionManager::CmdName##_CanExecute() const\
+	{\
+		return FInsightsManager::Get()->GetSession().IsValid();\
+	}\
+	\
+	ECheckBoxState FIoProfilerActionManager::CmdName##_GetCheckState() const\
+	{\
+		const bool b##IsEnabled = This->IsEnabled();\
+		return b##IsEnabled ? ECheckBoxState::Checked : ECheckBoxState::Unchecked;\
+	}
+
+IMPLEMENT_TOGGLE_COMMAND(ToggleTimingViewVisibility, IsTimingViewVisible, SetTimingViewVisible)
+
+#undef IMPLEMENT_TOGGLE_COMMAND
+
+////////////////////////////////////////////////////////////////////////////////////////////////////
+
+#undef LOCTEXT_NAMESPACE
