Copying //UE4/Dev-Core to Dev-Main (//UE4/Dev-Main)

diff --git a/Engine/Source/Developer/TraceInsights/Private/Insights/IoProfilerCommands.h b/Engine/Source/Developer/TraceInsights/Private/Insights/IoProfilerCommands.h
new file mode 100644
--- /dev/null
+++ b/Engine/Source/Developer/TraceInsights/Private/Insights/IoProfilerCommands.h
@@ -0,0 +1,93 @@
+// Copyright 1998-2019 Epic Games, Inc. All Rights Reserved.
+
+#pragma once
+
+#include "CoreMinimal.h"
+#include "Framework/Commands/UIAction.h"
+#include "Framework/Commands/Commands.h"
+
+class FMenuBuilder;
+
+////////////////////////////////////////////////////////////////////////////////////////////////////
+/**
+ * Class that holds all profiler commands.
+ */
+class FIoProfilerCommands : public TCommands<FIoProfilerCommands>
+{
+public:
+	/** Default constructor. */
+	FIoProfilerCommands();
+
+	/** Initialize commands. */
+	virtual void RegisterCommands() override;
+
+public:
+	//////////////////////////////////////////////////
+	// Global commands need to implement following method:
+	//     void Map_<CommandName>_Global();
+	// Custom commands needs to implement also the following method:
+	//     const FUIAction <CommandName>_Custom(...) const;
+	//////////////////////////////////////////////////
+
+	/** Toggles visibility for the Timing View. Global and custom command. */
+	TSharedPtr<FUICommandInfo> ToggleTimingViewVisibility;
+};
+
+////////////////////////////////////////////////////////////////////////////////////////////////////
+/**
+ * Menu builder. Helper class for adding a customized menu entry using the global UI command info.
+ */
+class FIoProfilerMenuBuilder
+{
+public:
+	/**
+	 * Helper method for adding a customized menu entry using the global UI command info.
+	 * FUICommandInfo cannot be executed with custom parameters, so we need to create a custom FUIAction,
+	 * but sometime we have global and local version for the UI command, so reuse data from the global UI command info.
+	 * Ex:
+	 *     SessionInstance_ToggleCapture          - Global version will toggle capture process for all active session instances
+	 *     SessionInstance_ToggleCapture_OneParam - Local version will toggle capture process only for the specified session instance
+	 *
+	 * @param MenuBuilder The menu to add items to
+	 * @param FUICommandInfo A shared pointer to the UI command info
+	 * @param UIAction Customized version of the UI command info stored in an UI action
+	 */
+	static void AddMenuEntry(FMenuBuilder& MenuBuilder, const TSharedPtr<FUICommandInfo>& UICommandInfo, const FUIAction& UIAction);
+};
+
+////////////////////////////////////////////////////////////////////////////////////////////////////
+/**
+ * Class that provides helper functions for the commands to avoid cluttering profiler manager with many small functions.
+ * Can't contain any variables. Directly operates on the profiler manager instance.
+ */
+class FIoProfilerActionManager
+{
+	friend class FIoProfilerManager;
+
+private:
+	/** Private constructor. */
+	FIoProfilerActionManager(class FIoProfilerManager* Instance)
+		: This(Instance)
+	{}
+
+	//////////////////////////////////////////////////
+	// Toggle Commands
+
+#define DECLARE_TOGGLE_COMMAND(CmdName)\
+public:\
+	void Map_##CmdName##_Global(); /**< Maps UI command info CmdName with the specified UI command list. */\
+	const FUIAction CmdName##_Custom(); /**< UI action for CmdName command. */\
+protected:\
+	void CmdName##_Execute(); /**< Handles FExecuteAction for CmdName. */\
+	bool CmdName##_CanExecute() const; /**< Handles FCanExecuteAction for CmdName. */\
+	ECheckBoxState CmdName##_GetCheckState() const; /**< Handles FGetActionCheckState for CmdName. */
+
+	DECLARE_TOGGLE_COMMAND(ToggleTimingViewVisibility)
+#undef DECLARE_TOGGLE_COMMAND
+
+	//////////////////////////////////////////////////
+
+protected:
+	/** Reference to the global instance of the profiler manager. */
+	class FIoProfilerManager* This;
+};
diff --git a/Engine/Source/Developer/TraceInsights/Private/Insights/IoProfilerManager.cpp b/Engine/Source/Developer/TraceInsights/Private/Insights/IoProfilerManager.cpp
new file mode 100644
--- /dev/null
+++ b/Engine/Source/Developer/TraceInsights/Private/Insights/IoProfilerManager.cpp
@@ -0,0 +1,139 @@
+// Copyright 1998-2019 Epic Games, Inc. All Rights Reserved.
+
+#include "IoProfilerManager.h"
+
+#include "Modules/ModuleManager.h"
+#include "Templates/ScopedPointer.h"
+#include "Templates/UniquePtr.h"
+#include "TraceServices/AnalysisService.h"
+#include "TraceServices/SessionService.h"
+
+// Insights
+#include "Insights/InsightsManager.h"
+//#include "Insights/IoProfilerCommon.h"
+#include "Insights/Widgets/STimingView.h"
+#include "Insights/Widgets/SIoProfilerWindow.h"
+
+////////////////////////////////////////////////////////////////////////////////////////////////////
+
+#define LOCTEXT_NAMESPACE "IoProfilerManager"
+
+//DEFINE_LOG_CATEGORY(IoProfiler);
+//
+//DEFINE_STAT(STAT_FT_OnPaint);
+//DEFINE_STAT(STAT_GT_OnPaint);
+//DEFINE_STAT(STAT_TT_OnPaint);
+//DEFINE_STAT(STAT_IOPM_Tick);
+
+TSharedPtr<FIoProfilerManager> FIoProfilerManager::Instance = nullptr;
+
+////////////////////////////////////////////////////////////////////////////////////////////////////
+
+FIoProfilerManager::FIoProfilerManager(TSharedRef<FUICommandList> InCommandList)
+	: CommandList(InCommandList)
+	, ActionManager(this)
+	, ProfilerWindow(nullptr)
+	, bIsTimingViewVisible(true)
+{
+}
+
+////////////////////////////////////////////////////////////////////////////////////////////////////
+
+void FIoProfilerManager::PostConstructor()
+{
+	// Register tick functions.
+	//OnTick = FTickerDelegate::CreateSP(this, &FIoProfilerManager::Tick);
+	//OnTickHandle = FTicker::GetCoreTicker().AddTicker(OnTick, 1.0f);
+
+	FIoProfilerCommands::Register();
+	BindCommands();
+}
+
+////////////////////////////////////////////////////////////////////////////////////////////////////
+
+void FIoProfilerManager::BindCommands()
+{
+	ActionManager.Map_ToggleTimingViewVisibility_Global();
+}
+
+////////////////////////////////////////////////////////////////////////////////////////////////////
+
+FIoProfilerManager::~FIoProfilerManager()
+{
+	FIoProfilerCommands::Unregister();
+
+	// Unregister tick function.
+	//FTicker::GetCoreTicker().RemoveTicker(OnTickHandle);
+}
+
+////////////////////////////////////////////////////////////////////////////////////////////////////
+
+TSharedPtr<FIoProfilerManager> FIoProfilerManager::Get()
+{
+	return FIoProfilerManager::Instance;
+}
+
+////////////////////////////////////////////////////////////////////////////////////////////////////
+
+const TSharedRef<FUICommandList> FIoProfilerManager::GetCommandList() const
+{
+	return CommandList;
+}
+
+////////////////////////////////////////////////////////////////////////////////////////////////////
+
+const FIoProfilerCommands& FIoProfilerManager::GetCommands()
+{
+	return FIoProfilerCommands::Get();
+}
+
+////////////////////////////////////////////////////////////////////////////////////////////////////
+
+FIoProfilerActionManager& FIoProfilerManager::GetActionManager()
+{
+	return FIoProfilerManager::Instance->ActionManager;
+}
+
+////////////////////////////////////////////////////////////////////////////////////////////////////
+
+bool FIoProfilerManager::Tick(float DeltaTime)
+{
+	//SCOPE_CYCLE_COUNTER(STAT_IOPM_Tick);
+
+	return true;
+}
+
+////////////////////////////////////////////////////////////////////////////////////////////////////
+
+void FIoProfilerManager::OnSessionChanged()
+{
+	TSharedPtr<SIoProfilerWindow> Wnd = GetProfilerWindow();
+	if (Wnd.IsValid())
+	{
+		TSharedPtr<STimingView> TimingView = Wnd->GetTimingView();
+		if (TimingView)
+		{
+			TimingView->Reset();
+		}
+	}
+}
+
+////////////////////////////////////////////////////////////////////////////////////////////////////
+
+void FIoProfilerManager::SetTimingViewVisible(const bool bTimingViewVisibleState)
+{
+	//bool bWasTimingViewVisible = bIsTimingViewVisible;
+
+	bIsTimingViewVisible = bTimingViewVisibleState;
+
+	// TODO
+	//TSharedPtr<SIoProfilerWindow> Wnd = GetProfilerWindow();
+	//if (Wnd.IsValid())
+	//{
+	//	Wnd->ShowHideTab(FIoProfilerTabs::TimingViewID, bIsTimingViewVisible);
+	//}
+}
+
+////////////////////////////////////////////////////////////////////////////////////////////////////
+
+#undef LOCTEXT_NAMESPACE
diff --git a/Engine/Source/Developer/TraceInsights/Private/Insights/IoProfilerManager.h b/Engine/Source/Developer/TraceInsights/Private/Insights/IoProfilerManager.h
new file mode 100644
--- /dev/null
+++ b/Engine/Source/Developer/TraceInsights/Private/Insights/IoProfilerManager.h
@@ -0,0 +1,133 @@
+// Copyright 1998-2019 Epic Games, Inc. All Rights Reserved.
+
+#pragma once
+
+#include "Containers/Ticker.h"
+#include "CoreMinimal.h"
+#include "Framework/Commands/UICommandList.h"
+
+// Insights
+#include "Insights/InsightsManager.h"
+#include "Insights/IoProfilerCommands.h"
+
+class SIoProfilerWindow;
+
+namespace Trace
+{
+	class ISessionService;
+	class IAnalysisService;
+	class IAnalysisSession;
+}
+
+////////////////////////////////////////////////////////////////////////////////////////////////////
+/**
+ * This class manages the Timing Profiler state and settings.
+ */
+class FIoProfilerManager : public TSharedFromThis<FIoProfilerManager>
+{
+	friend class FIoProfilerActionManager;
+
+public:
+	/** Creates the Timing Profiler manager, only one instance can exist. */
+	FIoProfilerManager(TSharedRef<FUICommandList> InCommandList);
+
+	/** Virtual destructor. */
+	virtual ~FIoProfilerManager();
+
+	/** Creates an instance of the profiler manager. */
+	static TSharedPtr<FIoProfilerManager> Initialize()
+	{
+		if (FIoProfilerManager::Instance.IsValid())
+		{
+			FIoProfilerManager::Instance.Reset();
+		}
+
+		FIoProfilerManager::Instance = MakeShareable(new FIoProfilerManager(FInsightsManager::Get()->GetCommandList()));
+		FIoProfilerManager::Instance->PostConstructor();
+
+		return FIoProfilerManager::Instance;
+	}
+
+	/** Shutdowns the Timing Profiler manager. */
+	void Shutdown()
+	{
+		FIoProfilerManager::Instance.Reset();
+	}
+
+protected:
+	/** Finishes initialization of the profiler manager. */
+	void PostConstructor();
+
+	/** Binds our UI commands to delegates. */
+	void BindCommands();
+
+public:
+	/**
+	 * @return the global instance of the Timing Profiler manager (FIoProfilerManager).
+	 * This is an internal singleton and cannot be used outside ProfilerModule.
+	 * For external use:
+	 *     IProfilerModule& ProfilerModule = FModuleManager::Get().LoadModuleChecked<IProfilerModule>("Profiler");
+	 *     ProfilerModule.GetProfilerManager();
+	 */
+	static TSharedPtr<FIoProfilerManager> Get();
+
+	/** @returns UI command list for the Timing Profiler manager. */
+	const TSharedRef<FUICommandList> GetCommandList() const;
+
+	/** @return an instance of the Timing Profiler commands. */
+	static const FIoProfilerCommands& GetCommands();
+
+	/** @return an instance of the Timing Profiler action manager. */
+	static FIoProfilerActionManager& GetActionManager();
+
+	void AssignProfilerWindow(const TSharedRef<SIoProfilerWindow>& InProfilerWindow)
+	{
+		ProfilerWindow = InProfilerWindow;
+	}
+
+	/**
+	 * Converts profiler window weak pointer to a shared pointer and returns it.
+	 * Make sure the returned pointer is valid before trying to dereference it.
+	 */
+	TSharedPtr<class SIoProfilerWindow> GetProfilerWindow() const
+	{
+		return ProfilerWindow.Pin();
+	}
+
+	////////////////////////////////////////////////////////////////////////////////////////////////////
+	// Getters and setters used by Toggle Commands.
+
+	/** @return true, if Timing View is visible */
+	const bool IsTimingViewVisible() const { return bIsTimingViewVisible; }
+	void SetTimingViewVisible(const bool bTimingViewVisibleState);
+
+	////////////////////////////////////////////////////////////////////////////////////////////////////
+
+	void OnSessionChanged();
+
+protected:
+	/** Updates this manager, done through FCoreTicker. */
+	bool Tick(float DeltaTime);
+
+protected:
+	/** The delegate to be invoked when this manager ticks. */
+	FTickerDelegate OnTick;
+
+	/** Handle to the registered OnTick. */
+	FDelegateHandle OnTickHandle;
+
+	/** List of UI commands for this manager. This will be filled by this and corresponding classes. */
+	TSharedRef<FUICommandList> CommandList;
+
+	/** An instance of the Timing Manager action manager. */
+	FIoProfilerActionManager ActionManager;
+
+	/** A weak pointer to the Timing Profiler window. */
+	TWeakPtr<class SIoProfilerWindow> ProfilerWindow;
+
+	/** If the Timing View is visible or hidden. */
+	bool bIsTimingViewVisible;
+
+	/** A shared pointer to the global instance of the profiler manager. */
+	static TSharedPtr<FIoProfilerManager> Instance;
+};
diff --git a/Engine/Source/Developer/TraceInsights/Private/Insights/TimingProfilerCommands.cpp b/Engine/Source/Developer/TraceInsights/Private/Insights/TimingProfilerCommands.cpp
new file mode 100644
--- /dev/null
+++ b/Engine/Source/Developer/TraceInsights/Private/Insights/TimingProfilerCommands.cpp
@@ -0,0 +1,112 @@
+// Copyright 1998-2019 Epic Games, Inc. All Rights Reserved.
+
+#include "Insights/TimingProfilerCommands.h"
+
+#include "DesktopPlatformModule.h"
+#include "EditorStyleSet.h"
+#include "Framework/Application/SlateApplication.h"
+#include "Framework/MultiBox/MultiBoxBuilder.h"
+
+// Insights
+#include "Insights/InsightsManager.h"
+#include "Insights/TimingProfilerManager.h"
+
+////////////////////////////////////////////////////////////////////////////////////////////////////
+
+#define LOCTEXT_NAMESPACE "FTimingProfilerCommands"
+
+////////////////////////////////////////////////////////////////////////////////////////////////////
+// FTimingProfilerMenuBuilder
+////////////////////////////////////////////////////////////////////////////////////////////////////
+
+void FTimingProfilerMenuBuilder::AddMenuEntry(FMenuBuilder& MenuBuilder, const TSharedPtr< FUICommandInfo >& UICommandInfo, const FUIAction& UIAction)
+{
+	MenuBuilder.AddMenuEntry
+	(
+		UICommandInfo->GetLabel(),
+		UICommandInfo->GetDescription(),
+		UICommandInfo->GetIcon(),
+		UIAction,
+		NAME_None,
+		UICommandInfo->GetUserInterfaceType()
+	);
+}
+
+////////////////////////////////////////////////////////////////////////////////////////////////////
+// FTimingProfilerCommands
+////////////////////////////////////////////////////////////////////////////////////////////////////
+
+FTimingProfilerCommands::FTimingProfilerCommands()
+	: TCommands<FTimingProfilerCommands>(
+		TEXT("TimingProfilerCommand"), // Context name for fast lookup
+		NSLOCTEXT("Contexts", "TimingProfilerCommand", "Timing Profiler Command"), // Localized context name for displaying
+		NAME_None, // Parent
+		FEditorStyle::GetStyleSetName() // Icon Style Set
+	)
+{
+}
+
+////////////////////////////////////////////////////////////////////////////////////////////////////
+
+// UI_COMMAND takes long for the compiler to optimize
+PRAGMA_DISABLE_OPTIMIZATION
+void FTimingProfilerCommands::RegisterCommands()
+{
+	UI_COMMAND(ToggleFramesTrackVisibility, "Frames", "Toggles the visibility of the Frames track", EUserInterfaceActionType::ToggleButton, FInputChord(EModifierKey::Control, EKeys::F));
+	UI_COMMAND(ToggleGraphTrackVisibility, "Graph", "Toggles the visibility of the Overview Graph track", EUserInterfaceActionType::ToggleButton, FInputChord(EModifierKey::Control, EKeys::G));
+	UI_COMMAND(ToggleTimingViewVisibility, "Timing", "Toggles the visibility of the main Timing view", EUserInterfaceActionType::ToggleButton, FInputChord(EModifierKey::Control, EKeys::T));
+	UI_COMMAND(ToggleTimersViewVisibility, "Timers", "Toggles the visibility of the Timers view", EUserInterfaceActionType::ToggleButton, FInputChord());
+	UI_COMMAND(ToggleStatsCountersViewVisibility, "Stats", "Toggles the visibility of the Stats Counters view", EUserInterfaceActionType::ToggleButton, FInputChord());
+	UI_COMMAND(ToggleLogViewVisibility, "Log", "Toggles the visibility of the Log view", EUserInterfaceActionType::ToggleButton, FInputChord(EKeys::L));
+}
+PRAGMA_ENABLE_OPTIMIZATION
+
+////////////////////////////////////////////////////////////////////////////////////////////////////
+// Toggle Commands
+////////////////////////////////////////////////////////////////////////////////////////////////////
+
+#define IMPLEMENT_TOGGLE_COMMAND(CmdName, IsEnabled, SetIsEnabled) \
+	\
+	void FTimingProfilerActionManager::Map_##CmdName##_Global()\
+	{\
+		This->CommandList->MapAction(This->GetCommands().CmdName, CmdName##_Custom());\
+	}\
+	\
+	const FUIAction FTimingProfilerActionManager::CmdName##_Custom() \
+	{\
+		FUIAction UIAction;\
+		UIAction.ExecuteAction = FExecuteAction::CreateRaw(this, &FTimingProfilerActionManager::CmdName##_Execute);\
+		UIAction.CanExecuteAction = FCanExecuteAction::CreateRaw(this, &FTimingProfilerActionManager::CmdName##_CanExecute);\
+		UIAction.GetActionCheckState = FGetActionCheckState::CreateRaw(this, &FTimingProfilerActionManager::CmdName##_GetCheckState);\
+		return UIAction;\
+	}\
+	\
+	void FTimingProfilerActionManager::CmdName##_Execute()\
+	{\
+		const bool b##IsEnabled = !This->IsEnabled();\
+		This->SetIsEnabled(b##IsEnabled);\
+	}\
+	\
+	bool FTimingProfilerActionManager::CmdName##_CanExecute() const\
+	{\
+		return FInsightsManager::Get()->GetSession().IsValid();\
+	}\
+	\
+	ECheckBoxState FTimingProfilerActionManager::CmdName##_GetCheckState() const\
+	{\
+		const bool b##IsEnabled = This->IsEnabled();\
+		return b##IsEnabled ? ECheckBoxState::Checked : ECheckBoxState::Unchecked;\
+	}
+
+IMPLEMENT_TOGGLE_COMMAND(ToggleFramesTrackVisibility, IsFramesTrackVisible, ShowHideFramesTrack)
+IMPLEMENT_TOGGLE_COMMAND(ToggleGraphTrackVisibility, IsGraphTrackVisible, ShowHideGraphTrack)
+IMPLEMENT_TOGGLE_COMMAND(ToggleTimingViewVisibility, IsTimingViewVisible, ShowHideTimingView)
+IMPLEMENT_TOGGLE_COMMAND(ToggleTimersViewVisibility, IsTimersViewVisible, ShowHideTimersView)
+IMPLEMENT_TOGGLE_COMMAND(ToggleStatsCountersViewVisibility, IsStatsCountersViewVisible, ShowHideStatsCountersView)
+IMPLEMENT_TOGGLE_COMMAND(ToggleLogViewVisibility, IsLogViewVisible, ShowHideLogView)
+
+#undef IMPLEMENT_TOGGLE_COMMAND
+
+////////////////////////////////////////////////////////////////////////////////////////////////////
+
+#undef LOCTEXT_NAMESPACE
diff --git a/Engine/Source/Developer/TraceInsights/Private/Insights/TimingProfilerCommands.h b/Engine/Source/Developer/TraceInsights/Private/Insights/TimingProfilerCommands.h
new file mode 100644
--- /dev/null
+++ b/Engine/Source/Developer/TraceInsights/Private/Insights/TimingProfilerCommands.h
@@ -0,0 +1,123 @@
+// Copyright 1998-2019 Epic Games, Inc. All Rights Reserved.
+
+#pragma once
+
+#include "CoreMinimal.h"
+#include "Framework/Commands/UIAction.h"
+#include "Framework/Commands/Commands.h"
+
+class FMenuBuilder;
+
+////////////////////////////////////////////////////////////////////////////////////////////////////
+/**
+ * Class that holds all profiler commands.
+ */
+class FTimingProfilerCommands : public TCommands<FTimingProfilerCommands>
+{
+public:
+	/** Default constructor. */
+	FTimingProfilerCommands();
+
+	/** Initialize commands. */
+	virtual void RegisterCommands() override;
+
+public:
+	//////////////////////////////////////////////////
+	// Global commands need to implement following method:
+	//     void Map_<CommandName>_Global();
+	// Custom commands needs to implement also the following method:
+	//     const FUIAction <CommandName>_Custom(...) const;
+	//////////////////////////////////////////////////
+
+	/** Toggles visibility for the Frames Track. Global and custom command. */
+	TSharedPtr<FUICommandInfo> ToggleFramesTrackVisibility;
+
+	/** Toggles visibility for the Graph Track. Global and custom command. */
+	TSharedPtr<FUICommandInfo> ToggleGraphTrackVisibility;
+
+	/** Toggles visibility for the Timing View. Global and custom command. */
+	TSharedPtr<FUICommandInfo> ToggleTimingViewVisibility;
+
+	/** Toggles visibility for the Timers View. Global and custom command. */
+	TSharedPtr<FUICommandInfo> ToggleTimersViewVisibility;
+
+	/** Toggles visibility for the Stats Counters View. Global and custom command. */
+	TSharedPtr<FUICommandInfo> ToggleStatsCountersViewVisibility;
+
+	/** Toggles visibility for the Log View. Global and custom command. */
+	TSharedPtr<FUICommandInfo> ToggleLogViewVisibility;
+};
+
+////////////////////////////////////////////////////////////////////////////////////////////////////
+/**
+ * Menu builder. Helper class for adding a customized menu entry using the global UI command info.
+ */
+class FTimingProfilerMenuBuilder
+{
+public:
+	/**
+	 * Helper method for adding a customized menu entry using the global UI command info.
+	 * FUICommandInfo cannot be executed with custom parameters, so we need to create a custom FUIAction,
+	 * but sometime we have global and local version for the UI command, so reuse data from the global UI command info.
+	 * Ex:
+	 *     SessionInstance_ToggleCapture          - Global version will toggle capture process for all active session instances
+	 *     SessionInstance_ToggleCapture_OneParam - Local version will toggle capture process only for the specified session instance
+	 *
+	 * @param MenuBuilder The menu to add items to
+	 * @param FUICommandInfo A shared pointer to the UI command info
+	 * @param UIAction Customized version of the UI command info stored in an UI action
+	 */
+	static void AddMenuEntry(FMenuBuilder& MenuBuilder, const TSharedPtr<FUICommandInfo>& UICommandInfo, const FUIAction& UIAction);
+};
+
+////////////////////////////////////////////////////////////////////////////////////////////////////
+/**
+ * Class that provides helper functions for the commands to avoid cluttering profiler manager with many small functions.
+ * Can't contain any variables. Directly operates on the profiler manager instance.
+ */
+class FTimingProfilerActionManager
+{
+	friend class FTimingProfilerManager;
+
+private:
+	/** Private constructor. */
+	FTimingProfilerActionManager(class FTimingProfilerManager* Instance)
+		: This(Instance)
+	{}
+
+	//////////////////////////////////////////////////
+	// Toggle Commands
+
+#define DECLARE_TOGGLE_COMMAND(CmdName)\
+public:\
+	void Map_##CmdName##_Global(); /**< Maps UI command info CmdName with the specified UI command list. */\
+	const FUIAction CmdName##_Custom(); /**< UI action for CmdName command. */\
+protected:\
+	void CmdName##_Execute(); /**< Handles FExecuteAction for CmdName. */\
+	bool CmdName##_CanExecute() const; /**< Handles FCanExecuteAction for CmdName. */\
+	ECheckBoxState CmdName##_GetCheckState() const; /**< Handles FGetActionCheckState for CmdName. */
+
+	DECLARE_TOGGLE_COMMAND(ToggleFramesTrackVisibility)
+	DECLARE_TOGGLE_COMMAND(ToggleGraphTrackVisibility)
+	DECLARE_TOGGLE_COMMAND(ToggleTimingViewVisibility)
+	DECLARE_TOGGLE_COMMAND(ToggleTimersViewVisibility)
+	DECLARE_TOGGLE_COMMAND(ToggleStatsCountersViewVisibility)
+	DECLARE_TOGGLE_COMMAND(ToggleLogViewVisibility)
+#undef DECLARE_TOGGLE_COMMAND
+
+	//////////////////////////////////////////////////
+	// OpenSettings
+
+public:
+	void Map_OpenSettings_Global(); /**< Maps UI command info OpenSettings with the specified UI command list. */
+	const FUIAction OpenSettings_Custom() const; /**< UI action for OpenSettings command. */
+protected:
+	void OpenSettings_Execute(); /**< Handles FExecuteAction for OpenSettings. */
+	bool OpenSettings_CanExecute() const; /**< Handles FCanExecuteAction for OpenSettings. */
+
+	//////////////////////////////////////////////////
+
+protected:
+	/** Reference to the global instance of the profiler manager. */
+	class FTimingProfilerManager* This;
+};
diff --git a/Engine/Source/Developer/TraceInsights/Private/Insights/TimingProfilerManager.cpp b/Engine/Source/Developer/TraceInsights/Private/Insights/TimingProfilerManager.cpp
new file mode 100644
--- /dev/null
+++ b/Engine/Source/Developer/TraceInsights/Private/Insights/TimingProfilerManager.cpp
@@ -0,0 +1,247 @@
+// Copyright 1998-2019 Epic Games, Inc. All Rights Reserved.
+
+#include "TimingProfilerManager.h"
+
+#include "Modules/ModuleManager.h"
+#include "Templates/ScopedPointer.h"
+#include "Templates/UniquePtr.h"
+#include "TraceServices/AnalysisService.h"
+#include "TraceServices/SessionService.h"
+
+// Insights
+#include "Insights/InsightsManager.h"
+#include "Insights/TimingProfilerCommon.h"
+#include "Insights/Widgets/SFrameTrack.h"
+#include "Insights/Widgets/SGraphTrack.h"
+#include "Insights/Widgets/SLogView.h"
+#include "Insights/Widgets/SStatsView.h"
+#include "Insights/Widgets/STimersView.h"
+#include "Insights/Widgets/STimingView.h"
+#include "Insights/Widgets/STimingProfilerWindow.h"
+
+////////////////////////////////////////////////////////////////////////////////////////////////////
+
+#define LOCTEXT_NAMESPACE "TimingProfilerManager"
+
+DEFINE_LOG_CATEGORY(TimingProfiler);
+
+DEFINE_STAT(STAT_FT_OnPaint);
+DEFINE_STAT(STAT_GT_OnPaint);
+DEFINE_STAT(STAT_TT_OnPaint);
+DEFINE_STAT(STAT_TPM_Tick);
+
+TSharedPtr<FTimingProfilerManager> FTimingProfilerManager::Instance = nullptr;
+
+////////////////////////////////////////////////////////////////////////////////////////////////////
+
+FTimingProfilerManager::FTimingProfilerManager(TSharedRef<FUICommandList> InCommandList)
+	: CommandList(InCommandList)
+	, ActionManager(this)
+	, ProfilerWindow(nullptr)
+	, bIsFramesTrackVisible(true)
+	, bIsGraphTrackVisible(false)
+	, bIsTimingViewVisible(true)
+	, bIsTimersViewVisible(true)
+	, bIsStatsCountersViewVisible(true)
+	, bIsLogViewVisible(true)
+{
+}
+
+////////////////////////////////////////////////////////////////////////////////////////////////////
+
+void FTimingProfilerManager::PostConstructor()
+{
+	// Register tick functions.
+	//OnTick = FTickerDelegate::CreateSP(this, &FTimingProfilerManager::Tick);
+	//OnTickHandle = FTicker::GetCoreTicker().AddTicker(OnTick, 1.0f);
+
+	FTimingProfilerCommands::Register();
+	BindCommands();
+}
+
+////////////////////////////////////////////////////////////////////////////////////////////////////
+
+void FTimingProfilerManager::BindCommands()
+{
+	ActionManager.Map_ToggleFramesTrackVisibility_Global();
+	ActionManager.Map_ToggleGraphTrackVisibility_Global();
+	ActionManager.Map_ToggleTimingViewVisibility_Global();
+	ActionManager.Map_ToggleTimersViewVisibility_Global();
+	ActionManager.Map_ToggleStatsCountersViewVisibility_Global();
+	ActionManager.Map_ToggleLogViewVisibility_Global();
+}
+
+////////////////////////////////////////////////////////////////////////////////////////////////////
+
+FTimingProfilerManager::~FTimingProfilerManager()
+{
+	FTimingProfilerCommands::Unregister();
+
+	// Unregister tick function.
+	//FTicker::GetCoreTicker().RemoveTicker(OnTickHandle);
+}
+
+////////////////////////////////////////////////////////////////////////////////////////////////////
+
+TSharedPtr<FTimingProfilerManager> FTimingProfilerManager::Get()
+{
+	return FTimingProfilerManager::Instance;
+}
+
+////////////////////////////////////////////////////////////////////////////////////////////////////
+
+const TSharedRef<FUICommandList> FTimingProfilerManager::GetCommandList() const
+{
+	return CommandList;
+}
+
+////////////////////////////////////////////////////////////////////////////////////////////////////
+
+const FTimingProfilerCommands& FTimingProfilerManager::GetCommands()
+{
+	return FTimingProfilerCommands::Get();
+}
+
+////////////////////////////////////////////////////////////////////////////////////////////////////
+
+FTimingProfilerActionManager& FTimingProfilerManager::GetActionManager()
+{
+	return FTimingProfilerManager::Instance->ActionManager;
+}
+
+////////////////////////////////////////////////////////////////////////////////////////////////////
+
+bool FTimingProfilerManager::Tick(float DeltaTime)
+{
+	SCOPE_CYCLE_COUNTER(STAT_TPM_Tick);
+
+	return true;
+}
+
+////////////////////////////////////////////////////////////////////////////////////////////////////
+
+void FTimingProfilerManager::OnSessionChanged()
+{
+	TSharedPtr<STimingProfilerWindow> Wnd = GetProfilerWindow();
+	if (Wnd.IsValid())
+	{
+		if (Wnd->FrameTrack)
+		{
+			Wnd->FrameTrack->Reset();
+		}
+		if (Wnd->GraphTrack)
+		{
+			Wnd->GraphTrack->Reset();
+		}
+		if (Wnd->TimingView)
+		{
+			Wnd->TimingView->Reset();
+		}
+		if (Wnd->TimersView)
+		{
+			Wnd->TimersView->RebuildTree();
+		}
+		if (Wnd->StatsView)
+		{
+			Wnd->StatsView->RebuildTree();
+		}
+		if (Wnd->LogView)
+		{
+			Wnd->LogView->Reset();
+		}
+	}
+}
+
+////////////////////////////////////////////////////////////////////////////////////////////////////
+
+void FTimingProfilerManager::ShowHideFramesTrack(const bool bFramesTrackVisibleState)
+{
+	bool bWasFramesTrackVisible = bIsFramesTrackVisible;
+
+	bIsFramesTrackVisible = bFramesTrackVisibleState;
+
+	TSharedPtr<STimingProfilerWindow> Wnd = GetProfilerWindow();
+	if (Wnd.IsValid())
+	{
+		Wnd->ShowHideTab(FTimingProfilerTabs::FramesTrackID, bIsFramesTrackVisible);
+	}
+}
+
+////////////////////////////////////////////////////////////////////////////////////////////////////
+
+void FTimingProfilerManager::ShowHideGraphTrack(const bool bGraphTrackVisibleState)
+{
+	bool bWasGraphTrackVisible = bIsGraphTrackVisible;
+
+	bIsGraphTrackVisible = bGraphTrackVisibleState;
+
+	TSharedPtr<STimingProfilerWindow> Wnd = GetProfilerWindow();
+	if (Wnd.IsValid())
+	{
+		Wnd->ShowHideTab(FTimingProfilerTabs::GraphTrackID, bIsGraphTrackVisible);
+	}
+}
+
+////////////////////////////////////////////////////////////////////////////////////////////////////
+
+void FTimingProfilerManager::ShowHideTimingView(const bool bTimingViewVisibleState)
+{
+	bool bWasTimingViewVisible = bIsTimingViewVisible;
+
+	bIsTimingViewVisible = bTimingViewVisibleState;
+
+	TSharedPtr<STimingProfilerWindow> Wnd = GetProfilerWindow();
+	if (Wnd.IsValid())
+	{
+		Wnd->ShowHideTab(FTimingProfilerTabs::TimingViewID, bIsTimingViewVisible);
+	}
+}
+
+////////////////////////////////////////////////////////////////////////////////////////////////////
+
+void FTimingProfilerManager::ShowHideTimersView(const bool bTimersViewVisibleState)
+{
+	bool bWasTimersViewVisible = bIsTimersViewVisible;
+
+	bIsTimersViewVisible = bTimersViewVisibleState;
+
+	TSharedPtr<STimingProfilerWindow> Wnd = GetProfilerWindow();
+	if (Wnd.IsValid())
+	{
+		Wnd->ShowHideTab(FTimingProfilerTabs::TimersID, bIsTimersViewVisible);
+	}
+}
+
+////////////////////////////////////////////////////////////////////////////////////////////////////
+
+void FTimingProfilerManager::ShowHideStatsCountersView(const bool bStatsCountersViewVisibleState)
+{
+	bool bWasStatsCountersViewVisible = bIsStatsCountersViewVisible;
+
+	bIsStatsCountersViewVisible = bStatsCountersViewVisibleState;
+
+	TSharedPtr<STimingProfilerWindow> Wnd = GetProfilerWindow();
+	if (Wnd.IsValid())
+	{
+		Wnd->ShowHideTab(FTimingProfilerTabs::StatsCountersID, bIsStatsCountersViewVisible);
+	}
+}
+
+////////////////////////////////////////////////////////////////////////////////////////////////////
+
+void FTimingProfilerManager::ShowHideLogView(const bool bLogViewVisibleState)
+{
+	bool bWasLogViewVisible = bIsLogViewVisible;
+
+	bIsLogViewVisible = bLogViewVisibleState;
+
+	TSharedPtr<STimingProfilerWindow> Wnd = GetProfilerWindow();
+	if (Wnd.IsValid())
+	{
+		Wnd->ShowHideTab(FTimingProfilerTabs::LogViewID, bIsLogViewVisible);
+	}
+}
+
+////////////////////////////////////////////////////////////////////////////////////////////////////
+
+#undef LOCTEXT_NAMESPACE
diff --git a/Engine/Source/Developer/TraceInsights/Private/Insights/TimingProfilerManager.h b/Engine/Source/Developer/TraceInsights/Private/Insights/TimingProfilerManager.h
new file mode 100644
--- /dev/null
+++ b/Engine/Source/Developer/TraceInsights/Private/Insights/TimingProfilerManager.h
@@ -0,0 +1,174 @@
+// Copyright 1998-2019 Epic Games, Inc. All Rights Reserved.
+
+#pragma once
+
+#include "Containers/Ticker.h"
+#include "CoreMinimal.h"
+#include "Framework/Commands/UICommandList.h"
+
+// Insights
+#include "Insights/InsightsManager.h"
+#include "Insights/TimingProfilerCommands.h"
+
+class STimingProfilerWindow;
+
+namespace Trace
+{
+	class ISessionService;
+	class IAnalysisService;
+	class IAnalysisSession;
+}
+
+////////////////////////////////////////////////////////////////////////////////////////////////////
+/**
+ * This class manages the Timing Profiler state and settings.
+ */
+class FTimingProfilerManager : public TSharedFromThis<FTimingProfilerManager>
+{
+	friend class FTimingProfilerActionManager;
+
+public:
+	/** Creates the Timing Profiler manager, only one instance can exist. */
+	FTimingProfilerManager(TSharedRef<FUICommandList> InCommandList);
+
+	/** Virtual destructor. */
+	virtual ~FTimingProfilerManager();
+
+	/** Creates an instance of the profiler manager. */
+	static TSharedPtr<FTimingProfilerManager> Initialize()
+	{
+		if (FTimingProfilerManager::Instance.IsValid())
+		{
+			FTimingProfilerManager::Instance.Reset();
+		}
+
+		FTimingProfilerManager::Instance = MakeShareable(new FTimingProfilerManager(FInsightsManager::Get()->GetCommandList()));
+		FTimingProfilerManager::Instance->PostConstructor();
+
+		return FTimingProfilerManager::Instance;
+	}
+
+	/** Shutdowns the Timing Profiler manager. */
+	void Shutdown()
+	{
+		FTimingProfilerManager::Instance.Reset();
+	}
+
+protected:
+	/** Finishes initialization of the profiler manager. */
+	void PostConstructor();
+
+	/** Binds our UI commands to delegates. */
+	void BindCommands();
+
+public:
+	/**
+	 * @return the global instance of the Timing Profiler manager (FTimingProfilerManager).
+	 * This is an internal singleton and cannot be used outside ProfilerModule.
+	 * For external use:
+	 *     IProfilerModule& ProfilerModule = FModuleManager::Get().LoadModuleChecked<IProfilerModule>("Profiler");
+	 *     ProfilerModule.GetProfilerManager();
+	 */
+	static TSharedPtr<FTimingProfilerManager> Get();
+
+	/** @returns UI command list for the Timing Profiler manager. */
+	const TSharedRef<FUICommandList> GetCommandList() const;
+
+	/** @return an instance of the Timing Profiler commands. */
+	static const FTimingProfilerCommands& GetCommands();
+
+	/** @return an instance of the Timing Profiler action manager. */
+	static FTimingProfilerActionManager& GetActionManager();
+
+	void AssignProfilerWindow(const TSharedRef<STimingProfilerWindow>& InProfilerWindow)
+	{
+		ProfilerWindow = InProfilerWindow;
+	}
+
+	/**
+	 * Converts profiler window weak pointer to a shared pointer and returns it.
+	 * Make sure the returned pointer is valid before trying to dereference it.
+	 */
+	TSharedPtr<class STimingProfilerWindow> GetProfilerWindow() const
+	{
+		return ProfilerWindow.Pin();
+	}
+
+	////////////////////////////////////////////////////////////////////////////////////////////////////
+	// Getters and setters used by Toggle Commands.
+
+	/** @return true, if Frames Track is visible */
+	const bool IsFramesTrackVisible() const { return bIsFramesTrackVisible; }
+	void SetFramesTrackVisible(const bool bFramesTrackVisibleState) { bIsFramesTrackVisible = bFramesTrackVisibleState; }
+	void ShowHideFramesTrack(const bool bFramesTrackVisibleState);
+
+	/** @return true, if Graph Track is visible */
+	const bool IsGraphTrackVisible() const { return bIsGraphTrackVisible; }
+	void SetGraphTrackVisible(const bool bGraphTrackVisibleState) { bIsGraphTrackVisible = bGraphTrackVisibleState; }
+	void ShowHideGraphTrack(const bool bGraphTrackVisibleState);
+
+	/** @return true, if Timing View is visible */
+	const bool IsTimingViewVisible() const { return bIsTimingViewVisible; }
+	void SetTimingViewVisible(const bool bTimingViewVisibleState) { bIsTimingViewVisible = bTimingViewVisibleState; }
+	void ShowHideTimingView(const bool bTimingViewVisibleState);
+
+	/** @return true, if Timers View is visible */
+	const bool IsTimersViewVisible() const { return bIsTimersViewVisible; }
+	void SetTimersViewVisible(const bool bTimersViewVisibleState) { bIsTimersViewVisible = bTimersViewVisibleState; }
+	void ShowHideTimersView(const bool bTimersViewVisibleState);
+
+	/** @return true, if Stats Counters View is visible */
+	const bool IsStatsCountersViewVisible() const { return bIsStatsCountersViewVisible; }
+	void SetStatsCountersViewVisible(const bool bStatsCountersViewVisibleState) { bIsStatsCountersViewVisible = bStatsCountersViewVisibleState; }
+	void ShowHideStatsCountersView(const bool bStatsCountersViewVisibleState);
+
+	/** @return true, if Log View is visible */
+	const bool IsLogViewVisible() const { return bIsLogViewVisible; }
+	void SetLogViewVisible(const bool bLogViewVisibleState) { bIsLogViewVisible = bLogViewVisibleState; }
+	void ShowHideLogView(const bool bLogViewVisibleState);
+
+	////////////////////////////////////////////////////////////////////////////////////////////////////
+
+	void OnSessionChanged();
+
+protected:
+	/** Updates this manager, done through FCoreTicker. */
+	bool Tick(float DeltaTime);
+
+protected:
+	/** The delegate to be invoked when this manager ticks. */
+	FTickerDelegate OnTick;
+
+	/** Handle to the registered OnTick. */
+	FDelegateHandle OnTickHandle;
+
+	/** List of UI commands for this manager. This will be filled by this and corresponding classes. */
+	TSharedRef<FUICommandList> CommandList;
+
+	/** An instance of the Timing Manager action manager. */
+	FTimingProfilerActionManager ActionManager;
+
+	/** A weak pointer to the Timing Profiler window. */
+	TWeakPtr<class STimingProfilerWindow> ProfilerWindow;
+
+	/** If the Frames Track is visible or hidden. */
+	bool bIsFramesTrackVisible;
+
+	/** If the Graph Track is visible or hidden. */
+	bool bIsGraphTrackVisible;
+
+	/** If the Timing View is visible or hidden. */
+	bool bIsTimingViewVisible;
+
+	/** If the Timers View is visible or hidden. */
+	bool bIsTimersViewVisible;
+
+	/** If the Stats Counters View is visible or hidden. */
+	bool bIsStatsCountersViewVisible;
+
+	/** If the Log View is visible or hidden. */
+	bool bIsLogViewVisible;
+
+	/** A shared pointer to the global instance of the profiler manager. */
+	static TSharedPtr<FTimingProfilerManager> Instance;
+};
diff --git a/Engine/Source/Developer/TraceInsights/Private/Insights/TraceInsightsModule.cpp b/Engine/Source/Developer/TraceInsights/Private/Insights/TraceInsightsModule.cpp
new file mode 100644
--- /dev/null
+++ b/Engine/Source/Developer/TraceInsights/Private/Insights/TraceInsightsModule.cpp
@@ -0,0 +1,247 @@
+// Copyright 1998-2019 Epic Games, Inc. All Rights Reserved.
+
+#include "CoreMinimal.h"
+#include "EditorStyleSet.h"
+#include "HAL/PlatformApplicationMisc.h"
+#include "Modules/ModuleManager.h"
+#include "TraceServices/ITraceServicesModule.h"
+#include "TraceServices/SessionService.h"
+#include "Widgets/DeclarativeSyntaxSupport.h"
+#include "Widgets/Docking/SDockTab.h"
+#include "Widgets/SWidget.h"
+#include "Widgets/Text/STextBlock.h"
+//#include "WorkspaceMenuStructure.h"
+//#include "WorkspaceMenuStructureModule.h"
+
+// Insights
+#include "Insights/InsightsManager.h"
+#include "Insights/InsightsStyle.h"
+#include "Insights/IoProfilerManager.h"
+#include "Insights/IUnrealInsightsModule.h"
+#include "Insights/TimingProfilerManager.h"
+#include "Insights/Widgets/SIoProfilerWindow.h"
+#include "Insights/Widgets/SStartPageWindow.h"
+#include "Insights/Widgets/STimingProfilerWindow.h"
+
+////////////////////////////////////////////////////////////////////////////////////////////////////
+/**
+ * Implements the Trace Insights module (TimingProfiler, IoProfiler, etc.).
+ */
+class FTraceInsightsModule : public IUnrealInsightsModule
+{
+public:
+	virtual void StartupModule() override;
+	virtual void ShutdownModule() override;
+
+	virtual bool SupportsDynamicReloading() override
+	{
+		return false;
+	}
+
+	virtual void OnNewLayout(TSharedRef<FTabManager::FLayout> NewLayout) override;
+	virtual void OnLayoutRestored(TSharedPtr<FTabManager> TabManager) override;
+
+protected:
+	/** Callback called when a major tab is closed. */
+	void OnTabBeingClosed(TSharedRef<SDockTab> TabBeingClosed);
+
+	/** Callback called when the Timing Profiler major tab is closed. */
+	void OnTimingProfilerTabBeingClosed(TSharedRef<SDockTab> TabBeingClosed);
+
+	/** Callback called when the Io Profiler major tab is closed. */
+	void OnIoProfilerTabBeingClosed(TSharedRef<SDockTab> TabBeingClosed);
+
+	/** Start Page */
+	TSharedRef<SDockTab> SpawnStartPageTab(const FSpawnTabArgs& Args);
+
+	/** Timing Profiler */
+	TSharedRef<SDockTab> SpawnTimingProfilerTab(const FSpawnTabArgs& Args);
+
+	/** I/O Profiler */
+	TSharedRef<SDockTab> SpawnIoProfilerTab(const FSpawnTabArgs& Args);
+
+protected:
+	TSharedPtr<Trace::IAnalysisService> TraceAnalysisService;
+	TSharedPtr<Trace::ISessionService> TraceSessionService;
+	TSharedPtr<Trace::IModuleService> TraceModuleService;
+};
+
+////////////////////////////////////////////////////////////////////////////////////////////////////
+
+IMPLEMENT_MODULE(FTraceInsightsModule, TraceInsights);
+
+////////////////////////////////////////////////////////////////////////////////////////////////////
+// FTraceInsightsModule
+////////////////////////////////////////////////////////////////////////////////////////////////////
+
+void FTraceInsightsModule::StartupModule()
+{
+	ITraceServicesModule& TraceServicesModule = FModuleManager::LoadModuleChecked<ITraceServicesModule>("TraceServices");
+	TraceAnalysisService = TraceServicesModule.GetAnalysisService();
+	TraceSessionService = TraceServicesModule.GetSessionService();
+	TraceModuleService = TraceServicesModule.GetModuleService();
+
+	// Starts the Trace Recorder service.
+	TraceSessionService->StartRecorderServer();
+
+	FInsightsStyle::Initialize();
+
+	FInsightsManager::Initialize(TraceAnalysisService.ToSharedRef(), TraceSessionService.ToSharedRef(), TraceModuleService.ToSharedRef());
+	FTimingProfilerManager::Initialize();
+	FIoProfilerManager::Initialize();
+
+	// Register tab spawner for the Start Page.
+	auto& StartPageTabSpawnerEntry = FGlobalTabmanager::Get()->RegisterTabSpawner(FInsightsManagerTabs::StartPageTabId,
+		FOnSpawnTab::CreateRaw(this, &FTraceInsightsModule::SpawnStartPageTab))
+		.SetDisplayName(NSLOCTEXT("FTraceInsightsModule", "StartPageTabTitle", "Unreal Insights"))
+		.SetTooltipText(NSLOCTEXT("FTraceInsightsModule", "StartPageTooltipText", "Open the start page for Unreal Insights."))
+		.SetIcon(FSlateIcon(FInsightsStyle::GetStyleSetName(), "StartPage.Icon.Small"));
+
+	// Register tab spawner for the Timing Insights.
+	auto& TimingProfilerTabSpawnerEntry = FGlobalTabmanager::Get()->RegisterTabSpawner(FInsightsManagerTabs::TimingProfilerTabId,
+		FOnSpawnTab::CreateRaw(this, &FTraceInsightsModule::SpawnTimingProfilerTab))
+		.SetDisplayName(NSLOCTEXT("FTraceInsightsModule", "TimingProfilerTabTitle", "Timing Insights"))
+		.SetTooltipText(NSLOCTEXT("FTraceInsightsModule", "TimingProfilerTooltipText", "Open the Timing Insights tab."))
+		.SetIcon(FSlateIcon(FInsightsStyle::GetStyleSetName(), "TimingInsights.Icon.Small"));
+
+//#if WITH_EDITOR
+//	TimingProfilerTabSpawnerEntry.SetGroup(WorkspaceMenu::GetMenuStructure().GetDeveloperToolsMiscCategory());
+//#else
+//	TimingProfilerTabSpawnerEntry.SetGroup(WorkspaceMenu::GetMenuStructure().GetToolsCategory());
+//#endif
+
+	// Register tab spawner for the Asset Loading Insights.
+	auto& IoProfilerTabSpawnerEntry = FGlobalTabmanager::Get()->RegisterTabSpawner(FInsightsManagerTabs::IoProfilerTabId,
+		FOnSpawnTab::CreateRaw(this, &FTraceInsightsModule::SpawnIoProfilerTab))
+		.SetDisplayName(NSLOCTEXT("FTraceInsightsModule", "IoProfilerTabTitle", "Asset Loading Insights"))
+		.SetTooltipText(NSLOCTEXT("FTraceInsightsModule", "IoProfilerTooltipText", "Open the Asset Loading Insights tab."))
+		.SetIcon(FSlateIcon(FInsightsStyle::GetStyleSetName(), "AssetLoadingInsights.Icon.Small"));
+
+//#if WITH_EDITOR
+//	IoProfilerTabSpawnerEntry.SetGroup(WorkspaceMenu::GetMenuStructure().GetDeveloperToolsMiscCategory());
+//#else
+//	IoProfilerTabSpawnerEntry.SetGroup(WorkspaceMenu::GetMenuStructure().GetToolsCategory());
+//#endif
+}
+
+////////////////////////////////////////////////////////////////////////////////////////////////////
+
+void FTraceInsightsModule::ShutdownModule()
+{
+	TraceSessionService->StopRecorderServer();
+
+	FGlobalTabmanager::Get()->UnregisterNomadTabSpawner(FInsightsManagerTabs::StartPageTabId);
+	FGlobalTabmanager::Get()->UnregisterNomadTabSpawner(FInsightsManagerTabs::TimingProfilerTabId);
+	FGlobalTabmanager::Get()->UnregisterNomadTabSpawner(FInsightsManagerTabs::IoProfilerTabId);
+
+	if (FTimingProfilerManager::Get().IsValid())
+	{
+		// Shutdown the Timing Profiler manager.
+		FTimingProfilerManager::Get()->Shutdown();
+	}
+
+	if (FIoProfilerManager::Get().IsValid())
+	{
+		// Shutdown the I/O Profiler manager.
+		FIoProfilerManager::Get()->Shutdown();
+	}
+
+	if (FInsightsManager::Get().IsValid())
+	{
+		// Shutdown the main manager.
+		FInsightsManager::Get()->Shutdown();
+	}
+}
+
+////////////////////////////////////////////////////////////////////////////////////////////////////
+
+void FTraceInsightsModule::OnNewLayout(TSharedRef<FTabManager::FLayout> NewLayout)
+{
+	const float DPIScaleFactor = FPlatformApplicationMisc::GetDPIScaleFactorAtPoint(10.0f, 10.0f);
+	NewLayout->AddArea
+	(
+		FTabManager::NewArea(1280.f * DPIScaleFactor, 720.0f * DPIScaleFactor)
+		->Split
+		(
+			FTabManager::NewStack()
+			->AddTab(FInsightsManagerTabs::StartPageTabId, ETabState::OpenedTab)
+			->AddTab(FInsightsManagerTabs::TimingProfilerTabId, ETabState::ClosedTab)
+			->AddTab(FInsightsManagerTabs::IoProfilerTabId, ETabState::ClosedTab)
+			->SetForegroundTab(FTabId(FInsightsManagerTabs::StartPageTabId))
+		)
+	);
+}
+
+////////////////////////////////////////////////////////////////////////////////////////////////////
+
+void FTraceInsightsModule::OnLayoutRestored(TSharedPtr<FTabManager> TabManager)
+{
+}
+
+////////////////////////////////////////////////////////////////////////////////////////////////////
+
+TSharedRef<SDockTab> FTraceInsightsModule::SpawnStartPageTab(const FSpawnTabArgs& Args)
+{
+	const TSharedRef<SDockTab> DockTab = SNew(SDockTab)
+		.TabRole(ETabRole::NomadTab);
+
+	// Create the Start Page widget.
+	TSharedRef<SStartPageWindow> Window = SNew(SStartPageWindow);
+	DockTab->SetContent(Window);
+
+	return DockTab;
+}
+
+////////////////////////////////////////////////////////////////////////////////////////////////////
+
+TSharedRef<SDockTab> FTraceInsightsModule::SpawnTimingProfilerTab(const FSpawnTabArgs& Args)
+{
+	const TSharedRef<SDockTab> DockTab = SNew(SDockTab)
+		.TabRole(ETabRole::NomadTab);
+
+	// Register OnTabClosed to handle Timing profiler manager shutdown.
+	DockTab->SetOnTabClosed(SDockTab::FOnTabClosedCallback::CreateRaw(this, &FTraceInsightsModule::OnTimingProfilerTabBeingClosed));
+
+	// Create the STimingProfilerWindow widget.
+	TSharedRef<STimingProfilerWindow> Window = SNew(STimingProfilerWindow, DockTab, Args.GetOwnerWindow());
+	FTimingProfilerManager::Get()->AssignProfilerWindow(Window);
+	DockTab->SetContent(Window);
+
+	return DockTab;
+}
+
+////////////////////////////////////////////////////////////////////////////////////////////////////
+
+void FTraceInsightsModule::OnTimingProfilerTabBeingClosed(TSharedRef<SDockTab> TabBeingClosed)
+{
+	// Disable TabClosed delegate.
+	TabBeingClosed->SetOnTabClosed(SDockTab::FOnTabClosedCallback());
+}
+
+////////////////////////////////////////////////////////////////////////////////////////////////////
+
+TSharedRef<SDockTab> FTraceInsightsModule::SpawnIoProfilerTab(const FSpawnTabArgs& Args)
+{
+	const TSharedRef<SDockTab> DockTab = SNew(SDockTab)
+		.TabRole(ETabRole::NomadTab);
+
+	// Register OnTabClosed to handle I/O profiler manager shutdown.
+	DockTab->SetOnTabClosed(SDockTab::FOnTabClosedCallback::CreateRaw(this, &FTraceInsightsModule::OnIoProfilerTabBeingClosed));
+
+	// Create the SIoProfilerWindow widget.
+	TSharedRef<SIoProfilerWindow> Window = SNew(SIoProfilerWindow);
+	FIoProfilerManager::Get()->AssignProfilerWindow(Window);
+	DockTab->SetContent(Window);
+
+	return DockTab;
+}
+
+////////////////////////////////////////////////////////////////////////////////////////////////////
+
+void FTraceInsightsModule::OnIoProfilerTabBeingClosed(TSharedRef<SDockTab> TabBeingClosed)
+{
+	// Disable TabClosed delegate.
+	TabBeingClosed->SetOnTabClosed(SDockTab::FOnTabClosedCallback());
+}
+
+////////////////////////////////////////////////////////////////////////////////////////////////////
diff --git a/Engine/Source/Developer/TraceInsights/Private/Insights/ViewModels/FrameTrackHelper.cpp b/Engine/Source/Developer/TraceInsights/Private/Insights/ViewModels/FrameTrackHelper.cpp
new file mode 100644
--- /dev/null
+++ b/Engine/Source/Developer/TraceInsights/Private/Insights/ViewModels/FrameTrackHelper.cpp
@@ -0,0 +1,322 @@
+// Copyright 1998-2019 Epic Games, Inc. All Rights Reserved.
+
+#include "FrameTrackHelper.h"
+
+#include "Brushes/SlateBorderBrush.h"
+#include "Brushes/SlateColorBrush.h"
+#include "EditorStyleSet.h"
+#include "Fonts/FontMeasure.h"
+#include "Fonts/SlateFontInfo.h"
+#include "Framework/Application/SlateApplication.h"
+#include "Rendering/DrawElements.h"
+#include "TraceServices/AnalysisService.h"
+
+// Insights
+#include "Insights/Common/PaintUtils.h"
+#include "Insights/ViewModels/FrameTrackViewport.h"
+
+#include <limits>
+
+////////////////////////////////////////////////////////////////////////////////////////////////////
+// FFrameTrackTimelineBuilder
+////////////////////////////////////////////////////////////////////////////////////////////////////
+
+FFrameTrackTimelineBuilder::FFrameTrackTimelineBuilder(FFrameTrackTimeline& InTimeline, const FFrameTrackViewport& InViewport)
+	: Timeline(InTimeline)
+	, Viewport(InViewport)
+	, NumAddedFrames(0)
+{
+	SampleW = Viewport.GetSampleWidth();
+	FramesPerSample = Viewport.GetNumFramesPerSample();
+	NumSamples = FMath::Max(0, FMath::CeilToInt(Viewport.Width / SampleW));
+	FirstFrameIndex = Viewport.GetFirstFrameIndex();
+
+	Timeline.NumAggregatedFrames = 0;
+	Timeline.Samples.Reset();
+	Timeline.Samples.AddDefaulted(NumSamples);
+}
+
+////////////////////////////////////////////////////////////////////////////////////////////////////
+
+void FFrameTrackTimelineBuilder::AddFrame(const Trace::FFrame& Frame)
+{
+	NumAddedFrames++;
+
+	const int32 FrameIndex = Frame.Index;
+
+	int32 SampleIndex = (FrameIndex - FirstFrameIndex) / FramesPerSample;
+	if (SampleIndex >= 0 && SampleIndex < NumSamples)
+	{
+		FFrameTrackSample& Sample = Timeline.Samples[SampleIndex];
+		Sample.NumFrames++;
+
+		double Duration = Frame.EndTime - Frame.StartTime;
+		Sample.TotalDuration += Duration;
+		if (Frame.StartTime < Sample.StartTime)
+		{
+			Sample.StartTime = Frame.StartTime;
+		}
+		if (Frame.EndTime > Sample.EndTime)
+		{
+			Sample.EndTime = Frame.EndTime;
+		}
+		if (Duration > Sample.LargestFrameDuration)
+		{
+			Sample.LargestFrameIndex = FrameIndex;
+			Sample.LargestFrameStartTime = Frame.StartTime;
+			Sample.LargestFrameDuration = Duration;
+		}
+
+		Timeline.NumAggregatedFrames++;
+	}
+}
+
+////////////////////////////////////////////////////////////////////////////////////////////////////
+// FFrameTrackDrawHelper
+////////////////////////////////////////////////////////////////////////////////////////////////////
+
+FFrameTrackDrawHelper::FFrameTrackDrawHelper(const FDrawContext& InDrawContext, const FFrameTrackViewport& InViewport)
+	: DrawContext(InDrawContext)
+	, Viewport(InViewport)
+	, WhiteBrush(FCoreStyle::Get().GetBrush("WhiteBrush"))
+	, BorderBrush(FEditorStyle::GetBrush("PlainBorder"))
+	, NumFrames(0)
+	, NumDrawSamples(0)
+{
+}
+
+////////////////////////////////////////////////////////////////////////////////////////////////////
+
+void FFrameTrackDrawHelper::DrawBackground() const
+{
+	const FSlateBrush* AreaBrush = WhiteBrush;
+	const FLinearColor ValidAreaColor(0.07f, 0.07f, 0.07f, 1.0f);
+	const FLinearColor InvalidAreaColor(0.1f, 0.07f, 0.07f, 1.0f);
+
+	float X0 = Viewport.MinX - Viewport.PosX;
+	float X1 = Viewport.MaxX - Viewport.PosX;
+	const float W = FMath::CeilToFloat(Viewport.Width);
+	const float H = FMath::CeilToFloat(Viewport.Height);
+
+	if (X0 >= W || X1 <= 0.0f)
+	{
+		// Draw invalid area (entire view).
+		DrawContext.DrawBox(0.0f, 0.0f, W, H, AreaBrush, InvalidAreaColor);
+	}
+	else // X0 < W && X1 > 0
+	{
+		if (X0 > 0.0f)
+		{
+			// Draw invalid area (left).
+			DrawContext.DrawBox(0.0f, 0.0f, X0, H, AreaBrush, InvalidAreaColor);
+		}
+
+		if (X1 < W)
+		{
+			// Draw invalid area (right).
+			DrawContext.DrawBox(X1, 0.0f, W - X1, H, AreaBrush, InvalidAreaColor);
+		}
+
+		X0 = FMath::Max(X0, 0.0f);
+		X1 = FMath::Min(X1, W);
+
+		if (X1 > X0)
+		{
+			// Draw valid area.
+			DrawContext.DrawBox(X0, 0.0f, X1 - X0, H, AreaBrush, ValidAreaColor);
+		}
+	}
+
+	DrawContext.LayerId++;
+}
+
+////////////////////////////////////////////////////////////////////////////////////////////////////
+
+FLinearColor FFrameTrackDrawHelper::GetColorById(int32 Id) const
+{
+	const float Alpha = 0.9;
+	switch (Id)
+	{
+	case 0: // Game Frames
+		return FLinearColor(0.75, 1.0, 1.0, Alpha);
+
+	case 1: // Rendering Frames
+		return FLinearColor(1.0, 0.75, 0.75, Alpha);
+
+	case 2:
+		return FLinearColor(0.75, 0.75, 1.0, Alpha);
+
+	default:
+		return FLinearColor(1.0, 1.0, 1.0, Alpha);
+	}
+}
+
+////////////////////////////////////////////////////////////////////////////////////////////////////
+
+void FFrameTrackDrawHelper::DrawCached(const FFrameTrackTimeline& Timeline) const
+{
+	if (Timeline.NumAggregatedFrames == 0)
+	{
+		return;
+	}
+
+	NumFrames += Timeline.NumAggregatedFrames;
+
+	FLinearColor TimelineColor = GetColorById(static_cast<int32>(Timeline.Id));
+
+	const float SampleW = Viewport.GetSampleWidth();
+	const int32 NumSamples = Timeline.Samples.Num();
+
+	const float Y1 = FMath::RoundToFloat(Viewport.GetViewportYForValue(0.0));
+
+	for (int32 SampleIndex = 0; SampleIndex < NumSamples; SampleIndex++)
+	{
+		const FFrameTrackSample& Sample = Timeline.Samples[SampleIndex];
+		if (Sample.NumFrames == 0)
+		{
+			continue;
+		}
+
+		NumDrawSamples++;
+
+		const float X = SampleIndex * SampleW;
+		float Y2;
+
+		FLinearColor ColorFill = TimelineColor;
+
+		if (Sample.LargestFrameDuration == std::numeric_limits<double>::infinity())
+		{
+			Y2 = Viewport.Height;
+			ColorFill.R = 0.0f;
+			ColorFill.G = 0.0f;
+			ColorFill.B = 0.0f;
+		}
+		else
+		{
+			Y2 = FMath::RoundToFloat(Viewport.GetViewportYForValue(Sample.LargestFrameDuration));
+			if (Sample.LargestFrameDuration > 1.0 / 30.0)
+			{
+				ColorFill.G *= 0.5f;
+				ColorFill.B *= 0.5f;
+			}
+			else if (Sample.LargestFrameDuration > 1.0 / 60.0)
+			{
+				ColorFill.B *= 0.5f;
+			}
+		}
+
+		const float Y = FMath::RoundToFloat(Viewport.Height) - Y2;
+		const float H = Y2 - Y1;
+
+		const FLinearColor ColorBorder(ColorFill.R * 0.75f, ColorFill.G * 0.75f, ColorFill.B * 0.75f, 1.0);
+
+		if (SampleW > 2.0f)
+		{
+			DrawContext.DrawBox(X + 1.0f, Y + 1.0f, SampleW - 2.0f, H - 2.0f, WhiteBrush, ColorFill);
+
+			// Draw border.
+			DrawContext.DrawBox(X, Y, 1.0, H, WhiteBrush, ColorBorder);
+			DrawContext.DrawBox(X + SampleW - 1.0f, Y, 1.0, H, WhiteBrush, ColorBorder);
+			DrawContext.DrawBox(X + 1.0f, Y, SampleW - 2.0f, 1.0f, WhiteBrush, ColorBorder);
+			DrawContext.DrawBox(X + 1.0f, Y + H - 1.0f, SampleW - 2.0f, 1.0f, WhiteBrush, ColorBorder);
+		}
+		else
+		{
+			DrawContext.DrawBox(X, Y, SampleW, H, WhiteBrush, ColorBorder);
+		}
+	}
+
+	DrawContext.LayerId++;
+}
+
+////////////////////////////////////////////////////////////////////////////////////////////////////
+
+void FFrameTrackDrawHelper::DrawHoveredSample(const FFrameTrackSample& Sample) const
+{
+	const float SampleW = Viewport.GetSampleWidth();
+	const int32 FramesPerSample = Viewport.GetNumFramesPerSample();
+	const int32 FirstFrameIndex = Viewport.GetFirstFrameIndex();
+	const int32 SampleIndex = (Sample.LargestFrameIndex - FirstFrameIndex) / FramesPerSample;
+	const float X = SampleIndex * SampleW;
+
+	const float Y1 = FMath::RoundToFloat(Viewport.GetViewportYForValue(0.0));
+	float Y2;
+	if (Sample.LargestFrameDuration == std::numeric_limits<double>::infinity())
+	{
+		Y2 = Viewport.Height;
+	}
+	else
+	{
+		Y2 = FMath::RoundToFloat(Viewport.GetViewportYForValue(Sample.LargestFrameDuration));
+	}
+	const float Y = FMath::RoundToFloat(Viewport.Height) - Y2;
+	const float H = Y2 - Y1;
+
+	const FLinearColor ColorBorder(1.0f, 1.0f, 0.0f, 1.0);
+	DrawContext.DrawBox(X - 1.0f, Y - 1.0f, SampleW + 2.0f, H + 2.0f, BorderBrush, ColorBorder);
+	DrawContext.LayerId++;
+}
+
+////////////////////////////////////////////////////////////////////////////////////////////////////
+
+void FFrameTrackDrawHelper::DrawHighlightedInterval(const FFrameTrackTimeline& Timeline, const double StartTime, const double EndTime) const
+{
+	const int32 NumSamples = Timeline.Samples.Num();
+
+	//TODO: binary search
+	int32 Index1 = 0;
+	int32 Index2 = NumSamples - 1;
+	while (Index1 < NumSamples && Timeline.Samples[Index1].EndTime < StartTime)
+	{
+		Index1++;
+	}
+	while (Index2 >= Index1 && Timeline.Samples[Index2].StartTime > EndTime)
+	{
+		Index2--;
+	}
+
+	if (Index1 <= Index2)
+	{
+		const float SampleW = Viewport.GetSampleWidth();
+		float X1 = Index1 * SampleW;
+		float X2 = Index2 * SampleW;
+
+		constexpr float Y1 = 12.0f; // allows 12px for the horizontal scrollbar (one displayed on top of the track)
+		const float Y2 = Viewport.Height;
+		constexpr float D = 2.0f; // line thickness (for both horizontal and vertical lines)
+		constexpr float H = 10.0f; // height of corner lines
+
+		const FLinearColor Color(1.0f, 1.0f, 1.0f, 1.0f);
+
+		if (X1 >= 0.0f && X1 < Viewport.Width - 2.0f)
+		{
+			// Draw left side vertical lines.
+			DrawContext.DrawBox(X1 - D, Y1, D, H, WhiteBrush, Color);
+			DrawContext.DrawBox(X1 - D, Y2 - H, D, H, WhiteBrush, Color);
+		}
+
+		if (X2 >= -2.0f && X2 < Viewport.Width)
+		{
+			// Draw right side vertical lines.
+			DrawContext.DrawBox(X2, Y1, D, H, WhiteBrush, Color);
+			DrawContext.DrawBox(X2, Y2 - H, D, H, WhiteBrush, Color);
+		}
+
+		if (X1 < 0)
+		{
+			X1 = 0.0f;
+		}
+		if (X2 > Viewport.Width)
+		{
+			X2 = Viewport.Width;
+		}
+		if (X1 < X2)
+		{
+			// Draw horizontal lines.
+			DrawContext.DrawBox(X1, Y1, X2 - X1, D, WhiteBrush, Color);
+			DrawContext.DrawBox(X1, Y2 - D, X2 - X1, D, WhiteBrush, Color);
+		}
+	}
+}
+
+////////////////////////////////////////////////////////////////////////////////////////////////////
diff --git a/Engine/Source/Developer/TraceInsights/Private/Insights/ViewModels/FrameTrackHelper.h b/Engine/Source/Developer/TraceInsights/Private/Insights/ViewModels/FrameTrackHelper.h
new file mode 100644
--- /dev/null
+++ b/Engine/Source/Developer/TraceInsights/Private/Insights/ViewModels/FrameTrackHelper.h
@@ -0,0 +1,137 @@
+// Copyright 1998-2019 Epic Games, Inc. All Rights Reserved.
+
+#pragma once
+
+#include "CoreMinimal.h"
+#include "Styling/WidgetStyle.h"
+
+enum class ESlateDrawEffect : uint8;
+
+struct FDrawContext;
+struct FGeometry;
+struct FSlateBrush;
+
+class FFrameTrackViewport;
+class FSlateWindowElementList;
+
+namespace Trace
+{
+	struct FFrame;
+}
+
+////////////////////////////////////////////////////////////////////////////////////////////////////
+
+enum class EFrameType
+{
+	Unknown = -1,
+
+	Game,
+	Render
+};
+
+////////////////////////////////////////////////////////////////////////////////////////////////////
+
+struct FFrameTrackSample
+{
+	int32 NumFrames;
+	double TotalDuration; // sum of durations of all frames in this sample
+	double StartTime; // min start time of all frames in this sample
+	double EndTime; // max end time of all frames in this sample
+	double LargestFrameIndex; // index of the largest frame
+	double LargestFrameStartTime; // start time of the largest frame
+	double LargestFrameDuration; // duration of the largest frame
+
+	FFrameTrackSample()
+		: NumFrames(0)
+		, TotalDuration(0.0)
+		, StartTime(DBL_MAX)
+		, EndTime(-DBL_MAX)
+		, LargestFrameIndex(0)
+		, LargestFrameStartTime(0.0)
+		, LargestFrameDuration(0.0)
+	{}
+};
+
+////////////////////////////////////////////////////////////////////////////////////////////////////
+
+struct FFrameTrackTimeline
+{
+	uint64 Id;
+	//TODO: EFrameType Type;
+
+	//TODO: double MinValue; // min value
+	//TODO: double MaxValue; // max value
+
+	//TODO: float SampleW; // width of a sample, in Slate units
+	//TODO: int32 FramesPerSample; // number of frames in a sample
+	//TODO: int32 FirstFrameIndex; // index of first frame in first sample; can be negative
+	//TODO: int32 NumSamples; // total number of samples
+	int32 NumAggregatedFrames; // total number of frames aggregated in samples; ie. sum of all Sample.NumFrames
+
+	TArray<FFrameTrackSample> Samples;
+};
+
+////////////////////////////////////////////////////////////////////////////////////////////////////
+
+class FFrameTrackTimelineBuilder
+{
+public:
+	FFrameTrackTimelineBuilder(FFrameTrackTimeline& InTimeline, const FFrameTrackViewport& InViewport);
+
+	/**
+	 * Non-copyable
+	 */
+	FFrameTrackTimelineBuilder(const FFrameTrackTimelineBuilder&) = delete;
+	FFrameTrackTimelineBuilder& operator=(const FFrameTrackTimelineBuilder&) = delete;
+
+	void AddFrame(const Trace::FFrame& Frame);
+
+	int32 GetNumAddedFrames() const { return NumAddedFrames; }
+
+private:
+	FFrameTrackTimeline& Timeline;
+	const FFrameTrackViewport& Viewport;
+
+	float SampleW; // width of a sample, in Slate units
+	int32 FramesPerSample; // number of frames in a sample
+	int32 FirstFrameIndex; // index of first frame in first sample; can be negative
+	int32 NumSamples; // total number of samples
+
+	// Debug stats.
+	int32 NumAddedFrames; // counts total number of added frame events
+};
+
+////////////////////////////////////////////////////////////////////////////////////////////////////
+
+class FFrameTrackDrawHelper
+{
+public:
+	FFrameTrackDrawHelper(const FDrawContext& InDrawContext, const FFrameTrackViewport& InViewport);
+
+	/**
+	 * Non-copyable
+	 */
+	FFrameTrackDrawHelper(const FFrameTrackDrawHelper&) = delete;
+	FFrameTrackDrawHelper& operator=(const FFrameTrackDrawHelper&) = delete;
+
+	void DrawBackground() const;
+	void DrawCached(const FFrameTrackTimeline& Timeline) const;
+	void DrawHoveredSample(const FFrameTrackSample& Sample) const;
+	void DrawHighlightedInterval(const FFrameTrackTimeline& Timeline, const double StartTime, const double EndTime) const;
+
+	FLinearColor GetColorById(int32 Id) const;
+
+	int32 GetNumFrames() const { return NumFrames; }
+	int32 GetNumDrawSamples() const { return NumDrawSamples; }
+
+private:
+	const FDrawContext& DrawContext;
+	const FFrameTrackViewport& Viewport;
+
+	const FSlateBrush* WhiteBrush;
+	const FSlateBrush* BorderBrush;
+
+	// Debug stats.
+	mutable int32 NumFrames;
+	mutable int32 NumDrawSamples;
+};
diff --git a/Engine/Source/Developer/TraceInsights/Private/Insights/ViewModels/FrameTrackViewport.cpp b/Engine/Source/Developer/TraceInsights/Private/Insights/ViewModels/FrameTrackViewport.cpp
new file mode 100644
--- /dev/null
+++ b/Engine/Source/Developer/TraceInsights/Private/Insights/ViewModels/FrameTrackViewport.cpp
@@ -0,0 +1,3 @@
+// Copyright 1998-2019 Epic Games, Inc. All Rights Reserved.
+
+#include "FrameTrackViewport.h"
diff --git a/Engine/Source/Developer/TraceInsights/Private/Insights/ViewModels/FrameTrackViewport.h b/Engine/Source/Developer/TraceInsights/Private/Insights/ViewModels/FrameTrackViewport.h
new file mode 100644
--- /dev/null
+++ b/Engine/Source/Developer/TraceInsights/Private/Insights/ViewModels/FrameTrackViewport.h
@@ -0,0 +1,228 @@
+// Copyright 1998-2019 Epic Games, Inc. All Rights Reserved.
+
+#pragma once
+
+#include "CoreMinimal.h"
+
+class FFrameTrackViewport
+{
+private:
+	const float SLATE_UNITS_TOLERANCE = 0.1f;
+
+public:
+	FFrameTrackViewport()
+	{
+		Reset();
+	}
+
+	void Reset()
+	{
+		Width = 0.0f;
+		Height = 0.0f;
+
+		MinIndex = 0;
+		MaxIndex = 0;
+
+		MinX = 0.0f;
+		MaxX = 0.0f;
+		PosX = 0.0f;
+		MinScaleX = 0.0001f;
+		MaxScaleX = 16.0f;
+		ScaleX = MaxScaleX;
+
+		MinValue = 0.0;
+		MaxValue = 1.0f;
+
+		MinY = 0.0f;
+		MaxY = 1.0f;
+		PosY = 0.0f;
+		MinScaleY = 0.001f;
+		MaxScaleY = 1000000.0f;
+		ScaleY = 1000.0f;
+
+		LastOnPaintMaxValue = 0.0;
+	}
+
+	void SetSize(const float InWidth, const float InHeight)
+	{
+		if (!FMath::IsNearlyEqual(Width, InWidth, SLATE_UNITS_TOLERANCE) ||
+			!FMath::IsNearlyEqual(Height, InHeight, SLATE_UNITS_TOLERANCE))
+		{
+			Width = InWidth;
+			Height = InHeight;
+			OnSizeChanged();
+		}
+	}
+
+	void SetMinMaxIndexInterval(const int32 InMinIndex, const int32 InMaxIndex)
+	{
+		MinIndex = InMinIndex;
+		MaxIndex = InMaxIndex;
+		UpdateMinMaxX();
+	}
+
+	int32 GetIndexAtPosX(const float X) const
+	{
+		return FMath::RoundToInt(X / ScaleX);
+	}
+
+	int32 GetIndexAtViewportX(const float VX) const
+	{
+		return FMath::RoundToInt((PosX + VX) / ScaleX);
+	}
+
+	float GetPosXForIndex(const int32 Index) const
+	{
+		return static_cast<float>(Index) * ScaleX;
+	}
+
+	float GetViewportXForIndex(const int32 Index) const
+	{
+		return static_cast<float>(Index) * ScaleX - PosX;
+	}
+
+	float GetRoundedViewportXForIndex(const int32 Index) const
+	{
+		return FMath::RoundToFloat(static_cast<float>(Index) * ScaleX - PosX);
+	}
+
+	void ScrollAtPosX(const float InPosX)
+	{
+		PosX = InPosX;
+		OnPosXChanged();
+	}
+
+	void ScrollAtIndex(const int32 Index)
+	{
+		PosX = static_cast<float>(Index) * ScaleX;
+		OnPosXChanged();
+	}
+
+	void CenterOnIndex(const int32 Index)
+	{
+		float W = FMath::Max(1.0f, ScaleX);
+		if (W > Width)
+		{
+			ScrollAtIndex(Index);
+		}
+		else
+		{
+			float VX = (Width - W) / 2.0f;
+			PosX = static_cast<float>(Index) * ScaleX - VX;
+			OnPosXChanged();
+		}
+	}
+
+	void CenterOnIndexInterval(const int32 IntervalStartIndex, const int32 IntervalEndIndex)
+	{
+		float W = FMath::Max(1.0f, ScaleX * static_cast<float>(IntervalEndIndex - IntervalStartIndex));
+		if (W > Width)
+		{
+			ScrollAtIndex(IntervalStartIndex);
+		}
+		else
+		{
+			float VX = (Width - W) / 2.0f;
+			PosX = static_cast<float>(IntervalStartIndex) * ScaleX - VX;
+			OnPosXChanged();
+		}
+	}
+
+	void ZoomWithFixedViewportX(const float InNewScaleX, const float VX)
+	{
+		float LocalNewScaleX = FMath::Clamp(InNewScaleX, MinScaleX, MaxScaleX);
+		if (LocalNewScaleX != ScaleX)
+		{
+			// Index at viewport X should remain the same. So we resolve equation:
+			//   (PosX + VX) / ScaleX = (NewPosX + VX) / NewScaleX
+			//   ==> NewPosX = (PosX + VX) / ScaleX * NewScaleX - VX
+			PosX = (PosX + VX) * (LocalNewScaleX / ScaleX) - VX;
+			ScaleX = LocalNewScaleX;
+
+			OnPosXChanged();
+			OnScaleXChanged();
+		}
+	}
+
+	float GetSampleWidth() const { return FMath::Max(1.0f, FMath::RoundToFloat(ScaleX)); }
+	int32 GetNumFramesPerSample() const { return FMath::Max(1, FMath::RoundToInt(1.0f / ScaleX)); }
+
+	int32 GetFirstFrameIndex() const { return GetIndexAtViewportX(0.0f); }
+
+	float GetPosYForValue(const double Value) const
+	{
+		return static_cast<float>(Value) * ScaleY;
+	}
+
+	float GetViewportYForValue(const double Value) const
+	{
+		return static_cast<float>(Value) * ScaleY - PosY;
+	}
+
+	void SetScaleY(const float InNewScaleY)
+	{
+		float LocalNewScaleY = FMath::Clamp(InNewScaleY, MinScaleY, MaxScaleY);
+		if (LocalNewScaleY != ScaleY)
+		{
+			ScaleY = LocalNewScaleY;
+			OnScaleYChanged();
+		}
+	}
+
+private:
+	void OnSizeChanged()
+	{
+	}
+
+	void OnPosXChanged()
+	{
+	}
+
+	void OnScaleXChanged()
+	{
+		UpdateMinMaxX();
+	}
+
+	void OnScaleYChanged()
+	{
+	}
+
+	void UpdateMinMaxX()
+	{
+		MinX = static_cast<float>(MinIndex) * ScaleX;
+		MaxX = static_cast<float>(MaxIndex) * ScaleX;
+
+		//float NewPosX = FMath::Clamp(PosX, MinX, MaxX);
+		//if (FMath::IsNearlyEqual(NewPosX, PosX, SLATE_UNITS_TOLERANCE))
+		//{
+		//	PosX = NewPosX;
+		//	OnPosXChanged();
+		//}
+	}
+
+public:
+	float Width; // width of viewport, in Slate units
+	float Height; // height of viewport, in Slate units
+
+	int32 MinIndex; // minimum index
+	int32 MaxIndex; // maximum index (exclusive)
+
+	float MinX; // minimum horizontal position, in Slate units
+	float MaxX; // maximum horizontal position, in Slate units
+	float PosX;
+	float MinScaleX; // minimum horizontal scale factor
+	float MaxScaleX; // maximum horizontal scale factor
+	float ScaleX; // scale factor between frame index and Slate units: if > 1, it represents number of Slate units for one frame index, otherwise it represents number of frames in one Slate unit
+
+	double MinValue;
+	double MaxValue;
+
+	float MinY; // minimum vertical scroll position, in Slate units
+	float MaxY; // maximum vertical scroll position, in Slate units
+	float PosY; // current vertical scroll position, in Slate units; origin at bottom
+	float MinScaleY; // minimum vertical scale factor
+	float MaxScaleY; // maximum vertical scale factor
+	float ScaleY; // scale factor between Value units and Slate units
+
+	mutable double LastOnPaintMaxValue;
+};
diff --git a/Engine/Source/Developer/TraceInsights/Private/Insights/ViewModels/GraphTrack.cpp b/Engine/Source/Developer/TraceInsights/Private/Insights/ViewModels/GraphTrack.cpp
new file mode 100644
--- /dev/null
+++ b/Engine/Source/Developer/TraceInsights/Private/Insights/ViewModels/GraphTrack.cpp
@@ -0,0 +1,675 @@
+// Copyright 1998-2019 Epic Games, Inc. All Rights Reserved.
+
+#include "GraphTrack.h"
+
+#include "Brushes/SlateColorBrush.h"
+#include "Brushes/SlateBorderBrush.h"
+#include "EditorStyleSet.h"
+#include "Rendering/DrawElements.h"
+#include "Styling/CoreStyle.h"
+#include "TraceServices/AnalysisService.h"
+#include "TraceServices/SessionService.h"
+
+// Insights
+#include "Insights/Common/PaintUtils.h"
+#include "Insights/InsightsManager.h"
+#include "Insights/TimingProfilerManager.h"
+#include "Insights/ViewModels/TimingTrackViewport.h"
+
+#define LOCTEXT_NAMESPACE "GraphTrack"
+
+////////////////////////////////////////////////////////////////////////////////////////////////////
+// FGraphTrackSeries
+////////////////////////////////////////////////////////////////////////////////////////////////////
+
+FGraphTrackSeries::FGraphTrackSeries()
+	: Color(0.0f, 0.5f, 1.0f, 1.0f)
+	, BorderColor(0.3f, 0.8f, 1.0f, 1.0f)
+	//, Points()
+	//, LinePoints()
+	//, Boxes()
+{
+}
+
+////////////////////////////////////////////////////////////////////////////////////////////////////
+
+FGraphTrackSeries::~FGraphTrackSeries()
+{
+}
+
+////////////////////////////////////////////////////////////////////////////////////////////////////
+// FGraphTrack
+////////////////////////////////////////////////////////////////////////////////////////////////////
+
+FGraphTrack::FGraphTrack(uint64 InTrackId)
+	: FBaseTimingTrack(InTrackId)
+	//, AllSeries()
+	, WhiteBrush(FCoreStyle::Get().GetBrush("WhiteBrush"))
+	, PointBrush(FEditorStyle::GetBrush("Graph.ExecutionBubble"))
+	, BorderBrush(new FSlateBorderBrush(NAME_None, FMargin(1.0f)))
+	//, BorderBrush(FEditorStyle::GetBrush("PlainBorder"))
+	, Font(FCoreStyle::GetDefaultFontStyle("Regular", 8))
+	, bDrawPoints(true)
+	, bDrawPointsWithBorder(true)
+	, bDrawLines(true)
+	, bDrawLinesWithDuration(false)
+	, bDrawBoxes(true)
+	, BaselineY(0.0f)
+	, ScaleY(1.0)
+	, NumAddedEvents(0)
+	, NumDrawPoints(0)
+	, NumDrawLines(0)
+	, NumDrawBoxes(0)
+{
+}
+
+////////////////////////////////////////////////////////////////////////////////////////////////////
+
+FGraphTrack::~FGraphTrack()
+{
+	delete BorderBrush;
+}
+
+////////////////////////////////////////////////////////////////////////////////////////////////////
+
+void FGraphTrack::UpdateHoveredState(float MouseX, float MouseY, const FTimingTrackViewport& Viewport)
+{
+	constexpr float HeaderWidth = 100.0f;
+	constexpr float HeaderHeight = 14.0f;
+
+	if (MouseY >= GetPosY() && MouseY < GetPosY() + GetHeight())
+	{
+		SetHoveredState(true);
+		SetHeaderHoveredState(MouseX < HeaderWidth && MouseY < GetPosY() + HeaderHeight);
+	}
+	else
+	{
+		SetHoveredState(false);
+	}
+}
+
+////////////////////////////////////////////////////////////////////////////////////////////////////
+
+void FGraphTrack::UpdateStats()
+{
+	NumDrawPoints = 0;
+	NumDrawLines = 0;
+	NumDrawBoxes = 0;
+
+	for (const FGraphTrackSeries& Series : AllSeries)
+	{
+		NumDrawPoints += Series.Points.Num();
+		NumDrawLines += Series.LinePoints.Num() / 2;
+		NumDrawBoxes += Series.Boxes.Num();
+	}
+}
+
+////////////////////////////////////////////////////////////////////////////////////////////////////
+
+void FGraphTrack::Draw(FDrawContext& DrawContext, const FTimingTrackViewport& Viewport) const
+{
+	DrawContext.DrawBox(0.0f, GetPosY(), Viewport.Width, GetHeight(), WhiteBrush, FLinearColor(0.1f, 0.1f, 0.1f, 1.0f));
+	DrawContext.LayerId++;
+
+	//TODO: Set clipping to (0, GetPosY(), Viewport.Width, GetHeight())!
+
+	for (const FGraphTrackSeries& Series : AllSeries)
+	{
+		DrawSeries(Series, DrawContext, Viewport);
+	}
+
+	// Draw baseline (Value == 0).
+	DrawContext.DrawBox(0.0f, GetPosY() + BaselineY - 1.0f, Viewport.Width, 1.0f, WhiteBrush, FLinearColor(0.05f, 0.05f, 0.05f, 1.0f));
+	DrawContext.LayerId++;
+}
+
+////////////////////////////////////////////////////////////////////////////////////////////////////
+
+void FGraphTrack::DrawSeries(const FGraphTrackSeries& Series, FDrawContext& DrawContext, const FTimingTrackViewport& Viewport) const
+{
+	if (bDrawBoxes)
+	{
+		int32 NumBoxes = Series.Boxes.Num();
+		for (int32 Index = 0; Index < NumBoxes; ++Index)
+		{
+			const FGraphBox& Box = Series.Boxes[Index];
+			DrawContext.DrawBox(Box.X, GetPosY() + Box.Y, Box.W, BaselineY - Box.Y, WhiteBrush, Series.Color);
+		}
+		DrawContext.LayerId++;
+	}
+
+	if (bDrawLines)
+	{
+		FPaintGeometry Geo = DrawContext.Geometry.ToPaintGeometry();
+		Geo.AppendTransform(FSlateLayoutTransform(FVector2D(0, GetPosY())));
+		FSlateDrawElement::MakeLines(DrawContext.ElementList, DrawContext.LayerId, Geo, Series.LinePoints, DrawContext.DrawEffects, Series.Color, false, 1.0f);
+		DrawContext.LayerId++;
+	}
+
+	if (bDrawPoints)
+	{
+		int32 NumPoints = Series.Points.Num();
+
+#define INSIGHTS_GRAPH_TRACK_DRAW_POINTS_AS_RECTANGLES 0
+#if !INSIGHTS_GRAPH_TRACK_DRAW_POINTS_AS_RECTANGLES
+
+		if (bDrawPointsWithBorder)
+		{
+			// Draw points (border).
+			for (int32 Index = 0; Index < NumPoints; ++Index)
+			{
+				const FVector2D& Pt = Series.Points[Index];
+				const float PtX = Pt.X - PointVisualSize / 2.0f - 1.0f;
+				const float PtY = GetPosY() + Pt.Y - PointVisualSize / 2.0f - 1.0f;
+				DrawContext.DrawBox(PtX, PtY, PointVisualSize + 2.0f, PointVisualSize + 2.0f, PointBrush, Series.BorderColor);
+			}
+			DrawContext.LayerId++;
+		}
+
+		// Draw points (interior).
+		for (int32 Index = 0; Index < NumPoints; ++Index)
+		{
+			const FVector2D& Pt = Series.Points[Index];
+			const float PtX = Pt.X - PointVisualSize / 2.0f;
+			const float PtY = GetPosY() + Pt.Y - PointVisualSize / 2.0f;
+			DrawContext.DrawBox(PtX, PtY, PointVisualSize, PointVisualSize, PointBrush, Series.Color);
+		}
+		DrawContext.LayerId++;
+
+#else // Alternative way of drawing points; kept here for debugging purposes.
+
+		//const float Angle = FMath::DegreesToRadians(45.0f);
+
+		if (bDrawPointsWithBorder)
+		{
+			// Draw borders.
+			const float BorderPtSize = PointVisualSize;
+			//FVector2D BorderRotationPoint(BorderPtSize / 2.0f, BorderPtSize / 2.0f);
+			for (int32 Index = 0; Index < NumPoints; ++Index)
+			{
+				const FVector2D& Pt = Series.Points[Index];
+				const float PtX = Pt.X - BorderPtSize / 2.0f;
+				const float PtY = GetPosY() + Pt.Y - BorderPtSize / 2.0f;
+				DrawContext.DrawBox(PtX, PtY, BorderPtSize, BorderPtSize, BorderBrush, Series.BorderColor);
+				//DrawContext.DrawRotatedBox(PtX, PtY, BorderPtSize, BorderPtSize, BorderBrush, Series.BorderColor, Angle, BorderRotationPoint);
+			}
+			DrawContext.LayerId++;
+		}
+
+		// Draw points as rectangles.
+		const float PtSize = PointVisualSize - 2.0f;
+		//FVector2D RotationPoint(PtSize / 2.0f, PtSize / 2.0f);
+		for (int32 Index = 0; Index < NumPoints; ++Index)
+		{
+			const FVector2D& Pt = Series.Points[Index];
+			const float PtX = Pt.X - PtSize / 2.0f;
+			const float PtY = GetPosY() + Pt.Y - PtSize / 2.0f;
+			DrawContext.DrawBox(PtX, PtY, PtSize, PtSize, PointBrush, Series.Color);
+			//DrawContext.DrawRotatedBox(PtX, PtY, PtSize, PtSize, WhiteBrush, Series.Color, Angle, RotationPoint);
+		}
+		DrawContext.LayerId++;
+
+#endif
+	}
+}
+
+////////////////////////////////////////////////////////////////////////////////////////////////////
+// FRandomGraphTrack
+////////////////////////////////////////////////////////////////////////////////////////////////////
+
+FRandomGraphTrack::FRandomGraphTrack(uint64 InTrackId)
+	: FGraphTrack(InTrackId)
+{
+	bDrawPoints = true;
+	bDrawPointsWithBorder = true;
+	bDrawLines = true;
+	bDrawLinesWithDuration = false;
+	bDrawBoxes = false;
+
+	AllSeries.AddDefaulted();
+	FGraphTrackSeries& Series = AllSeries[0];
+	Series.SetColor(FLinearColor(0.0f, 0.5f, 1.0f, 1.0f), FLinearColor(0.3f, 0.8f, 1.0f, 1.0f));
+}
+
+////////////////////////////////////////////////////////////////////////////////////////////////////
+
+FRandomGraphTrack::~FRandomGraphTrack()
+{
+}
+
+////////////////////////////////////////////////////////////////////////////////////////////////////
+
+void FRandomGraphTrack::Update(const FTimingTrackViewport& Viewport)
+{
+	NumAddedEvents = 0;
+
+	// TODO: vertical panning and zooming
+	BaselineY = GetHeight();
+	ScaleY = 1.0;
+
+	ensure(AllSeries.Num() == 1);
+	GenerateSeries(AllSeries[0], Viewport, 1000000);
+
+	UpdateStats();
+}
+
+////////////////////////////////////////////////////////////////////////////////////////////////////
+
+void FRandomGraphTrack::GenerateSeries(FGraphTrackSeries& Series, const FTimingTrackViewport& Viewport, const int32 EventCount)
+{
+	//////////////////////////////////////////////////
+	// Generate random events.
+
+	constexpr double MinDeltaTime = 0.0000001; // 100ns
+	constexpr double MaxDeltaTime = 0.01; // 100ms
+	const float MinValue = 0;
+	const float MaxValue = GetHeight();
+
+	struct FGraphEvent
+	{
+		double Time;
+		double Duration;
+		double Value;
+	};
+
+	TArray<FGraphEvent> Events;
+	Events.Reserve(EventCount);
+
+	FRandomStream RandomStream(0);
+	double NextT = 0.0;
+	for (int32 Index = 0; Index < EventCount; ++Index)
+	{
+		FGraphEvent Ev;
+		Ev.Time = NextT;
+		const double TimeAdvance = RandomStream.GetFraction() * (MaxDeltaTime - MinDeltaTime);
+		NextT += MinDeltaTime + TimeAdvance;
+		Ev.Duration = MinDeltaTime + RandomStream.GetFraction() * TimeAdvance;
+		Ev.Value = MinValue + RandomStream.GetFraction() * (MaxValue - MinValue);
+		Events.Add(Ev);
+	}
+
+	//////////////////////////////////////////////////
+	// Optimize and build draw lists.
+	{
+		FGraphTrackBuilder Builder(*this, Series, Viewport);
+
+		int32 Index = 0;
+		while (Index < EventCount && Events[Index].Time < Viewport.StartTime)
+		{
+			++Index;
+		}
+		if (Index > 0)
+		{
+			Index--; // one point outside screen (left side)
+		}
+		while (Index < EventCount)
+		{
+			const FGraphEvent& Ev = Events[Index];
+			Builder.AddEvent(Ev.Time, Ev.Duration, Ev.Value);
+
+			if (Ev.Time > Viewport.EndTime)
+			{
+				// one point outside screen (right side)
+				break;
+			}
+
+			++Index;
+		}
+	}
+}
+
+////////////////////////////////////////////////////////////////////////////////////////////////////
+// FFramesGraphTrack
+////////////////////////////////////////////////////////////////////////////////////////////////////
+
+FFramesGraphTrack::FFramesGraphTrack(uint64 InTrackId)
+	: FGraphTrack(InTrackId)
+{
+	bDrawPoints = true;
+	bDrawPointsWithBorder = true;
+	bDrawLines = true;
+	bDrawLinesWithDuration = true;
+	bDrawBoxes = false;
+
+	AllSeries.AddDefaulted(2);
+
+	FGraphTrackSeries& GameFramesSeries = AllSeries[0];
+	GameFramesSeries.SetColor(FLinearColor(0.3f, 0.3f, 1.0f, 1.0f), FLinearColor(1.0f, 1.0f, 1.0f, 1.0f));
+
+	FGraphTrackSeries& RenderFramesSeries = AllSeries[1];
+	RenderFramesSeries.SetColor(FLinearColor(1.0f, 0.3f, 0.3f, 1.0f), FLinearColor(1.0f, 1.0f, 1.0f, 1.0f));
+}
+
+////////////////////////////////////////////////////////////////////////////////////////////////////
+
+FFramesGraphTrack::~FFramesGraphTrack()
+{
+}
+
+////////////////////////////////////////////////////////////////////////////////////////////////////
+
+void FFramesGraphTrack::Update(const FTimingTrackViewport& Viewport)
+{
+	NumAddedEvents = 0;
+
+	// TODO: Vertical panning and zooming needs to be moved out in a Viewport like controller.
+	BaselineY = GetHeight();
+	ScaleY = 200.0 / 0.1; // 200px = 100ms
+
+	ensure(AllSeries.Num() == 2);
+	UpdateSeries(AllSeries[0], Viewport, TraceFrameType_Game);
+	UpdateSeries(AllSeries[1], Viewport, TraceFrameType_Rendering);
+
+	UpdateStats();
+}
+
+////////////////////////////////////////////////////////////////////////////////////////////////////
+
+void FFramesGraphTrack::UpdateSeries(FGraphTrackSeries& Series, const FTimingTrackViewport& Viewport, ETraceFrameType FrameType)
+{
+	FGraphTrackBuilder Builder(*this, Series, Viewport);
+
+	TSharedPtr<const Trace::IAnalysisSession> Session = FInsightsManager::Get()->GetSession();
+	if (Session.IsValid())
+	{
+		Trace::FAnalysisSessionReadScope SessionReadScope(*Session.Get());
+		const Trace::IFrameProvider& FramesProvider = ReadFrameProvider(*Session.Get());
+		uint64 FrameCount = FramesProvider.GetFrameCount(FrameType);
+		FramesProvider.EnumerateFrames(FrameType, 0, FrameCount - 1, [&Builder](const Trace::FFrame& Frame)
+		{
+			const double Duration = Frame.EndTime - Frame.StartTime;
+			Builder.AddEvent(Frame.StartTime, Duration, Duration);
+		});
+	}
+}
+
+////////////////////////////////////////////////////////////////////////////////////////////////////
+// FGraphTrackBuilder
+////////////////////////////////////////////////////////////////////////////////////////////////////
+
+FGraphTrackBuilder::FGraphTrackBuilder(FGraphTrack& InTrack, FGraphTrackSeries& InSeries, const FTimingTrackViewport& InViewport)
+	: Track(InTrack)
+	, Series(InSeries)
+	, Viewport(InViewport)
+{
+	BeginPoints();
+	BeginConnectedLines();
+	BeginBoxes();
+}
+
+////////////////////////////////////////////////////////////////////////////////////////////////////
+
+FGraphTrackBuilder::~FGraphTrackBuilder()
+{
+	EndPoints();
+	EndConnectedLines();
+	EndBoxes();
+}
+
+////////////////////////////////////////////////////////////////////////////////////////////////////
+
+void FGraphTrackBuilder::AddEvent(double Time, double Duration, double Value)
+{
+	Track.NumAddedEvents++;
+
+	if (Track.bDrawPoints)
+	{
+		AddPoint(Time, Value);
+	}
+
+	if (Track.bDrawLines)
+	{
+		AddConnectedLine(Time, Value);
+
+		if (Track.bDrawLinesWithDuration)
+		{
+			AddConnectedLine(Time + Duration, Value);
+		}
+	}
+
+	if (Track.bDrawBoxes)
+	{
+		AddBox(Time, Duration, Value);
+	}
+}
+
+////////////////////////////////////////////////////////////////////////////////////////////////////
+// FGraphTrackBuilder - Points
+////////////////////////////////////////////////////////////////////////////////////////////////////
+
+void FGraphTrackBuilder::BeginPoints()
+{
+	Series.Points.Reset();
+
+	PointsCurrentX = -DBL_MAX;
+
+	PointsAtCurrentX.Reset();
+	int32 MaxPointsPerLineScan = FMath::CeilToInt(Track.GetHeight() / FGraphTrack::PointSizeY);
+	if (MaxPointsPerLineScan > 0)
+	{
+		PointsAtCurrentX.AddDefaulted(MaxPointsPerLineScan);
+	}
+}
+
+////////////////////////////////////////////////////////////////////////////////////////////////////
+
+void FGraphTrackBuilder::AddPoint(double Time, double Value)
+{
+	const float X = Viewport.TimeToSlateUnitsRounded(Time);
+	if (X < -FGraphTrack::PointVisualSize / 2.0f || X >= Viewport.Width + FGraphTrack::PointVisualSize / 2.0f)
+	{
+		return;
+	}
+
+	// Align the X with a grid of GraphTrackPointDX pixels in size, in the global space (i.e. scroll independent).
+	const double AlignedX = FMath::RoundToDouble(Time * Viewport.ScaleX / FGraphTrack::PointSizeX) * FGraphTrack::PointSizeX;
+
+	if (AlignedX > PointsCurrentX + FGraphTrack::PointSizeX - 0.5)
+	{
+		FlushPoints();
+		PointsCurrentX = AlignedX;
+	}
+
+	const float Y = Track.GetYForValue(Value);
+
+	int32 Index = FMath::RoundToInt(Y / FGraphTrack::PointSizeY);
+	if (Index >= 0 && Index < PointsAtCurrentX.Num())
+	{
+		FPointInfo& Pt = PointsAtCurrentX[Index];
+		Pt.bValid = true;
+		Pt.X = X;
+		Pt.Y = Y;
+	}
+}
+
+////////////////////////////////////////////////////////////////////////////////////////////////////
+
+void FGraphTrackBuilder::FlushPoints()
+{
+	for (int32 Index = 0; Index < PointsAtCurrentX.Num(); ++Index)
+	{
+		FPointInfo& Pt = PointsAtCurrentX[Index];
+		if (Pt.bValid)
+		{
+			Pt.bValid = false;
+			Series.Points.Add(FVector2D(Pt.X, Pt.Y));
+		}
+	}
+}
+
+////////////////////////////////////////////////////////////////////////////////////////////////////
+
+void FGraphTrackBuilder::EndPoints()
+{
+	FlushPoints();
+}
+
+////////////////////////////////////////////////////////////////////////////////////////////////////
+// FGraphTrackBuilder - Connected Lines
+////////////////////////////////////////////////////////////////////////////////////////////////////
+
+void FGraphTrackBuilder::BeginConnectedLines()
+{
+	Series.LinePoints.Reset();
+
+	LinesCurrentX = -FLT_MAX;
+	LinesMinY = FLT_MAX;
+	LinesMaxY = -FLT_MAX;
+	LinesFirstY = FLT_MAX;
+	LinesLastY = FLT_MAX;
+	bIsLastLineAdded = false;
+}
+
+////////////////////////////////////////////////////////////////////////////////////////////////////
+
+void FGraphTrackBuilder::AddConnectedLine(double Time, double Value)
+{
+	if (bIsLastLineAdded)
+	{
+		return;
+	}
+
+	const float Y = Track.GetYForValue(Value);
+
+	const float X = Viewport.TimeToSlateUnitsRounded(Viewport.RestrictEndTime(Time));
+
+	ensure(X >= LinesCurrentX); // we are assuming events are already sorted by Time
+
+	if (X < 0)
+	{
+		LinesCurrentX = X;
+		LinesLastY = Y;
+		return;
+	}
+
+	if (X >= Viewport.Width)
+	{
+		if (!bIsLastLineAdded)
+		{
+			bIsLastLineAdded = true;
+
+			if (LinesLastY != FLT_MAX)
+			{
+				FlushConnectedLine();
+
+				Series.LinePoints.Add(FVector2D(LinesCurrentX, LinesLastY));
+				Series.LinePoints.Add(FVector2D(X, Y));
+			}
+
+			// Reset the "reduction line" so last FlushConnectedLine() call will do nothing.
+			LinesMinY = Y;
+			LinesMaxY = Y;
+		}
+		return;
+	}
+
+	if (X > LinesCurrentX)
+	{
+		if (LinesLastY != FLT_MAX)
+		{
+			FlushConnectedLine();
+
+			Series.LinePoints.Add(FVector2D(LinesCurrentX, LinesLastY));
+			Series.LinePoints.Add(FVector2D(X, Y));
+		}
+
+		// Advance the "reduction line".
+		LinesCurrentX = X;
+		LinesMinY = Y;
+		LinesMaxY = Y;
+		LinesFirstY = Y;
+		LinesLastY = Y;
+	}
+	else
+	{
+		// Merge current line with the "reduction line".
+		if (Y < LinesMinY)
+		{
+			LinesMinY = Y;
+		}
+		if (Y > LinesMaxY)
+		{
+			LinesMaxY = Y;
+		}
+		LinesLastY = Y;
+	}
+}
+
+////////////////////////////////////////////////////////////////////////////////////////////////////
+
+void FGraphTrackBuilder::FlushConnectedLine()
+{
+	if (LinesCurrentX >= 0.0f && LinesMinY != LinesMaxY)
+	{
+		Series.LinePoints.Add(FVector2D(LinesCurrentX, LinesMaxY));
+		Series.LinePoints.Add(FVector2D(LinesCurrentX, LinesMinY));
+	}
+}
+
+////////////////////////////////////////////////////////////////////////////////////////////////////
+
+void FGraphTrackBuilder::EndConnectedLines()
+{
+	FlushConnectedLine();
+}
+
+////////////////////////////////////////////////////////////////////////////////////////////////////
+// FGraphTrackBuilder - Boxes
+////////////////////////////////////////////////////////////////////////////////////////////////////
+
+void FGraphTrackBuilder::BeginBoxes()
+{
+	Series.Boxes.Reset();
+}
+
+////////////////////////////////////////////////////////////////////////////////////////////////////
+
+void FGraphTrackBuilder::AddBox(double Time, double Duration, double Value)
+{
+	float X1 = Viewport.TimeToSlateUnitsRounded(Time);
+	if (X1 > Viewport.Width)
+	{
+		return;
+	}
+
+	double EndTime = Viewport.RestrictEndTime(Time + Duration);
+	float X2 = Viewport.TimeToSlateUnitsRounded(EndTime);
+	if (X2 < 0)
+	{
+		return;
+	}
+
+	float W = X2 - X1;
+	ensure(W >= 0); // we expect events to be sorted
+
+	// Timing events are displayed with minimum 1px (including empty ones).
+	if (W == 0)
+	{
+		W = 1.0f;
+	}
+
+	// TODO: reduction algorithm
+	FGraphBox Box;
+	Box.X = X1;
+	Box.W = W;
+	Box.Y = Track.GetYForValue(Value);
+	Series.Boxes.Add(Box);
+}
+
+////////////////////////////////////////////////////////////////////////////////////////////////////
+
+void FGraphTrackBuilder::FlushBox()
+{
+	// TODO: reduction algorithm
+}
+
+////////////////////////////////////////////////////////////////////////////////////////////////////
+
+void FGraphTrackBuilder::EndBoxes()
+{
+	FlushBox();
+}
+
+////////////////////////////////////////////////////////////////////////////////////////////////////
+
+#undef LOCTEXT_NAMESPACE
diff --git a/Engine/Source/Developer/TraceInsights/Private/Insights/ViewModels/GraphTrack.h b/Engine/Source/Developer/TraceInsights/Private/Insights/ViewModels/GraphTrack.h
new file mode 100644
--- /dev/null
+++ b/Engine/Source/Developer/TraceInsights/Private/Insights/ViewModels/GraphTrack.h
@@ -0,0 +1,214 @@
+// Copyright 1998-2019 Epic Games, Inc. All Rights Reserved.
+
+#pragma once
+
+#include "CoreMinimal.h"
+#include "Fonts/SlateFontInfo.h"
+
+// Insights
+#include "Insights/ViewModels/BaseTimingTrack.h"
+
+struct FDrawContext;
+struct FSlateBrush;
+class FTimingTrackViewport;
+
+////////////////////////////////////////////////////////////////////////////////////////////////////
+
+struct FGraphBox
+{
+	float X;
+	float W;
+	float Y;
+};
+
+////////////////////////////////////////////////////////////////////////////////////////////////////
+
+class FGraphTrackSeries
+{
+	friend class FGraphTrack;
+	friend class FGraphTrackBuilder;
+
+public:
+	FGraphTrackSeries();
+	~FGraphTrackSeries();
+
+	void SetColor(FLinearColor InColor, FLinearColor InBorderColor)
+	{
+		Color = InColor;
+		BorderColor = InBorderColor;
+	}
+
+protected:
+	//FText Name;
+	//FText Description;
+
+	FLinearColor Color;
+	FLinearColor BorderColor;
+
+	TArray<FVector2D> Points;
+	TArray<FVector2D> LinePoints;
+	TArray<FGraphBox> Boxes;
+
+	//bool bIsVisible;
+};
+
+////////////////////////////////////////////////////////////////////////////////////////////////////
+
+class FGraphTrack : public FBaseTimingTrack
+{
+	friend class FGraphTrackBuilder;
+
+private:
+	// Visual size of points (in pixels).
+	static constexpr float PointVisualSize = 7.0f;
+
+	// Size of points (in pixels) used in reduction algorithm.
+	static constexpr double PointSizeX = 3.0f;
+	static constexpr float PointSizeY = 3.0f;
+
+public:
+	FGraphTrack(uint64 InTrackId);
+	virtual ~FGraphTrack();
+
+	virtual void UpdateHoveredState(float MouseX, float MouseY, const FTimingTrackViewport& Viewport);
+
+	virtual void Update(const FTimingTrackViewport& Viewport) override = 0;
+
+	void Draw(FDrawContext& DrawContext, const FTimingTrackViewport& Viewport) const;
+
+	int32 GetNumAddedEvents() const { return NumAddedEvents; }
+	int32 GetNumDrawPoints() const { return NumDrawPoints; }
+	int32 GetNumDrawLines() const { return NumDrawLines; }
+	int32 GetNumDrawBoxes() const { return NumDrawBoxes; }
+
+protected:
+	void UpdateStats();
+
+	void DrawSeries(const FGraphTrackSeries& Series, FDrawContext& DrawContext, const FTimingTrackViewport& Viewport) const;
+
+	float GetYForValue(double Value) const
+	{
+		return BaselineY - static_cast<float>(Value * ScaleY); // TODO: vertical zooming and panning
+	}
+
+protected:
+	TArray<FGraphTrackSeries> AllSeries;
+
+	// Slate resources
+	const FSlateBrush* WhiteBrush;
+	const FSlateBrush* PointBrush;
+	const FSlateBrush* BorderBrush;
+	const FSlateFontInfo Font;
+
+	bool bDrawPoints;
+	bool bDrawPointsWithBorder;
+	bool bDrawLines;
+	bool bDrawLinesWithDuration;
+	bool bDrawBoxes;
+
+	float BaselineY; // Y position (in viewport local space) of the baseline (with Value == 0); in pixels (Slate units)
+	double ScaleY; // scale between Value units and viewport units; in pixels (Slate units) / Value unit
+
+	// Stats
+	int32 NumAddedEvents;
+	int32 NumDrawPoints;
+	int32 NumDrawLines;
+	int32 NumDrawBoxes;
+};
+
+////////////////////////////////////////////////////////////////////////////////////////////////////
+
+class FRandomGraphTrack : public FGraphTrack
+{
+public:
+	FRandomGraphTrack(uint64 InTrackId);
+	virtual ~FRandomGraphTrack();
+
+	virtual void Update(const FTimingTrackViewport& Viewport) override;
+
+protected:
+	void GenerateSeries(FGraphTrackSeries& Series, const FTimingTrackViewport& Viewport, const int32 EventCount);
+};
+
+////////////////////////////////////////////////////////////////////////////////////////////////////
+
+class FFramesGraphTrack : public FGraphTrack
+{
+public:
+	FFramesGraphTrack(uint64 InTrackId);
+	virtual ~FFramesGraphTrack();
+
+	virtual void Update(const FTimingTrackViewport& Viewport) override;
+
+protected:
+	void UpdateSeries(FGraphTrackSeries& Series, const FTimingTrackViewport& Viewport, ETraceFrameType FrameType);
+};
+
+////////////////////////////////////////////////////////////////////////////////////////////////////
+
+class FGraphTrackBuilder
+{
+private:
+	struct FPointInfo
+	{
+		bool bValid;
+		float X;
+		float Y;
+
+		FPointInfo() : bValid(false) {}
+	};
+
+public:
+	FGraphTrackBuilder(FGraphTrack& InTrack, FGraphTrackSeries& InSeries, const FTimingTrackViewport& InViewport);
+	~FGraphTrackBuilder();
+
+	/**
+	 * Non-copyable
+	 */
+	FGraphTrackBuilder(const FGraphTrackBuilder&) = delete;
+	FGraphTrackBuilder& operator=(const FGraphTrackBuilder&) = delete;
+
+	FGraphTrack& GetTrack() const { return Track; }
+	FGraphTrackSeries& GetSeries() const { return Series; }
+	const FTimingTrackViewport& GetViewport() const { return Viewport; }
+
+	void AddEvent(double Time, double Duration, double Value);
+
+private:
+	void BeginPoints();
+	void AddPoint(double Time, double Value);
+	void FlushPoints();
+	void EndPoints();
+
+	void BeginConnectedLines();
+	void AddConnectedLine(double Time, double Value);
+	void FlushConnectedLine();
+	void EndConnectedLines();
+
+	void BeginBoxes();
+	void AddBox(double Time, double Duration, double Value);
+	void FlushBox();
+	void EndBoxes();
+
+private:
+	FGraphTrack& Track;
+	FGraphTrackSeries& Series;
+	const FTimingTrackViewport& Viewport;
+
+	// Used by the point reduction algorithm.
+	double PointsCurrentX;
+	TArray<FPointInfo> PointsAtCurrentX;
+
+	// Used by the line reduction algorithm.
+	float LinesCurrentX;
+	float LinesMinY;
+	float LinesMaxY;
+	float LinesFirstY;
+	float LinesLastY;
+	bool bIsLastLineAdded;
+
+	// Used by the box reduction algorithm.
+	//...
+};
+
+////////////////////////////////////////////////////////////////////////////////////////////////////
