Copying //UE4/Dev-Core to Dev-Main (//UE4/Dev-Main)

diff --git a/Engine/Source/Developer/TraceInsights/Private/Insights/ViewModels/MarkersTimingTrack.cpp b/Engine/Source/Developer/TraceInsights/Private/Insights/ViewModels/MarkersTimingTrack.cpp
new file mode 100644
--- /dev/null
+++ b/Engine/Source/Developer/TraceInsights/Private/Insights/ViewModels/MarkersTimingTrack.cpp
@@ -0,0 +1,526 @@
+// Copyright 1998-2019 Epic Games, Inc. All Rights Reserved.
+
+#include "MarkersTimingTrack.h"
+
+#include "Brushes/SlateColorBrush.h"
+#include "Fonts/FontMeasure.h"
+#include "Framework/Application/SlateApplication.h"
+#include "Styling/CoreStyle.h"
+#include "TraceServices/AnalysisService.h"
+
+// Insights
+#include "Insights/Common/PaintUtils.h"
+#include "Insights/InsightsManager.h"
+#include "Insights/TimingProfilerManager.h"
+#include "Insights/ViewModels/TimingTrackViewport.h"
+
+////////////////////////////////////////////////////////////////////////////////////////////////////
+
+#define LOCTEXT_NAMESPACE "MarkersTimingTrack"
+
+////////////////////////////////////////////////////////////////////////////////////////////////////
+// FMarkersTimingTrack
+////////////////////////////////////////////////////////////////////////////////////////////////////
+
+FMarkersTimingTrack::FMarkersTimingTrack(uint64 InTrackId)
+	: FBaseTimingTrack(InTrackId)
+	//, TimeMarkerBoxes()
+	//, TimeMarkerTexts()
+	, bIsCollapsed(false)
+	, bUseOnlyBookmarks(true)
+	, bIsDirty(true)
+	, TargetHoveredAnimPercent(0.0f)
+	, CurrentHoveredAnimPercent(0.0f)
+	, NumLogMessages(0)
+	, NumDrawBoxes(0)
+	, NumDrawTexts(0)
+	, WhiteBrush(FCoreStyle::Get().GetBrush("WhiteBrush"))
+	, Font(FCoreStyle::GetDefaultFontStyle("Regular", 8))
+{
+}
+
+////////////////////////////////////////////////////////////////////////////////////////////////////
+
+FMarkersTimingTrack::~FMarkersTimingTrack()
+{
+}
+
+////////////////////////////////////////////////////////////////////////////////////////////////////
+
+void FMarkersTimingTrack::Reset()
+{
+	FBaseTimingTrack::Reset();
+
+	TimeMarkerBoxes.Reset();
+	TimeMarkerTexts.Reset();
+
+	bIsCollapsed = false;
+	bUseOnlyBookmarks = true;
+	bIsDirty = true;
+
+	TargetHoveredAnimPercent = 0.0f;
+	CurrentHoveredAnimPercent = 0.0f;
+
+	NumLogMessages = 0;
+	NumDrawBoxes = 0;
+	NumDrawTexts = 0;
+
+	UpdateHeight();
+}
+
+////////////////////////////////////////////////////////////////////////////////////////////////////
+
+void FMarkersTimingTrack::UpdateHeight()
+{
+	constexpr float BookmarksTrackHeight = 14.0f;
+	constexpr float TimeMarkersTrackHeight = 28.0f;
+
+	if (bUseOnlyBookmarks)
+	{
+		SetHeight(BookmarksTrackHeight);
+	}
+	else
+	{
+		SetHeight(TimeMarkersTrackHeight);
+	}
+}
+
+////////////////////////////////////////////////////////////////////////////////////////////////////
+
+void FMarkersTimingTrack::UpdateHoveredState(float MouseX, float MouseY, const FTimingTrackViewport& Viewport)
+{
+	constexpr float HeaderWidth = 80.0f;
+	constexpr float HeaderHeight = 14.0f;
+
+	if (MouseY >= GetPosY() && MouseY < GetPosY() + GetHeight())
+	{
+		SetHoveredState(true);
+		SetHeaderHoveredState(MouseX < HeaderWidth && MouseY < GetPosY() + HeaderHeight);
+		TargetHoveredAnimPercent = 1.0f;
+	}
+	else
+	{
+		SetHoveredState(false);
+		TargetHoveredAnimPercent = 0.0f;
+	}
+}
+
+////////////////////////////////////////////////////////////////////////////////////////////////////
+
+void FMarkersTimingTrack::Tick(const double InCurrentTime, const float InDeltaTime)
+{
+	if (CurrentHoveredAnimPercent != TargetHoveredAnimPercent)
+	{
+		if (CurrentHoveredAnimPercent < TargetHoveredAnimPercent)
+		{
+			const float ShowAnimSpeed = 2 * InDeltaTime;
+			CurrentHoveredAnimPercent += ShowAnimSpeed;
+			if (CurrentHoveredAnimPercent > TargetHoveredAnimPercent)
+			{
+				CurrentHoveredAnimPercent = TargetHoveredAnimPercent;
+			}
+		}
+		else
+		{
+			const float HideAnimSpeed = 3 * InDeltaTime;
+			CurrentHoveredAnimPercent -= HideAnimSpeed;
+			if (CurrentHoveredAnimPercent < TargetHoveredAnimPercent)
+			{
+				CurrentHoveredAnimPercent = TargetHoveredAnimPercent;
+			}
+		}
+	}
+}
+
+////////////////////////////////////////////////////////////////////////////////////////////////////
+
+void FMarkersTimingTrack::Update(const FTimingTrackViewport& InViewport)
+{
+	FTimeMarkerTrackBuilder Builder(*this, InViewport);
+
+	TSharedPtr<const Trace::IAnalysisSession> Session = FInsightsManager::Get()->GetSession();
+	if (Session.IsValid())
+	{
+		Trace::FAnalysisSessionReadScope SessionReadScope(*Session.Get());
+
+		const Trace::ILogProvider& LogProvider = Trace::ReadLogProvider(*Session.Get());
+		Builder.BeginLog(LogProvider);
+
+		LogProvider.EnumerateMessages(
+			Builder.GetViewport().StartTime,
+			Builder.GetViewport().EndTime,
+			[&Builder](const Trace::FLogMessage& Message) { Builder.AddLogMessage(Message); });
+
+		Builder.EndLog();
+	}
+}
+
+////////////////////////////////////////////////////////////////////////////////////////////////////
+
+void FMarkersTimingTrack::Draw(FDrawContext& DrawContext, const FTimingTrackViewport& Viewport) const
+{
+	// Draw background.
+	DrawContext.DrawBox(0.0f, GetPosY(), Viewport.Width, GetHeight(), WhiteBrush, FLinearColor(0.04f, 0.04f, 0.04f, 1.0f));
+	DrawContext.LayerId++;
+
+	// Draw the track's header, in background.
+	if (!IsHovered() || CurrentHoveredAnimPercent < 1.0)
+	{
+		DrawHeader(DrawContext, true);
+	}
+
+	//////////////////////////////////////////////////
+	// Draw vertical lines.
+	// Multiple adjacent vertical lines with same color are merged into a single box.
+
+	float BoxY, BoxH;
+	if (IsCollapsed())
+	{
+		BoxY = GetPosY();
+		BoxH = GetHeight();
+	}
+	else
+	{
+		BoxY = 0.0f;
+		BoxH = Viewport.Height;
+	}
+
+	int32 NumBoxes = TimeMarkerBoxes.Num();
+	for (int32 BoxIndex = 0; BoxIndex < NumBoxes; BoxIndex++)
+	{
+		const FTimeMarkerBoxInfo& Box = TimeMarkerBoxes[BoxIndex];
+		DrawContext.DrawBox(Box.X, BoxY, Box.W, BoxH, WhiteBrush, Box.Color);
+	}
+	DrawContext.LayerId++;
+	NumDrawBoxes = NumBoxes;
+
+	//////////////////////////////////////////////////
+	// Draw texts (strings are already truncated).
+
+	const float CategoryY = GetPosY() + 2.0f;
+	const float MessageY = GetPosY() + (IsBookmarksTrack() ? 1.0f : 14.0f);
+
+	int32 NumTexts = TimeMarkerTexts.Num();
+	for (int32 TextIndex = 0; TextIndex < NumTexts; TextIndex++)
+	{
+		const FTimeMarkerTextInfo& TextInfo = TimeMarkerTexts[TextIndex];
+
+		if (!IsBookmarksTrack() && TextInfo.Category.Len() > 0)
+		{
+			DrawContext.DrawText(TextInfo.X, CategoryY, TextInfo.Category, Font, TextInfo.Color);
+			NumDrawTexts++;
+		}
+
+		if (TextInfo.Message.Len() > 0)
+		{
+			DrawContext.DrawText(TextInfo.X, MessageY, TextInfo.Message, Font, TextInfo.Color);
+			NumDrawTexts++;
+		}
+	}
+	DrawContext.LayerId++;
+
+	//////////////////////////////////////////////////
+
+	// When hovered, the track's header is draw on top.
+	if (IsHovered() || CurrentHoveredAnimPercent > 0.0)
+	{
+		DrawHeader(DrawContext, false);
+	}
+}
+
+////////////////////////////////////////////////////////////////////////////////////////////////////
+
+void FMarkersTimingTrack::DrawHeader(FDrawContext& DrawContext, bool bFirstDraw) const
+{
+	const FString Name = IsBookmarksTrack() ? TEXT("Bookmarks") : TEXT("Logs");
+	//const float ArrowSize = 6.0f;
+	const float ArrowSizeX = 4.0f;
+	const float ArrowSizeY = 8.0f;
+	const float ArrowX = IsBookmarksTrack() ? 64.0f : 31.0f;
+	const float ArrowY = GetPosY() + 3.0f;
+	const float HeaderW = ArrowX + ArrowSizeX + 4.0f;
+	const float HeaderH = 14.0f;
+
+	if (!bFirstDraw)
+	{
+		DrawContext.DrawBox(0.0f, GetPosY(), HeaderW, HeaderH, WhiteBrush, FLinearColor(0.04f, 0.04f, 0.04f, CurrentHoveredAnimPercent));
+		DrawContext.LayerId++;
+	}
+
+	FLinearColor Color;
+	if (bFirstDraw || CurrentHoveredAnimPercent == 0.0)
+	{
+		Color = FLinearColor(0.07f, 0.07f, 0.07f, 1.0f);
+	}
+	else if (IsHeaderHovered())
+	{
+		Color = FLinearColor(1.0f, 1.0f, 0.0f, CurrentHoveredAnimPercent);
+	}
+	else
+	{
+		Color = FLinearColor(1.0f, 1.0f, 1.0f, CurrentHoveredAnimPercent);
+	}
+
+	// Draw "Bookmarks" or "Logs" text.
+	DrawContext.DrawText(2.0f, GetPosY() + 1.0f, Name, Font, Color);
+
+	if (IsCollapsed())
+	{
+		// Draw "right empty arrow".
+		TArray<FVector2D> Points =
+		{
+			FVector2D(ArrowX, ArrowY),
+			FVector2D(ArrowX, ArrowY + ArrowSizeY),
+			FVector2D(ArrowX + ArrowSizeX, ArrowY + ArrowSizeY / 2.0f),
+			FVector2D(ArrowX, ArrowY)
+		};
+		FSlateDrawElement::MakeLines(DrawContext.ElementList, DrawContext.LayerId, DrawContext.Geometry.ToPaintGeometry(), Points, DrawContext.DrawEffects, Color, false, 1.0f);
+	}
+	else
+	{
+		// Draw "down-right filled arrow".
+		for (float A = 1.0f; A < ArrowSizeY; A += 1.0f)
+		{
+			DrawContext.DrawBox(ArrowX - 3.0 + ArrowSizeY - A, ArrowY + A - 1.0f, A, 1.0f, WhiteBrush, Color);
+		}
+	}
+
+	DrawContext.LayerId++;
+}
+
+////////////////////////////////////////////////////////////////////////////////////////////////////
+// FTimeMarkerTrackBuilder
+////////////////////////////////////////////////////////////////////////////////////////////////////
+
+FTimeMarkerTrackBuilder::FTimeMarkerTrackBuilder(FMarkersTimingTrack& InTrack, const FTimingTrackViewport& InViewport)
+	: Track(InTrack)
+	, Viewport(InViewport)
+	, FontMeasureService(FSlateApplication::Get().GetRenderer()->GetFontMeasureService())
+	, Font(FCoreStyle::GetDefaultFontStyle("Regular", 8))
+{
+	Track.ResetCache();
+	Track.NumLogMessages = 0;
+}
+
+////////////////////////////////////////////////////////////////////////////////////////////////////
+
+void FTimeMarkerTrackBuilder::BeginLog(const Trace::ILogProvider& LogProvider)
+{
+	LogProviderPtr = &LogProvider;
+
+	LastX1 = -1000.0f;
+	LastX2 = -1000.0f;
+	LastLogIndex = 0;
+	LastVerbosity = ELogVerbosity::NoLogging;
+	LastCategory = nullptr;
+	LastMessage = nullptr;
+}
+
+////////////////////////////////////////////////////////////////////////////////////////////////////
+
+void FTimeMarkerTrackBuilder::AddLogMessage(const Trace::FLogMessage& Message)
+{
+	Track.NumLogMessages++;
+
+	// Add also the log message imediately on the left of the screen (if any).
+	if (Track.NumLogMessages == 1 && Message.Index > 0)
+	{
+		// Note: Reading message at Index-1 will not work as expected when using filter!
+		//TODO: Search API like: LogProviderPtr->SearchMessage(StartIndex, ESearchDirection::Backward, LambdaPredicate, bResolveFormatString);
+		LogProviderPtr->ReadMessage(
+			Message.Index - 1,
+			[this](const Trace::FLogMessage& Message) { AddLogMessage(Message); });
+	}
+
+	if (!Track.bUseOnlyBookmarks || FCString::Strcmp(Message.Category->Name, TEXT("LogBookmark")) == 0)
+	{
+		float X = Viewport.TimeToSlateUnitsRounded(Message.Time);
+		if (X < 0)
+		{
+			X = -1.0f;
+		}
+		AddTimeMarker(X, Message.Index, Message.Verbosity, Message.Category->Name, Message.Message);
+	}
+}
+
+////////////////////////////////////////////////////////////////////////////////////////////////////
+
+FLinearColor FTimeMarkerTrackBuilder::GetColorByCategory(const TCHAR* const Category)
+{
+	// Strip the "Log" prefix.
+	FString CategoryStr(Category);
+	if (CategoryStr.StartsWith(TEXT("Log")))
+	{
+		CategoryStr = CategoryStr.RightChop(3);
+	}
+
+	uint32 Hash = 0;
+	for (const TCHAR* c = *CategoryStr; *c; ++c)
+	{
+		Hash = (Hash + *c) * 0x2c2c57ed;
+	}
+
+	// Divided by 128.0 in order to force bright colors.
+	return FLinearColor(((Hash >> 16) & 0xFF) / 128.0f, ((Hash >> 8) & 0xFF) / 128.0f, (Hash & 0xFF) / 128.0f, 1.0f);
+}
+
+////////////////////////////////////////////////////////////////////////////////////////////////////
+
+FLinearColor FTimeMarkerTrackBuilder::GetColorByVerbosity(const ELogVerbosity::Type Verbosity)
+{
+	static FLinearColor Colors[] =
+	{
+		FLinearColor(0.0, 0.0, 0.0, 1.0), // NoLogging
+		FLinearColor(1.0, 0.0, 0.0, 1.0), // Fatal
+		FLinearColor(1.0, 0.3, 0.0, 1.0), // Error
+		FLinearColor(0.7, 0.5, 0.0, 1.0), // Warning
+		FLinearColor(0.0, 0.7, 0.0, 1.0), // Display
+		FLinearColor(0.0, 0.7, 1.0, 1.0), // Log
+		FLinearColor(0.7, 0.7, 0.7, 1.0), // Verbose
+		FLinearColor(1.0, 1.0, 1.0, 1.0), // VeryVerbose
+	};
+	static_assert(sizeof(Colors) / sizeof(FLinearColor) == (int)ELogVerbosity::Type::All + 1, "ELogVerbosity::Type has changed!?");
+	//return Colors[Verbosity & ELogVerbosity::VerbosityMask];
+	return Colors[Verbosity & 7]; // using 7 instead of ELogVerbosity::VerbosityMask (15) to make static analyzer happy
+}
+
+////////////////////////////////////////////////////////////////////////////////////////////////////
+
+void FTimeMarkerTrackBuilder::Flush(float AvailableTextW)
+{
+	// Is last marker valid?
+	if (LastCategory != nullptr)
+	{
+		const FLinearColor Color = GetColorByCategory(LastCategory);
+
+		bool bAddNewBox = true;
+		if (Track.TimeMarkerBoxes.Num() > 0)
+		{
+			FTimeMarkerBoxInfo& PrevBox = Track.TimeMarkerBoxes[Track.TimeMarkerBoxes.Num() - 1];
+			if (PrevBox.X + PrevBox.W == LastX1 &&
+				PrevBox.Color.R == Color.R &&
+				PrevBox.Color.G == Color.G &&
+				PrevBox.Color.B == Color.B)
+			{
+				// Extend previous box instead.
+				PrevBox.W += LastX2 - LastX1;
+				bAddNewBox = false;
+			}
+		}
+
+		if (bAddNewBox)
+		{
+			// Add new Box info to cache.
+			FTimeMarkerBoxInfo& Box = Track.TimeMarkerBoxes[Track.TimeMarkerBoxes.AddDefaulted()];
+			Box.X = LastX1;
+			Box.W = LastX2 - LastX1;
+			Box.Color = Color;
+			Box.Color.A = 0.25f;
+		}
+
+		if (AvailableTextW > 6.0f)
+		{
+			// Strip the "Log" prefix.
+			FString CategoryStr(LastCategory);
+			if (CategoryStr.StartsWith(TEXT("Log")))
+			{
+				CategoryStr = CategoryStr.RightChop(3);
+			}
+
+			const int32 LastWholeCharacterIndexCategory = FontMeasureService->FindLastWholeCharacterIndexBeforeOffset(CategoryStr, Font, FMath::RoundToInt(AvailableTextW - 2.0f));
+			const int32 LastWholeCharacterIndexMessage = FontMeasureService->FindLastWholeCharacterIndexBeforeOffset(LastMessage, Font, FMath::RoundToInt(AvailableTextW - 2.0f));
+
+			if (LastWholeCharacterIndexCategory >= 0 ||
+				LastWholeCharacterIndexMessage >= 0)
+			{
+				// Add new Text info to cache.
+				FTimeMarkerTextInfo& TextInfo = Track.TimeMarkerTexts[Track.TimeMarkerTexts.AddDefaulted()];
+				TextInfo.X = LastX2 + 2.0f;
+				TextInfo.Color = Color;
+				if (LastWholeCharacterIndexCategory >= 0)
+				{
+					TextInfo.Category = CategoryStr.Left(LastWholeCharacterIndexCategory + 1);
+				}
+				if (LastWholeCharacterIndexMessage >= 0)
+				{
+					TextInfo.Message.AppendChars(LastMessage, LastWholeCharacterIndexMessage + 1);
+				}
+			}
+		}
+	}
+}
+
+////////////////////////////////////////////////////////////////////////////////////////////////////
+
+void FTimeMarkerTrackBuilder::AddTimeMarker(const float X, const uint64 LogIndex, const ELogVerbosity::Type Verbosity, const TCHAR* const Category, const TCHAR* Message)
+{
+	float W = X - LastX2;
+
+	if (W > 0.0f) // There is at least 1px from previous box?
+	{
+		// Flush previous marker (if any).
+		Flush(W);
+
+		// Begin new marker info.
+		LastX1 = X;
+		LastX2 = X + 1.0f;
+	}
+	else if (W == 0.0f) // Adjacent to previous box?
+	{
+		// Same color as previous marker?
+		if (Category == LastCategory)
+		{
+			// Extend previous box.
+			LastX2++;
+		}
+		else
+		{
+			// Flush previous marker (if any).
+			Flush(0.0f);
+
+			// Begin new box.
+			LastX1 = X;
+			LastX2 = X + 1.0f;
+		}
+	}
+	else // Overlaps previous box?
+	{
+		// Same color as previous marker?
+		if (Category == LastCategory)
+		{
+			// Keep previous box.
+		}
+		else
+		{
+			// Shrink previous box.
+			LastX2--;
+
+			if (LastX2 > LastX1)
+			{
+				// Flush previous marker (if any).
+				Flush(0.0f);
+			}
+
+			// Begin new box.
+			LastX1 = X;
+			LastX2 = X + 1.0f;
+		}
+	}
+
+	// Save marker.
+	LastCategory = Category;
+	LastVerbosity = Verbosity;
+	LastLogIndex = LogIndex;
+	LastMessage = Message;
+}
+
+////////////////////////////////////////////////////////////////////////////////////////////////////
+
+void FTimeMarkerTrackBuilder::EndLog()
+{
+	Flush(Viewport.Width - LastX2);
+}
+
+////////////////////////////////////////////////////////////////////////////////////////////////////
+
+#undef LOCTEXT_NAMESPACE
diff --git a/Engine/Source/Developer/TraceInsights/Private/Insights/ViewModels/MarkersTimingTrack.h b/Engine/Source/Developer/TraceInsights/Private/Insights/ViewModels/MarkersTimingTrack.h
new file mode 100644
--- /dev/null
+++ b/Engine/Source/Developer/TraceInsights/Private/Insights/ViewModels/MarkersTimingTrack.h
@@ -0,0 +1,152 @@
+// Copyright 1998-2019 Epic Games, Inc. All Rights Reserved.
+
+#pragma once
+
+#include "CoreMinimal.h"
+#include "Fonts/SlateFontInfo.h"
+
+// Insights
+#include "Insights/ViewModels/BaseTimingTrack.h"
+
+namespace Trace
+{
+	struct FLogMessage;
+	class ILogProvider;
+}
+
+struct FDrawContext;
+struct FSlateBrush;
+class FTimingTrackViewport;
+
+////////////////////////////////////////////////////////////////////////////////////////////////////
+
+struct FTimeMarkerBoxInfo
+{
+	float X;
+	float W;
+	FLinearColor Color;
+};
+
+////////////////////////////////////////////////////////////////////////////////////////////////////
+
+struct FTimeMarkerTextInfo
+{
+	float X;
+	FLinearColor Color;
+	FString Category; // truncated Category string
+	FString Message; // truncated Message string
+};
+
+////////////////////////////////////////////////////////////////////////////////////////////////////
+
+class FMarkersTimingTrack : public FBaseTimingTrack
+{
+	friend class FTimeMarkerTrackBuilder;
+
+public:
+	FMarkersTimingTrack(uint64 InTrackId);
+	virtual ~FMarkersTimingTrack();
+
+	virtual void Reset() override;
+
+	bool IsCollapsed() const { return bIsCollapsed; }
+	void ToggleCollapsed() { bIsCollapsed = !bIsCollapsed; }
+
+	bool IsBookmarksTrack() const { return bUseOnlyBookmarks; }
+	void SetBookmarksTrackFlag(bool bInUseOnlyBookmarks)
+	{
+		bUseOnlyBookmarks = bInUseOnlyBookmarks;
+		UpdateHeight();
+	}
+
+	bool IsDirty() const { return bIsDirty; }
+	void SetDirtyFlag() { bIsDirty = true; }
+	void ClearDirtyFlag() { bIsDirty = false; }
+
+	// Stats
+	int32 GetNumLogMessages() const { return NumLogMessages; }
+	int32 GetNumBoxes() const { return TimeMarkerBoxes.Num(); }
+	int32 GetNumTexts() const { return TimeMarkerTexts.Num(); }
+
+	virtual void UpdateHoveredState(float MouseX, float MouseY, const FTimingTrackViewport& Viewport) override;
+	virtual void Tick(const double CurrentTime, const float DeltaTime) override;
+	virtual void Update(const FTimingTrackViewport& Viewport) override;
+
+	void Draw(FDrawContext& DrawContext, const FTimingTrackViewport& Viewport) const;
+
+private:
+	void ResetCache()
+	{
+		TimeMarkerBoxes.Reset();
+		TimeMarkerTexts.Reset();
+	}
+
+	void UpdateHeight();
+
+	void DrawHeader(FDrawContext& DrawContext, bool bFirstDraw) const;
+
+private:
+	TArray<FTimeMarkerBoxInfo> TimeMarkerBoxes;
+	TArray<FTimeMarkerTextInfo> TimeMarkerTexts;
+
+	bool bIsCollapsed; // If false, the vertical lines will extend to entire viewport height; otherwise will be limited to this track's height.
+	bool bUseOnlyBookmarks; // If true, uses only bookmarks; otherwise it uses all log messages.
+	bool bIsDirty; // Source list of bookmarks/markers has changed. Cached Boxes and Texts are dirty.
+
+	float TargetHoveredAnimPercent; // [0.0 .. 1.0], 0.0 = hidden, 1.0 = visible
+	float CurrentHoveredAnimPercent;
+
+	// Stats
+	int32 NumLogMessages;
+	mutable int32 NumDrawBoxes;
+	mutable int32 NumDrawTexts;
+
+	// Slate resources
+	const FSlateBrush* WhiteBrush;
+	const FSlateFontInfo Font;
+};
+
+////////////////////////////////////////////////////////////////////////////////////////////////////
+
+class FTimeMarkerTrackBuilder
+{
+public:
+	FTimeMarkerTrackBuilder(FMarkersTimingTrack& InTrack, const FTimingTrackViewport& InViewport);
+
+	/**
+	 * Non-copyable
+	 */
+	FTimeMarkerTrackBuilder(const FTimeMarkerTrackBuilder&) = delete;
+	FTimeMarkerTrackBuilder& operator=(const FTimeMarkerTrackBuilder&) = delete;
+
+	const FTimingTrackViewport& GetViewport() { return Viewport; }
+
+	void BeginLog(const Trace::ILogProvider& LogProvider);
+	void AddLogMessage(const Trace::FLogMessage& Message);
+	void EndLog();
+
+	static FLinearColor GetColorByCategory(const TCHAR* const Category);
+	static FLinearColor GetColorByVerbosity(const ELogVerbosity::Type Verbosity);
+
+private:
+	void Flush(float AvailableTextW);
+	void AddTimeMarker(const float X, const uint64 LogIndex, const ELogVerbosity::Type Verbosity, const TCHAR* const Category, const TCHAR* Message);
+
+private:
+	FMarkersTimingTrack& Track;
+	const FTimingTrackViewport& Viewport;
+
+	const TSharedRef<class FSlateFontMeasure> FontMeasureService;
+	const FSlateFontInfo Font;
+
+	const Trace::ILogProvider* LogProviderPtr; // valid only between BeginLog() and EndLog()
+
+	float LastX1;
+	float LastX2;
+	uint64 LastLogIndex;
+	ELogVerbosity::Type LastVerbosity;
+	const TCHAR* LastCategory;
+	const TCHAR* LastMessage;
+};
+
+////////////////////////////////////////////////////////////////////////////////////////////////////
diff --git a/Engine/Source/Developer/TraceInsights/Private/Insights/ViewModels/StatsNode.cpp b/Engine/Source/Developer/TraceInsights/Private/Insights/ViewModels/StatsNode.cpp
new file mode 100644
--- /dev/null
+++ b/Engine/Source/Developer/TraceInsights/Private/Insights/ViewModels/StatsNode.cpp
@@ -0,0 +1,191 @@
+// Copyright 1998-2019 Epic Games, Inc. All Rights Reserved.
+
+#include "StatsNode.h"
+
+// Insights
+#include "Insights/Common/TimeUtils.h"
+
+#define LOCTEXT_NAMESPACE "StatsNode"
+
+////////////////////////////////////////////////////////////////////////////////////////////////////
+
+void FStatsNode::ResetAggregatedStats()
+{
+	AggregatedStats.Reset();
+}
+
+////////////////////////////////////////////////////////////////////////////////////////////////////
+
+void FStatsNode::SetAggregatedStats(FAggregatedStats& InAggregatedStats)
+{
+	AggregatedStats = InAggregatedStats;
+}
+
+////////////////////////////////////////////////////////////////////////////////////////////////////
+
+void FStatsNode::ResetAggregatedIntegerStats()
+{
+	AggregatedIntegerStats.Reset();
+}
+
+////////////////////////////////////////////////////////////////////////////////////////////////////
+
+void FStatsNode::SetAggregatedIntegerStats(FAggregatedIntegerStats& InAggregatedIntegerStats)
+{
+	AggregatedIntegerStats = InAggregatedIntegerStats;
+
+	// Sorting and display of Count value uses the "float" stats.
+	AggregatedStats.Count = AggregatedIntegerStats.Count;
+
+	AggregatedStats.Sum           = static_cast<double>(AggregatedIntegerStats.Sum);
+	AggregatedStats.Min           = static_cast<double>(AggregatedIntegerStats.Min);
+	AggregatedStats.Max           = static_cast<double>(AggregatedIntegerStats.Max);
+	AggregatedStats.Average       = static_cast<double>(AggregatedIntegerStats.Average);
+	AggregatedStats.Median        = static_cast<double>(AggregatedIntegerStats.Median);
+	AggregatedStats.LowerQuartile = static_cast<double>(AggregatedIntegerStats.LowerQuartile);
+	AggregatedStats.UpperQuartile = static_cast<double>(AggregatedIntegerStats.UpperQuartile);
+}
+
+////////////////////////////////////////////////////////////////////////////////////////////////////
+
+const FText FStatsNode::FormatAggregatedStatsValue(double ValueDbl, int64 ValueInt) const
+{
+	if (GetType() == EStatsNodeType::Float)
+	{
+		if (AggregatedStats.Count == 0)
+		{
+			return LOCTEXT("AggregatedStatsNA", "N/A");
+		}
+		else
+		{
+			//TODO: if (GetDisplayHint() == FName(TEXT("Seconds")))
+			if (GetMetaGroupName() == FName(TEXT("Time")))
+			{
+				return FText::FromString(TimeUtils::FormatTimeAuto(ValueDbl));
+			}
+			else
+			{
+				return FText::AsNumber(ValueDbl);
+			}
+		}
+	}
+	else
+	{
+		if (AggregatedIntegerStats.Count == 0)
+		{
+			return LOCTEXT("AggregatedStatsNA", "N/A");
+		}
+		else
+		{
+			//TODO: if (GetDisplayHint() == FName(TEXT("Bytes")))
+			if (GetMetaGroupName() == FName(TEXT("Memory")))
+			{
+				if (ValueInt > 0)
+				{
+					return FText::AsMemory(ValueInt);
+				}
+				else if (ValueInt == 0)
+				{
+					return FText::FromString(TEXT("0"));
+				}
+				else
+				{
+					return FText::FromString(FString::Printf(TEXT("-%s"), *FText::AsMemory(-ValueInt).ToString()));
+				}
+			}
+			else
+			{
+				return FText::AsNumber(ValueInt);
+			}
+		}
+	}
+}
+
+////////////////////////////////////////////////////////////////////////////////////////////////////
+
+const FText FStatsNode::GetTextForAggregatedStatsSum() const
+{
+	return FormatAggregatedStatsValue(AggregatedStats.Sum, AggregatedIntegerStats.Sum);
+}
+
+////////////////////////////////////////////////////////////////////////////////////////////////////
+
+const FText FStatsNode::GetTextForAggregatedStatsMin() const
+{
+	return FormatAggregatedStatsValue(AggregatedStats.Min, AggregatedIntegerStats.Min);
+}
+
+////////////////////////////////////////////////////////////////////////////////////////////////////
+
+const FText FStatsNode::GetTextForAggregatedStatsMax() const
+{
+	return FormatAggregatedStatsValue(AggregatedStats.Max, AggregatedIntegerStats.Max);
+}
+
+////////////////////////////////////////////////////////////////////////////////////////////////////
+
+const FText FStatsNode::GetTextForAggregatedStatsAverage() const
+{
+	return FormatAggregatedStatsValue(AggregatedStats.Average, AggregatedIntegerStats.Average);
+}
+
+////////////////////////////////////////////////////////////////////////////////////////////////////
+
+const FText FStatsNode::GetTextForAggregatedStatsMedian() const
+{
+	return FormatAggregatedStatsValue(AggregatedStats.Median, AggregatedIntegerStats.Median);
+}
+
+////////////////////////////////////////////////////////////////////////////////////////////////////
+
+const FText FStatsNode::GetTextForAggregatedStatsLowerQuartile() const
+{
+	return FormatAggregatedStatsValue(AggregatedStats.LowerQuartile, AggregatedIntegerStats.LowerQuartile);
+}
+
+////////////////////////////////////////////////////////////////////////////////////////////////////
+
+const FText FStatsNode::GetTextForAggregatedStatsUpperQuartile() const
+{
+	return FormatAggregatedStatsValue(AggregatedStats.UpperQuartile, AggregatedIntegerStats.UpperQuartile);
+}
+
+////////////////////////////////////////////////////////////////////////////////////////////////////
+
+const FText FStatsNode::GetNameEx() const
+{
+	FText Text = FText::GetEmpty();
+
+	if (IsGroup())
+	{
+		const int32 NumChildren = Children.Num();
+		const int32 NumFilteredChildren = FilteredChildren.Num();
+
+		if (NumFilteredChildren == NumChildren)
+		{
+			Text = FText::Format(LOCTEXT("StatsNodeGroupTextFmt1", "{0} ({1})"), FText::FromName(Name), FText::AsNumber(NumChildren));
+		}
+		else
+		{
+			Text = FText::Format(LOCTEXT("StatsNodeGroupTextFmt2", "{0} ({1} / {2})"), FText::FromName(Name), FText::AsNumber(NumFilteredChildren), FText::AsNumber(NumChildren));
+		}
+	}
+	else
+	{
+		//const bool bIsStatTracked = FProfilerManager::Get()->IsStatTracked(GroupOrStatNode->GetStatID());
+		//if (bIsStatTracked)
+		//{
+		//	Text = FText::Format(LOCTEXT("StatsNodeTrackedTextFmt", "{0}*"), FText::FromName(Name));
+		//}
+		//else
+		{
+			Text = FText::FromName(Name);
+		}
+	}
+
+	return Text;
+}
+
+////////////////////////////////////////////////////////////////////////////////////////////////////
+
+#undef LOCTEXT_NAMESPACE
diff --git a/Engine/Source/Developer/TraceInsights/Private/Insights/ViewModels/StatsNode.h b/Engine/Source/Developer/TraceInsights/Private/Insights/ViewModels/StatsNode.h
new file mode 100644
--- /dev/null
+++ b/Engine/Source/Developer/TraceInsights/Private/Insights/ViewModels/StatsNode.h
@@ -0,0 +1,292 @@
+// Copyright 1998-2019 Epic Games, Inc. All Rights Reserved.
+
+#pragma once
+
+#include "CoreMinimal.h"
+
+class FStatsNode;
+
+/** Type definition for shared pointers to instances of FStatsNode. */
+typedef TSharedPtr<class FStatsNode> FStatsNodePtr;
+
+/** Type definition for shared references to instances of FStatsNode. */
+typedef TSharedRef<class FStatsNode> FStatsNodeRef;
+
+/** Type definition for shared references to const instances of FStatsNode. */
+typedef TSharedRef<const class FStatsNode> FStatsNodeRefConst;
+
+/** Type definition for weak references to instances of FStatsNode. */
+typedef TWeakPtr<class FStatsNode> FStatsNodeWeak;
+
+enum class EStatsNodeType
+{
+	/** The StatsNode is a floating number stats. */
+	Float,
+
+	/** The StatsNode is an integer number stats. */
+	Int64,
+
+	/** The StatsNode is a group node. */
+	Group,
+
+	/** Invalid enum type, may be used as a number of enumerations. */
+	InvalidOrMax,
+};
+
+template<typename Type>
+struct TAggregatedStats
+{
+	uint64 Count; /** Number of values. */
+
+	Type Sum; /** Sum of all values. */
+	Type Min; /** Min value. */
+	Type Max; /** Max value. */
+	Type Average; /** Average value. */
+	Type Median; /** Median value. */
+	Type LowerQuartile; /** Lower Quartile value. */
+	Type UpperQuartile; /** Upper Quartile value. */
+
+	TAggregatedStats()
+		: Count(0)
+		, Sum(0)
+		, Min(0)
+		, Max(0)
+		, Average(0)
+		, Median(0)
+		, LowerQuartile(0)
+		, UpperQuartile(0)
+	{
+	}
+
+	void Reset()
+	{
+		Count = 0;
+		Sum = 0;
+		Min = 0;
+		Max = 0;
+		Average = 0;
+		Median = 0;
+		LowerQuartile = 0;
+		UpperQuartile = 0;
+	}
+};
+
+typedef TAggregatedStats<double> FAggregatedStats;
+typedef TAggregatedStats<int64> FAggregatedIntegerStats;
+
+/**
+ * Class used to store information about a timer node (used in timers' tree view).
+ */
+class FStatsNode : public TSharedFromThis<FStatsNode>
+{
+public:
+	static const uint64 InvalidId = -1;
+
+public:
+	/** Initialization constructor for the timer node. */
+	FStatsNode(uint64 InId, const FName InName, const FName InMetaGroupName, EStatsNodeType InType)
+		: Id(InId)
+		, Name(InName)
+		, MetaGroupName(InMetaGroupName)
+		, Type(InType)
+		, bForceExpandGroupNode(false)
+	{
+		ResetAggregatedStats();
+	}
+
+	/** Initialization constructor for the group node. */
+	FStatsNode(const FName InGroupName)
+		: Id(0)
+		, Name(InGroupName)
+		, Type(EStatsNodeType::Group)
+		, bForceExpandGroupNode(false)
+	{
+		ResetAggregatedStats();
+	}
+
+	/**
+	 * @return an Id of this timer, valid only for timer nodes.
+	 */
+	const uint64 GetId() const
+	{
+		return Id;
+	}
+
+	/**
+	 * @return a name of this node, group or timer.
+	 */
+	const FName& GetName() const
+	{
+		return Name;
+	}
+
+	/**
+	 * @return a name of this node, group or timer + addditional info, to display in Stats tree view.
+	*/
+	const FText GetNameEx() const;
+
+	/**
+	 * @return a name of the group that this timer node belongs to, taken from the metadata.
+	 */
+	const FName& GetMetaGroupName() const
+	{
+		return MetaGroupName;
+	}
+
+	/**
+	 * @return a type of this timer or EStatsNodeType::Group for group nodes.
+	 */
+	const EStatsNodeType& GetType() const
+	{
+		return Type;
+	}
+
+	/**
+	 * @return true, if this node is a group node.
+	 */
+	bool IsGroup() const
+	{
+		return Type == EStatsNodeType::Group;
+	}
+
+	/**
+	 * @return the aggregated stats of this counter (if counter is a "float number" type).
+	 */
+	const FAggregatedStats& GetAggregatedStats() const
+	{
+		return AggregatedStats;
+	}
+
+	void ResetAggregatedStats();
+
+	void SetAggregatedStats(FAggregatedStats& AggregatedStats);
+
+	/**
+	 * @return the aggregated stats of this counter  (if counter is an "integer number" type).
+	 */
+	const FAggregatedIntegerStats& GetAggregatedIntegerStats() const
+	{
+		return AggregatedIntegerStats;
+	}
+
+	void ResetAggregatedIntegerStats();
+
+	void SetAggregatedIntegerStats(FAggregatedIntegerStats& AggregatedIntegerStats);
+	
+	const FText FormatValue(double Value) const;
+	const FText FormatValue(int64 Value) const;
+
+	const FText FormatAggregatedStatsValue(double ValueDbl, int64 ValueInt) const;
+
+	const FText GetTextForAggregatedStatsSum() const;
+	const FText GetTextForAggregatedStatsMin() const;
+	const FText GetTextForAggregatedStatsMax() const;
+	const FText GetTextForAggregatedStatsAverage() const;
+	const FText GetTextForAggregatedStatsMedian() const;
+	const FText GetTextForAggregatedStatsLowerQuartile() const;
+	const FText GetTextForAggregatedStatsUpperQuartile() const;
+
+	/**
+	 * @return a const reference to the child nodes of this group.
+	 */
+	FORCEINLINE_DEBUGGABLE const TArray<FStatsNodePtr>& GetChildren() const
+	{
+		return Children;
+	}
+
+	/**
+	 * @return a const reference to the child nodes that should be visible to the UI based on filtering.
+	 */
+	FORCEINLINE_DEBUGGABLE const TArray<FStatsNodePtr>& GetFilteredChildren() const
+	{
+		return FilteredChildren;
+	}
+
+	/**
+	 * @return a weak reference to the group of this timer node, may be invalid.
+	 */
+	FStatsNodeWeak GetGroupPtr() const
+	{
+		return GroupPtr;
+	}
+
+	/**
+	 * @return a name of the fake group that this timer node belongs to.
+	 */
+	const FName& GetGroupName() const
+	{
+		return GroupPtr.Pin()->Name;
+	}
+
+	/**
+	 * @return a name of the fake group that this timer node belongs to.
+	 */
+	//const FName& GetGroupNameSafe() const
+	//{
+	//	return GroupPtr.IsValid() ? GroupPtr.Pin()->Name : NAME_None;
+	//}
+
+	bool IsFiltered() const
+	{
+		return false; // TODO
+	}
+
+public:
+	/** Sorts children using the specified class instance. */
+	template<typename TSortingClass>
+	void SortChildren(const TSortingClass& Instance)
+	{
+		Children.Sort(Instance);
+	}
+
+	/** Adds specified child to the children and sets group for it. */
+	FORCEINLINE_DEBUGGABLE void AddChildAndSetGroupPtr(const FStatsNodePtr& ChildPtr)
+	{
+		ChildPtr->GroupPtr = AsShared();
+		Children.Add(ChildPtr);
+	}
+
+	/** Adds specified child to the filtered children. */
+	FORCEINLINE_DEBUGGABLE void AddFilteredChild(const FStatsNodePtr& ChildPtr)
+	{
+		FilteredChildren.Add(ChildPtr);
+	}
+
+	/** Clears filtered children. */
+	void ClearFilteredChildren()
+	{
+		FilteredChildren.Reset();
+	}
+
+protected:
+	/** The Id of this timer or group. */
+	const uint64 Id;
+
+	/** The name of this timer or group. */
+	const FName Name;
+
+	/** The name of the group that this timer belongs to, based on the timer's metadata; only valid for timer nodes. */
+	const FName MetaGroupName;
+
+	/** Holds the type of this timer; for the group, this is Group. */
+	const EStatsNodeType Type;
+
+	/** Aggregated stats (double). */
+	FAggregatedStats AggregatedStats;
+
+	/** Aggregated stats (int64). */
+	FAggregatedIntegerStats AggregatedIntegerStats;
+
+	/** Children of this node. */
+	TArray<FStatsNodePtr> Children;
+
+	/** Filtered children of this node. */
+	TArray<FStatsNodePtr> FilteredChildren;
+
+	/** A weak pointer to the group/parent of this node. */
+	FStatsNodeWeak GroupPtr;
+
+public:
+	/** Whether this group node should be expanded when the text filtering is enabled. */
+	bool bForceExpandGroupNode;
+};
diff --git a/Engine/Source/Developer/TraceInsights/Private/Insights/ViewModels/StatsNodeHelper.cpp b/Engine/Source/Developer/TraceInsights/Private/Insights/ViewModels/StatsNodeHelper.cpp
new file mode 100644
--- /dev/null
+++ b/Engine/Source/Developer/TraceInsights/Private/Insights/ViewModels/StatsNodeHelper.cpp
@@ -0,0 +1,123 @@
+// Copyright 1998-2019 Epic Games, Inc. All Rights Reserved.
+
+#include "StatsNodeHelper.h"
+#include "EditorStyleSet.h"
+
+#define LOCTEXT_NAMESPACE "StatsNode"
+
+////////////////////////////////////////////////////////////////////////////////////////////////////
+// StatsNode Type Helper
+////////////////////////////////////////////////////////////////////////////////////////////////////
+
+FText StatsNodeTypeHelper::ToName(const EStatsNodeType NodeType)
+{
+	static_assert(static_cast<int>(EStatsNodeType::InvalidOrMax) == 3, "Not all cases are handled in switch below!?");
+	switch (NodeType)
+	{
+		case EStatsNodeType::Float:		return LOCTEXT("Stats_Name_Float", "Float");
+		case EStatsNodeType::Int64:		return LOCTEXT("Stats_Name_Int64", "Integer");
+		case EStatsNodeType::Group:		return LOCTEXT("Stats_Name_Group", "Group");
+		default:						return LOCTEXT("InvalidOrMax", "InvalidOrMax");
+	}
+}
+
+////////////////////////////////////////////////////////////////////////////////////////////////////
+
+FText StatsNodeTypeHelper::ToDescription(const EStatsNodeType NodeType)
+{
+	static_assert(static_cast<int>(EStatsNodeType::InvalidOrMax) == 3, "Not all cases are handled in switch below!?");
+	switch (NodeType)
+	{
+		case EStatsNodeType::Float:		return LOCTEXT("Stats_Desc_Float", "Float number stats");
+		case EStatsNodeType::Int64:		return LOCTEXT("Stats_Desc_Int64", "Integer number stats");
+		case EStatsNodeType::Group:		return LOCTEXT("Stats_Desc_Group", "Group stats node");
+		default:						return LOCTEXT("InvalidOrMax", "InvalidOrMax");
+	}
+}
+
+////////////////////////////////////////////////////////////////////////////////////////////////////
+
+FName StatsNodeTypeHelper::ToBrushName(const EStatsNodeType NodeType)
+{
+	static_assert(static_cast<int>(EStatsNodeType::InvalidOrMax) == 3, "Not all cases are handled in switch below!?");
+	switch (NodeType)
+	{
+		case EStatsNodeType::Float:		return TEXT("Profiler.FiltersAndPresets.StatTypeIcon"); //TODO: "Icons.StatsType.Float"
+		case EStatsNodeType::Int64:		return TEXT("Profiler.FiltersAndPresets.StatTypeIcon"); //TODO: "Icons.StatsType.Int64"
+		case EStatsNodeType::Group:		return TEXT("Profiler.Misc.GenericGroup"); //TODO: "Icons.GenericGroup"
+		default:						return NAME_None;
+	}
+}
+
+////////////////////////////////////////////////////////////////////////////////////////////////////
+
+const FSlateBrush* StatsNodeTypeHelper::GetIconForGroup()
+{
+	return FEditorStyle::GetBrush(TEXT("Profiler.Misc.GenericGroup")); //TODO: FInsightsStyle::GetBrush(TEXT("Icons.GenericGroup"));
+}
+
+////////////////////////////////////////////////////////////////////////////////////////////////////
+
+const FSlateBrush* StatsNodeTypeHelper::GetIconForStatsNodeType(const EStatsNodeType NodeType)
+{
+	static_assert(static_cast<int>(EStatsNodeType::InvalidOrMax) == 3, "Not all cases are handled in switch below!?");
+	switch (NodeType)
+	{
+		case EStatsNodeType::Float:		return FEditorStyle::GetBrush(TEXT("Profiler.Type.NumberFloat")); //TODO: FInsightsStyle::GetBrush(TEXT("Icons.StatsType.Float"));
+		case EStatsNodeType::Int64:		return FEditorStyle::GetBrush(TEXT("Profiler.Type.NumberInt")); //TODO: FInsightsStyle::GetBrush(TEXT("Icons.StatsType.Int64"));
+		case EStatsNodeType::Group:		return FEditorStyle::GetBrush(TEXT("Profiler.Misc.GenericGroup")); //TODO: FInsightsStyle::GetBrush(TEXT("Icons.GenericGroup"));
+		default:						return nullptr;
+	}
+}
+
+////////////////////////////////////////////////////////////////////////////////////////////////////
+// StatsNode Grouping Helper
+////////////////////////////////////////////////////////////////////////////////////////////////////
+
+FText StatsNodeGroupingHelper::ToName(const EStatsGroupingMode GroupingMode)
+{
+	static_assert(static_cast<int>(EStatsGroupingMode::InvalidOrMax) == 4, "Not all cases are handled in switch below!?");
+	switch (GroupingMode)
+	{
+		case EStatsGroupingMode::Flat:				return LOCTEXT("Grouping_Name_Flat",			"Flat");
+		case EStatsGroupingMode::ByName:			return LOCTEXT("Grouping_Name_ByName",			"Stats Name");
+		case EStatsGroupingMode::ByMetaGroupName:	return LOCTEXT("Grouping_Name_MetaGroupName",	"Meta Group Name");
+		case EStatsGroupingMode::ByType:			return LOCTEXT("Grouping_Name_Type",			"Stats Type");
+		default:									return LOCTEXT("InvalidOrMax", "InvalidOrMax");
+	}
+}
+
+////////////////////////////////////////////////////////////////////////////////////////////////////
+
+FText StatsNodeGroupingHelper::ToDescription(const EStatsGroupingMode GroupingMode)
+{
+	static_assert(static_cast<int>(EStatsGroupingMode::InvalidOrMax) == 4, "Not all cases are handled in switch below!?");
+	switch (GroupingMode)
+	{
+		case EStatsGroupingMode::Flat:				return LOCTEXT("Grouping_Desc_Flat",			"Creates a single group. Includes all stats.");
+		case EStatsGroupingMode::ByName:			return LOCTEXT("Grouping_Desc_ByName",			"Creates one group for one letter.");
+		case EStatsGroupingMode::ByMetaGroupName:	return LOCTEXT("Grouping_Desc_MetaGroupName",	"Creates groups based on metadata group names of stats.");
+		case EStatsGroupingMode::ByType:			return LOCTEXT("Grouping_Desc_Type",			"Creates one group for each stats type.");
+		default:									return LOCTEXT("InvalidOrMax", "InvalidOrMax");
+	}
+}
+
+////////////////////////////////////////////////////////////////////////////////////////////////////
+
+FName StatsNodeGroupingHelper::ToBrushName(const EStatsGroupingMode GroupingMode)
+{
+	static_assert(static_cast<int>(EStatsGroupingMode::InvalidOrMax) == 4, "Not all cases are handled in switch below!?");
+	switch (GroupingMode)
+	{
+		//TODO: "UnrealProfiler.*Icon"
+		case EStatsGroupingMode::Flat:				return TEXT("Profiler.FiltersAndPresets.GroupNameIcon"); //TODO: "Icons.Grouping.Flat"
+		case EStatsGroupingMode::ByName:			return TEXT("Profiler.FiltersAndPresets.Group1NameIcon"); //TODO: "Icons.Grouping.ByName"
+		case EStatsGroupingMode::ByMetaGroupName:	return TEXT("Profiler.FiltersAndPresets.StatNameIcon"); //TODO
+		case EStatsGroupingMode::ByType:			return TEXT("Profiler.FiltersAndPresets.StatTypeIcon"); //TODO
+		default:									return NAME_None;
+	}
+}
+
+////////////////////////////////////////////////////////////////////////////////////////////////////
+
+#undef LOCTEXT_NAMESPACE
diff --git a/Engine/Source/Developer/TraceInsights/Private/Insights/ViewModels/StatsNodeHelper.h b/Engine/Source/Developer/TraceInsights/Private/Insights/ViewModels/StatsNodeHelper.h
new file mode 100644
--- /dev/null
+++ b/Engine/Source/Developer/TraceInsights/Private/Insights/ViewModels/StatsNodeHelper.h
@@ -0,0 +1,259 @@
+// Copyright 1998-2019 Epic Games, Inc. All Rights Reserved.
+
+#pragma once
+
+#include "CoreMinimal.h"
+#include "Styling/SlateBrush.h"
+
+// Insights
+#include "Insights/ViewModels/StatsNode.h"
+
+////////////////////////////////////////////////////////////////////////////////////////////////////
+
+/** Enumerates types of grouping or sorting for the stats nodes. */
+enum class EStatsGroupingMode
+{
+	/** Creates a single group for all timers. */
+	Flat,
+
+	/** Creates one group for one letter. */
+	ByName,
+
+	/** Creates groups based on stats metadata group names. */
+	ByMetaGroupName,
+
+	/** Creates one group for each stats type. */
+	ByType,
+
+	/** Invalid enum type, may be used as a number of enumerations. */
+	InvalidOrMax,
+};
+
+/** Type definition for shared pointers to instances of EStatsGroupingMode. */
+typedef TSharedPtr<EStatsGroupingMode> EStatsGroupingModePtr;
+
+////////////////////////////////////////////////////////////////////////////////////////////////////
+
+/** Helper struct that contains helper functions and classes for EStatsNodeType enum. */
+struct StatsNodeTypeHelper
+{
+	/**
+	 * @param Type - The value to get the text for.
+	 *
+	 * @return text representation of the specified EStatsNodeType value.
+	 */
+	static FText ToName(const EStatsNodeType Type);
+
+	/**
+	 * @param Type - The value to get the text for.
+	 *
+	 * @return text representation with more detailed explanation of the specified EStatsNodeType value.
+	 */
+	static FText ToDescription(const EStatsNodeType Type);
+
+	/**
+	 * @param Type - The value to get the brush name for.
+	 *
+	 * @return brush name of the specified EStatsNodeType value.
+	 */
+	static FName ToBrushName(const EStatsNodeType Type);
+
+	static const FSlateBrush* GetIconForGroup();
+	static const FSlateBrush* GetIconForStatsNodeType(const EStatsNodeType Type);
+};
+
+////////////////////////////////////////////////////////////////////////////////////////////////////
+
+/** Helper struct that contains grouping static functions and classes. */
+struct StatsNodeGroupingHelper
+{
+	/**
+	 * @param StatsGroupingMode - The value to get the text for.
+	 *
+	 * @return text representation of the specified EStatsGroupingMode value.
+	 */
+	static FText ToName(const EStatsGroupingMode StatsGroupingMode);
+
+	/**
+	 * @param StatsGroupingMode - The value to get the text for.
+	 *
+	 * @return text representation with more detailed explanation of the specified EStatsGroupingMode value.
+	 */
+	static FText ToDescription(const EStatsGroupingMode StatsGroupingMode);
+
+	/**
+	 * @param StatsGroupingMode - The value to get the brush name for.
+	 *
+	 * @return brush name of the specified EStatsGroupingMode value.
+	 */
+	static FName ToBrushName(const EStatsGroupingMode StatsGroupingMode);
+};
+
+////////////////////////////////////////////////////////////////////////////////////////////////////
+
+/** Helper struct that contains sorting static functions and classes. */
+struct StatsNodeSortingHelper
+{
+	//////////////////////////////////////////////////
+	// Sorting by stats name.
+
+	struct ByNameAscending
+	{
+		FORCEINLINE_DEBUGGABLE bool operator()(const FStatsNodePtr& A, const FStatsNodePtr& B) const
+		{
+			return A->GetName().LexicalLess(B->GetName());
+		}
+	};
+
+	struct ByNameDescending
+	{
+		FORCEINLINE_DEBUGGABLE bool operator()(const FStatsNodePtr& A, const FStatsNodePtr& B) const
+		{
+			return B->GetName().LexicalLess(A->GetName());
+		}
+	};
+
+	//////////////////////////////////////////////////
+	// Sorting by stats meta group name.
+	// If meta group names are the same then sort by name.
+
+	struct ByMetaGroupNameAscending
+	{
+		FORCEINLINE_DEBUGGABLE bool operator()(const FStatsNodePtr& A, const FStatsNodePtr& B) const
+		{
+			if (A->GetMetaGroupName() == B->GetMetaGroupName())
+			{
+				// Sort by stats name (ascending).
+				return A->GetName().LexicalLess(B->GetName());
+			}
+			else
+			{
+				return A->GetMetaGroupName().LexicalLess(B->GetMetaGroupName());
+			}
+		}
+	};
+
+	struct ByMetaGroupNameDescending
+	{
+		FORCEINLINE_DEBUGGABLE bool operator()(const FStatsNodePtr& A, const FStatsNodePtr& B) const
+		{
+			if (A->GetMetaGroupName() == B->GetMetaGroupName())
+			{
+				// Sort by stats name (ascending).
+				return A->GetName().LexicalLess(B->GetName());
+			}
+			else
+			{
+				return B->GetMetaGroupName().LexicalLess(A->GetMetaGroupName());
+			}
+		}
+	};
+
+	//////////////////////////////////////////////////
+	// Sorting by stats type.
+	// If types are the same then sort by name.
+
+	struct ByTypeAscending
+	{
+		FORCEINLINE_DEBUGGABLE bool operator()(const FStatsNodePtr& A, const FStatsNodePtr& B) const
+		{
+			const EStatsNodeType& TypeA = A->GetType();
+			const EStatsNodeType& TypeB = B->GetType();
+
+			if (TypeA == TypeB)
+			{
+				// Sort by stats name (ascending).
+				return A->GetName().LexicalLess(B->GetName());
+			}
+			else
+			{
+				return TypeA < TypeB;
+			}
+		}
+	};
+
+	struct ByTypeDescending
+	{
+		FORCEINLINE_DEBUGGABLE bool operator()(const FStatsNodePtr& A, const FStatsNodePtr& B) const
+		{
+			const EStatsNodeType& TypeA = A->GetType();
+			const EStatsNodeType& TypeB = B->GetType();
+
+			if (TypeA == TypeB)
+			{
+				// Sort by stats name (ascending).
+				return A->GetName().LexicalLess(B->GetName());
+			}
+			else
+			{
+				return TypeA > TypeB;
+			}
+		}
+	};
+
+	//////////////////////////////////////////////////
+	// Sorting by an aggregated stats value.
+	// If aggregated stats values are the same then sort by name.
+
+	#define SORT_BY_STATS_ASCENDING(Type, SortName, AggregatedStatsMember) \
+		struct SortName \
+		{ \
+			FORCEINLINE_DEBUGGABLE bool operator()(const FStatsNodePtr& A, const FStatsNodePtr& B) const \
+			{ \
+				const Type ValueA = A->GetAggregatedStats().AggregatedStatsMember; \
+				const Type ValueB = B->GetAggregatedStats().AggregatedStatsMember; \
+				\
+				if (ValueA == ValueB) \
+				{ \
+					/* Sort by name (ascending). */ \
+					return A->GetName().LexicalLess(B->GetName()); \
+				} \
+				else \
+				{ \
+					return ValueA < ValueB; \
+				} \
+			} \
+		};
+
+	#define SORT_BY_STATS_DESCENDING(Type, SortName, AggregatedStatsMember) \
+		struct SortName \
+		{ \
+			FORCEINLINE_DEBUGGABLE bool operator()(const FStatsNodePtr& A, const FStatsNodePtr& B) const \
+			{ \
+				const Type ValueA = A->GetAggregatedStats().AggregatedStatsMember; \
+				const Type ValueB = B->GetAggregatedStats().AggregatedStatsMember; \
+				\
+				if (ValueA == ValueB) \
+				{ \
+					/* Sort by name (ascending). */ \
+					return A->GetName().LexicalLess(B->GetName()); \
+				} \
+				else \
+				{ \
+					return ValueA > ValueB; \
+				} \
+			} \
+		};
+
+	#define SORT_BY_STATS(Type, SortName, AggregatedStatsMember) \
+		SORT_BY_STATS_ASCENDING(Type, SortName##Ascending, AggregatedStatsMember) \
+		SORT_BY_STATS_DESCENDING(Type, SortName##Descending, AggregatedStatsMember)
+
+	SORT_BY_STATS(uint64, ByCount, Count);
+	SORT_BY_STATS(double, BySum, Sum);
+	SORT_BY_STATS(double, ByMin, Min);
+	SORT_BY_STATS(double, ByMax, Max);
+	SORT_BY_STATS(double, ByAverage, Average);
+	SORT_BY_STATS(double, ByMedian, Median);
+	SORT_BY_STATS(double, ByLowerQuartile, LowerQuartile);
+	SORT_BY_STATS(double, ByUpperQuartile, UpperQuartile);
+
+
+	#undef SORT_BY_STATS_ASCENDING
+	#undef SORT_BY_STATS_DESCENDING
+	#undef SORT_BY_STATS
+
+	//////////////////////////////////////////////////
+};
+
+////////////////////////////////////////////////////////////////////////////////////////////////////
diff --git a/Engine/Source/Developer/TraceInsights/Private/Insights/ViewModels/StatsViewColumn.h b/Engine/Source/Developer/TraceInsights/Private/Insights/ViewModels/StatsViewColumn.h
new file mode 100644
--- /dev/null
+++ b/Engine/Source/Developer/TraceInsights/Private/Insights/ViewModels/StatsViewColumn.h
@@ -0,0 +1,124 @@
+// Copyright 1998-2019 Epic Games, Inc. All Rights Reserved.
+
+#pragma once
+
+#include "CoreMinimal.h"
+#include "Types/SlateEnums.h"
+#include "Misc/EnumClassFlags.h"
+
+class FStatsNode; // TODO: IInsightsTreeNode; generic formatter; see also TimersViewColumn.h
+
+////////////////////////////////////////////////////////////////////////////////////////////////////
+
+enum class EStatsViewColumnFlags : uint32
+{
+	None = 0,
+
+	CanBeHidden = (1 << 0),
+	CanBeSorted = (1 << 1),
+	CanBeFiltered = (1 << 2),
+};
+ENUM_CLASS_FLAGS(EStatsViewColumnFlags);
+
+////////////////////////////////////////////////////////////////////////////////////////////////////
+
+/** Holds information about a column in the Stats Counters view widget. */
+class FStatsViewColumn
+{
+	friend struct FStatsViewColumnFactory;
+
+public:
+	typedef TFunction<FText(const FStatsViewColumn& Column, const FStatsNode& StatsNode)> FGetFormattedValueFn;
+
+public:
+	/** Whether this column can be hidden. */
+	bool bCanBeHidden() const { return EnumHasAnyFlags(Flags, EStatsViewColumnFlags::CanBeHidden); }
+
+	/** Whether this column cab be used for sorting. */
+	bool bCanBeSorted() const { return EnumHasAnyFlags(Flags, EStatsViewColumnFlags::CanBeSorted); }
+
+	/** Where this column can be used to filtering displayed results. */
+	bool bCanBeFiltered() const { return EnumHasAnyFlags(Flags, EStatsViewColumnFlags::CanBeFiltered); }
+
+	/** If MinColumnWidth == MaxColumnWidth, this column has fixed width and cannot be resized. */
+	bool bIsFixedColumnWidth() const { return MinColumnWidth == MaxColumnWidth; }
+
+	FText GetFormattedValue(const FStatsNode& StatsNode) const
+	{
+		return GetFormattedValueFn(*this, StatsNode);
+	}
+
+protected:
+	/** No default constructor. */
+	FStatsViewColumn() = delete;
+
+	/** Initialization constructor, only used in FStatsViewColumnFactory. */
+	FStatsViewColumn
+	(
+		int32 InOrder,
+		const FName InId,
+		const FName InSearchId,
+		FText InShortName,
+		FText InTitleName,
+		FText InDescription,
+		bool bInIsVisible,
+		const EStatsViewColumnFlags InFlags,
+		const EHorizontalAlignment InHorizontalAlignment,
+		const float InInitialColumnWidth,
+		const float InMinColumnWidth,
+		const float InMaxColumnWidth,
+		const FGetFormattedValueFn InGetFormattedValueFn
+	)
+		: Order(InOrder)
+		, Id(InId)
+		, SearchId(InSearchId)
+		, ShortName(MoveTemp(InShortName))
+		, TitleName(MoveTemp(InTitleName))
+		, Description(MoveTemp(InDescription))
+		, bIsVisible(bInIsVisible)
+		, Flags(InFlags)
+		, HorizontalAlignment(InHorizontalAlignment)
+		, InitialColumnWidth(InInitialColumnWidth)
+		, MinColumnWidth(InMinColumnWidth)
+		, MaxColumnWidth(InMaxColumnWidth)
+		, GetFormattedValueFn(InGetFormattedValueFn)
+	{
+	}
+
+public:
+	/** Order value, to sort columns in the tree view. */
+	int32 Order;
+
+	/** Name of the column, name of the property. */
+	FName Id;
+
+	/** Name of the column used by the searching system. */
+	FName SearchId;
+
+	/** Short name of the column, displayed in the column header. */
+	FText ShortName;
+
+	/** Title name of the column, displayed as title in the column tooltip. */
+	FText TitleName;
+
+	/** Long name of the column, displayed in the column tooltip. */
+	FText Description;
+
+	/** Is this column visible? */
+	bool bIsVisible;
+
+	/** On/off switches. */
+	EStatsViewColumnFlags Flags;
+
+	/** Horizontal alignment of the content in this column. */
+	EHorizontalAlignment HorizontalAlignment;
+
+	float InitialColumnWidth; /**< Initial column width. */
+	float MinColumnWidth; /**< Minimum column width. */
+	float MaxColumnWidth; /**< Maximum column width. */
+
+	/** Custom function used to format (as an FText) the value of the stats counter node to be displayed by this column. */
+	FGetFormattedValueFn GetFormattedValueFn;
+};
+
+////////////////////////////////////////////////////////////////////////////////////////////////////
diff --git a/Engine/Source/Developer/TraceInsights/Private/Insights/ViewModels/StatsViewColumnFactory.cpp b/Engine/Source/Developer/TraceInsights/Private/Insights/ViewModels/StatsViewColumnFactory.cpp
new file mode 100644
--- /dev/null
+++ b/Engine/Source/Developer/TraceInsights/Private/Insights/ViewModels/StatsViewColumnFactory.cpp
@@ -0,0 +1,287 @@
+// Copyright 1998-2019 Epic Games, Inc. All Rights Reserved.
+
+#include "StatsViewColumnFactory.h"
+
+// Insights
+#include "Insights/Common/TimeUtils.h"
+#include "Insights/ViewModels/StatsNodeHelper.h"
+#include "Insights/ViewModels/StatsViewColumn.h"
+
+#define LOCTEXT_NAMESPACE "SStatsView"
+
+////////////////////////////////////////////////////////////////////////////////////////////////////
+
+// Column identifiers
+const FName FStatsViewColumns::NameColumnID(TEXT("Name"));
+const FName FStatsViewColumns::MetaGroupNameColumnID(TEXT("MetaGroupName"));
+const FName FStatsViewColumns::TypeColumnID(TEXT("Type"));
+const FName FStatsViewColumns::CountColumnID(TEXT("Count"));
+const FName FStatsViewColumns::SumColumnID(TEXT("Sum"));
+const FName FStatsViewColumns::MaxColumnID(TEXT("Max"));
+const FName FStatsViewColumns::UpperQuartileColumnID(TEXT("UpperQuartile"));
+const FName FStatsViewColumns::AverageColumnID(TEXT("Average"));
+const FName FStatsViewColumns::MedianColumnID(TEXT("Median"));
+const FName FStatsViewColumns::LowerQuartileColumnID(TEXT("LowerQuartile"));
+const FName FStatsViewColumns::MinColumnID(TEXT("Min"));
+
+////////////////////////////////////////////////////////////////////////////////////////////////////
+
+FStatsViewColumnFactory::FStatsViewColumnFactory()
+{
+	uint32 Index = 0;
+
+	// FStatsViewColumn
+	// (
+	//     Order,
+	//     Id,
+	//     SearchId,
+	//     ShortName,
+	//     TitleName,
+	//     Description,
+	//     bIsVisible,
+	//     Flags,
+	//     HorizontalAlignment,
+	//     InitialWidth, MinWidth, MaxWidth,
+	//     GetFormattedValueFn
+	// )
+
+	//////////////////////////////////////////////////
+	// Name
+
+	const EStatsViewColumnFlags NameColumnFlags = EStatsViewColumnFlags::CanBeSorted |
+												  EStatsViewColumnFlags::CanBeFiltered;
+
+	Collection.Add(new FStatsViewColumn
+	(
+		Index++,
+		FStatsViewColumns::NameColumnID,
+		TEXT("name"),
+		LOCTEXT("StatsNameColumnName", "Name"),
+		LOCTEXT("StatsNameColumnTitle", "Stats or Group Name"),
+		LOCTEXT("StatsNameColumnDesc", "Name of the stats or group"),
+		true,
+		NameColumnFlags,
+		HAlign_Left,
+		206.0f, 10.0f, FLT_MAX,
+		[](const FStatsViewColumn& Column, const FStatsNode& StatsNode) -> FText
+		{ return FText::FromName(StatsNode.GetName()); }
+	));
+
+	//////////////////////////////////////////////////
+	// Meta Group Name
+
+	const EStatsViewColumnFlags MetaGroupNameColumnFlags = EStatsViewColumnFlags::CanBeHidden |
+														   EStatsViewColumnFlags::CanBeSorted |
+														   EStatsViewColumnFlags::CanBeFiltered;
+
+	Collection.Add(new FStatsViewColumn
+	(
+		Index++,
+		FStatsViewColumns::MetaGroupNameColumnID,
+		TEXT("metagroupname"),
+		LOCTEXT("StatsMetaGroupNameColumnName", "Meta Group"),
+		LOCTEXT("StatsMetaGroupNameColumnTitle", "Meta Group Name"),
+		LOCTEXT("StatsMetaGroupNameColumnDesc", "Name of the meta group"),
+		false,
+		MetaGroupNameColumnFlags,
+		HAlign_Left,
+		100.0f, 10.0f, FLT_MAX,
+		[](const FStatsViewColumn& Column, const FStatsNode& StatsNode) -> FText
+		{ return FText::FromName(StatsNode.GetMetaGroupName()); }
+	));
+
+	//////////////////////////////////////////////////
+	// Type
+
+	const EStatsViewColumnFlags TypeColumnFlags = EStatsViewColumnFlags::CanBeHidden |
+												  EStatsViewColumnFlags::CanBeSorted |
+												  EStatsViewColumnFlags::CanBeFiltered;
+
+	Collection.Add(new FStatsViewColumn
+	(
+		Index++,
+		FStatsViewColumns::TypeColumnID,
+		TEXT("type"),
+		LOCTEXT("StatsTypeColumnName", "Type"),
+		LOCTEXT("StatsTypeColumnTitle", "Type"),
+		LOCTEXT("StatsTypeColumnDesc", "Type of the stats or group"),
+		false,
+		TypeColumnFlags,
+		HAlign_Left,
+		60.0f, 10.0f, FLT_MAX,
+		[](const FStatsViewColumn& Column, const FStatsNode& StatsNode) -> FText
+		{ return StatsNodeTypeHelper::ToName(StatsNode.GetType()); }
+	));
+		
+	//////////////////////////////////////////////////
+
+	const EStatsViewColumnFlags AggregatedStatsColumnFlags = EStatsViewColumnFlags::CanBeHidden |
+															 EStatsViewColumnFlags::CanBeSorted |
+															 EStatsViewColumnFlags::CanBeFiltered;
+
+	//////////////////////////////////////////////////
+	// Count
+
+	Collection.Add(new FStatsViewColumn
+	(
+		Index++,
+		FStatsViewColumns::CountColumnID,
+		TEXT("count"),
+		LOCTEXT("CountName", "Count"),
+		LOCTEXT("CountTitle", "Count"),
+		LOCTEXT("CountDesc", "Number of selected values"),
+		true,
+		AggregatedStatsColumnFlags,
+		HAlign_Right,
+		60.0f, 0.0f, FLT_MAX,
+		[](const FStatsViewColumn& Column, const FStatsNode& StatsNode) -> FText
+		{ return FText::AsNumber(StatsNode.GetAggregatedStats().Count); }
+	));
+
+	//////////////////////////////////////////////////
+	// Aggregated stats
+
+	const float StatsInitialColumnWidth = 80.0f;
+
+	Collection.Add(new FStatsViewColumn
+	(
+		Index++,
+		FStatsViewColumns::SumColumnID,
+		TEXT("sum"),
+		LOCTEXT("SumName", "Sum"),
+		LOCTEXT("SumTitle", "Sum"),
+		LOCTEXT("SumDesc", "Sum of selected values"),
+		false,
+		AggregatedStatsColumnFlags,
+		HAlign_Right,
+		StatsInitialColumnWidth, 4.0f, FLT_MAX,
+		[](const FStatsViewColumn& Column, const FStatsNode& StatsNode) -> FText
+		{ return StatsNode.GetTextForAggregatedStatsSum(); }
+	));
+
+	Collection.Add(new FStatsViewColumn
+	(
+		Index++,
+		FStatsViewColumns::MaxColumnID,
+		TEXT("max"),
+		LOCTEXT("MaxName", "Max"),
+		LOCTEXT("MaxTitle", "Maximum"),
+		LOCTEXT("MaxDesc", "Maximum of selected values"),
+		true,
+		AggregatedStatsColumnFlags,
+		HAlign_Right,
+		StatsInitialColumnWidth, 4.0f, FLT_MAX,
+		[](const FStatsViewColumn& Column, const FStatsNode& StatsNode) -> FText
+		{ return StatsNode.GetTextForAggregatedStatsMax(); }
+	));
+
+	Collection.Add(new FStatsViewColumn
+	(
+		Index++,
+		FStatsViewColumns::UpperQuartileColumnID,
+		TEXT("upperquartile"),
+		LOCTEXT("UpperQuartileName", "Upper"),
+		LOCTEXT("UpperQuartileTitle", "Upper Quartile"),
+		LOCTEXT("UpperQuartileDesc", "Upper quartile (Q3; third quartile; 75th percentile) of selected values"),
+		false,
+		AggregatedStatsColumnFlags,
+		HAlign_Right,
+		StatsInitialColumnWidth, 4.0f, FLT_MAX,
+		[](const FStatsViewColumn& Column, const FStatsNode& StatsNode) -> FText
+		{ return StatsNode.GetTextForAggregatedStatsUpperQuartile(); }
+	));
+
+	Collection.Add(new FStatsViewColumn
+	(
+		Index++,
+		FStatsViewColumns::AverageColumnID,
+		TEXT("avg"),
+		LOCTEXT("AverageName", "Avg."),
+		LOCTEXT("AverageTitle", "Average"),
+		LOCTEXT("AverageDesc", "Average (arithmetic mean) of selected values"),
+		false,
+		AggregatedStatsColumnFlags,
+		HAlign_Right,
+		StatsInitialColumnWidth, 4.0f, FLT_MAX,
+		[](const FStatsViewColumn& Column, const FStatsNode& StatsNode) -> FText
+		{ return StatsNode.GetTextForAggregatedStatsAverage(); }
+	));
+
+	Collection.Add(new FStatsViewColumn
+	(
+		Index++,
+		FStatsViewColumns::MedianColumnID,
+		TEXT("med"),
+		LOCTEXT("MedianName", "Med."),
+		LOCTEXT("MedianTitle", "Median"),
+		LOCTEXT("MedianDesc", "Median (Q2; second quartile; 50th percentile) of selected values"),
+		false,
+		AggregatedStatsColumnFlags,
+		HAlign_Right,
+		StatsInitialColumnWidth, 4.0f, FLT_MAX,
+		[](const FStatsViewColumn& Column, const FStatsNode& StatsNode) -> FText
+		{ return StatsNode.GetTextForAggregatedStatsMedian(); }
+	));
+
+	Collection.Add(new FStatsViewColumn
+	(
+		Index++,
+		FStatsViewColumns::LowerQuartileColumnID,
+		TEXT("lowerquartile"),
+		LOCTEXT("LowerQuartileName", "Lower"),
+		LOCTEXT("LowerQuartileTitle", "Lower Quartile"),
+		LOCTEXT("LowerQuartileDesc", "Lower quartile (Q1; first quartile; 25th percentile) of selected values"),
+		false,
+		AggregatedStatsColumnFlags,
+		HAlign_Right,
+		StatsInitialColumnWidth, 4.0f, FLT_MAX,
+		[](const FStatsViewColumn& Column, const FStatsNode& StatsNode) -> FText
+		{ return StatsNode.GetTextForAggregatedStatsLowerQuartile(); }
+	));
+
+	Collection.Add(new FStatsViewColumn
+	(
+		Index++,
+		FStatsViewColumns::MinColumnID,
+		TEXT("min"),
+		LOCTEXT("MinName", "Min"),
+		LOCTEXT("MinTitle", "Minimum"),
+		LOCTEXT("MinDesc", "Minimum of selected values"),
+		true,
+		AggregatedStatsColumnFlags,
+		HAlign_Right,
+		StatsInitialColumnWidth, 4.0f, FLT_MAX,
+		[](const FStatsViewColumn& Column, const FStatsNode& StatsNode) -> FText
+		{ return StatsNode.GetTextForAggregatedStatsMin(); }
+	));
+
+	//////////////////////////////////////////////////
+
+	for (int32 MapIndex = 0; MapIndex < Collection.Num(); MapIndex++)
+	{
+		ColumnIdToPtrMapping.Add(Collection[MapIndex]->Id, Collection[MapIndex]);
+	}
+}
+
+////////////////////////////////////////////////////////////////////////////////////////////////////
+
+FStatsViewColumnFactory::~FStatsViewColumnFactory()
+{
+	for (const FStatsViewColumn* Column : Collection)
+	{
+		delete Column;
+	}
+	Collection.Reset();
+}
+
+////////////////////////////////////////////////////////////////////////////////////////////////////
+
+const FStatsViewColumnFactory& FStatsViewColumnFactory::Get()
+{
+	static FStatsViewColumnFactory Instance;
+	return Instance;
+}
+
+////////////////////////////////////////////////////////////////////////////////////////////////////
+
+#undef LOCTEXT_NAMESPACE
diff --git a/Engine/Source/Developer/TraceInsights/Private/Insights/ViewModels/StatsViewColumnFactory.h b/Engine/Source/Developer/TraceInsights/Private/Insights/ViewModels/StatsViewColumnFactory.h
new file mode 100644
--- /dev/null
+++ b/Engine/Source/Developer/TraceInsights/Private/Insights/ViewModels/StatsViewColumnFactory.h
@@ -0,0 +1,49 @@
+// Copyright 1998-2019 Epic Games, Inc. All Rights Reserved.
+
+#pragma once
+
+#include "CoreMinimal.h"
+
+#include "Insights/ViewModels/StatsViewColumn.h"
+
+////////////////////////////////////////////////////////////////////////////////////////////////////
+
+struct FStatsViewColumns
+{
+	// Column identifiers
+	static const FName NameColumnID;
+	static const FName MetaGroupNameColumnID;
+	static const FName TypeColumnID;
+	static const FName CountColumnID;
+	static const FName SumColumnID;
+	static const FName MaxColumnID;
+	static const FName UpperQuartileColumnID;
+	static const FName AverageColumnID;
+	static const FName MedianColumnID;
+	static const FName LowerQuartileColumnID;
+	static const FName MinColumnID;
+};
+
+////////////////////////////////////////////////////////////////////////////////////////////////////
+
+struct FStatsViewColumnFactory
+{
+private:
+	/** Default constructor. */
+	FStatsViewColumnFactory();
+
+	/** Destructor. */
+	~FStatsViewColumnFactory();
+
+public:
+	/** Contains basic information about columns used in the Stats Counters view widget. Names should be localized. */
+	TArray<FStatsViewColumn*> Collection;
+
+	/** Mapping between column IDs and FStatsViewColumn pointers. */
+	TMap<FName, const FStatsViewColumn*> ColumnIdToPtrMapping;
+
+	/** Singleton */
+	static const FStatsViewColumnFactory& Get();
+};
+
+////////////////////////////////////////////////////////////////////////////////////////////////////
diff --git a/Engine/Source/Developer/TraceInsights/Private/Insights/ViewModels/TimeRulerTrack.cpp b/Engine/Source/Developer/TraceInsights/Private/Insights/ViewModels/TimeRulerTrack.cpp
new file mode 100644
--- /dev/null
+++ b/Engine/Source/Developer/TraceInsights/Private/Insights/ViewModels/TimeRulerTrack.cpp
@@ -0,0 +1,232 @@
+// Copyright 1998-2019 Epic Games, Inc. All Rights Reserved.
+
+#include "TimeRulerTrack.h"
+
+#include "Brushes/SlateColorBrush.h"
+#include "Fonts/FontMeasure.h"
+#include "Fonts/SlateFontInfo.h"
+#include "Framework/Application/SlateApplication.h"
+#include "Rendering/DrawElements.h"
+
+// Insights
+#include "Insights/Common/PaintUtils.h"
+#include "Insights/Common/TimeUtils.h"
+#include "Insights/ViewModels/TimingTrackViewport.h"
+
+#define LOCTEXT_NAMESPACE "TimeRulerTrack"
+
+////////////////////////////////////////////////////////////////////////////////////////////////////
+
+FTimeRulerTrack::FTimeRulerTrack(uint64 InTrackId)
+	: FBaseTimingTrack(InTrackId)
+	, WhiteBrush(FCoreStyle::Get().GetBrush("WhiteBrush"))
+	, Font(FCoreStyle::GetDefaultFontStyle("Regular", 8))
+{
+}
+
+////////////////////////////////////////////////////////////////////////////////////////////////////
+
+FTimeRulerTrack::~FTimeRulerTrack()
+{
+}
+
+////////////////////////////////////////////////////////////////////////////////////////////////////
+
+void FTimeRulerTrack::Reset()
+{
+	FBaseTimingTrack::Reset();
+
+	constexpr float TimeRulerHeight = 22.0f;
+	SetHeight(TimeRulerHeight);
+}
+
+////////////////////////////////////////////////////////////////////////////////////////////////////
+
+void FTimeRulerTrack::UpdateHoveredState(float MouseX, float MouseY, const FTimingTrackViewport& Viewport)
+{
+	//constexpr float HeaderWidth = 100.0f;
+	//constexpr float HeaderHeight = 14.0f;
+
+	if (MouseY >= GetPosY() && MouseY < GetPosY() + GetHeight())
+	{
+		SetHoveredState(true);
+		//SetHeaderHoveredState(MouseX < HeaderWidth && MouseY < GetPosY() + HeaderHeight);
+	}
+	else
+	{
+		SetHoveredState(false);
+	}
+}
+
+////////////////////////////////////////////////////////////////////////////////////////////////////
+
+void FTimeRulerTrack::DrawBackground(FDrawContext& DrawContext, const FTimingTrackViewport& Viewport) const
+{
+	const float X0 = Viewport.TimeToSlateUnitsRounded(0.0);
+	const float X1 = Viewport.TimeToSlateUnitsRounded(Viewport.MaxValidTime);
+	const float W = FMath::CeilToFloat(Viewport.Width);
+	const float H = GetHeight();
+
+	const FLinearColor InvalidAreaColor(0.08f, 0.07f, 0.07f, 1.0f);
+	const FLinearColor ValidAreaColor(0.09f, 0.09f, 0.09f, 1.0f);
+
+	if (X0 >= W || X1 <= 0.0f)
+	{
+		// Draw invalid area (entire view).
+		DrawContext.DrawBox(0.0f, 0.0f, W, H, WhiteBrush, InvalidAreaColor);
+	}
+	else // X0 < W && X1 > 0
+	{
+		if (X0 > 0.0f)
+		{
+			// Draw invalid area (left).
+			DrawContext.DrawBox(0.0f, 0.0f, X0, H, WhiteBrush, InvalidAreaColor);
+		}
+
+		if (X1 < W)
+		{
+			// Draw invalid area (right).
+			DrawContext.DrawBox(X1, 0.0f, W - X1, H, WhiteBrush, InvalidAreaColor);
+		}
+
+		float ValidX0 = FMath::Max(X0, 0.0f);
+		float ValidX1 = FMath::Min(X1, W);
+
+		if (ValidX1 > ValidX0)
+		{
+			// Draw valid area.
+			DrawContext.DrawBox(ValidX0, 0.0f, ValidX1 - ValidX0, H, WhiteBrush, ValidAreaColor);
+		}
+	}
+
+	DrawContext.LayerId++;
+}
+
+////////////////////////////////////////////////////////////////////////////////////////////////////
+
+void FTimeRulerTrack::Draw(FDrawContext& DrawContext, const FTimingTrackViewport& Viewport, const FVector2D& MousePosition, const bool bIsSelecting, const double SelectionStartTime, const double SelectionEndTime) const
+{
+	const float MinorTickMark = 5.0f;
+	const float MajorTickMark = 20 * MinorTickMark;
+
+	const float MinorTickMarkHeight = 5.0f;
+	const float MajorTickMarkHeight = 9.0f;
+
+	const float TextY = GetPosY() + MajorTickMarkHeight;
+
+	double MinorTickMarkTime = Viewport.GetDurationForViewportDX(MinorTickMark);
+	double MajorTickMarkTime = Viewport.GetDurationForViewportDX(MajorTickMark);
+
+	double VX = Viewport.StartTime * Viewport.ScaleX;
+	double MinorN = FMath::FloorToDouble(VX / static_cast<double>(MinorTickMark));
+	double MajorN = FMath::FloorToDouble(VX / static_cast<double>(MajorTickMark));
+	float MinorOX = static_cast<float>(FMath::RoundToDouble(MinorN * static_cast<double>(MinorTickMark) - VX));
+	float MajorOX = static_cast<float>(FMath::RoundToDouble(MajorN * static_cast<double>(MajorTickMark) - VX));
+
+	// Draw the time ruler's background.
+	DrawBackground(DrawContext, Viewport);
+
+	// Draw the minor tick marks.
+	for (float X = MinorOX; X < Viewport.Width; X += MinorTickMark)
+	{
+		const bool bIsTenth = ((int32)(((X - MajorOX) / MinorTickMark) + 0.4f) % 2 == 0);
+		const float MinorTickH = bIsTenth ? MinorTickMarkHeight : MinorTickMarkHeight - 1.0f;
+		DrawContext.DrawBox(X, GetPosY(), 1.0f, MinorTickH, WhiteBrush,
+			bIsTenth ? FLinearColor(0.3f, 0.3f, 0.3f, 1.0f) : FLinearColor(0.25f, 0.25f, 0.25f, 1.0f));
+	}
+	// Draw the major tick marks.
+	for (float X = MajorOX; X < Viewport.Width; X += MajorTickMark)
+	{
+		DrawContext.DrawBox(X, GetPosY(), 1.0f, MajorTickMarkHeight, WhiteBrush, FLinearColor(0.4f, 0.4f, 0.4f, 1.0f));
+	}
+	DrawContext.LayerId++;
+
+	const double DT = static_cast<double>(MajorTickMark) / Viewport.ScaleX;
+	const double Precision = FMath::Max(DT / 10.0, TimeUtils::Nanosecond);
+
+	const TSharedRef<FSlateFontMeasure> FontMeasureService = FSlateApplication::Get().GetRenderer()->GetFontMeasureService();
+
+	// Draw the time at major tick marks.
+	for (float X = MajorOX; X < Viewport.Width + MajorTickMark; X += MajorTickMark)
+	{
+		const double T = Viewport.SlateUnitsToTime(X);
+		FString Text = TimeUtils::FormatTime(T, Precision);
+		const float TextWidth = FontMeasureService->Measure(Text, Font).X;
+		DrawContext.DrawText(X - TextWidth / 2, TextY, Text, Font,
+			(T < 0 || T >= Viewport.MaxValidTime) ? FLinearColor(0.7f, 0.5f, 0.5f, 1.0f) : FLinearColor(0.8f, 0.8f, 0.8f, 1.0f));
+	}
+	DrawContext.LayerId++;
+
+	bool bShowMousePos = !MousePosition.IsZero();
+	if (bShowMousePos)
+	{
+		const FLinearColor MousePosLineColor(0.9f, 0.9f, 0.9f, 0.1f);
+		const FLinearColor MousePosTextBackgroundColor(0.9f, 0.9f, 0.9f, 1.0f);
+		FLinearColor MousePosTextForegroundColor(0.1f, 0.1f, 0.1f, 1.0f);
+
+		// Time at current mouse position.
+		FString MousePosText;
+
+		const double MousePosTime = Viewport.SlateUnitsToTime(MousePosition.X);
+		const double MousePosPrecision = FMath::Max(DT / 100.0, TimeUtils::Nanosecond);
+		if (MousePosition.Y >= GetPosY() && MousePosition.Y < GetPosY() + GetHeight())
+		{
+			// If mouse is hovering the time ruller, format time with a better precision (split seconds in ms, us, ns and ps).
+			MousePosText = TimeUtils::FormatTimeSplit(MousePosTime, MousePosPrecision);
+		}
+		else
+		{
+			// Format current time with one more digit than the time at major tick marks.
+			MousePosText = TimeUtils::FormatTime(MousePosTime, MousePosPrecision);
+		}
+
+		const float MousePosTextWidth = FMath::RoundToFloat(FontMeasureService->Measure(MousePosText, Font).X);
+		static float CrtMousePosTextWidth = 0.0f;
+
+		if (!FMath::IsNearlyEqual(CrtMousePosTextWidth, MousePosTextWidth))
+		{
+			// Animate the box's width (to avoid flickering).
+			CrtMousePosTextWidth = CrtMousePosTextWidth * 0.6f + MousePosTextWidth * 0.4f;
+		}
+
+		float X = MousePosition.X;
+		float W = CrtMousePosTextWidth + 4.0f;
+		if (bIsSelecting && SelectionStartTime < SelectionEndTime)
+		{
+			float SelectionX1 = Viewport.TimeToSlateUnitsRounded(SelectionStartTime);
+			float SelectionX2 = Viewport.TimeToSlateUnitsRounded(SelectionEndTime);
+			if (X - SelectionX1 > 1.0f)
+			{
+				X = SelectionX2 + W / 2;
+			}
+			if (SelectionX2 - X > 1.0f)
+			{
+				X = SelectionX1 - W / 2;
+			}
+			MousePosTextForegroundColor = FLinearColor(FColor(32, 64, 128, 255));
+		}
+		else
+		{
+			// Draw horizontal line at mouse position.
+			//DrawContext.DrawBox(0.0f, MousePosition.Y, Viewport.Width, 1.0f, WhiteBrush, MousePosLineColor);
+
+			// Draw vertical line at mouse position.
+			DrawContext.DrawBox(MousePosition.X, 0.0f, 1.0f, Viewport.Height, WhiteBrush, MousePosLineColor);
+
+			// Stroke the vertical line above current time box.
+			DrawContext.DrawBox(MousePosition.X, 0.0f, 1.0f, TextY, WhiteBrush, MousePosTextBackgroundColor);
+		}
+
+		// Fill the current time box.
+		DrawContext.DrawBox(X - W / 2, TextY, W, 12.0f, WhiteBrush, MousePosTextBackgroundColor);
+		DrawContext.LayerId++;
+
+		// Draw current time text.
+		DrawContext.DrawText(X - MousePosTextWidth / 2, TextY, MousePosText, Font, MousePosTextForegroundColor);
+		DrawContext.LayerId++;
+	}
+}
+
+////////////////////////////////////////////////////////////////////////////////////////////////////
+
+#undef LOCTEXT_NAMESPACE
diff --git a/Engine/Source/Developer/TraceInsights/Private/Insights/ViewModels/TimeRulerTrack.h b/Engine/Source/Developer/TraceInsights/Private/Insights/ViewModels/TimeRulerTrack.h
new file mode 100644
--- /dev/null
+++ b/Engine/Source/Developer/TraceInsights/Private/Insights/ViewModels/TimeRulerTrack.h
@@ -0,0 +1,41 @@
+// Copyright 1998-2019 Epic Games, Inc. All Rights Reserved.
+
+#pragma once
+
+#include "CoreMinimal.h"
+#include "Fonts/SlateFontInfo.h"
+
+// Insights
+#include "Insights/ViewModels/BaseTimingTrack.h"
+
+struct FDrawContext;
+struct FSlateBrush;
+class FTimingTrackViewport;
+
+////////////////////////////////////////////////////////////////////////////////////////////////////
+
+class FTimeRulerTrack : public FBaseTimingTrack
+{
+public:
+	FTimeRulerTrack(uint64 InTrackId);
+	virtual ~FTimeRulerTrack();
+
+	virtual void Reset() override;
+	virtual void UpdateHoveredState(float MouseX, float MouseY, const FTimingTrackViewport& Viewport);
+
+	void Draw(FDrawContext& DrawContext, const FTimingTrackViewport& Viewport,
+				const FVector2D& MousePosition = FVector2D(0.0f, 0.0f),
+				const bool bIsSelecting = false,
+				const double SelectionStartTime = 0.0,
+				const double SelectionEndTime = 0.0) const;
+
+private:
+	void DrawBackground(FDrawContext& DrawContext, const FTimingTrackViewport& Viewport) const;
+
+public:
+	// Slate resources
+	const FSlateBrush* WhiteBrush;
+	const FSlateFontInfo Font;
+};
+
+////////////////////////////////////////////////////////////////////////////////////////////////////
diff --git a/Engine/Source/Developer/TraceInsights/Private/Insights/ViewModels/TimerGroupingAndSorting.cpp b/Engine/Source/Developer/TraceInsights/Private/Insights/ViewModels/TimerGroupingAndSorting.cpp
new file mode 100644
--- /dev/null
+++ b/Engine/Source/Developer/TraceInsights/Private/Insights/ViewModels/TimerGroupingAndSorting.cpp
@@ -0,0 +1,9 @@
+// Copyright 1998-2019 Epic Games, Inc. All Rights Reserved.
+
+#include "TimerGroupingAndSorting.h"
+
+
+#define LOCTEXT_NAMESPACE "TimerNode"
+
+
+#undef LOCTEXT_NAMESPACE
diff --git a/Engine/Source/Developer/TraceInsights/Private/Insights/ViewModels/TimerNode.cpp b/Engine/Source/Developer/TraceInsights/Private/Insights/ViewModels/TimerNode.cpp
new file mode 100644
--- /dev/null
+++ b/Engine/Source/Developer/TraceInsights/Private/Insights/ViewModels/TimerNode.cpp
@@ -0,0 +1,59 @@
+// Copyright 1998-2019 Epic Games, Inc. All Rights Reserved.
+
+#include "TimerNode.h"
+
+#define LOCTEXT_NAMESPACE "TimerNode"
+
+////////////////////////////////////////////////////////////////////////////////////////////////////
+
+void FTimerNode::ResetAggregatedStats()
+{
+	AggregatedStats = Trace::FAggregatedTimingStats();
+}
+
+////////////////////////////////////////////////////////////////////////////////////////////////////
+
+void FTimerNode::SetAggregatedStats(const Trace::FAggregatedTimingStats& InAggregatedStats)
+{
+	AggregatedStats = InAggregatedStats;
+}
+
+////////////////////////////////////////////////////////////////////////////////////////////////////
+
+const FText FTimerNode::GetNameEx() const
+{
+	FText Text = FText::GetEmpty();
+
+	if (IsGroup())
+	{
+		const int32 NumChildren = Children.Num();
+		const int32 NumFilteredChildren = FilteredChildren.Num();
+
+		if (NumFilteredChildren == NumChildren)
+		{
+			Text = FText::Format(LOCTEXT("TimerNodeGroupTextFmt1", "{0} ({1})"), FText::FromName(Name), FText::AsNumber(NumChildren));
+		}
+		else
+		{
+			Text = FText::Format(LOCTEXT("TimerNodeGroupTextFmt2", "{0} ({1} / {2})"), FText::FromName(Name), FText::AsNumber(NumFilteredChildren), FText::AsNumber(NumChildren));
+		}
+	}
+	else
+	{
+		//const bool bIsStatTracked = FProfilerManager::Get()->IsStatTracked(GroupOrStatNode->GetStatID());
+		//if (bIsStatTracked)
+		//{
+		//	Text = FText::Format(LOCTEXT("TimerNodeTrackedTextFmt", "{0}*"), FText::FromName(Name));
+		//}
+		//else
+		{
+			Text = FText::FromName(Name);
+		}
+	}
+
+	return Text;
+}
+
+////////////////////////////////////////////////////////////////////////////////////////////////////
+
+#undef LOCTEXT_NAMESPACE
diff --git a/Engine/Source/Developer/TraceInsights/Private/Insights/ViewModels/TimerNode.h b/Engine/Source/Developer/TraceInsights/Private/Insights/ViewModels/TimerNode.h
new file mode 100644
--- /dev/null
+++ b/Engine/Source/Developer/TraceInsights/Private/Insights/ViewModels/TimerNode.h
@@ -0,0 +1,227 @@
+// Copyright 1998-2019 Epic Games, Inc. All Rights Reserved.
+
+#pragma once
+
+#include "CoreMinimal.h"
+#include "TraceServices/AnalysisService.h"
+
+class FTimerNode;
+
+/** Type definition for shared pointers to instances of FTimerNode. */
+typedef TSharedPtr<class FTimerNode> FTimerNodePtr;
+
+/** Type definition for shared references to instances of FTimerNode. */
+typedef TSharedRef<class FTimerNode> FTimerNodeRef;
+
+/** Type definition for shared references to const instances of FTimerNode. */
+typedef TSharedRef<const class FTimerNode> FTimerNodeRefConst;
+
+/** Type definition for weak references to instances of FTimerNode. */
+typedef TWeakPtr<class FTimerNode> FTimerNodeWeak;
+
+enum class ETimerNodeType
+{
+	/** The TimerNode is a CPU Scope timer. */
+	CpuScope,
+
+	/** The TimerNode is a GPU Scope timer. */
+	GpuScope,
+
+	/** The TimerNode is a Compute Scope timer. */
+	ComputeScope,
+
+	/** The TimerNode is a group node. */
+	Group,
+
+	/** Invalid enum type, may be used as a number of enumerations. */
+	InvalidOrMax,
+};
+
+/**
+ * Class used to store information about a timer node (used in timers' tree view).
+ */
+class FTimerNode : public TSharedFromThis<FTimerNode>
+{
+public:
+	static const uint64 InvalidId = -1;
+
+public:
+	/** Initialization constructor for the timer node. */
+	FTimerNode(uint64 InId, const FName InName, const FName InMetaGroupName, ETimerNodeType InType)
+		: Id(InId)
+		, Name(InName)
+		, MetaGroupName(InMetaGroupName)
+		, Type(InType)
+		, bForceExpandGroupNode(false)
+	{
+		ResetAggregatedStats();
+	}
+
+	/** Initialization constructor for the group node. */
+	FTimerNode(const FName InGroupName)
+		: Id(0)
+		, Name(InGroupName)
+		, Type(ETimerNodeType::Group)
+		, bForceExpandGroupNode(false)
+	{
+		ResetAggregatedStats();
+	}
+
+	/**
+	 * @return an Id of this timer, valid only for timer nodes.
+	 */
+	const uint64 GetId() const
+	{
+		return Id;
+	}
+
+	/**
+	 * @return a name of this node, group or timer.
+	 */
+	const FName& GetName() const
+	{
+		return Name;
+	}
+
+	/**
+	 * @return a name of this node, group or timer + addditional info, to display in Timers tree view.
+	*/
+	const FText GetNameEx() const;
+
+	/**
+	 * @return a name of the group that this timer node belongs to, taken from the metadata.
+	 */
+	const FName& GetMetaGroupName() const
+	{
+		return MetaGroupName;
+	}
+
+	/**
+	 * @return a type of this timer or ETimerNodeType::Group for group nodes.
+	 */
+	const ETimerNodeType& GetType() const
+	{
+		return Type;
+	}
+
+	/**
+	 * @return true, if this node is a group node.
+	 */
+	bool IsGroup() const
+	{
+		return Type == ETimerNodeType::Group;
+	}
+
+	/**
+	 * @return the aggregated stats for this timer.
+	 */
+	const Trace::FAggregatedTimingStats& GetAggregatedStats() const
+	{
+		return AggregatedStats;
+	}
+
+	void ResetAggregatedStats();
+
+	void SetAggregatedStats(const Trace::FAggregatedTimingStats& AggregatedStats);
+
+	/**
+	 * @return a const reference to the child nodes of this group.
+	 */
+	FORCEINLINE_DEBUGGABLE const TArray<FTimerNodePtr>& GetChildren() const
+	{
+		return Children;
+	}
+
+	/**
+	 * @return a const reference to the child nodes that should be visible to the UI based on filtering.
+	 */
+	FORCEINLINE_DEBUGGABLE const TArray<FTimerNodePtr>& GetFilteredChildren() const
+	{
+		return FilteredChildren;
+	}
+
+	/**
+	 * @return a weak reference to the group of this timer node, may be invalid.
+	 */
+	FTimerNodeWeak GetGroupPtr() const
+	{
+		return GroupPtr;
+	}
+
+	/**
+	 * @return a name of the fake group that this timer node belongs to.
+	 */
+	const FName& GetGroupName() const
+	{
+		return GroupPtr.Pin()->Name;
+	}
+
+	/**
+	 * @return a name of the fake group that this timer node belongs to.
+	 */
+	//const FName& GetGroupNameSafe() const
+	//{
+	//	return GroupPtr.IsValid() ? GroupPtr.Pin()->Name : NAME_None;
+	//}
+
+	bool IsFiltered() const
+	{
+		return false; // TODO
+	}
+
+public:
+	/** Sorts children using the specified class instance. */
+	template<typename TSortingClass>
+	void SortChildren(const TSortingClass& Instance)
+	{
+		Children.Sort(Instance);
+	}
+
+	/** Adds specified child to the children and sets group for it. */
+	FORCEINLINE_DEBUGGABLE void AddChildAndSetGroupPtr(const FTimerNodePtr& ChildPtr)
+	{
+		ChildPtr->GroupPtr = AsShared();
+		Children.Add(ChildPtr);
+	}
+
+	/** Adds specified child to the filtered children. */
+	FORCEINLINE_DEBUGGABLE void AddFilteredChild(const FTimerNodePtr& ChildPtr)
+	{
+		FilteredChildren.Add(ChildPtr);
+	}
+
+	/** Clears filtered children. */
+	void ClearFilteredChildren()
+	{
+		FilteredChildren.Reset();
+	}
+
+protected:
+	/** The Id of this timer or group. */
+	const uint64 Id;
+
+	/** The name of this timer or group. */
+	const FName Name;
+
+	/** The name of the group that this timer belongs to, based on the timer's metadata; only valid for timer nodes. */
+	const FName MetaGroupName;
+
+	/** Holds the type of this timer; for the group, this is Group. */
+	const ETimerNodeType Type;
+
+	/** Aggregated stats. */
+	Trace::FAggregatedTimingStats AggregatedStats;
+
+	/** Children of this node. */
+	TArray<FTimerNodePtr> Children;
+
+	/** Filtered children of this node. */
+	TArray<FTimerNodePtr> FilteredChildren;
+
+	/** A weak pointer to the group/parent of this node. */
+	FTimerNodeWeak GroupPtr;
+
+public:
+	/** Whether this group node should be expanded when the text filtering is enabled. */
+	bool bForceExpandGroupNode;
+};
diff --git a/Engine/Source/Developer/TraceInsights/Private/Insights/ViewModels/TimerNodeHelper.cpp b/Engine/Source/Developer/TraceInsights/Private/Insights/ViewModels/TimerNodeHelper.cpp
new file mode 100644
--- /dev/null
+++ b/Engine/Source/Developer/TraceInsights/Private/Insights/ViewModels/TimerNodeHelper.cpp
@@ -0,0 +1,135 @@
+// Copyright 1998-2019 Epic Games, Inc. All Rights Reserved.
+
+#include "TimerNodeHelper.h"
+#include "EditorStyleSet.h"
+
+#define LOCTEXT_NAMESPACE "TimerNode"
+
+////////////////////////////////////////////////////////////////////////////////////////////////////
+// TimerNode Type Helper
+////////////////////////////////////////////////////////////////////////////////////////////////////
+
+FText TimerNodeTypeHelper::ToName(const ETimerNodeType NodeType)
+{
+	static_assert(static_cast<int>(ETimerNodeType::InvalidOrMax) == 4, "Not all cases are handled in switch below!?");
+	switch (NodeType)
+	{
+		case ETimerNodeType::GpuScope:		return LOCTEXT("Timer_Name_Gpu", "GPU");
+		case ETimerNodeType::ComputeScope:	return LOCTEXT("Timer_Name_Compute", "Compute");
+		case ETimerNodeType::CpuScope:		return LOCTEXT("Timer_Name_Cpu", "CPU");
+		case ETimerNodeType::Group:			return LOCTEXT("Timer_Name_Group", "Group");
+		default:							return LOCTEXT("InvalidOrMax", "InvalidOrMax");
+	}
+}
+
+////////////////////////////////////////////////////////////////////////////////////////////////////
+
+FText TimerNodeTypeHelper::ToDescription(const ETimerNodeType NodeType)
+{
+	static_assert(static_cast<int>(ETimerNodeType::InvalidOrMax) == 4, "Not all cases are handled in switch below!?");
+	switch (NodeType)
+	{
+		case ETimerNodeType::GpuScope:		return LOCTEXT("Timer_Desc_Gpu", "GPU scope timer");
+		case ETimerNodeType::ComputeScope:	return LOCTEXT("Timer_Desc_Compute", "Compute scope timer");
+		case ETimerNodeType::CpuScope:		return LOCTEXT("Timer_Desc_Cpu", "CPU scope timer");
+		case ETimerNodeType::Group:			return LOCTEXT("Timer_Desc_Group", "Group timer node");
+		default:							return LOCTEXT("InvalidOrMax", "InvalidOrMax");
+	}
+}
+
+////////////////////////////////////////////////////////////////////////////////////////////////////
+
+FName TimerNodeTypeHelper::ToBrushName(const ETimerNodeType NodeType)
+{
+	static_assert(static_cast<int>(ETimerNodeType::InvalidOrMax) == 4, "Not all cases are handled in switch below!?");
+	switch (NodeType)
+	{
+		case ETimerNodeType::GpuScope:		return TEXT("Profiler.FiltersAndPresets.StatTypeIcon"); //TODO: "Icons.GpuTimer"
+		case ETimerNodeType::ComputeScope:	return TEXT("Profiler.FiltersAndPresets.StatTypeIcon"); //TODO: "Icons.ComputeTimer"
+		case ETimerNodeType::CpuScope:		return TEXT("Profiler.FiltersAndPresets.StatTypeIcon"); //TODO: "Icons.CpuTimer"
+		case ETimerNodeType::Group:			return TEXT("Profiler.Misc.GenericGroup"); //TODO: "Icons.GenericGroup"
+		default:							return NAME_None;
+	}
+}
+
+////////////////////////////////////////////////////////////////////////////////////////////////////
+
+const FSlateBrush* TimerNodeTypeHelper::GetIconForGroup()
+{
+	return FEditorStyle::GetBrush(TEXT("Profiler.Misc.GenericGroup")); //TODO: FInsightsStyle::GetBrush(TEXT("Icons.GenericGroup"));
+}
+
+////////////////////////////////////////////////////////////////////////////////////////////////////
+
+const FSlateBrush* TimerNodeTypeHelper::GetIconForTimerNodeType(const ETimerNodeType NodeType)
+{
+	static_assert(static_cast<int>(ETimerNodeType::InvalidOrMax) == 4, "Not all cases are handled in switch below!?");
+	switch (NodeType)
+	{
+		case ETimerNodeType::GpuScope:		return FEditorStyle::GetBrush(TEXT("Profiler.Type.NumberFloat")); //TODO: FInsightsStyle::GetBrush(TEXT("Icons.GpuTimer"));
+		case ETimerNodeType::ComputeScope:	return FEditorStyle::GetBrush(TEXT("Profiler.Type.NumberFloat")); //TODO: FInsightsStyle::GetBrush(TEXT("Icons.ComputeTimer"));
+		case ETimerNodeType::CpuScope:		return FEditorStyle::GetBrush(TEXT("Profiler.Type.NumberFloat")); //TODO: FInsightsStyle::GetBrush(TEXT("Icons.CpuTimer"));
+		case ETimerNodeType::Group:			return FEditorStyle::GetBrush(TEXT("Profiler.Misc.GenericGroup")); //TODO: FInsightsStyle::GetBrush(TEXT("Icons.GenericGroup"));
+		default:							return nullptr;
+	}
+}
+
+////////////////////////////////////////////////////////////////////////////////////////////////////
+// TimerNode Grouping Helper
+////////////////////////////////////////////////////////////////////////////////////////////////////
+
+FText TimerNodeGroupingHelper::ToName(const ETimerGroupingMode GroupingMode)
+{
+	static_assert(static_cast<int>(ETimerGroupingMode::InvalidOrMax) == 7, "Not all cases are handled in switch below!?");
+	switch (GroupingMode)
+	{
+		case ETimerGroupingMode::Flat:					return LOCTEXT("Grouping_Name_Flat",				"Flat");
+		case ETimerGroupingMode::ByName:				return LOCTEXT("Grouping_Name_ByName",				"Timer Name");
+		case ETimerGroupingMode::ByMetaGroupName:		return LOCTEXT("Grouping_Name_MetaGroupName",		"Meta Group Name");
+		case ETimerGroupingMode::ByType:				return LOCTEXT("Grouping_Name_Type",				"Timer Type");
+		case ETimerGroupingMode::ByTotalInclusiveTime:	return LOCTEXT("Grouping_Name_TotalInclusiveTime",	"Total Inclusive Time");
+		case ETimerGroupingMode::ByTotalExclusiveTime:	return LOCTEXT("Grouping_Name_TotalExclusiveTime",	"Total Exclusive Time");
+		case ETimerGroupingMode::ByInstanceCount:		return LOCTEXT("Grouping_Name_InstanceCount",		"Instance Count");
+		default:										return LOCTEXT("InvalidOrMax", "InvalidOrMax");
+	}
+}
+
+////////////////////////////////////////////////////////////////////////////////////////////////////
+
+FText TimerNodeGroupingHelper::ToDescription(const ETimerGroupingMode GroupingMode)
+{
+	static_assert(static_cast<int>(ETimerGroupingMode::InvalidOrMax) == 7, "Not all cases are handled in switch below!?");
+	switch (GroupingMode)
+	{
+		case ETimerGroupingMode::Flat:					return LOCTEXT("Grouping_Desc_Flat",				"Creates a single group. Includes all timers.");
+		case ETimerGroupingMode::ByName:				return LOCTEXT("Grouping_Desc_ByName",				"Creates one group for one letter.");
+		case ETimerGroupingMode::ByMetaGroupName:		return LOCTEXT("Grouping_Desc_MetaGroupName",		"Creates groups based on metadata group names of timers.");
+		case ETimerGroupingMode::ByType:				return LOCTEXT("Grouping_Desc_Type",				"Creates one group for each timer type.");
+		case ETimerGroupingMode::ByTotalInclusiveTime:	return LOCTEXT("Grouping_Desc_TotalInclusiveTime",	"Creates one group for each logarithmic range ie. 0.001 - 0.01, 0.01 - 0.1, 0.1 - 1.0, 1.0 - 10.0 etc");
+		case ETimerGroupingMode::ByTotalExclusiveTime:	return LOCTEXT("Grouping_Desc_TotalExclusiveTime",	"Creates one group for each logarithmic range ie. 0.001 - 0.01, 0.01 - 0.1, 0.1 - 1.0, 1.0 - 10.0 etc");
+		case ETimerGroupingMode::ByInstanceCount:		return LOCTEXT("Grouping_Desc_InstanceCount",		"Creates one group for each logarithmic range ie. 0, 1 - 10, 10 - 100, 100 - 1000, etc");
+		default:										return LOCTEXT("InvalidOrMax", "InvalidOrMax");
+	}
+}
+
+////////////////////////////////////////////////////////////////////////////////////////////////////
+
+FName TimerNodeGroupingHelper::ToBrushName(const ETimerGroupingMode GroupingMode)
+{
+	static_assert(static_cast<int>(ETimerGroupingMode::InvalidOrMax) == 7, "Not all cases are handled in switch below!?");
+	switch (GroupingMode)
+	{
+		case ETimerGroupingMode::Flat:					return TEXT("Profiler.FiltersAndPresets.GroupNameIcon"); //TODO: "Icons.Grouping.Flat"
+		case ETimerGroupingMode::ByName:				return TEXT("Profiler.FiltersAndPresets.GroupNameIcon"); //TODO: "Icons.Grouping.ByName"
+		case ETimerGroupingMode::ByMetaGroupName	:	return TEXT("Profiler.FiltersAndPresets.StatNameIcon"); //TODO
+		case ETimerGroupingMode::ByType:				return TEXT("Profiler.FiltersAndPresets.StatTypeIcon"); //TODO
+		case ETimerGroupingMode::ByTotalInclusiveTime:	return TEXT("Profiler.FiltersAndPresets.StatValueIcon"); //TODO
+		case ETimerGroupingMode::ByTotalExclusiveTime:	return TEXT("Profiler.FiltersAndPresets.StatValueIcon"); //TODO
+		case ETimerGroupingMode::ByInstanceCount:		return TEXT("Profiler.FiltersAndPresets.StatValueIcon"); //TODO
+		default:										return NAME_None;
+	}
+}
+
+////////////////////////////////////////////////////////////////////////////////////////////////////
+
+#undef LOCTEXT_NAMESPACE
diff --git a/Engine/Source/Developer/TraceInsights/Private/Insights/ViewModels/TimerNodeHelper.h b/Engine/Source/Developer/TraceInsights/Private/Insights/ViewModels/TimerNodeHelper.h
new file mode 100644
--- /dev/null
+++ b/Engine/Source/Developer/TraceInsights/Private/Insights/ViewModels/TimerNodeHelper.h
@@ -0,0 +1,240 @@
+// Copyright 1998-2019 Epic Games, Inc. All Rights Reserved.
+
+#pragma once
+
+#include "CoreMinimal.h"
+#include "Styling/SlateBrush.h"
+
+// Insights
+#include "Insights/ViewModels/TimerNode.h"
+#include "Insights/ViewModels/TimerGroupingAndSorting.h"
+
+////////////////////////////////////////////////////////////////////////////////////////////////////
+
+/** Helper struct that contains helper functions and classes for ETimerNodeType enum. */
+struct TimerNodeTypeHelper
+{
+	/**
+	 * @param Type - The value to get the text for.
+	 *
+	 * @return text representation of the specified ETimerNodeType value.
+	 */
+	static FText ToName(const ETimerNodeType Type);
+
+	/**
+	 * @param Type - The value to get the text for.
+	 *
+	 * @return text representation with more detailed explanation of the specified ETimerNodeType value.
+	 */
+	static FText ToDescription(const ETimerNodeType Type);
+
+	/**
+	 * @param Type - The value to get the brush name for.
+	 *
+	 * @return brush name of the specified ETimerNodeType value.
+	 */
+	static FName ToBrushName(const ETimerNodeType Type);
+
+	static const FSlateBrush* GetIconForGroup();
+	static const FSlateBrush* GetIconForTimerNodeType(const ETimerNodeType Type);
+};
+
+////////////////////////////////////////////////////////////////////////////////////////////////////
+
+/** Helper struct that contains grouping static functions and classes. */
+struct TimerNodeGroupingHelper
+{
+	/**
+	 * @param TimerGroupingMode - The value to get the text for.
+	 *
+	 * @return text representation of the specified ETimerGroupingMode value.
+	 */
+	static FText ToName(const ETimerGroupingMode TimerGroupingMode);
+
+	/**
+	 * @param TimerGroupingMode - The value to get the text for.
+	 *
+	 * @return text representation with more detailed explanation of the specified ETimerGroupingMode value.
+	 */
+	static FText ToDescription(const ETimerGroupingMode TimerGroupingMode);
+
+	/**
+	 * @param TimerGroupingMode - The value to get the brush name for.
+	 *
+	 * @return brush name of the specified ETimerGroupingMode value.
+	 */
+	static FName ToBrushName(const ETimerGroupingMode TimerGroupingMode);
+};
+
+////////////////////////////////////////////////////////////////////////////////////////////////////
+
+/** Helper struct that contains sorting static functions and classes. */
+struct TimerNodeSortingHelper
+{
+	//////////////////////////////////////////////////
+	// Sorting by timer's name.
+
+	struct ByNameAscending
+	{
+		FORCEINLINE_DEBUGGABLE bool operator()(const FTimerNodePtr& A, const FTimerNodePtr& B) const
+		{
+			return A->GetName().LexicalLess(B->GetName());
+		}
+	};
+
+	struct ByNameDescending
+	{
+		FORCEINLINE_DEBUGGABLE bool operator()(const FTimerNodePtr& A, const FTimerNodePtr& B) const
+		{
+			return B->GetName().LexicalLess(A->GetName());
+		}
+	};
+
+	//////////////////////////////////////////////////
+	// Sorting by timer's meta group name.
+	// If meta group names are the same then sort by name.
+
+	struct ByMetaGroupNameAscending
+	{
+		FORCEINLINE_DEBUGGABLE bool operator()(const FTimerNodePtr& A, const FTimerNodePtr& B) const
+		{
+			if (A->GetMetaGroupName() == B->GetMetaGroupName())
+			{
+				// Sort by timer name (ascending).
+				return A->GetName().LexicalLess(B->GetName());
+			}
+			else
+			{
+				return A->GetMetaGroupName().LexicalLess(B->GetMetaGroupName());
+			}
+		}
+	};
+
+	struct ByMetaGroupNameDescending
+	{
+		FORCEINLINE_DEBUGGABLE bool operator()(const FTimerNodePtr& A, const FTimerNodePtr& B) const
+		{
+			if (A->GetMetaGroupName() == B->GetMetaGroupName())
+			{
+				// Sort by timer name (ascending).
+				return A->GetName().LexicalLess(B->GetName());
+			}
+			else
+			{
+				return B->GetMetaGroupName().LexicalLess(A->GetMetaGroupName());
+			}
+		}
+	};
+
+	//////////////////////////////////////////////////
+	// Sorting by timer's type.
+	// If types are the same then sort by name.
+
+	struct ByTypeAscending
+	{
+		FORCEINLINE_DEBUGGABLE bool operator()(const FTimerNodePtr& A, const FTimerNodePtr& B) const
+		{
+			const ETimerNodeType& TypeA = A->GetType();
+			const ETimerNodeType& TypeB = B->GetType();
+
+			if (TypeA == TypeB)
+			{
+				// Sort by timer name (ascending).
+				return A->GetName().LexicalLess(B->GetName());
+			}
+			else
+			{
+				return TypeA < TypeB;
+			}
+		}
+	};
+
+	struct ByTypeDescending
+	{
+		FORCEINLINE_DEBUGGABLE bool operator()(const FTimerNodePtr& A, const FTimerNodePtr& B) const
+		{
+			const ETimerNodeType& TypeA = A->GetType();
+			const ETimerNodeType& TypeB = B->GetType();
+
+			if (TypeA == TypeB)
+			{
+				// Sort by timer name (ascending).
+				return A->GetName().LexicalLess(B->GetName());
+			}
+			else
+			{
+				return TypeA > TypeB;
+			}
+		}
+	};
+
+	//////////////////////////////////////////////////
+	// Sorting by an aggregated stats value.
+	// If aggregated stats values are the same then sort by name.
+
+	#define SORT_BY_STATS_ASCENDING(Type, SortName, AggregatedStatsMember) \
+		struct SortName \
+		{ \
+			FORCEINLINE_DEBUGGABLE bool operator()(const FTimerNodePtr& A, const FTimerNodePtr& B) const \
+			{ \
+				const Type ValueA = A->GetAggregatedStats().AggregatedStatsMember; \
+				const Type ValueB = B->GetAggregatedStats().AggregatedStatsMember; \
+				\
+				if (ValueA == ValueB) \
+				{ \
+					/* Sort by name (ascending). */ \
+					return A->GetName().LexicalLess(B->GetName()); \
+				} \
+				else \
+				{ \
+					return ValueA < ValueB; \
+				} \
+			} \
+		};
+
+	#define SORT_BY_STATS_DESCENDING(Type, SortName, AggregatedStatsMember) \
+		struct SortName \
+		{ \
+			FORCEINLINE_DEBUGGABLE bool operator()(const FTimerNodePtr& A, const FTimerNodePtr& B) const \
+			{ \
+				const Type ValueA = A->GetAggregatedStats().AggregatedStatsMember; \
+				const Type ValueB = B->GetAggregatedStats().AggregatedStatsMember; \
+				\
+				if (ValueA == ValueB) \
+				{ \
+					/* Sort by name (ascending). */ \
+					return A->GetName().LexicalLess(B->GetName()); \
+				} \
+				else \
+				{ \
+					return ValueA > ValueB; \
+				} \
+			} \
+		};
+
+	#define SORT_BY_STATS(Type, SortName, AggregatedStatsMember) \
+		SORT_BY_STATS_ASCENDING(Type, SortName##Ascending, AggregatedStatsMember) \
+		SORT_BY_STATS_DESCENDING(Type, SortName##Descending, AggregatedStatsMember)
+
+	SORT_BY_STATS(uint64, ByInstanceCount, InstanceCount);
+
+	SORT_BY_STATS(double, ByTotalInclusiveTime, TotalInclusiveTime);
+	SORT_BY_STATS(double, ByMinInclusiveTime, MinInclusiveTime);
+	SORT_BY_STATS(double, ByMaxInclusiveTime, MaxInclusiveTime);
+	SORT_BY_STATS(double, ByAverageInclusiveTime, AverageInclusiveTime);
+	SORT_BY_STATS(double, ByMedianInclusiveTime, MedianInclusiveTime);
+
+	SORT_BY_STATS(double, ByTotalExclusiveTime, TotalExclusiveTime);
+	SORT_BY_STATS(double, ByMinExclusiveTime, MinExclusiveTime);
+	SORT_BY_STATS(double, ByMaxExclusiveTime, MaxExclusiveTime);
+	SORT_BY_STATS(double, ByAverageExclusiveTime, AverageExclusiveTime);
+	SORT_BY_STATS(double, ByMedianExclusiveTime, MedianExclusiveTime);
+
+	#undef SORT_BY_STATS_ASCENDING
+	#undef SORT_BY_STATS_DESCENDING
+	#undef SORT_BY_STATS
+
+	//////////////////////////////////////////////////
+};
+
+////////////////////////////////////////////////////////////////////////////////////////////////////
