Copying //UE4/Dev-Core to Dev-Main (//UE4/Dev-Main)

diff --git a/Engine/Source/Developer/TraceInsights/Private/Insights/ViewModels/TimersViewColumn.h b/Engine/Source/Developer/TraceInsights/Private/Insights/ViewModels/TimersViewColumn.h
new file mode 100644
--- /dev/null
+++ b/Engine/Source/Developer/TraceInsights/Private/Insights/ViewModels/TimersViewColumn.h
@@ -0,0 +1,118 @@
+// Copyright 1998-2019 Epic Games, Inc. All Rights Reserved.
+
+#pragma once
+
+#include "CoreMinimal.h"
+#include "Types/SlateEnums.h"
+#include "Misc/EnumClassFlags.h"
+
+class FTimerNode; // TODO: IInsightsTreeNode; generic formatter; see also StatsViewColumn.h
+
+enum class ETimersViewColumnFlags : uint32
+{
+	None = 0,
+
+	CanBeHidden = (1 << 0),
+	CanBeSorted = (1 << 1),
+	CanBeFiltered = (1 << 2),
+};
+ENUM_CLASS_FLAGS(ETimersViewColumnFlags);
+
+/** Holds information about a column in the timer view widget. */
+class FTimersViewColumn
+{
+	friend struct FTimersViewColumnFactory;
+
+public:
+	typedef TFunction<FText(const FTimersViewColumn& Column, const FTimerNode& TimerNode)> FGetFormattedValueFn;
+
+public:
+	/** Whether this column can be hidden. */
+	bool bCanBeHidden() const { return EnumHasAnyFlags(Flags, ETimersViewColumnFlags::CanBeHidden); }
+
+	/** Whether this column cab be used for sorting. */
+	bool bCanBeSorted() const { return EnumHasAnyFlags(Flags, ETimersViewColumnFlags::CanBeSorted); }
+
+	/** Where this column can be used to filtering displayed results. */
+	bool bCanBeFiltered() const { return EnumHasAnyFlags(Flags, ETimersViewColumnFlags::CanBeFiltered); }
+
+	/** If MinColumnWidth == MaxColumnWidth, this column has fixed width and cannot be resized. */
+	bool bIsFixedColumnWidth() const { return MinColumnWidth == MaxColumnWidth; }
+
+	FText GetFormattedValue(const FTimerNode& TimerNode) const
+	{
+		return GetFormattedValueFn(*this, TimerNode);
+	}
+
+protected:
+	/** No default constructor. */
+	FTimersViewColumn() = delete;
+
+	/** Initialization constructor, only used in FStatsViewColumnFactory. */
+	FTimersViewColumn
+	(
+		int32 InOrder,
+		const FName InId,
+		const FName InSearchId,
+		FText InShortName,
+		FText InTitleName,
+		FText InDescription,
+		bool bInIsVisible,
+		const ETimersViewColumnFlags InFlags,
+		const EHorizontalAlignment InHorizontalAlignment,
+		const float InInitialColumnWidth,
+		const float InMinColumnWidth,
+		const float InMaxColumnWidth,
+		const FGetFormattedValueFn InGetFormattedValueFn
+	)
+		: Order(InOrder)
+		, Id(InId)
+		, SearchId(InSearchId)
+		, ShortName(MoveTemp(InShortName))
+		, TitleName(MoveTemp(InTitleName))
+		, Description(MoveTemp(InDescription))
+		, bIsVisible(bInIsVisible)
+		, Flags(InFlags)
+		, HorizontalAlignment(InHorizontalAlignment)
+		, InitialColumnWidth(InInitialColumnWidth)
+		, MinColumnWidth(InMinColumnWidth)
+		, MaxColumnWidth(InMaxColumnWidth)
+		, GetFormattedValueFn(InGetFormattedValueFn)
+	{
+	}
+
+public:
+	/** Order value, to sort columns in the tree view. */
+	int32 Order;
+
+	/** Name of the column, name of the property. */
+	FName Id;
+
+	/** Name of the column used by the searching system. */
+	FName SearchId;
+
+	/** Short name of the column, displayed in the column header. */
+	FText ShortName;
+
+	/** Title name of the column, displayed as title in the column tooltip. */
+	FText TitleName;
+
+	/** Long name of the column, displayed in the column tooltip. */
+	FText Description;
+
+	/** Is this column visible? */
+	bool bIsVisible;
+
+	/** On/off switches. */
+	ETimersViewColumnFlags Flags;
+
+	/** Horizontal alignment of the content in this column. */
+	EHorizontalAlignment HorizontalAlignment;
+
+	float InitialColumnWidth; /**< Initial column width. */
+	float MinColumnWidth; /**< Minimum column width. */
+	float MaxColumnWidth; /**< Maximum column width. */
+
+	/** Custom function used to format (as an FText) the value of the timer node to be displayed by this column. */
+	FGetFormattedValueFn GetFormattedValueFn;
+};
diff --git a/Engine/Source/Developer/TraceInsights/Private/Insights/ViewModels/TimersViewColumnFactory.cpp b/Engine/Source/Developer/TraceInsights/Private/Insights/ViewModels/TimersViewColumnFactory.cpp
new file mode 100644
--- /dev/null
+++ b/Engine/Source/Developer/TraceInsights/Private/Insights/ViewModels/TimersViewColumnFactory.cpp
@@ -0,0 +1,356 @@
+// Copyright 1998-2019 Epic Games, Inc. All Rights Reserved.
+
+#include "TimersViewColumnFactory.h"
+
+// Insights
+#include "Insights/Common/TimeUtils.h"
+#include "Insights/ViewModels/TimerNodeHelper.h"
+
+#define LOCTEXT_NAMESPACE "STimerView"
+
+////////////////////////////////////////////////////////////////////////////////////////////////////
+// Column identifiers
+
+const FName FTimersViewColumns::NameColumnID(TEXT("Name"));
+const FName FTimersViewColumns::MetaGroupNameColumnID(TEXT("MetaGroupName"));
+const FName FTimersViewColumns::TypeColumnID(TEXT("Type"));
+const FName FTimersViewColumns::InstanceCountColumnID(TEXT("Count"));
+
+// Inclusive Time columns
+const FName FTimersViewColumns::TotalInclusiveTimeColumnID(TEXT("TotalInclTime"));
+const FName FTimersViewColumns::MaxInclusiveTimeColumnID(TEXT("MaxInclTime"));
+const FName FTimersViewColumns::UpperQuartileInclusiveTimeColumnID(TEXT("UpperQuartileInclTime"));
+const FName FTimersViewColumns::AverageInclusiveTimeColumnID(TEXT("AverageInclTime"));
+const FName FTimersViewColumns::MedianInclusiveTimeColumnID(TEXT("MedianInclTime"));
+const FName FTimersViewColumns::LowerQuartileInclusiveTimeColumnID(TEXT("LowerQuartileInclTime"));
+const FName FTimersViewColumns::MinInclusiveTimeColumnID(TEXT("MinInclTime"));
+
+// Exclusive Time columns
+const FName FTimersViewColumns::TotalExclusiveTimeColumnID(TEXT("TotalExclTime"));
+const FName FTimersViewColumns::MaxExclusiveTimeColumnID(TEXT("MaxExclTime"));
+const FName FTimersViewColumns::UpperQuartileExclusiveTimeColumnID(TEXT("UpperQuartileExclTime"));
+const FName FTimersViewColumns::AverageExclusiveTimeColumnID(TEXT("AverageExclTime"));
+const FName FTimersViewColumns::MedianExclusiveTimeColumnID(TEXT("MedianExclTime"));
+const FName FTimersViewColumns::LowerQuartileExclusiveTimeColumnID(TEXT("LowerQuartileExclTime"));
+const FName FTimersViewColumns::MinExclusiveTimeColumnID(TEXT("MinExclTime"));
+
+////////////////////////////////////////////////////////////////////////////////////////////////////
+
+FTimersViewColumnFactory::FTimersViewColumnFactory()
+{
+	uint32 Index = 0;
+
+	// FTimerViewColumn
+	// (
+	//     Order,
+	//     Id,
+	//     SearchId,
+	//     ShortName,
+	//     TitleName,
+	//     Description,
+	//     bIsVisible,
+	//     Flags,
+	//     HorizontalAlignment,
+	//     InitialWidth, MinWidth, MaxWidth,
+	//     GetFormattedValueFn
+	// )
+
+	//////////////////////////////////////////////////
+	// Name
+
+	const ETimersViewColumnFlags NameColumnFlags = ETimersViewColumnFlags::CanBeSorted |
+												   ETimersViewColumnFlags::CanBeFiltered;
+
+	Collection.Add(new FTimersViewColumn
+	(
+		Index++,
+		FTimersViewColumns::NameColumnID,
+		TEXT("name"),
+		LOCTEXT("TimerNameColumnName", "Name"),
+		LOCTEXT("TimerNameColumnTitle", "Timer or Group Name"),
+		LOCTEXT("TimerNameColumnDesc", "Name of the timer or group"),
+		true,
+		NameColumnFlags,
+		HAlign_Left,
+		246.0f, 10.0f, FLT_MAX,
+		[](const FTimersViewColumn& Column, const FTimerNode& TimerNode) -> FText
+		{ return FText::FromName(TimerNode.GetName()); }
+	));
+
+	//////////////////////////////////////////////////
+	// Meta Group Name
+
+	const ETimersViewColumnFlags MetaGroupNameColumnFlags = ETimersViewColumnFlags::CanBeHidden |
+															ETimersViewColumnFlags::CanBeSorted |
+															ETimersViewColumnFlags::CanBeFiltered;
+
+	Collection.Add(new FTimersViewColumn
+	(
+		Index++,
+		FTimersViewColumns::MetaGroupNameColumnID,
+		TEXT("metagroupname"),
+		LOCTEXT("TimerMetaGroupNameColumnName", "Meta Group"),
+		LOCTEXT("TimerMetaGroupNameColumnTitle", "Meta Group Name"),
+		LOCTEXT("TimerMetaGroupNameColumnDesc", "Name of the meta group"),
+		false,
+		MetaGroupNameColumnFlags,
+		HAlign_Left,
+		100.0f, 10.0f, FLT_MAX,
+		[](const FTimersViewColumn& Column, const FTimerNode& TimerNode) -> FText
+		{ return FText::FromName(TimerNode.GetMetaGroupName()); }
+	));
+
+	//////////////////////////////////////////////////
+	// Type
+
+	const ETimersViewColumnFlags TypeColumnFlags = ETimersViewColumnFlags::CanBeHidden |
+												   ETimersViewColumnFlags::CanBeSorted |
+												   ETimersViewColumnFlags::CanBeFiltered;
+
+	Collection.Add(new FTimersViewColumn
+	(
+		Index++,
+		FTimersViewColumns::TypeColumnID,
+		TEXT("type"),
+		LOCTEXT("TimerTypeColumnName", "Type"),
+		LOCTEXT("TimerTypeColumnTitle", "Type"),
+		LOCTEXT("TimerTypeColumnDesc", "Type of timer or group"),
+		false,
+		TypeColumnFlags,
+		HAlign_Left,
+		60.0f, 10.0f, FLT_MAX,
+		[](const FTimersViewColumn& Column, const FTimerNode& TimerNode) -> FText
+		{ return TimerNodeTypeHelper::ToName(TimerNode.GetType()); }
+	));
+		
+	//////////////////////////////////////////////////
+
+	const ETimersViewColumnFlags AggregatedStatsColumnFlags = ETimersViewColumnFlags::CanBeHidden |
+															  ETimersViewColumnFlags::CanBeSorted |
+															  ETimersViewColumnFlags::CanBeFiltered;
+
+	//////////////////////////////////////////////////
+	// Instance Count
+
+	Collection.Add(new FTimersViewColumn
+	(
+		Index++,
+		FTimersViewColumns::InstanceCountColumnID,
+		TEXT("count"),
+		LOCTEXT("InstanceCountName", "Count"),
+		LOCTEXT("InstanceCountTitle", "Instance Count"),
+		LOCTEXT("InstanceCountDesc", "Number of timer's instances"),
+		true,
+		AggregatedStatsColumnFlags,
+		HAlign_Right,
+		60.0f, 0.0f, FLT_MAX,
+		[](const FTimersViewColumn& Column, const FTimerNode& TimerNode) -> FText
+		{ return FText::AsNumber(TimerNode.GetAggregatedStats().InstanceCount); }
+	));
+
+	//////////////////////////////////////////////////
+
+	constexpr float TotalTimeColumnInitialWidth = 60.0f;
+	constexpr float TotalTimeColumnMinWidth = 4.0f;
+	constexpr float TotalTimeColumnMaxWidth = FLT_MAX;
+
+	constexpr float TimeMsColumnInitialWidth = 50.0f;
+	constexpr float TimeMsColumnMinWidth = 4.0f;
+	constexpr float TimeMsColumnMaxWidth = FLT_MAX;
+
+	//////////////////////////////////////////////////
+	// Inclusive Time stats
+
+	Collection.Add(new FTimersViewColumn
+	(
+		Index++,
+		FTimersViewColumns::TotalInclusiveTimeColumnID,
+		TEXT("inc"),
+		LOCTEXT("TotalInclusiveTimeName", "Incl"),
+		LOCTEXT("TotalInclusiveTimeTitle", "Total Inclusive Time"),
+		LOCTEXT("TotalInclusiveTimeDesc", "Total inclusive duration of selected timer's instances"),
+		true,
+		AggregatedStatsColumnFlags,
+		HAlign_Right,
+		TotalTimeColumnInitialWidth, TotalTimeColumnMinWidth, TotalTimeColumnMaxWidth,
+		[](const FTimersViewColumn& Column, const FTimerNode& TimerNode) -> FText
+		{ return FText::FromString(TimeUtils::FormatTimeAuto(TimerNode.GetAggregatedStats().TotalInclusiveTime)); }
+	));
+
+	Collection.Add(new FTimersViewColumn
+	(
+		Index++,
+		FTimersViewColumns::MaxInclusiveTimeColumnID,
+		TEXT("maxinc"),
+		LOCTEXT("MaxInclusiveTimeName", "I.Max"),
+		LOCTEXT("MaxInclusiveTimeTitle", "Max Inclusive Time (ms)"),
+		LOCTEXT("MaxInclusiveTimeDesc", "Maximum inclusive duration of selected timer's instances, in milliseconds"),
+		false,
+		AggregatedStatsColumnFlags,
+		HAlign_Right,
+		TimeMsColumnInitialWidth, TimeMsColumnMinWidth, TimeMsColumnMaxWidth,
+		[](const FTimersViewColumn& Column, const FTimerNode& TimerNode) -> FText
+		{ return FText::FromString(TimeUtils::FormatTimeMs(TimerNode.GetAggregatedStats().MaxInclusiveTime)); }
+	));
+
+	Collection.Add(new FTimersViewColumn
+	(
+		Index++,
+		FTimersViewColumns::AverageInclusiveTimeColumnID,
+		TEXT("avginc"),
+		LOCTEXT("AvgInclusiveTimeName", "I.Avg"),
+		LOCTEXT("AvgInclusiveTimeTitle", "Average Inclusive Time (ms)"),
+		LOCTEXT("AvgInclusiveTimeDesc", "Average inclusive duration of selected timer's instances, in milliseconds"),
+		false,
+		AggregatedStatsColumnFlags,
+		HAlign_Right,
+		TimeMsColumnInitialWidth, TimeMsColumnMinWidth, TimeMsColumnMaxWidth,
+		[](const FTimersViewColumn& Column, const FTimerNode& TimerNode) -> FText
+		{ return FText::FromString(TimeUtils::FormatTimeMs(TimerNode.GetAggregatedStats().AverageInclusiveTime)); }
+	));
+
+	Collection.Add(new FTimersViewColumn
+	(
+		Index++,
+		FTimersViewColumns::MedianInclusiveTimeColumnID,
+		TEXT("medinc"),
+		LOCTEXT("MedInclusiveTimeName", "I.Med"),
+		LOCTEXT("MedInclusiveTimeTitle", "Median Inclusive Time (ms)"),
+		LOCTEXT("MedInclusiveTimeDesc", "Median inclusive duration of selected timer's instances, in milliseconds"),
+		false,
+		AggregatedStatsColumnFlags,
+		HAlign_Right,
+		TimeMsColumnInitialWidth, TimeMsColumnMinWidth, TimeMsColumnMaxWidth,
+		[](const FTimersViewColumn& Column, const FTimerNode& TimerNode) -> FText
+		{ return FText::FromString(TimeUtils::FormatTimeMs(TimerNode.GetAggregatedStats().MedianInclusiveTime)); }
+	));
+
+	Collection.Add(new FTimersViewColumn
+	(
+		Index++,
+		FTimersViewColumns::MinInclusiveTimeColumnID,
+		TEXT("mininc"),
+		LOCTEXT("MinInclusiveTimeName", "I.Min"),
+		LOCTEXT("MinInclusiveTimeTitle", "Min Inclusive Time (ms)"),
+		LOCTEXT("MinInclusiveTimeDesc", "Minimum inclusive duration of selected timer's instances, in milliseconds"),
+		false,
+		AggregatedStatsColumnFlags,
+		HAlign_Right,
+		TimeMsColumnInitialWidth, TimeMsColumnMinWidth, TimeMsColumnMaxWidth,
+		[](const FTimersViewColumn& Column, const FTimerNode& TimerNode) -> FText
+		{ return FText::FromString(TimeUtils::FormatTimeMs(TimerNode.GetAggregatedStats().MinInclusiveTime)); }
+	));
+
+	//////////////////////////////////////////////////
+	// Exclusive Time stats
+
+	Collection.Add(new FTimersViewColumn
+	(
+		Index++,
+		FTimersViewColumns::TotalExclusiveTimeColumnID,
+		TEXT("exc"),
+		LOCTEXT("TotalExclusiveTimeName", "Excl"),
+		LOCTEXT("TotalExclusiveTimeTitle", "Total Exclusive Time"),
+		LOCTEXT("TotalExclusiveTimeDesc", "Total exclusive duration of selected timer's instances"),
+		true,
+		AggregatedStatsColumnFlags,
+		HAlign_Right,
+		TotalTimeColumnInitialWidth, TotalTimeColumnMinWidth, TotalTimeColumnMaxWidth,
+		[](const FTimersViewColumn& Column, const FTimerNode& TimerNode) -> FText
+		{ return FText::FromString(TimeUtils::FormatTimeAuto(TimerNode.GetAggregatedStats().TotalExclusiveTime)); }
+	));
+
+	Collection.Add(new FTimersViewColumn
+	(
+		Index++,
+		FTimersViewColumns::MaxExclusiveTimeColumnID,
+		TEXT("maxexc"),
+		LOCTEXT("MaxExclusiveTimeName", "E.Max"),
+		LOCTEXT("MaxExclusiveTimeTitle", "Max Exclusive Time (ms)"),
+		LOCTEXT("MaxExclusiveTimeDesc", "Maximum exclusive duration of selected timer's instances, in milliseconds"),
+		false,
+		AggregatedStatsColumnFlags,
+		HAlign_Right,
+		TimeMsColumnInitialWidth, TimeMsColumnMinWidth, TimeMsColumnMaxWidth,
+		[](const FTimersViewColumn& Column, const FTimerNode& TimerNode) -> FText
+		{ return FText::FromString(TimeUtils::FormatTimeMs(TimerNode.GetAggregatedStats().MaxExclusiveTime)); }
+	));
+
+	Collection.Add(new FTimersViewColumn
+	(
+		Index++,
+		FTimersViewColumns::AverageExclusiveTimeColumnID,
+		TEXT("avgexc"),
+		LOCTEXT("AvgExclusiveTimeName", "E.Avg"),
+		LOCTEXT("AvgExclusiveTimeTitle", "Average Exclusive Time (ms)"),
+		LOCTEXT("AvgExclusiveTimeDesc", "Average exclusive duration of selected timer's instances, in milliseconds"),
+		false,
+		AggregatedStatsColumnFlags,
+		HAlign_Right,
+		TimeMsColumnInitialWidth, TimeMsColumnMinWidth, TimeMsColumnMaxWidth,
+		[](const FTimersViewColumn& Column, const FTimerNode& TimerNode) -> FText
+		{ return FText::FromString(TimeUtils::FormatTimeMs(TimerNode.GetAggregatedStats().AverageExclusiveTime)); }
+	));
+
+	Collection.Add(new FTimersViewColumn
+	(
+		Index++,
+		FTimersViewColumns::MedianExclusiveTimeColumnID,
+		TEXT("medexc"),
+		LOCTEXT("MedExclusiveTimeName", "E.Med"),
+		LOCTEXT("MedExclusiveTimeTitle", "Median Exclusive Time (ms)"),
+		LOCTEXT("MedExclusiveTimeDesc", "Median exclusive duration of timer's instances, in milliseconds"),
+		false,
+		AggregatedStatsColumnFlags,
+		HAlign_Right,
+		TimeMsColumnInitialWidth, TimeMsColumnMinWidth, TimeMsColumnMaxWidth,
+		[](const FTimersViewColumn& Column, const FTimerNode& TimerNode) -> FText
+		{ return FText::FromString(TimeUtils::FormatTimeMs(TimerNode.GetAggregatedStats().MedianExclusiveTime)); }
+	));
+
+	Collection.Add(new FTimersViewColumn
+	(
+		Index++,
+		FTimersViewColumns::MinExclusiveTimeColumnID,
+		TEXT("minexc"),
+		LOCTEXT("MinExclusiveTimeName", "E.Min"),
+		LOCTEXT("MinExclusiveTimeTitle", "Min Exclusive Time (ms)"),
+		LOCTEXT("MinExclusiveTimeDesc", "Minimum exclusive duration of selected timer's instances, in milliseconds"),
+		false,
+		AggregatedStatsColumnFlags,
+		HAlign_Right,
+		TimeMsColumnInitialWidth, TimeMsColumnMinWidth, TimeMsColumnMaxWidth,
+		[](const FTimersViewColumn& Column, const FTimerNode& TimerNode) -> FText
+		{ return FText::FromString(TimeUtils::FormatTimeMs(TimerNode.GetAggregatedStats().MinExclusiveTime)); }
+	));
+
+	//////////////////////////////////////////////////
+
+	for (int32 MapIndex = 0; MapIndex < Collection.Num(); MapIndex++)
+	{
+		ColumnIdToPtrMapping.Add(Collection[MapIndex]->Id, Collection[MapIndex]);
+	}
+}
+
+////////////////////////////////////////////////////////////////////////////////////////////////////
+
+FTimersViewColumnFactory::~FTimersViewColumnFactory()
+{
+	for (const FTimersViewColumn* Column : Collection)
+	{
+		delete Column;
+	}
+	Collection.Reset();
+}
+
+////////////////////////////////////////////////////////////////////////////////////////////////////
+
+const FTimersViewColumnFactory& FTimersViewColumnFactory::Get()
+{
+	static FTimersViewColumnFactory Instance;
+	return Instance;
+}
+
+////////////////////////////////////////////////////////////////////////////////////////////////////
+
+#undef LOCTEXT_NAMESPACE
diff --git a/Engine/Source/Developer/TraceInsights/Private/Insights/ViewModels/TimersViewColumnFactory.h b/Engine/Source/Developer/TraceInsights/Private/Insights/ViewModels/TimersViewColumnFactory.h
new file mode 100644
--- /dev/null
+++ b/Engine/Source/Developer/TraceInsights/Private/Insights/ViewModels/TimersViewColumnFactory.h
@@ -0,0 +1,62 @@
+// Copyright 1998-2019 Epic Games, Inc. All Rights Reserved.
+
+#pragma once
+
+#include "CoreMinimal.h"
+
+#include "Insights/ViewModels/TimersViewColumn.h"
+
+////////////////////////////////////////////////////////////////////////////////////////////////////
+
+struct FTimersViewColumns
+{
+	//////////////////////////////////////////////////
+	// Column identifiers
+
+	static const FName NameColumnID;
+	static const FName MetaGroupNameColumnID;
+	static const FName TypeColumnID;
+	static const FName InstanceCountColumnID;
+
+	// Inclusive Time columns
+	static const FName TotalInclusiveTimeColumnID;
+	static const FName MaxInclusiveTimeColumnID;
+	static const FName UpperQuartileInclusiveTimeColumnID;
+	static const FName AverageInclusiveTimeColumnID;
+	static const FName MedianInclusiveTimeColumnID;
+	static const FName LowerQuartileInclusiveTimeColumnID;
+	static const FName MinInclusiveTimeColumnID;
+
+	// Exclusive Time columns
+	static const FName TotalExclusiveTimeColumnID;
+	static const FName MaxExclusiveTimeColumnID;
+	static const FName UpperQuartileExclusiveTimeColumnID;
+	static const FName AverageExclusiveTimeColumnID;
+	static const FName MedianExclusiveTimeColumnID;
+	static const FName LowerQuartileExclusiveTimeColumnID;
+	static const FName MinExclusiveTimeColumnID;
+};
+
+////////////////////////////////////////////////////////////////////////////////////////////////////
+
+struct FTimersViewColumnFactory
+{
+private:
+	/** Default constructor. */
+	FTimersViewColumnFactory();
+
+	/** Destructor. */
+	~FTimersViewColumnFactory();
+
+public:
+	/** Contains basic information about columns used in the Timers view widget. Names should be localized. */
+	TArray<FTimersViewColumn*> Collection;
+
+	/** Mapping between column IDs and FTimersViewColumn pointers. */
+	TMap<FName, const FTimersViewColumn*> ColumnIdToPtrMapping;
+
+	/** Singleton */
+	static const FTimersViewColumnFactory& Get();
+};
+
+////////////////////////////////////////////////////////////////////////////////////////////////////
diff --git a/Engine/Source/Developer/TraceInsights/Private/Insights/ViewModels/TimingEventsTrack.cpp b/Engine/Source/Developer/TraceInsights/Private/Insights/ViewModels/TimingEventsTrack.cpp
new file mode 100644
--- /dev/null
+++ b/Engine/Source/Developer/TraceInsights/Private/Insights/ViewModels/TimingEventsTrack.cpp
@@ -0,0 +1,132 @@
+// Copyright 1998-2019 Epic Games, Inc. All Rights Reserved.
+
+#include "TimingEventsTrack.h"
+
+#define LOCTEXT_NAMESPACE "TimingEventsTrack"
+
+////////////////////////////////////////////////////////////////////////////////////////////////////
+// FTimingEventsTrackLayout
+////////////////////////////////////////////////////////////////////////////////////////////////////
+
+void FTimingEventsTrackLayout::ForceNormalMode()
+{
+	bIsCompactMode = false;
+	EventH = NormalLayoutEventH;
+	EventDY = NormalLayoutEventDY;
+	TimelineDY = NormalLayoutTimelineDY;
+	MinTimelineH = NormalLayoutMinTimelineH;
+	TargetMinTimelineH = NormalLayoutMinTimelineH;
+}
+
+////////////////////////////////////////////////////////////////////////////////////////////////////
+
+void FTimingEventsTrackLayout::ForceCompactMode()
+{
+	bIsCompactMode = true;
+	EventH = CompactLayoutEventH;
+	EventDY = CompactLayoutEventDY;
+	TimelineDY = CompactLayoutTimelineDY;
+	MinTimelineH = CompactLayoutMinTimelineH;
+	TargetMinTimelineH = CompactLayoutMinTimelineH;
+}
+
+////////////////////////////////////////////////////////////////////////////////////////////////////
+
+void FTimingEventsTrackLayout::Update()
+{
+	constexpr float LayoutTransitionSpeed = 0.25f;
+
+	if (bIsCompactMode)
+	{
+		if (EventH > CompactLayoutEventH)
+		{
+			EventH -= LayoutTransitionSpeed;
+		}
+		if (EventDY > CompactLayoutEventDY)
+		{
+			EventDY -= LayoutTransitionSpeed;
+		}
+		if (TimelineDY > CompactLayoutTimelineDY)
+		{
+			TimelineDY -= LayoutTransitionSpeed;
+		}
+	}
+	else
+	{
+		if (EventH < NormalLayoutEventH)
+		{
+			EventH += LayoutTransitionSpeed;
+		}
+		if (EventDY < NormalLayoutEventDY)
+		{
+			EventDY += LayoutTransitionSpeed;
+		}
+		if (TimelineDY < NormalLayoutTimelineDY)
+		{
+			TimelineDY += LayoutTransitionSpeed;
+		}
+	}
+
+	if (MinTimelineH > TargetMinTimelineH)
+	{
+		MinTimelineH -= LayoutTransitionSpeed;
+	}
+	else if (MinTimelineH < TargetMinTimelineH)
+	{
+		MinTimelineH += LayoutTransitionSpeed;
+	}
+}
+
+////////////////////////////////////////////////////////////////////////////////////////////////////
+// FTimingEventsTrack
+////////////////////////////////////////////////////////////////////////////////////////////////////
+
+FTimingEventsTrack::FTimingEventsTrack(uint64 InTrackId, ETimingEventsTrackType InType, const FString& InName, const TCHAR* InGroupName)
+	: FBaseTimingTrack(InTrackId)
+	, Type(InType)
+	, Name(InName)
+	, GroupName(InGroupName)
+	, ThreadId(0)
+	, Order(0)
+	, Depth(-1)
+	, bIsCollapsed(false)
+{
+}
+
+////////////////////////////////////////////////////////////////////////////////////////////////////
+
+FTimingEventsTrack::~FTimingEventsTrack()
+{
+}
+
+////////////////////////////////////////////////////////////////////////////////////////////////////
+
+void FTimingEventsTrack::Reset()
+{
+	FBaseTimingTrack::Reset();
+
+	Depth = -1;
+	bIsCollapsed = false;
+}
+
+////////////////////////////////////////////////////////////////////////////////////////////////////
+
+void FTimingEventsTrack::UpdateHoveredState(float MouseX, float MouseY, const FTimingTrackViewport& Viewport)
+{
+	constexpr float HeaderWidth = 100.0f;
+	constexpr float HeaderHeight = 14.0f;
+
+	if (MouseY >= GetPosY() && MouseY < GetPosY() + GetHeight())
+	{
+		SetHoveredState(true);
+		SetHeaderHoveredState(MouseX < HeaderWidth && MouseY < GetPosY() + HeaderHeight);
+	}
+	else
+	{
+		SetHoveredState(false);
+	}
+}
+
+////////////////////////////////////////////////////////////////////////////////////////////////////
+
+#undef LOCTEXT_NAMESPACE
diff --git a/Engine/Source/Developer/TraceInsights/Private/Insights/ViewModels/TimingEventsTrack.h b/Engine/Source/Developer/TraceInsights/Private/Insights/ViewModels/TimingEventsTrack.h
new file mode 100644
--- /dev/null
+++ b/Engine/Source/Developer/TraceInsights/Private/Insights/ViewModels/TimingEventsTrack.h
@@ -0,0 +1,93 @@
+// Copyright 1998-2019 Epic Games, Inc. All Rights Reserved.
+
+#pragma once
+
+#include "CoreMinimal.h"
+
+#include "Insights/ViewModels/BaseTimingTrack.h"
+
+////////////////////////////////////////////////////////////////////////////////////////////////////
+
+static const float RealMinTimelineH = 13.0f;
+
+static const float NormalLayoutEventH = 14.0f;
+static const float NormalLayoutEventDY = 2.0f;
+static const float NormalLayoutTimelineDY = 14.0f;
+static const float NormalLayoutMinTimelineH = 0.0f;
+
+static const float CompactLayoutEventH = 2.0f;
+static const float CompactLayoutEventDY = 1.0f;
+static const float CompactLayoutTimelineDY = 3.0f;
+static const float CompactLayoutMinTimelineH = 0.0f;
+
+////////////////////////////////////////////////////////////////////////////////////////////////////
+
+struct FTimingEventsTrackLayout
+{
+	bool bIsCompactMode;
+
+	float EventH; // height of a timing event, in Slate units
+	float EventDY; // vertical distance between two timing event sub-tracks, in Slate units
+	float TimelineDY; // space at top and bottom of each timeline, in Slate units
+	float MinTimelineH;
+	float TargetMinTimelineH;
+
+	float GetLaneY(uint32 Depth) const { return 1.0f + TimelineDY + Depth * (EventDY + EventH); }
+
+	void ForceNormalMode();
+	void ForceCompactMode();
+	void Update();
+};
+
+////////////////////////////////////////////////////////////////////////////////////////////////////
+
+enum class ETimingEventsTrackType
+{
+	Cpu,
+	Gpu,
+	Loading,
+	Io,
+};
+
+////////////////////////////////////////////////////////////////////////////////////////////////////
+
+class FTimingEventsTrack : public FBaseTimingTrack
+{
+public:
+	FTimingEventsTrack(uint64 InTrackId, ETimingEventsTrackType InType, const FString& InName, const TCHAR* InGroupName);
+	virtual ~FTimingEventsTrack();
+
+	ETimingEventsTrackType GetType() const { return Type; }
+	const FString& GetName() const { return Name; }
+
+	const TCHAR* GetGroupName() const { return GroupName; };
+
+	void SetThreadId(uint32 InThreadId) { ThreadId = InThreadId; }
+	uint32 GetThreadId() const { return ThreadId; }
+
+	void SetOrder(int32 InOrder) { Order = InOrder; }
+	int32 GetOrder() const { return Order; }
+
+	int32 GetDepth() const { return Depth; }
+	int32 GetNumLanes() const { return Depth + 1; }
+
+	virtual void Reset() override;
+	virtual void UpdateHoveredState(float MX, float MY, const FTimingTrackViewport& Viewport);
+
+public:
+	ETimingEventsTrackType Type;
+	FString Name;
+	const TCHAR* GroupName;
+	uint32 ThreadId;
+	int32 Order;
+	int32 Depth; // number of lanes == Depth + 1
+	bool bIsCollapsed;
+
+	// Cached OnPaint state.
+	//TArray<FEventBoxInfo> Boxes;
+	//TArray<FEventBoxInfo> MergedBorderBoxes;
+	//TArray<FEventBoxInfo> Borders;
+	//TArray<FTextInfo> Texts;
+};
+
+////////////////////////////////////////////////////////////////////////////////////////////////////
diff --git a/Engine/Source/Developer/TraceInsights/Private/Insights/ViewModels/TimingTrackViewport.h b/Engine/Source/Developer/TraceInsights/Private/Insights/ViewModels/TimingTrackViewport.h
new file mode 100644
--- /dev/null
+++ b/Engine/Source/Developer/TraceInsights/Private/Insights/ViewModels/TimingTrackViewport.h
@@ -0,0 +1,109 @@
+// Copyright 1998-2019 Epic Games, Inc. All Rights Reserved.
+
+#pragma once
+
+#include "CoreMinimal.h"
+
+class FTimingTrackViewport
+{
+public:
+	float Width; // width of viewport, in Slate units
+	float Height; // height of viewport, in Slate units
+
+	double MinValidTime; // min session time, in seconds
+	double MaxValidTime; // max session time, in seconds
+	double StartTime; // time of viewport's left side, in seconds
+	double EndTime; // time of viewport's right side, in seconds; computed when StartTime or Width changes
+	double MinScaleX;
+	double MaxScaleX;
+	double ScaleX; // scale factor between seconds and pixels (Slate units), in pixels (Slate units) per second
+
+	float TopOffset; // top offset (to allow the time ruller to be visible), in pixels (Slate units)
+	float ScrollHeight; // height of the vertical scrollable area, in pixels (Slate units)
+	float ScrollPosY; // current vertical scroll position, in pixels (Slate units)
+
+public:
+	FTimingTrackViewport()
+	{
+		Reset();
+	}
+
+	void Reset()
+	{
+		Width = 0.0f;
+		Height = 0.0f;
+
+		MinValidTime = 0.0;
+		MaxValidTime = 0.0;
+		StartTime = 0.0;
+		EndTime = 0.0;
+		MinScaleX = (5 * 20) / 3600.0; // 1h between major tick marks
+		MaxScaleX = 1.0E10; // 10ns between major tick marks
+		ScaleX = (5 * 20) / 5.0; // 5s between major tick marks
+
+		TopOffset = 0.0f;
+		ScrollHeight = 1.0f;
+		ScrollPosY = 0.0f;
+	}
+
+	float TimeToSlateUnits(const double Time) const
+	{
+		return static_cast<float>((Time - StartTime) * ScaleX);
+	}
+
+	float TimeToSlateUnitsRounded(const double Time) const
+	{
+		return static_cast<float>(FMath::RoundToDouble((Time - StartTime) * ScaleX));
+	}
+
+	float GetViewportDXForDuration(const double DT) const
+	{
+		return static_cast<float>(DT * ScaleX);
+	}
+
+	double GetDurationForViewportDX(const float DX) const
+	{
+		return static_cast<double>(DX) / ScaleX;
+	}
+
+	double SlateUnitsToTime(const float X) const
+	{
+		return StartTime + static_cast<double>(X) / ScaleX;
+	}
+
+	bool UpdateSize(const float InWidth, const float InHeight)
+	{
+		if (Width != InWidth || Height != InHeight)
+		{
+			Width = InWidth;
+			Height = InHeight;
+			EndTime = SlateUnitsToTime(Width);
+			return true;
+		}
+		return false;
+	}
+
+	double AlignTimeToPixel(const double InTime, const double InScaleX) const
+	{
+		return FMath::RoundToDouble(InTime * InScaleX) / InScaleX;
+	}
+
+	double AlignTimeToPixel(const double Time) const
+	{
+		return AlignTimeToPixel(Time, ScaleX);
+	}
+
+	bool ScrollAtTime(const double Time);
+	bool CenterOnTimeInterval(const double Time, const double Duration);
+
+	bool ZoomOnTimeInterval(const double Time, const double Duration);
+	bool ZoomWithFixedX(const double NewScaleX, const float X);
+
+	double RestrictEndTime(const double InEndTime) const;
+	double RestrictDuration(const double InStartTime, const double InEndTime) const;
+
+	void GetHorizontalScrollLimits(double& OutMinT, double& OutMaxT);
+	bool EnforceHorizontalScrollLimits(const double U);
+
+	float GetViewportY(const float Y) const { return TopOffset + Y - ScrollPosY; }
+};
diff --git a/Engine/Source/Developer/TraceInsights/Private/Insights/ViewModels/TimingViewDrawHelper.cpp b/Engine/Source/Developer/TraceInsights/Private/Insights/ViewModels/TimingViewDrawHelper.cpp
new file mode 100644
--- /dev/null
+++ b/Engine/Source/Developer/TraceInsights/Private/Insights/ViewModels/TimingViewDrawHelper.cpp
@@ -0,0 +1,571 @@
+// Copyright 1998-2019 Epic Games, Inc. All Rights Reserved.
+
+#include "TimingViewDrawHelper.h"
+
+#include "Brushes/SlateBorderBrush.h"
+#include "Brushes/SlateColorBrush.h"
+#include "EditorStyleSet.h"
+#include "Fonts/FontMeasure.h"
+#include "Fonts/SlateFontInfo.h"
+#include "Framework/Application/SlateApplication.h"
+#include "Rendering/DrawElements.h"
+#include "Styling/CoreStyle.h"
+#include "TraceServices/AnalysisService.h"
+
+// Insights
+#include "Insights/Common/PaintUtils.h"
+#include "Insights/Common/TimeUtils.h"
+#include "Insights/ViewModels/MarkersTimingTrack.h"
+#include "Insights/ViewModels/TimingEventsTrack.h"
+#include "Insights/ViewModels/TimingTrackViewport.h"
+
+////////////////////////////////////////////////////////////////////////////////////////////////////
+
+FTimingViewDrawHelper::FTimingViewDrawHelper(const FDrawContext& InDC, const FTimingTrackViewport& InViewport, const FTimingEventsTrackLayout& InLayout)
+	: DrawContext(InDC)
+	, Viewport(InViewport)
+	, Layout(InLayout)
+	, WhiteBrush(FCoreStyle::Get().GetBrush("WhiteBrush"))
+	, BorderBrush(FEditorStyle::GetBrush("PlainBorder"))
+	, EventsBorderBrush(new FSlateBorderBrush(NAME_None, FMargin(1.0f)))
+	, BackgroundAreaBrush(WhiteBrush)
+	, ValidAreaColor(0.07f, 0.07f, 0.07f, 1.0f)
+	, InvalidAreaColor(0.1f, 0.07f, 0.07f, 1.0f)
+	, EventFont(FCoreStyle::GetDefaultFontStyle("Regular", 8))
+	, Stats()
+{
+}
+
+////////////////////////////////////////////////////////////////////////////////////////////////////
+
+FTimingViewDrawHelper::~FTimingViewDrawHelper()
+{
+	delete EventsBorderBrush; //im: is it safe to delete brush (i.e in OnPaint, imediately after use)?
+}
+
+////////////////////////////////////////////////////////////////////////////////////////////////////
+
+void FTimingViewDrawHelper::DrawBackground() const
+{
+	const float X0 = Viewport.TimeToSlateUnitsRounded(0.0);
+	const float X1 = Viewport.TimeToSlateUnitsRounded(Viewport.MaxValidTime);
+	const float W = FMath::CeilToFloat(Viewport.Width);
+	const float H = FMath::CeilToFloat(Viewport.Height);
+
+	if (X0 >= W || X1 <= 0.0f)
+	{
+		ValidX0 = 0.0f;
+		ValidX1 = W;
+
+		// Draw invalid area (entire view).
+		DrawContext.DrawBox(0.0f, 0.0f, W, H, BackgroundAreaBrush, InvalidAreaColor);
+	}
+	else // X0 < W && X1 > 0
+	{
+		if (X0 > 0.0f)
+		{
+			// Draw invalid area (left).
+			DrawContext.DrawBox(0.0f, 0.0f, X0, H, BackgroundAreaBrush, InvalidAreaColor);
+		}
+
+		if (X1 < W)
+		{
+			// Draw invalid area (right).
+			DrawContext.DrawBox(X1, 0.0f, W - X1, H, BackgroundAreaBrush, InvalidAreaColor);
+		}
+
+		ValidX0 = FMath::Max(X0, 0.0f);
+		ValidX1 = FMath::Min(X1, W);
+
+		if (ValidX1 > ValidX0)
+		{
+			// Draw valid area.
+			DrawContext.DrawBox(ValidX0, 0.0f, ValidX1 - ValidX0, H, BackgroundAreaBrush, ValidAreaColor);
+		}
+	}
+
+	DrawContext.LayerId++;
+}
+
+////////////////////////////////////////////////////////////////////////////////////////////////////
+
+void FTimingViewDrawHelper::BeginTimelines()
+{
+	TimelineIndex = -1;
+	TimelineTopY = -Viewport.ScrollPosY;
+
+	float Y = TimelineTopY + Viewport.TopOffset;
+
+	if (Y > 0.0f && ValidX1 > ValidX0)
+	{
+		Y = FMath::Min(Y, Viewport.Height);
+
+		// Draw invalid area (top).
+		DrawContext.DrawBox(DrawContext.LayerId + ToInt32(EDrawLayer::TimelineHeader), ValidX0, 0.0f, ValidX1 - ValidX0, Y, BackgroundAreaBrush, InvalidAreaColor);
+	}
+}
+
+////////////////////////////////////////////////////////////////////////////////////////////////////
+
+bool FTimingViewDrawHelper::BeginTimeline(FTimingEventsTrack& Track)
+{
+	TimelineIndex++;
+
+	Track.SetPosY(TimelineTopY + Viewport.ScrollPosY);
+
+	if (Track.GetHeight() < Layout.MinTimelineH)
+		Track.SetHeight(Layout.MinTimelineH);
+
+	if (TimelineTopY + Track.GetHeight() < -1.0f)
+	{
+		TimelineTopY += Track.GetHeight();
+		return false;
+	}
+
+	if (TimelineTopY > Viewport.Height - Viewport.TopOffset)
+	{
+		return false;
+	}
+
+	MaxDepth = -1;
+
+	// +1.0f is for horizontal line between timelines
+	TimelineY = Viewport.TopOffset + TimelineTopY + 1.0f + Layout.TimelineDY;
+
+	// Reset "last event X2" for all depths.
+	LastEventX2.Reset();
+
+	// Reset "last box" for all depths.
+	LastBox.Reset();
+
+	return true;
+}
+
+////////////////////////////////////////////////////////////////////////////////////////////////////
+
+void FTimingViewDrawHelper::EndTimeline(FTimingEventsTrack& Track)
+{
+	// Flush merged boxes.
+	for (int32 Depth = 0; Depth <= MaxDepth; ++Depth)
+	{
+		const FBoxData& Box = LastBox[Depth];
+		if (Box.X1 < Box.X2)
+		{
+			const float EventY = TimelineY + (Layout.EventH + Layout.EventDY) * Depth;
+			DrawBox(Box, EventY, Layout.EventH);
+		}
+	}
+
+	Track.Depth = MaxDepth;
+
+	float NewH;
+	if (MaxDepth < 0)
+	{
+		NewH = Layout.MinTimelineH;
+	}
+	else //if (MaxDepth >= 0)
+	{
+		// 1.0f is for horizontal line between timelines
+		NewH = 1.0f + Layout.TimelineDY + Layout.EventH * (MaxDepth + 1) + Layout.EventDY * MaxDepth + Layout.TimelineDY;
+
+		if (NewH < RealMinTimelineH)
+		{
+			NewH = RealMinTimelineH;
+		}
+	}
+
+	if (Track.GetHeight() < NewH)
+	{
+		float NewTrackHeight;
+		if (Layout.bIsCompactMode)
+		{
+			NewTrackHeight = FMath::CeilToFloat(Track.GetHeight() * 0.5f + NewH * 0.5f);
+		}
+		else
+		{
+			NewTrackHeight = FMath::CeilToFloat(Track.GetHeight() * 0.9f + NewH * 0.1f);
+		}
+		Track.SetHeight(NewTrackHeight);
+	}
+	else if (Track.GetHeight() > NewH)
+	{
+		float NewTrackHeight;
+		if (Layout.bIsCompactMode)
+		{
+			NewTrackHeight = FMath::FloorToFloat(Track.GetHeight() * 0.5f + NewH * 0.5f);
+		}
+		else
+		{
+			NewTrackHeight = FMath::FloorToFloat(Track.GetHeight() * 0.9f + NewH * 0.1f);
+		}
+		Track.SetHeight(NewTrackHeight);
+	}
+
+	if (Track.GetHeight() > 0)
+	{
+		const float Y = Viewport.TopOffset + TimelineTopY;
+
+		FLinearColor Color(0.05f, 0.05f, 0.05f, 1.0f);
+		FLinearColor TextColor(1.0f, 1.0f, 1.0f, 1.0f);
+
+		if (Track.IsSelected())
+		{
+			if (Track.IsHovered())
+			{
+				TextColor = FLinearColor(1.0f, 1.0f, 0.0f, 1.0f);
+			}
+			else
+			{
+				TextColor = FLinearColor(1.0f, 1.0f, 0.5f, 1.0f);
+			}
+		}
+		else if (Track.IsHovered())
+		{
+			TextColor = FLinearColor(1.0f, 1.0f, 0.0f, 1.0f);
+		}
+
+		// Draw a horizontal line between timelines.
+		DrawContext.DrawBox(DrawContext.LayerId + ToInt32(EDrawLayer::TimelineHeader), 0, Y, Viewport.Width, 1.0f, WhiteBrush, Color);
+
+		// Draw name of timeline.
+		//const FString Name = FString::Printf(TEXT("%d Y: %g, H: %g, VO: %g, VY: %g"), TimelineIndex, Track.Y, Track.H, Viewport.TopOffset, Viewport.ScrollPosY); //debug
+		const TSharedRef<FSlateFontMeasure> FontMeasureService = FSlateApplication::Get().GetRenderer()->GetFontMeasureService();
+		float NameWidth = FontMeasureService->Measure(Track.GetName(), EventFont).X;
+		DrawContext.DrawBox(DrawContext.LayerId + ToInt32(EDrawLayer::TimelineHeader), 0.0f, Y + 1.0f, NameWidth + 4.0f, 12.0f, WhiteBrush, Color);
+		DrawContext.DrawText(DrawContext.LayerId + ToInt32(EDrawLayer::TimelineText), 2.0f, Y, Track.GetName(), EventFont, TextColor);
+	}
+
+	TimelineTopY += Track.GetHeight();
+}
+
+////////////////////////////////////////////////////////////////////////////////////////////////////
+
+void FTimingViewDrawHelper::AddEvent(double EventStartTime, double EventEndTime, uint32 EventDepth, const TCHAR* EventName, uint32 Color)
+{
+	Stats.NumEvents++;
+
+	float EventX1 = Viewport.TimeToSlateUnitsRounded(EventStartTime);
+	if (EventX1 > Viewport.Width)
+	{
+		return;
+	}
+
+	double RestrictedEndTime = Viewport.RestrictEndTime(EventEndTime);
+	float EventX2 = Viewport.TimeToSlateUnitsRounded(RestrictedEndTime);
+	if (EventX2 < 0)
+	{
+		return;
+	}
+
+	// Timing events are displayed with minimum 1px (including empty ones).
+	if (EventX1 == EventX2)
+	{
+		EventX2 = EventX1 + 1.0f;
+	}
+
+	const int32 Depth = static_cast<int32>(EventDepth);
+	if (Depth > MaxDepth)
+	{
+		MaxDepth = Depth;
+	}
+
+	// Ensure we have enough slots in array. See LastBox[Depth] usage.
+	if (LastBox.Num() <= Depth)
+	{
+		LastBox.AddDefaulted(Depth + 1 - LastBox.Num());
+	}
+
+	float EventY = TimelineY + (Layout.EventH + Layout.EventDY) * Depth;
+	if (EventY < -Layout.EventH || EventY > Viewport.Height)
+	{
+		return;
+	}
+
+	// Ensure we have enough slots in array. See LastEventX2[Depth] usage.
+	while (LastEventX2.Num() <= Depth)
+	{
+		LastEventX2.Add(-1.0f);
+	}
+
+	// Limit event width on the viewport's left side.
+	// This also makes the text to be displayed in viewport, for very long events.
+	const float MinX = -2.0f; // -2 allows event border to remain outside screen
+	if (EventX1 < MinX)
+	{
+		EventX1 = MinX;
+	}
+
+	// Limit event width on the viewport's right side.
+	const float MaxX = Viewport.Width + 2.0f; // +2 allows event border to remain outside screen
+	if (EventX2 > MaxX)
+	{
+		EventX2 = MaxX;
+	}
+
+	float EventW = EventX2 - EventX1;
+
+	// Optimization...
+	if (EventW == 1.0f && EventX1 == LastEventX2[Depth] - 1.0f)
+	{
+		// Do no draw 1 pixel event if the last event was ended on that pixel.
+		return;
+	}
+
+	//////////////////////////////////////////////////
+	// Coloring
+
+	FLinearColor EventColorFill;
+
+	//im:TODO: EventColorFill = GetEventColorFn(Event);
+	if (Color == 0)
+	{
+		uint32 NameHash = 0;
+		for (const TCHAR* c = EventName; *c; ++c)
+		{
+			NameHash = (NameHash + *c) * 0x2c2c57ed;
+		}
+
+		Color = NameHash | 0xFF000000;
+
+		// Increase brightness.
+		//Color = ((NameHash | ((NameHash & 0x00808080) >> 1) | ((NameHash & 0x00808080) >> 2) | ((NameHash & 0x00808080) >> 3)) << 1) | 0xFF000000;
+
+		// Increase brightness.
+		//EventColorFill.R = ((NameHash >> 16) & 0xFF) / 128.0f;
+		//EventColorFill.G = ((NameHash >> 8) & 0xFF) / 128.0f;
+		//EventColorFill.B = (NameHash & 0xFF) / 128.0f;
+		//EventColorFill.A = 1.0f;
+	}
+	//else
+	{
+		EventColorFill.R = ((Color >> 16) & 0xFF) / 255.0f;
+		EventColorFill.G = ((Color >>  8) & 0xFF) / 255.0f;
+		EventColorFill.B = ((Color      ) & 0xFF) / 255.0f;
+		EventColorFill.A = ((Color >> 24) & 0xFF) / 255.0f;
+	}
+
+	// Maps hash to color directly.
+	//EventColorFill.R = ((Color >> 16) & 0xFF) / 255.0f;
+	//EventColorFill.G = ((Color >> 8) & 0xFF) / 255.0f;
+	//EventColorFill.B = (Color & 0xFF) / 255.0f;
+	//EventColorFill.A = 1.0f;
+
+	// Maps hash to color (sRGB to linear).
+	//const FColor ColorFill(Color);
+	//EventColorFill = ColorFill;
+
+	// Invereted color.
+	//EventColorFill.R = 1.0f - ((Color >> 16) & 0xFF) / 255.0f;
+	//EventColorFill.G = 1.0f - ((Color >> 8) & 0xFF) / 255.0f;
+	//EventColorFill.B = 1.0f - (Color & 0xFF) / 255.0f;
+	//EventColorFill.A = 1.0f;
+
+	// Bright color.
+	//EventColorFill.R = ((Color >> 16) & 0xFF) / 128.0f;
+	//EventColorFill.G = ((Color >> 8) & 0xFF) / 128.0f;
+	//EventColorFill.B = (Color & 0xFF) / 128.0f;
+	//EventColorFill.A = 1.0f;
+
+	// Bright color (adjusted R, G, B).
+	//const float BrightnessFactor = 0.59f / 255.0f; // [0 .. 0.59f / 255.0f]
+	//EventColorFill.R = ((Color >> 16) & 0xFF) * (BrightnessFactor / 0.30f);
+	//EventColorFill.G = ((Color >>  8) & 0xFF) * (BrightnessFactor / 0.59f);
+	//EventColorFill.B = ((Color      ) & 0xFF) * (BrightnessFactor / 0.11f);
+	//EventColorFill.A = 1.0f;
+
+	// Bright desaturated color.
+	//const float MinL = 0.3f;
+	//const float MulL = (1.0f - MinL) / 255.0f;
+	//EventColorFill.R = MinL + ((Color >> 16) & 0xFF) * MulL;
+	//EventColorFill.G = MinL + ((Color >> 8) & 0xFF) * MulL;
+	//EventColorFill.B = MinL + (Color & 0xFF) * MulL;
+	//EventColorFill.A = 1.0f;
+
+	// Dark desaturated color.
+	//const float MaxL = 0.5f;
+	//const float MulL = MaxL / 255.0f;
+	//EventColorFill.R = (((Color >> 16) & 0xFF) * MulL;
+	//EventColorFill.G = ((Color >> 8) & 0xFF) * MulL;
+	//EventColorFill.B = (Color & 0xFF) * MulL;
+	//EventColorFill.A = 1.0f;
+
+	//constexpr float BorderColorFactor = 0.75f; // darker border
+	constexpr float BorderColorFactor = 1.25f; // brighter border
+
+	//////////////////////////////////////////////////
+
+	if (EventW > 2.0f && Layout.EventH > 2.0f)
+	{
+		// Fill inside of the timing event box.
+		DrawContext.DrawBox(DrawContext.LayerId + ToInt32(EDrawLayer::EventFill), EventX1 + 1.0f, EventY + 1.0f, EventW - 2.0f, Layout.EventH - 2.0f, WhiteBrush, EventColorFill);
+		Stats.NumDrawBoxes++;
+	}
+
+	// Save X2, for current depth.
+	LastEventX2[Depth] = EventX2;
+
+	// Draw border around the timing event box.
+	if (EventW > 2.0f)
+	{
+		FBoxData& Box = LastBox[Depth];
+		if (Box.X1 < Box.X2)
+		{
+			DrawBox(Box, EventY, Layout.EventH);
+			Box.Reset();
+		}
+
+		const FLinearColor EventColorBorder(EventColorFill.R * BorderColorFactor, EventColorFill.G * BorderColorFactor, EventColorFill.B * BorderColorFactor, 1.0f);
+		DrawContext.DrawBox(DrawContext.LayerId + ToInt32(EDrawLayer::EventBorder), EventX1, EventY, EventW, Layout.EventH, EventsBorderBrush, EventColorBorder);
+		Stats.NumDrawBorders++;
+	}
+	else // 1px or 2px boxes
+	{
+		FBoxData& Box = LastBox[Depth];
+
+		// Check if we can merge this box with previous one, if any.
+		// Note: We are assuming events are processed in sorted order by X1.
+		if (Color == Box.Color && // same color
+			EventX1 <= Box.X2) // overlapping or adjacent
+		{
+			// Merge it with previous box.
+			Box.X2 = EventX2;
+			Stats.NumMergedBoxes++;
+		}
+		else
+		{
+			// Flush previous box, if any.
+			if (Box.X1 < Box.X2)
+			{
+				DrawBox(Box, EventY, Layout.EventH);
+			}
+
+			// Start new "merge box".
+			Box.X1 = EventX1;
+			Box.X2 = EventX2;
+			Box.Color = Color;
+			Box.LinearColor = FLinearColor(EventColorFill.R * BorderColorFactor, EventColorFill.G * BorderColorFactor, EventColorFill.B * BorderColorFactor, 1.0f);
+		}
+	}
+
+	// Draw the name of the timing event.
+	if (EventW > 8.0f && Layout.EventH > 10.0f)
+	{
+		FString Name = EventName;// +TEXT(" [") + FText::AsNumber(Event.Type->Id).ToString() + TEXT("]");
+		if (EventW > Name.Len() * 2.0f + 48.0f)
+		{
+			const double Duration = EventEndTime - EventStartTime;
+			Name += TEXT(" (");
+			Name += TimeUtils::FormatTimeAuto(Duration);
+			Name += TEXT(")");
+		}
+
+		const TSharedRef<FSlateFontMeasure> FontMeasureService = FSlateApplication::Get().GetRenderer()->GetFontMeasureService();
+		const int32 LastWholeCharacterIndex = FontMeasureService->FindLastWholeCharacterIndexBeforeOffset(Name, EventFont, FMath::RoundToInt(EventW - 2.0f));
+
+		if (LastWholeCharacterIndex >= 0)
+		{
+			// Grey threshold is shifted toward black (0.4 instead of 0.5 in test below) due to "area rule":
+			// a large gray surface (background of a timing event in this case) is perceived lighter than a smaller area (text pixels).
+			// Ref: https://books.google.ro/books?id=0pVr7dhmdWYC
+			const bool bIsDarkColor = (EventColorFill.ComputeLuminance() < 0.4f);
+
+			DrawContext.DrawText(DrawContext.LayerId + ToInt32(EDrawLayer::EventText), EventX1 + 2.0f, EventY + 1.0f, Name, 0, LastWholeCharacterIndex + 1, EventFont, bIsDarkColor ? FLinearColor::White : FLinearColor::Black);
+			Stats.NumDrawTexts++;
+		}
+	}
+}
+
+////////////////////////////////////////////////////////////////////////////////////////////////////
+
+void FTimingViewDrawHelper::DrawBox(const FBoxData& Box, const float EventY, const float EventH)
+{
+	DrawContext.DrawBox(DrawContext.LayerId + ToInt32(EDrawLayer::EventBorder), Box.X1, EventY, Box.X2 - Box.X1, EventH, WhiteBrush, Box.LinearColor);
+	Stats.NumDrawBoxes++;
+}
+
+////////////////////////////////////////////////////////////////////////////////////////////////////
+
+void FTimingViewDrawHelper::EndTimelines()
+{
+	float Y = Viewport.TopOffset + TimelineTopY;
+
+	// Draw a last horizontal line.
+	DrawContext.DrawBox(DrawContext.LayerId + ToInt32(EDrawLayer::TimelineHeader), 0, Y, Viewport.Width, 1.0f, WhiteBrush, FLinearColor(0.05f, 0.05f, 0.05f, 1.0f));
+	Y += 1.0f;
+
+	if (Y < Viewport.Height && ValidX1 > ValidX0)
+	{
+		Y = FMath::Max(Y, 0.0f);
+
+		// Draw invalid area (bottom).
+		DrawContext.DrawBox(DrawContext.LayerId + ToInt32(EDrawLayer::TimelineHeader), ValidX0, Y, ValidX1 - ValidX0, Viewport.Height - Y, BackgroundAreaBrush, InvalidAreaColor);
+	}
+
+	DrawContext.LayerId += ToInt32(EDrawLayer::Count);
+}
+
+////////////////////////////////////////////////////////////////////////////////////////////////////
+
+void FTimingViewDrawHelper::DrawTimingEventHighlight(double StartTime, double EndTime, float Y, EHighlightMode Mode)
+{
+	float EventX1 = Viewport.TimeToSlateUnitsRounded(StartTime);
+	if (EventX1 > Viewport.Width)
+	{
+		return;
+	}
+
+	EndTime = Viewport.RestrictEndTime(EndTime);
+	float EventX2 = Viewport.TimeToSlateUnitsRounded(EndTime);
+	if (EventX2 < 0)
+	{
+		return;
+	}
+
+	if (EventX1 == EventX2)
+	{
+		EventX2 = EventX1 + 1.0f;
+	}
+
+	// Limit event width on the viewport's left side.
+	const float MinX = -2.0f; // -2 allows event border to remain outside screen
+	if (EventX1 < MinX)
+	{
+		EventX1 = MinX;
+	}
+
+	// Limit event width on the viewport's right side.
+	const float MaxX = Viewport.Width + 2.0f; // +2 allows event border to remain outside screen
+	if (EventX2 > MaxX)
+	{
+		EventX2 = MaxX;
+	}
+
+	float EventW = EventX2 - EventX1;
+
+	if (Mode == EHighlightMode::Hovered)
+	{
+		const FSlateBrush* HighlightBrush = BorderBrush;
+
+		const FLinearColor Color(1.0f, 1.0f, 0.0f, 1.0f); // yellow
+
+		// Draw border around the timing event box.
+		DrawContext.DrawBox(EventX1 - 2.0f, Y - 2.0f, EventW + 4.0f, Layout.EventH + 4.0f, HighlightBrush, Color);
+	}
+	else // EHighlightMode::Selected or EHighlightMode::SelectedAndHovered
+	{
+		const FSlateBrush* SelectedBrush = BorderBrush;
+
+		// Animate color from white (if selected and hovered) or yellow (if only selected) to black, using a squared sine function.
+		const double Time = static_cast<double>(FPlatformTime::Cycles64()) * FPlatformTime::GetSecondsPerCycle64();
+		float S = FMath::Sin(2.0 * Time);
+		S = S * S; // squared, to ensure only positive [0 - 1] values
+		const float Blue = (Mode == EHighlightMode::SelectedAndHovered) ? 0.0f : S;
+		const FLinearColor Color(S, S, Blue, 1.0f);
+
+		// Draw border around the timing event box.
+		DrawContext.DrawBox(EventX1 - 2.0f, Y - 2.0f, EventW + 4.0f, Layout.EventH + 4.0f, SelectedBrush, Color);
+	}
+	DrawContext.LayerId++;
+}
+
+////////////////////////////////////////////////////////////////////////////////////////////////////
diff --git a/Engine/Source/Developer/TraceInsights/Private/Insights/ViewModels/TimingViewDrawHelper.h b/Engine/Source/Developer/TraceInsights/Private/Insights/ViewModels/TimingViewDrawHelper.h
new file mode 100644
--- /dev/null
+++ b/Engine/Source/Developer/TraceInsights/Private/Insights/ViewModels/TimingViewDrawHelper.h
@@ -0,0 +1,157 @@
+// Copyright 1998-2019 Epic Games, Inc. All Rights Reserved.
+
+#pragma once
+
+#include "CoreMinimal.h"
+#include "Fonts/SlateFontInfo.h"
+#include "Math/Color.h"
+
+////////////////////////////////////////////////////////////////////////////////////////////////////
+
+struct FSlateBrush;
+struct FDrawContext;
+
+class FTimingEventsTrack;
+struct FTimingEventsTrackLayout;
+class FTimingTrackViewport;
+
+////////////////////////////////////////////////////////////////////////////////////////////////////
+
+class FTimingViewDrawHelper
+{
+public:
+	enum class EHighlightMode
+	{
+		Hovered,
+		Selected,
+		SelectedAndHovered
+	};
+
+private:
+	enum class EDrawLayer : int32
+	{
+		EventBorder,
+		EventFill,
+		EventText,
+		TimelineHeader,
+		TimelineText,
+
+		Count,
+	};
+	static int32 ToInt32(EDrawLayer Layer) { return static_cast<int32>(Layer); }
+
+	struct FBoxData
+	{
+		float X1;
+		float X2;
+		uint32 Color;
+		FLinearColor LinearColor;
+
+		FBoxData() : X1(0.0f), X2(0.0f), Color(0) {}
+		void Reset() { X1 = 0.0f; X2 = 0.0f; Color = 0; }
+	};
+
+	//struct FEventBoxInfo
+	//{
+	//	float X1;
+	//	float X2;
+	//	int32 Depth;
+	//	uint32 Color;
+	//}
+
+	//struct FEventTextInfo
+	//{
+	//	float X;
+	//	float Y;
+	//	FString Text;
+	//	bool bUseDarkTextColor; // true if text needs to be displayed in Black, otherwise will be displayed in White
+	//};
+
+	struct FStats
+	{
+		int32 NumEvents;
+		int32 NumDrawBoxes;
+		int32 NumMergedBoxes;
+		int32 NumDrawBorders;
+		int32 NumDrawTexts;
+		int32 NumDrawTimeMarkerBoxes;
+		int32 NumDrawTimeMarkerTexts;
+
+		FStats()
+			: NumEvents(0)
+			, NumDrawBoxes(0)
+			, NumMergedBoxes(0)
+			, NumDrawBorders(0)
+			, NumDrawTexts(0)
+			, NumDrawTimeMarkerBoxes(0)
+			, NumDrawTimeMarkerTexts(0)
+		{}
+	};
+
+public:
+	FTimingViewDrawHelper(const FDrawContext& InDrawContext, const FTimingTrackViewport& InViewport, const FTimingEventsTrackLayout& InLayout);
+	~FTimingViewDrawHelper();
+
+	/**
+	 * Non-copyable
+	 */
+	FTimingViewDrawHelper(const FTimingViewDrawHelper&) = delete;
+	FTimingViewDrawHelper& operator=(const FTimingViewDrawHelper&) = delete;
+
+	const FDrawContext& GetDrawContext() const { return DrawContext; }
+	const FTimingTrackViewport& GetViewport() const { return Viewport; }
+	const FTimingEventsTrackLayout& GetLayout() const { return Layout; }
+
+	int32 GetNumEvents() const              { return Stats.NumEvents; }
+	int32 GetNumDrawBoxes() const           { return Stats.NumDrawBoxes; }
+	int32 GetNumMergedBoxes() const         { return Stats.NumMergedBoxes; }
+	int32 GetNumDrawBorders() const         { return Stats.NumDrawBorders; }
+	int32 GetNumDrawTexts() const           { return Stats.NumDrawTexts; }
+	int32 GetNumDrawTimeMarkerBoxes() const { return Stats.NumDrawTimeMarkerBoxes; }
+	int32 GetNumDrawTimeMarkerTexts() const { return Stats.NumDrawTimeMarkerTexts; }
+
+	void DrawBackground() const;
+	void DrawTimingEventHighlight(double StartTime, double EndTime, float Y, EHighlightMode Mode);
+
+	//TODO: move the following in a Builder class
+	void BeginTimelines();
+	bool BeginTimeline(FTimingEventsTrack& Track);
+	void AddEvent(double StartTime, double EndTime, uint32 Depth, const TCHAR* EventName, uint32 Color = 0);
+	void EndTimeline(FTimingEventsTrack& Track);
+	void EndTimelines();
+
+private:
+	void DrawBox(const FBoxData& Box, const float EventY, const float EventH);
+
+private:
+	const FDrawContext& DrawContext;
+	const FTimingTrackViewport& Viewport;
+	const FTimingEventsTrackLayout& Layout;
+
+	const FSlateBrush* WhiteBrush;
+	const FSlateBrush* BorderBrush;
+	const FSlateBrush* EventsBorderBrush;
+	const FSlateBrush* BackgroundAreaBrush;
+	const FLinearColor ValidAreaColor;
+	const FLinearColor InvalidAreaColor;
+	const FSlateFontInfo EventFont;
+
+	//static const FSlateColorBrush SolidWhiteBrush = FSlateColorBrush(FColorList::White);
+	//static const FSlateBrush BorderBrush = FSlateBorderBrush(NAME_None, FMargin(1.0f));
+
+	mutable float ValidX0;
+	mutable float ValidX1;
+
+	float TimelineTopY;
+	float TimelineY;
+	int32 MaxDepth;
+	int32 TimelineIndex;
+
+	TArray<float> LastEventX2; // X2 value for last event on each depth, for current timeline
+	TArray<FBoxData> LastBox;
+
+	/** Debug stats */
+	mutable FStats Stats;
+};
+
+////////////////////////////////////////////////////////////////////////////////////////////////////
diff --git a/Engine/Source/Developer/TraceInsights/Private/Insights/Widgets/SFrameTrack.cpp b/Engine/Source/Developer/TraceInsights/Private/Insights/Widgets/SFrameTrack.cpp
new file mode 100644
--- /dev/null
+++ b/Engine/Source/Developer/TraceInsights/Private/Insights/Widgets/SFrameTrack.cpp
@@ -0,0 +1,882 @@
+// Copyright 1998-2019 Epic Games, Inc. All Rights Reserved.
+
+#include "SFrameTrack.h"
+
+#include "Brushes/SlateBorderBrush.h"
+#include "Brushes/SlateColorBrush.h"
+#include "EditorStyleSet.h"
+#include "Fonts/FontMeasure.h"
+#include "Fonts/SlateFontInfo.h"
+#include "Framework/Application/SlateApplication.h"
+#include "HAL/PlatformTime.h"
+#include "Rendering/DrawElements.h"
+#include "TraceServices/AnalysisService.h"
+#include "TraceServices/SessionService.h"
+#include "Widgets/Layout/SScrollBar.h"
+
+// Insights
+#include "Insights/Common/PaintUtils.h"
+#include "Insights/Common/Stopwatch.h"
+#include "Insights/Common/TimeUtils.h"
+#include "Insights/InsightsManager.h"
+#include "Insights/TimingProfilerCommon.h"
+#include "Insights/TimingProfilerManager.h"
+#include "Insights/ViewModels/FrameTrackHelper.h"
+#include "Insights/Widgets/STimingProfilerWindow.h"
+#include "Insights/Widgets/STimingView.h"
+
+#include <limits>
+
+////////////////////////////////////////////////////////////////////////////////////////////////////
+
+#define LOCTEXT_NAMESPACE "SFrameTrack"
+
+////////////////////////////////////////////////////////////////////////////////////////////////////
+
+SFrameTrack::SFrameTrack()
+	: HoveredSample(nullptr, nullptr)
+{
+	Reset();
+}
+
+////////////////////////////////////////////////////////////////////////////////////////////////////
+
+SFrameTrack::~SFrameTrack()
+{
+}
+
+////////////////////////////////////////////////////////////////////////////////////////////////////
+
+void SFrameTrack::Reset()
+{
+	Viewport.Reset();
+	bIsViewportDirty = true;
+
+	CachedTimelines.Reset();
+	TimelinesOrder.Reset();
+	TimelinesOrder.Add(0);
+	TimelinesOrder.Add(1);
+	TimelinesOrder.Add(2);
+	bIsStateDirty = true;
+
+	AnalysisSyncNextTimestamp = 0;
+
+	MousePosition = FVector2D::ZeroVector;
+
+	MousePositionOnButtonDown = FVector2D::ZeroVector;
+	ViewportPosXOnButtonDown = 0.0f;
+	ViewportPosYOnButtonDown = 0.0f;
+
+	MousePositionOnButtonUp = FVector2D::ZeroVector;
+
+	bIsLMB_Pressed = false;
+	bIsRMB_Pressed = false;
+
+	bIsRMB_Scrolling = false;
+
+	SelectionStartFrameIndex = 0;
+	SelectionEndFrameIndex = 0;
+
+	HoveredSample.Timeline = nullptr;
+	HoveredSample.Sample = nullptr;
+
+	//ThisGeometry
+
+	CursorType = EFrameTrackCursor::Default;
+
+	NumUpdatedFrames = 0;
+	UpdateDurationHistory.Reset();
+	DrawDurationHistory.Reset();
+	OnPaintDurationHistory.Reset();
+	LastOnPaintTime = FPlatformTime::Cycles64();
+}
+
+////////////////////////////////////////////////////////////////////////////////////////////////////
+
+void SFrameTrack::Construct(const FArguments& InArgs)
+{
+	ChildSlot
+	[
+		SNew(SOverlay)
+		.Visibility(EVisibility::SelfHitTestInvisible)
+
+		+ SOverlay::Slot()
+		.VAlign(VAlign_Top)
+		.Padding(FMargin(0, 0, 0, 0))
+		[
+			SAssignNew(HorizontalScrollBar, SScrollBar)
+			.Orientation(Orient_Horizontal)
+			.AlwaysShowScrollbar(false)
+			.Visibility(EVisibility::Visible)
+			.Thickness(FVector2D(5.0f, 5.0f))
+			.RenderOpacity(0.75)
+			.OnUserScrolled(this, &SFrameTrack::HorizontalScrollBar_OnUserScrolled)
+		]
+	];
+
+	UpdateHorizontalScrollBar();
+
+	BindCommands();
+}
+
+////////////////////////////////////////////////////////////////////////////////////////////////////
+
+void SFrameTrack::Tick(const FGeometry& AllottedGeometry, const double InCurrentTime, const float InDeltaTime)
+{
+	if (ThisGeometry != AllottedGeometry || bIsViewportDirty)
+	{
+		bIsViewportDirty = false;
+		const float ViewWidth = AllottedGeometry.GetLocalSize().X;
+		const float ViewHeight = AllottedGeometry.GetLocalSize().Y;
+		Viewport.SetSize(ViewWidth, ViewHeight);
+		bIsStateDirty = true;
+	}
+
+	ThisGeometry = AllottedGeometry;
+
+	if (!bIsRMB_Scrolling)
+	{
+		// Elastic snap to horizontal limits (+15%).
+
+		float MinX, MaxX;
+		if (Viewport.MaxX - Viewport.MinX < Viewport.Width)
+		{
+			MinX = Viewport.MaxX - Viewport.Width;
+			MaxX = Viewport.MinX;
+		}
+		else
+		{
+			MinX = Viewport.MinX - 0.15f * Viewport.Width;
+			MaxX = Viewport.MaxX - 0.85f * Viewport.Width;
+		}
+		const float U = 0.5f;
+
+		float PosX = Viewport.PosX;
+		if (PosX < MinX)
+		{
+			PosX = PosX * U + (1.0f - U) * MinX;
+			if (FMath::IsNearlyEqual(PosX, MinX, 1.0f))
+				PosX = MinX;
+		}
+		else if (PosX > MaxX)
+		{
+			PosX = PosX * U + (1.0f - U) * MaxX;
+			if (FMath::IsNearlyEqual(PosX, MaxX, 1.0f))
+				PosX = MaxX;
+			if (PosX < MinX)
+				PosX = MinX;
+		}
+
+		if (Viewport.PosX != PosX)
+		{
+			//Viewport.ScrollAtPosX(PosX);
+			Viewport.ScrollAtIndex(Viewport.GetIndexAtPosX(PosX));
+			bIsStateDirty = true;
+		}
+	}
+
+	uint64 Time = FPlatformTime::Cycles64();
+	if (Time > AnalysisSyncNextTimestamp)
+	{
+		const uint64 WaitTime = static_cast<uint64>(0.1 / FPlatformTime::GetSecondsPerCycle64()); // 100ms
+		AnalysisSyncNextTimestamp = Time + WaitTime;
+
+		TSharedPtr<const Trace::IAnalysisSession> Session = FInsightsManager::Get()->GetSession();
+		if (Session.IsValid())
+		{
+			Trace::FAnalysisSessionReadScope SessionReadScope(*Session.Get());
+
+			const Trace::IFrameProvider& FramesProvider = Trace::ReadFrameProvider(*Session.Get());
+
+			for (int32 FrameType = 0; FrameType < TraceFrameType_Count; ++FrameType)
+			{
+				FFrameTrackTimeline& CachedTimeline = CachedTimelines.FindOrAdd(FrameType);
+				CachedTimeline.Id = FrameType;
+
+				int32 NumFrames = FramesProvider.GetFrameCount(static_cast<ETraceFrameType>(FrameType));
+				if (NumFrames > Viewport.MaxIndex)
+				{
+					Viewport.SetMinMaxIndexInterval(0, NumFrames);
+					UpdateHorizontalScrollBar();
+					bIsStateDirty = true;
+
+					// Auto zoom out.
+					float SampleW = Viewport.GetSampleWidth();
+					int32 NumSamples = FMath::CeilToInt(Viewport.Width / SampleW);
+					if (NumFrames > NumSamples * Viewport.GetNumFramesPerSample())
+					{
+						ZoomHorizontally(-0.1f, 0.0f);
+						Viewport.ScrollAtPosX(0.0f);
+					}
+				}
+			}
+		}
+	}
+
+	if (bIsStateDirty)
+	{
+		bIsStateDirty = false;
+		UpdateState();
+	}
+}
+
+////////////////////////////////////////////////////////////////////////////////////////////////////
+
+void SFrameTrack::UpdateState()
+{
+	FStopwatch Stopwatch;
+	Stopwatch.Start();
+
+	// Reset stats.
+	for (TPair<uint64, FFrameTrackTimeline>& KeyValuePair : CachedTimelines)
+	{
+		FFrameTrackTimeline& Timeline = KeyValuePair.Value;
+		Timeline.NumAggregatedFrames = 0;
+	}
+	NumUpdatedFrames = 0;
+
+	TSharedPtr<const Trace::IAnalysisSession> Session = FInsightsManager::Get()->GetSession();
+	if (Session.IsValid())
+	{
+		Trace::FAnalysisSessionReadScope SessionReadScope(*Session.Get());
+
+		const Trace::IFrameProvider& FramesProvider = Trace::ReadFrameProvider(*Session.Get());
+
+		for (int32 FrameType = 0; FrameType < TraceFrameType_Count; ++FrameType)
+		{
+			FFrameTrackTimeline& Timeline = CachedTimelines.FindOrAdd(FrameType);
+			Timeline.Id = FrameType;
+
+			FFrameTrackTimelineBuilder Builder(Timeline, Viewport);
+
+			uint64 StartIndex = static_cast<uint64>(FMath::Max(0, Viewport.GetIndexAtViewportX(0.0f)));
+			uint64 EndIndex = static_cast<uint64>(Viewport.GetIndexAtViewportX(Viewport.Width));
+
+			FramesProvider.EnumerateFrames(static_cast<ETraceFrameType>(FrameType), StartIndex, EndIndex, [&Builder](const Trace::FFrame& Frame)
+			{
+				Builder.AddFrame(Frame);
+			});
+
+			NumUpdatedFrames += Builder.GetNumAddedFrames();
+		}
+	}
+
+	Stopwatch.Stop();
+	UpdateDurationHistory.AddValue(Stopwatch.AccumulatedTime);
+}
+
+////////////////////////////////////////////////////////////////////////////////////////////////////
+
+FSampleRef SFrameTrack::GetSampleAtMousePosition(float X, float Y)
+{
+	float SampleW = Viewport.GetSampleWidth();
+	int32 SampleIndex = FMath::FloorToInt(X / SampleW);
+	if (SampleIndex >= 0)
+	{
+		const float ViewportHeight = FMath::RoundToFloat(Viewport.Height);
+		const float SampleY2 = ViewportHeight - FMath::RoundToFloat(Viewport.GetViewportYForValue(0.0));
+
+		// Search in reverse paint order.
+		for (int32 TimelineIndex = TimelinesOrder.Num() - 1; TimelineIndex >= 0; --TimelineIndex)
+		{
+			int32 TimelineId = TimelinesOrder[TimelineIndex];
+			if (CachedTimelines.Contains(TimelineId))
+			{
+				const FFrameTrackTimeline& CachedTimeline = CachedTimelines[TimelineId];
+				if (CachedTimeline.NumAggregatedFrames > 0 &&
+					SampleIndex < CachedTimeline.Samples.Num())
+				{
+					const FFrameTrackSample& Sample = CachedTimeline.Samples[SampleIndex];
+					if (Sample.NumFrames > 0)
+					{
+						float SampleY1;
+						if (Sample.LargestFrameDuration == std::numeric_limits<double>::infinity())
+						{
+							SampleY1 = 0.0;
+						}
+						else
+						{
+							SampleY1 = ViewportHeight - FMath::RoundToFloat(Viewport.GetViewportYForValue(Sample.LargestFrameDuration));
+						}
+						if (Y >= SampleY1 && Y <= SampleY2)
+						{
+							return FSampleRef(&CachedTimeline, &Sample);
+						}
+					}
+				}
+			}
+		}
+	}
+	return FSampleRef(nullptr, nullptr);
+}
+
+////////////////////////////////////////////////////////////////////////////////////////////////////
+
+void SFrameTrack::SelectFrameAtMousePosition(float X, float Y)
+{
+	FSampleRef SampleRef = GetSampleAtMousePosition(X, Y);
+	if (!SampleRef.IsValid())
+	{
+		SampleRef = GetSampleAtMousePosition(X - 1.0f, Y);
+	}
+	if (!SampleRef.IsValid())
+	{
+		SampleRef = GetSampleAtMousePosition(X + 1.0f, Y);
+	}
+
+	if (SampleRef.IsValid())
+	{
+		TSharedPtr<STimingProfilerWindow> Window = FTimingProfilerManager::Get()->GetProfilerWindow();
+		if (Window.IsValid())
+		{
+			const double StartTime = SampleRef.Sample->LargestFrameStartTime;
+			const double Duration = SampleRef.Sample->LargestFrameDuration;
+			Window->TimingView->CenterOnTimeInterval(StartTime, Duration);
+			Window->TimingView->SelectTimeInterval(StartTime, Duration);
+			FSlateApplication::Get().SetKeyboardFocus(Window->TimingView);
+		}
+	}
+}
+
+////////////////////////////////////////////////////////////////////////////////////////////////////
+
+int32 SFrameTrack::OnPaint(const FPaintArgs& Args, const FGeometry& AllottedGeometry, const FSlateRect& MyCullingRect, FSlateWindowElementList& OutDrawElements, int32 LayerId, const FWidgetStyle& InWidgetStyle, bool bParentEnabled) const
+{
+	const bool bEnabled = ShouldBeEnabled(bParentEnabled);
+	const ESlateDrawEffect DrawEffects = bEnabled ? ESlateDrawEffect::None : ESlateDrawEffect::DisabledEffect;
+	FDrawContext DrawContext(AllottedGeometry, MyCullingRect, InWidgetStyle, DrawEffects, OutDrawElements, LayerId);
+
+	const TSharedRef<FSlateFontMeasure> FontMeasureService = FSlateApplication::Get().GetRenderer()->GetFontMeasureService();
+
+	const FSlateBrush* WhiteBrush = FCoreStyle::Get().GetBrush("WhiteBrush");
+
+	const float ViewWidth = AllottedGeometry.Size.X;
+	const float ViewHeight = AllottedGeometry.Size.Y;
+
+	int32 NumDrawSamples = 0;
+
+	{
+		FStopwatch Stopwatch;
+		Stopwatch.Start();
+
+		FFrameTrackDrawHelper Helper(DrawContext, Viewport);
+
+		Helper.DrawBackground();
+
+		// Draw frames, for each visible timeline.
+		for (const TPair<uint64, FFrameTrackTimeline>& KeyValuePair : CachedTimelines)
+		{
+			const FFrameTrackTimeline& Timeline = KeyValuePair.Value;
+			Helper.DrawCached(Timeline);
+		}
+		NumDrawSamples = Helper.GetNumDrawSamples();
+
+		// Highlight the mouse hovered sample.
+		if (HoveredSample.IsValid())
+		{
+			Helper.DrawHoveredSample(*HoveredSample.Sample);
+
+			FString Text = FString::Format(TEXT("{0} frame {1} ({2})"),
+				{
+					(HoveredSample.Timeline->Id == 0) ? TEXT("Game") :
+					(HoveredSample.Timeline->Id == 1) ? TEXT("Render") : TEXT("Misc"),
+					FText::AsNumber(HoveredSample.Sample->LargestFrameIndex).ToString(),
+					TimeUtils::FormatTimeAuto(HoveredSample.Sample->LargestFrameDuration)
+				});
+
+			FSlateFontInfo SummaryFont = FCoreStyle::GetDefaultFontStyle("Regular", 8);
+			FVector2D TextSize = FontMeasureService->Measure(Text, SummaryFont);
+
+			const float DX = 2.0f;
+			const float W2 = TextSize.X / 2 + DX;
+
+			float X1 = Viewport.GetViewportXForIndex(HoveredSample.Sample->LargestFrameIndex);
+			float CX = X1 + FMath::RoundToFloat(Viewport.GetSampleWidth() / 2);
+			if (CX + W2 > Viewport.Width)
+			{
+				CX = FMath::RoundToFloat(Viewport.Width - W2);
+			}
+			if (CX - W2 < 0)
+			{
+				CX = W2;
+			}
+
+			const float Y = 10.0f;
+			const float H = 14.0f;
+			DrawContext.DrawBox(CX - W2, Y, 2 * W2, H, WhiteBrush, FLinearColor(0.7, 0.7, 0.7, 0.9));
+			DrawContext.DrawText(CX - W2 + DX, Y + 1.0f, Text, SummaryFont, FLinearColor(0.0, 0.0, 0.0, 0.9));
+		}
+
+		TSharedPtr<STimingProfilerWindow> Window = FTimingProfilerManager::Get()->GetProfilerWindow();
+		if (Window && Window->TimingView && CachedTimelines.Num() > 0)
+		{
+			// Highlight the area corresponding to viewport of Timing View.
+			const double StartTime = Window->TimingView->GetViewport().StartTime;
+			const double EndTime = Window->TimingView->GetViewport().EndTime;
+			Helper.DrawHighlightedInterval(CachedTimelines[0], StartTime, EndTime);
+		}
+
+		Stopwatch.Stop();
+		DrawDurationHistory.AddValue(Stopwatch.AccumulatedTime);
+	}
+
+	//////////////////////////////////////////////////
+
+	const bool bShouldDisplayDebugInfo = FInsightsManager::Get()->IsDebugInfoEnabled();
+	if (bShouldDisplayDebugInfo)
+	{
+		FSlateFontInfo SummaryFont = FCoreStyle::GetDefaultFontStyle("Regular", 8);
+		const float MaxFontCharHeight = FontMeasureService->Measure(TEXT("!"), SummaryFont).Y;
+		const float DbgDY = MaxFontCharHeight;
+
+		const float DbgW = 240.0f;
+		const float DbgH = DbgDY * 4 + 3.0f;
+		const float DbgX = ViewWidth - DbgW - 20.0f;
+		float DbgY = 7.0f;
+
+		++LayerId;
+		FSlateDrawElement::MakeBox(
+			OutDrawElements,
+			LayerId,
+			MAKE_PAINT_GEOMETRY_RC(AllottedGeometry, DbgX - 2.0f, DbgY - 2.0f, DbgW, DbgH),
+			WhiteBrush,
+			DrawEffects,
+			FLinearColor(1.0, 1.0, 1.0, 0.9)
+		);
+
+		++LayerId;
+		FLinearColor DbgTextColor(0.0, 0.0, 0.0, 0.9);
+
+		// Time interval since last OnPaint call.
+		const uint64 CurrentTime = FPlatformTime::Cycles64();
+		const uint64 OnPaintDuration = CurrentTime - LastOnPaintTime;
+		LastOnPaintTime = CurrentTime;
+		OnPaintDurationHistory.AddValue(OnPaintDuration); // saved for last 32 OnPaint calls
+		const uint64 AvgOnPaintDuration = OnPaintDurationHistory.ComputeAverage();
+		const uint64 AvgOnPaintDurationMs = FStopwatch::Cycles64ToMilliseconds(AvgOnPaintDuration);
+		const double AvgOnPaintFps = AvgOnPaintDurationMs != 0 ? 1.0 / FStopwatch::Cycles64ToSeconds(AvgOnPaintDuration) : 0.0;
+
+		const uint64 AvgUpdateDurationMs = FStopwatch::Cycles64ToMilliseconds(UpdateDurationHistory.ComputeAverage());
+		const uint64 AvgDrawDurationMs = FStopwatch::Cycles64ToMilliseconds(DrawDurationHistory.ComputeAverage());
+
+		// Draw performance info.
+		FSlateDrawElement::MakeText
+		(
+			OutDrawElements,
+			LayerId,
+			MAKE_PAINT_GEOMETRY_PT(AllottedGeometry, DbgX, DbgY),
+			FString::Printf(TEXT("U: %llu ms, D: %llu ms + %llu ms = %llu ms (%d fps)"),
+				AvgUpdateDurationMs, // caching time
+				AvgDrawDurationMs, // drawing time
+				AvgOnPaintDurationMs - AvgDrawDurationMs, // other overhead to OnPaint calls
+				AvgOnPaintDurationMs, // average time between two OnPaint calls
+				FMath::RoundToInt(AvgOnPaintFps)), // framerate of OnPaint calls
+			SummaryFont,
+			DrawEffects,
+			DbgTextColor
+		);
+		DbgY += DbgDY;
+
+		// Draw number of draw calls.
+		FSlateDrawElement::MakeText
+		(
+			OutDrawElements,
+			LayerId,
+			MAKE_PAINT_GEOMETRY_PT(AllottedGeometry, DbgX, DbgY),
+			FString::Format(TEXT("U: {0} frames, D: {1} samples"),
+				{
+					FText::AsNumber(NumUpdatedFrames).ToString(),
+					FText::AsNumber(NumDrawSamples).ToString()
+				}),
+			SummaryFont,
+			DrawEffects,
+			DbgTextColor
+		);
+		DbgY += DbgDY;
+
+		// Draw viewport's horizontal info.
+		FSlateDrawElement::MakeText
+		(
+			OutDrawElements,
+			LayerId,
+			MAKE_PAINT_GEOMETRY_PT(AllottedGeometry, DbgX, DbgY),
+			FString::Printf(TEXT("SX: %.3f (%d %s), X: %.2f"),
+				Viewport.ScaleX,
+				(Viewport.ScaleX >= 1.0f) ? FMath::RoundToInt(Viewport.ScaleX) : FMath::RoundToInt(1.0f / Viewport.ScaleX),
+				(Viewport.ScaleX >= 1.0f) ? TEXT("px/frame") : TEXT("frames/px"),
+				Viewport.PosX),
+			SummaryFont,
+			DrawEffects,
+			DbgTextColor
+		);
+		DbgY += DbgDY;
+
+		// Draw viewport's vertical info.
+		FSlateDrawElement::MakeText
+		(
+			OutDrawElements,
+			LayerId,
+			MAKE_PAINT_GEOMETRY_PT(AllottedGeometry, DbgX, DbgY),
+			FString::Printf(TEXT("SY: %g, Y: %.2f"), Viewport.ScaleY, Viewport.PosY),
+			SummaryFont,
+			DrawEffects,
+			DbgTextColor
+		);
+		DbgY += DbgDY;
+	}
+
+	return SCompoundWidget::OnPaint(Args, AllottedGeometry, MyCullingRect, OutDrawElements, LayerId, InWidgetStyle, bParentEnabled && IsEnabled());
+}
+
+////////////////////////////////////////////////////////////////////////////////////////////////////
+
+FReply SFrameTrack::OnMouseButtonDown(const FGeometry& MyGeometry, const FPointerEvent& MouseEvent)
+{
+	FReply Reply = FReply::Unhandled();
+
+	if (!IsReady())
+	{
+		return Reply;
+	}
+
+	MousePositionOnButtonDown = MyGeometry.AbsoluteToLocal(MouseEvent.GetScreenSpacePosition());
+	ViewportPosXOnButtonDown = Viewport.PosX;
+	ViewportPosYOnButtonDown = Viewport.PosY;
+
+	if (MouseEvent.GetEffectingButton() == EKeys::LeftMouseButton)
+	{
+		bIsLMB_Pressed = true;
+
+		// Capture mouse.
+		Reply = FReply::Handled().CaptureMouse(SharedThis(this));
+	}
+	else if (MouseEvent.GetEffectingButton() == EKeys::RightMouseButton)
+	{
+		bIsRMB_Pressed = true;
+
+		// Capture mouse, so we can scroll outside this widget.
+		Reply = FReply::Handled().CaptureMouse(SharedThis(this));
+	}
+
+	return Reply;
+}
+
+////////////////////////////////////////////////////////////////////////////////////////////////////
+
+FReply SFrameTrack::OnMouseButtonUp(const FGeometry& MyGeometry, const FPointerEvent& MouseEvent)
+{
+	FReply Reply = FReply::Unhandled();
+
+	if (!IsReady())
+	{
+		return Reply;
+	}
+
+	MousePositionOnButtonUp = MyGeometry.AbsoluteToLocal(MouseEvent.GetScreenSpacePosition());
+
+	const bool bIsValidForMouseClick = MousePositionOnButtonUp.Equals(MousePositionOnButtonDown, MOUSE_SNAP_DISTANCE);
+
+	if (MouseEvent.GetEffectingButton() == EKeys::LeftMouseButton)
+	{
+		if (bIsLMB_Pressed)
+		{
+			// Release the mouse.
+			Reply = FReply::Handled().ReleaseMouseCapture();
+
+			if (bIsValidForMouseClick)
+			{
+				SelectFrameAtMousePosition(MousePositionOnButtonUp.X, MousePositionOnButtonUp.Y);
+			}
+
+			bIsLMB_Pressed = false;
+		}
+	}
+	else if (MouseEvent.GetEffectingButton() == EKeys::RightMouseButton)
+	{
+		if (bIsRMB_Pressed)
+		{
+			// Release mouse as we no longer scroll.
+			Reply = FReply::Handled().ReleaseMouseCapture();
+
+			if (bIsRMB_Scrolling)
+			{
+				CursorType = EFrameTrackCursor::Default;
+				bIsRMB_Scrolling = false;
+			}
+			else if (bIsValidForMouseClick)
+			{
+				ShowContextMenu(MouseEvent.GetScreenSpacePosition());
+				Reply = FReply::Handled();
+			}
+
+			bIsRMB_Pressed = false;
+		}
+	}
+
+	return Reply;
+}
+
+////////////////////////////////////////////////////////////////////////////////////////////////////
+
+FReply SFrameTrack::OnMouseMove(const FGeometry& MyGeometry, const FPointerEvent& MouseEvent)
+{
+	FReply Reply = FReply::Unhandled();
+
+	if (!IsReady())
+	{
+		return Reply;
+	}
+
+	MousePosition = MyGeometry.AbsoluteToLocal(MouseEvent.GetScreenSpacePosition());
+
+	if (MouseEvent.IsMouseButtonDown(EKeys::LeftMouseButton))
+	{
+		if (HasMouseCapture() && !MouseEvent.GetCursorDelta().IsZero())
+		{
+			// Inform other widgets that we have moved the selection box.
+			//SelectionBoxChangedEvent.Broadcast(SelectionBoxFrameStart, SelectionBoxFrameEnd);
+
+			Reply = FReply::Handled();
+		}
+	}
+	else if (MouseEvent.IsMouseButtonDown(EKeys::RightMouseButton))
+	{
+		if (HasMouseCapture() && !MouseEvent.GetCursorDelta().IsZero())
+		{
+			bIsRMB_Scrolling = true;
+			CursorType = EFrameTrackCursor::Hand;
+
+			HoveredSample.Reset();
+
+			float PosX = ViewportPosXOnButtonDown + (MousePositionOnButtonDown.X - MousePosition.X);
+			Viewport.ScrollAtIndex(Viewport.GetIndexAtPosX(PosX));
+			UpdateHorizontalScrollBar();
+			bIsStateDirty = true;
+
+			Reply = FReply::Handled();
+		}
+	}
+	else
+	{
+		HoveredSample = GetSampleAtMousePosition(MousePosition.X, MousePosition.Y);
+		if (!HoveredSample.IsValid())
+			HoveredSample = GetSampleAtMousePosition(MousePosition.X - 1.0f, MousePosition.Y);
+		if (!HoveredSample.IsValid())
+			HoveredSample = GetSampleAtMousePosition(MousePosition.X + 1.0f, MousePosition.Y);
+	}
+
+	return Reply;
+}
+
+////////////////////////////////////////////////////////////////////////////////////////////////////
+
+void SFrameTrack::OnMouseEnter(const FGeometry& MyGeometry, const FPointerEvent& MouseEvent)
+{
+
+}
+
+////////////////////////////////////////////////////////////////////////////////////////////////////
+
+void SFrameTrack::OnMouseLeave(const FPointerEvent& MouseEvent)
+{
+	if (!HasMouseCapture())
+	{
+		bIsLMB_Pressed = false;
+		bIsRMB_Pressed = false;
+
+		HoveredSample.Reset();
+
+		CursorType = EFrameTrackCursor::Default;
+	}
+}
+
+////////////////////////////////////////////////////////////////////////////////////////////////////
+
+FReply SFrameTrack::OnMouseWheel(const FGeometry& MyGeometry, const FPointerEvent& MouseEvent)
+{
+	MousePosition = MyGeometry.AbsoluteToLocal(MouseEvent.GetScreenSpacePosition());
+
+	if (MouseEvent.GetModifierKeys().IsShiftDown())
+	{
+		// Zoom in/out vertically.
+		const float Delta = MouseEvent.GetWheelDelta();
+		constexpr float ZoomStep = 0.25f; // as percent
+		float ScaleY;
+
+		if (Delta > 0)
+		{
+			ScaleY = Viewport.ScaleY * FMath::Pow(1.0f + ZoomStep, Delta);
+		}
+		else
+		{
+			ScaleY = Viewport.ScaleY * FMath::Pow(1.0f / (1.0f + ZoomStep), -Delta);
+		}
+
+		Viewport.SetScaleY(ScaleY);
+		//UpdateVerticalScrollBar();
+	}
+	else //if (MouseEvent.GetModifierKeys().IsControlDown())
+	{
+		// Zoom in/out horizontally.
+		const float Delta = MouseEvent.GetWheelDelta();
+		ZoomHorizontally(Delta, MousePosition.X);
+	}
+
+	return FReply::Handled();
+}
+
+////////////////////////////////////////////////////////////////////////////////////////////////////
+
+void SFrameTrack::ZoomHorizontally(const float Delta, const float X)
+{
+	constexpr float ZoomStep = 0.25f; // as percent
+	float ScaleX;
+
+	if (Delta > 0)
+	{
+		ScaleX = Viewport.ScaleX * FMath::Pow(1.0f + ZoomStep, Delta);
+	}
+	else
+	{
+		ScaleX = Viewport.ScaleX * FMath::Pow(1.0f / (1.0f + ZoomStep), -Delta);
+	}
+
+	// Snap to integer value of either: "number of frames per pixel (Slate unit)" or "number of pixels per frame".
+	if (ScaleX < 1.0f)
+	{
+		// frames per Slate unit
+		if (Delta > 0)
+			ScaleX = 1.0f / FMath::FloorToFloat(1.0f / ScaleX);
+		else
+			ScaleX = 1.0f / FMath::CeilToFloat(1.0f / ScaleX);
+	}
+	else
+	{
+		// Slate units per frame
+		if (Delta > 0)
+			ScaleX = FMath::CeilToFloat(ScaleX);
+		else
+			ScaleX = FMath::FloorToFloat(ScaleX);
+	}
+
+	//UE_LOG(TimingProfiler, Log, TEXT("%.2f, %.2f, %.2f"), Delta, Viewport.ScaleX, ScaleX);
+	Viewport.ZoomWithFixedViewportX(ScaleX, MousePosition.X);
+	UpdateHorizontalScrollBar();
+	bIsStateDirty = true;
+}
+
+////////////////////////////////////////////////////////////////////////////////////////////////////
+
+FReply SFrameTrack::OnMouseButtonDoubleClick(const FGeometry& MyGeometry, const FPointerEvent& MouseEvent)
+{
+	return FReply::Unhandled();
+}
+
+////////////////////////////////////////////////////////////////////////////////////////////////////
+
+FCursorReply SFrameTrack::OnCursorQuery(const FGeometry& MyGeometry, const FPointerEvent& CursorEvent) const
+{
+	FCursorReply CursorReply = FCursorReply::Unhandled();
+
+	if (CursorType == EFrameTrackCursor::Arrow)
+	{
+		CursorReply = FCursorReply::Cursor(EMouseCursor::ResizeLeftRight);
+	}
+	else if (CursorType == EFrameTrackCursor::Hand)
+	{
+		CursorReply = FCursorReply::Cursor(EMouseCursor::GrabHand);
+	}
+
+	return CursorReply;
+}
+
+////////////////////////////////////////////////////////////////////////////////////////////////////
+
+void SFrameTrack::ShowContextMenu(const FVector2D& ScreenSpacePosition)
+{
+
+}
+
+////////////////////////////////////////////////////////////////////////////////////////////////////
+
+void SFrameTrack::BindCommands()
+{
+
+}
+
+////////////////////////////////////////////////////////////////////////////////////////////////////
+
+void SFrameTrack::HorizontalScrollBar_OnUserScrolled(float ScrollOffset)
+{
+	const float SX = 1.0 / (Viewport.MaxX - Viewport.MinX);
+	const float ThumbSizeFraction = FMath::Clamp<float>(Viewport.Width * SX, 0.0f, 1.0f);
+	const float OffsetFraction = FMath::Clamp<float>(ScrollOffset, 0.0f, 1.0f - ThumbSizeFraction);
+
+	const float PosX = Viewport.MinX + OffsetFraction * (Viewport.MaxX - Viewport.MinX);
+	Viewport.ScrollAtPosX(PosX);
+	bIsStateDirty = true;
+
+	HorizontalScrollBar->SetState(OffsetFraction, ThumbSizeFraction);
+}
+
+////////////////////////////////////////////////////////////////////////////////////////////////////
+
+void SFrameTrack::UpdateHorizontalScrollBar()
+{
+	const float SX = 1.0 / (Viewport.MaxX - Viewport.MinX);
+	const float ThumbSizeFraction = FMath::Clamp<float>(Viewport.Width * SX, 0.0f, 1.0f);
+	const float ScrollOffset = (Viewport.PosX - Viewport.MinX) * SX;
+	const float OffsetFraction = FMath::Clamp<float>(ScrollOffset, 0.0f, 1.0f - ThumbSizeFraction);
+
+	HorizontalScrollBar->SetState(OffsetFraction, ThumbSizeFraction);
+}
+
+////////////////////////////////////////////////////////////////////////////////////////////////////
+/*
+void SFrameTrack::AddThreadTime(int32 InFrameIndex, const TMap<uint32, float>& InThreadMS, const TSharedRef<FProfilerStatMetaData>& InStatMetaData)
+{
+	FFrameThreadTimes FrameThreadTimes;
+	FrameThreadTimes.FrameNumber = InFrameIndex;
+	FrameThreadTimes.ThreadTimes = InThreadMS;
+	RecentlyAddedFrames.Add(FrameThreadTimes);
+
+	if (!bIsActiveTimerRegistered)
+	{
+		bIsActiveTimerRegistered = true;
+		RegisterActiveTimer(0.f, FWidgetActiveTimerDelegate::CreateSP(this, &SFrameTrack::EnsureDataUpdateDuringPreview));
+	}
+
+	StatMetadata = InStatMetaData;
+}
+*/
+////////////////////////////////////////////////////////////////////////////////////////////////////
+/*
+EActiveTimerReturnType SFrameTrack::EnsureDataUpdateDuringPreview(double InCurrentTime, float InDeltaTime)
+{
+	//if (RecentlyAddedFrames.Num() > 0)
+	//{
+	//	bUpdateData = true;
+	//	return EActiveTimerReturnType::Continue;
+	//}
+
+	bIsActiveTimerRegistered = false;
+	return EActiveTimerReturnType::Stop;
+}
+*/
+////////////////////////////////////////////////////////////////////////////////////////////////////
+/*
+void SFrameTrack::MoveSelectionBox(int32 FrameIndex)
+{
+	const int32 SelectionBoxSize = SelectionBoxFrameEnd - SelectionBoxFrameStart;
+	const int32 SelectionBoxHalfSize = SelectionBoxSize / 2;
+	const int32 CenterFrameIndex = FMath::Clamp(FrameIndex - SelectionBoxHalfSize, 0, Frames.Num() - 1 - SelectionBoxSize);
+
+	// Inform other widgets that we have moved the selection box.
+	SelectionBoxFrameStart = CenterFrameIndex;
+	SelectionBoxFrameEnd = CenterFrameIndex + SelectionBoxSize;
+	SelectionBoxChangedEvent.Broadcast(SelectionBoxFrameStart, SelectionBoxFrameEnd);
+}
+*/
+////////////////////////////////////////////////////////////////////////////////////////////////////
+
+#undef LOCTEXT_NAMESPACE
diff --git a/Engine/Source/Developer/TraceInsights/Private/Insights/Widgets/SFrameTrack.h b/Engine/Source/Developer/TraceInsights/Private/Insights/Widgets/SFrameTrack.h
new file mode 100644
--- /dev/null
+++ b/Engine/Source/Developer/TraceInsights/Private/Insights/Widgets/SFrameTrack.h
@@ -0,0 +1,190 @@
+// Copyright 1998-2019 Epic Games, Inc. All Rights Reserved.
+
+#pragma once
+
+#include "CoreMinimal.h"
+#include "Input/CursorReply.h"
+#include "Input/Reply.h"
+#include "Layout/Geometry.h"
+#include "Rendering/RenderingCommon.h"
+#include "Widgets/DeclarativeSyntaxSupport.h"
+#include "Widgets/SCompoundWidget.h"
+
+// Insights
+#include "Insights/Common/FixedCircularBuffer.h"
+#include "Insights/ViewModels/FrameTrackHelper.h"
+#include "Insights/ViewModels/FrameTrackViewport.h"
+
+class SScrollBar;
+
+////////////////////////////////////////////////////////////////////////////////////////////////////
+
+struct FSampleRef
+{
+	const FFrameTrackTimeline* Timeline;
+	const FFrameTrackSample* Sample;
+
+	FSampleRef(const FFrameTrackTimeline* InTimeline, const FFrameTrackSample* InSample)
+		: Timeline(InTimeline), Sample(InSample)
+	{
+	}
+
+	void Reset() { Timeline = nullptr; Sample = nullptr; }
+	bool IsValid() const { return Timeline != nullptr && Sample != nullptr; }
+};
+
+////////////////////////////////////////////////////////////////////////////////////////////////////
+/**
+ * Widget used to present frames data in a track.
+ */
+class SFrameTrack : public SCompoundWidget
+{
+	enum
+	{
+		/** Number of pixels. */
+		MOUSE_SNAP_DISTANCE = 4,
+	};
+
+	enum class EFrameTrackCursor
+	{
+		Default,
+		Arrow,
+		Hand,
+	};
+
+public:
+	/** Default constructor. */
+	SFrameTrack();
+
+	/** Virtual destructor. */
+	virtual ~SFrameTrack();
+
+	/** Resets internal widget's data to the default one. */
+	void Reset();
+
+	SLATE_BEGIN_ARGS(SFrameTrack)
+	{
+		_Clipping = EWidgetClipping::ClipToBounds;
+	}
+	SLATE_END_ARGS()
+
+	/**
+	 * Construct this widget
+	 * @param InArgs The declaration data for this widget
+	 */
+	void Construct(const FArguments& InArgs);
+
+	/**
+	 * Ticks this widget. Override in derived classes, but always call the parent implementation.
+	 * @param  AllottedGeometry The space allotted for this widget
+	 * @param  InCurrentTime  Current absolute real time
+	 * @param  InDeltaTime  Real time passed since last tick
+	 */
+	virtual void Tick(const FGeometry& AllottedGeometry, const double InCurrentTime, const float InDeltaTime) override;
+
+	virtual int32 OnPaint(const FPaintArgs& Args, const FGeometry& AllottedGeometry, const FSlateRect& MyCullingRect, FSlateWindowElementList& OutDrawElements, int32 LayerId, const FWidgetStyle& InWidgetStyle, bool bParentEnabled) const override;
+
+	virtual FReply OnMouseButtonDown(const FGeometry& MyGeometry, const FPointerEvent& MouseEvent) override;
+	virtual FReply OnMouseButtonUp(const FGeometry& MyGeometry, const FPointerEvent& MouseEvent) override;
+	virtual FReply OnMouseMove(const FGeometry& MyGeometry, const FPointerEvent& MouseEvent) override;
+	virtual void OnMouseEnter(const FGeometry& MyGeometry, const FPointerEvent& MouseEvent) override;
+	virtual void OnMouseLeave(const FPointerEvent& MouseEvent) override;
+	virtual FReply OnMouseWheel(const FGeometry& MyGeometry, const FPointerEvent& MouseEvent) override;
+	virtual FReply OnMouseButtonDoubleClick(const FGeometry& MyGeometry, const FPointerEvent& MouseEvent) override;
+
+	virtual FCursorReply OnCursorQuery(const FGeometry& MyGeometry, const FPointerEvent& CursorEvent) const override;
+
+protected:
+	bool IsReady() { return true; }
+
+	void UpdateState();
+
+	FSampleRef GetSampleAtMousePosition(float X, float Y);
+	void SelectFrameAtMousePosition(float X, float Y);
+
+	void ShowContextMenu(const FVector2D& ScreenSpacePosition);
+
+	/** Binds our UI commands to delegates. */
+	void BindCommands();
+
+	/**
+	 * Called when the user scrolls the horizontal scrollbar.
+	 * @param ScrollOffset Scroll offset as a fraction between 0 and 1.
+	 */
+	void HorizontalScrollBar_OnUserScrolled(float ScrollOffset);
+	void UpdateHorizontalScrollBar();
+
+	void ZoomHorizontally(const float Delta, const float X);
+
+	//////////////////////////////////////////////////
+	// SelectionBoxChanged Event
+
+public:
+	/** The event to execute when the selection box has been changed. */
+	DECLARE_EVENT_TwoParams(SFrameTrack, FSelectionBoxChangedEvent, int32 /*FrameStart*/, int32 /*FrameEnd*/);
+	FSelectionBoxChangedEvent& OnSelectionBoxChanged() { return SelectionBoxChangedEvent; }
+protected:
+	/** The event to execute when the selection box has been changed. */
+	FSelectionBoxChangedEvent SelectionBoxChangedEvent;
+
+	//////////////////////////////////////////////////
+
+protected:
+	/** The track's viewport. Encapsulates info about position and scale. */
+	FFrameTrackViewport Viewport;
+	bool bIsViewportDirty;
+
+	/** Cached info for timelines. */
+	TMap<uint64, FFrameTrackTimeline> CachedTimelines;
+	TArray<int32> TimelinesOrder;
+
+	bool bIsStateDirty;
+
+	uint64 AnalysisSyncNextTimestamp;
+
+	//////////////////////////////////////////////////
+
+	TSharedPtr<SScrollBar> HorizontalScrollBar;
+
+	//////////////////////////////////////////////////
+	// Panning, Zooming and Selection behaviors
+
+	/** The current mouse position. */
+	FVector2D MousePosition;
+
+	/** Mouse position during the call on mouse button down. */
+	FVector2D MousePositionOnButtonDown;
+	float ViewportPosXOnButtonDown;
+	float ViewportPosYOnButtonDown;
+
+	/** Mouse position during the call on mouse button up. */
+	FVector2D MousePositionOnButtonUp;
+
+	bool bIsLMB_Pressed;
+	bool bIsRMB_Pressed;
+
+	/** True, if the user is currently interactively scrolling the view by holding the right mouse button and dragging. */
+	bool bIsRMB_Scrolling;
+
+	//////////////////////////////////////////////////
+
+	int32 SelectionStartFrameIndex;
+	int32 SelectionEndFrameIndex;
+
+	FSampleRef HoveredSample;
+
+	//////////////////////////////////////////////////
+	// Misc
+
+	FGeometry ThisGeometry;
+
+	/** Cursor type. */
+	EFrameTrackCursor CursorType;
+
+	// Debug stats
+	int32 NumUpdatedFrames;
+	TFixedCircularBuffer<uint64, 32> UpdateDurationHistory;
+	mutable TFixedCircularBuffer<uint64, 32> DrawDurationHistory;
+	mutable TFixedCircularBuffer<uint64, 32> OnPaintDurationHistory;
+	mutable uint64 LastOnPaintTime;
+};
diff --git a/Engine/Source/Developer/TraceInsights/Private/Insights/Widgets/SGraphTrack.cpp b/Engine/Source/Developer/TraceInsights/Private/Insights/Widgets/SGraphTrack.cpp
new file mode 100644
--- /dev/null
+++ b/Engine/Source/Developer/TraceInsights/Private/Insights/Widgets/SGraphTrack.cpp
@@ -0,0 +1,220 @@
+// Copyright 1998-2019 Epic Games, Inc. All Rights Reserved.
+
+#include "SGraphTrack.h"
+
+#include "Brushes/SlateBorderBrush.h"
+#include "Brushes/SlateBoxBrush.h"
+#include "Brushes/SlateColorBrush.h"
+#include "Containers/ArrayBuilder.h"
+#include "Containers/MapBuilder.h"
+#include "EditorStyleSet.h"
+#include "Fonts/FontMeasure.h"
+#include "Fonts/SlateFontInfo.h"
+#include "Framework/Application/MenuStack.h"
+#include "Framework/Application/SlateApplication.h"
+#include "Framework/MultiBox/MultiBoxBuilder.h"
+#include "Layout/WidgetPath.h"
+#include "Misc/Paths.h"
+#include "Rendering/DrawElements.h"
+#include "SlateOptMacros.h"
+#include "Styling/CoreStyle.h"
+#include "Widgets/Input/SButton.h"
+#include "Widgets/Layout/SSpacer.h"
+#include "Widgets/SBoxPanel.h"
+#include "Widgets/SOverlay.h"
+#include "Widgets/Text/STextBlock.h"
+
+// Insights
+#include "Insights/Common/PaintUtils.h"
+#include "Insights/TimingProfilerManager.h"
+
+////////////////////////////////////////////////////////////////////////////////////////////////////
+
+#define LOCTEXT_NAMESPACE "SGraphTrack"
+
+////////////////////////////////////////////////////////////////////////////////////////////////////
+// SGraphTrack
+////////////////////////////////////////////////////////////////////////////////////////////////////
+
+SGraphTrack::SGraphTrack()
+	: TimeRulerTrack(0)
+	, GraphTrack(1)
+{
+	Reset();
+}
+
+////////////////////////////////////////////////////////////////////////////////////////////////////
+
+SGraphTrack::~SGraphTrack()
+{
+}
+
+////////////////////////////////////////////////////////////////////////////////////////////////////
+
+void SGraphTrack::Reset()
+{
+	TimeRulerTrack.Reset();
+	GraphTrack.Reset();
+	Viewport.ScaleX = (5 * 20) / 0.1; // 100ms between major tick marks
+}
+
+////////////////////////////////////////////////////////////////////////////////////////////////////
+
+void SGraphTrack::Construct(const FArguments& InArgs)
+{
+	//TODO: Add toggle buttons for Points/Lines/Bars.
+	//TODO: Add "show list of series" expand button + toggle buttons for each available graph series.
+}
+
+////////////////////////////////////////////////////////////////////////////////////////////////////
+
+void SGraphTrack::Tick(const FGeometry& AllottedGeometry, const double InCurrentTime, const float InDeltaTime)
+{
+	//ThisGeometry = AllottedGeometry;
+
+	float TrackWidth = AllottedGeometry.GetAbsoluteSize().X;
+	float TrackHeight = AllottedGeometry.GetAbsoluteSize().Y;
+
+	bool bIsGraphDirty = false;
+
+	if (Viewport.UpdateSize(TrackWidth, TrackHeight))
+	{
+		bIsGraphDirty = true;
+	}
+
+	if (GraphTrack.GetHeight() != TrackHeight)
+	{
+		bIsGraphDirty = true;
+	}
+
+	if (bIsGraphDirty)
+	{
+		bIsGraphDirty = false;
+
+		GraphTrack.SetPosY(TimeRulerTrack.GetHeight());
+		GraphTrack.SetHeight(Viewport.Height - TimeRulerTrack.GetHeight());
+		GraphTrack.Update(Viewport);
+	}
+}
+
+////////////////////////////////////////////////////////////////////////////////////////////////////
+
+int32 SGraphTrack::OnPaint(const FPaintArgs& Args, const FGeometry& AllottedGeometry, const FSlateRect& MyCullingRect, FSlateWindowElementList& OutDrawElements, int32 LayerId, const FWidgetStyle& InWidgetStyle, bool bParentEnabled) const
+{
+	const bool bEnabled = ShouldBeEnabled(bParentEnabled);
+	const ESlateDrawEffect DrawEffects = bEnabled ? ESlateDrawEffect::None : ESlateDrawEffect::DisabledEffect;
+	FDrawContext DrawContext(AllottedGeometry, MyCullingRect, InWidgetStyle, DrawEffects, OutDrawElements, LayerId);
+
+	GraphTrack.Draw(DrawContext, Viewport);
+	TimeRulerTrack.Draw(DrawContext, Viewport);
+
+	return SCompoundWidget::OnPaint(Args, AllottedGeometry, MyCullingRect, OutDrawElements, LayerId, InWidgetStyle, bParentEnabled && IsEnabled());
+}
+
+////////////////////////////////////////////////////////////////////////////////////////////////////
+
+FReply SGraphTrack::OnMouseButtonDown(const FGeometry& MyGeometry, const FPointerEvent& MouseEvent)
+{
+	//...
+
+	return FReply::Unhandled();
+}
+
+////////////////////////////////////////////////////////////////////////////////////////////////////
+
+FReply SGraphTrack::OnMouseButtonUp(const FGeometry& MyGeometry, const FPointerEvent& MouseEvent)
+{
+	//...
+
+	return FReply::Unhandled();
+}
+
+////////////////////////////////////////////////////////////////////////////////////////////////////
+
+FReply SGraphTrack::OnMouseMove(const FGeometry& MyGeometry, const FPointerEvent& MouseEvent)
+{
+	//...
+
+	return FReply::Unhandled();
+}
+
+////////////////////////////////////////////////////////////////////////////////////////////////////
+
+void SGraphTrack::OnMouseEnter(const FGeometry& MyGeometry, const FPointerEvent& MouseEvent)
+{
+	//...
+}
+
+////////////////////////////////////////////////////////////////////////////////////////////////////
+
+void SGraphTrack::OnMouseLeave(const FPointerEvent& MouseEvent)
+{
+	//...
+}
+
+////////////////////////////////////////////////////////////////////////////////////////////////////
+
+FReply SGraphTrack::OnMouseWheel(const FGeometry& MyGeometry, const FPointerEvent& MouseEvent)
+{
+	//...
+
+	return FReply::Handled();
+}
+
+////////////////////////////////////////////////////////////////////////////////////////////////////
+
+FReply SGraphTrack::OnMouseButtonDoubleClick(const FGeometry& MyGeometry, const FPointerEvent& MouseEvent)
+{
+	//...
+
+	return FReply::Unhandled();
+}
+
+////////////////////////////////////////////////////////////////////////////////////////////////////
+
+void SGraphTrack::OnDragEnter(const FGeometry& MyGeometry, const FDragDropEvent& DragDropEvent)
+{
+	SCompoundWidget::OnDragEnter(MyGeometry, DragDropEvent);
+
+	//...
+}
+
+////////////////////////////////////////////////////////////////////////////////////////////////////
+
+void SGraphTrack::OnDragLeave(const FDragDropEvent& DragDropEvent)
+{
+	SCompoundWidget::OnDragLeave(DragDropEvent);
+
+	//...
+}
+
+////////////////////////////////////////////////////////////////////////////////////////////////////
+
+FReply SGraphTrack::OnDragOver(const FGeometry& MyGeometry, const FDragDropEvent& DragDropEvent)
+{
+	//...
+
+	return SCompoundWidget::OnDragOver(MyGeometry,DragDropEvent);
+}
+
+////////////////////////////////////////////////////////////////////////////////////////////////////
+
+FReply SGraphTrack::OnDrop(const FGeometry& MyGeometry, const FDragDropEvent& DragDropEvent)
+{
+	//...
+
+	return SCompoundWidget::OnDrop(MyGeometry,DragDropEvent);
+}
+
+////////////////////////////////////////////////////////////////////////////////////////////////////
+
+FCursorReply SGraphTrack::OnCursorQuery(const FGeometry& MyGeometry, const FPointerEvent& CursorEvent) const
+{
+	//...
+
+	return FCursorReply::Unhandled();
+}
+
+////////////////////////////////////////////////////////////////////////////////////////////////////
+
+#undef LOCTEXT_NAMESPACE
