Copying //UE4/Dev-Core to Dev-Main (//UE4/Dev-Main)

diff --git a/Engine/Source/Developer/TraceInsights/Private/Insights/Widgets/SGraphTrack.h b/Engine/Source/Developer/TraceInsights/Private/Insights/Widgets/SGraphTrack.h
new file mode 100644
--- /dev/null
+++ b/Engine/Source/Developer/TraceInsights/Private/Insights/Widgets/SGraphTrack.h
@@ -0,0 +1,187 @@
+// Copyright 1998-2019 Epic Games, Inc. All Rights Reserved.
+
+#pragma once
+
+#include "CoreMinimal.h"
+#include "Widgets/DeclarativeSyntaxSupport.h"
+#include "Layout/Geometry.h"
+#include "Input/CursorReply.h"
+#include "Input/Reply.h"
+#include "Widgets/SWidget.h"
+#include "Widgets/SCompoundWidget.h"
+
+// Insights
+#include "Insights/ViewModels/GraphTrack.h"
+#include "Insights/ViewModels/TimeRulerTrack.h"
+#include "Insights/ViewModels/TimingTrackViewport.h"
+
+class FPaintArgs;
+class FSlateWindowElementList;
+
+/** A custom widget used to display graphs. */
+class SGraphTrack : public SCompoundWidget
+{
+public:
+	/** Default constructor. */
+	SGraphTrack();
+
+	/** Virtual destructor. */
+	virtual ~SGraphTrack();
+
+	/** Resets internal widget's data to the default one. */
+	void Reset();
+
+	SLATE_BEGIN_ARGS(SGraphTrack)
+		{
+			_Clipping = EWidgetClipping::ClipToBounds;
+		}
+	SLATE_END_ARGS()
+
+	/**
+	 * Construct this widget
+	 *
+	 * @param	InArgs	The declaration data for this widget
+	 */
+	void Construct(const FArguments& InArgs);
+
+	/**
+	 * Ticks this widget.  Override in derived classes, but always call the parent implementation.
+	 *
+	 * @param  AllottedGeometry The space allotted for this widget
+	 * @param  InCurrentTime  Current absolute real time
+	 * @param  InDeltaTime  Real time passed since last tick
+	 */
+	virtual void Tick(const FGeometry& AllottedGeometry, const double InCurrentTime, const float InDeltaTime) override;
+
+	virtual int32 OnPaint(const FPaintArgs& Args, const FGeometry& AllottedGeometry, const FSlateRect& MyCullingRect, FSlateWindowElementList& OutDrawElements, int32 LayerId, const FWidgetStyle& InWidgetStyle, bool bParentEnabled) const override;
+
+	/**
+	 * The system calls this method to notify the widget that a mouse button was pressed within it. This event is bubbled.
+	 *
+	 * @param MyGeometry The Geometry of the widget receiving the event
+	 * @param MouseEvent Information about the input event
+	 *
+	 * @return Whether the event was handled along with possible requests for the system to take action.
+	 */
+	virtual FReply OnMouseButtonDown(const FGeometry& MyGeometry, const FPointerEvent& MouseEvent) override;
+
+	/**
+	 * The system calls this method to notify the widget that a mouse button was release within it. This event is bubbled.
+	 *
+	 * @param MyGeometry The Geometry of the widget receiving the event
+	 * @param MouseEvent Information about the input event
+	 *
+	 * @return Whether the event was handled along with possible requests for the system to take action.
+	 */
+	virtual FReply OnMouseButtonUp(const FGeometry& MyGeometry, const FPointerEvent& MouseEvent) override;
+
+	/**
+	 * The system calls this method to notify the widget that a mouse moved within it. This event is bubbled.
+	 *
+	 * @param MyGeometry The Geometry of the widget receiving the event
+	 * @param MouseEvent Information about the input event
+	 *
+	 * @return Whether the event was handled along with possible requests for the system to take action.
+	 */
+	virtual FReply OnMouseMove(const FGeometry& MyGeometry, const FPointerEvent& MouseEvent) override;
+
+	/**
+	 * The system will use this event to notify a widget that the cursor has entered it. This event is NOT bubbled.
+	 *
+	 * @param MyGeometry The Geometry of the widget receiving the event
+	 * @param MouseEvent Information about the input event
+	 */
+	virtual void OnMouseEnter(const FGeometry& MyGeometry, const FPointerEvent& MouseEvent) override;
+
+	/**
+	 * The system will use this event to notify a widget that the cursor has left it. This event is NOT bubbled.
+	 *
+	 * @param MouseEvent Information about the input event
+	 */
+	virtual void OnMouseLeave(const FPointerEvent& MouseEvent) override;
+
+	/**
+	 * Called when the mouse wheel is spun. This event is bubbled.
+	 *
+	 * @param  MouseEvent  Mouse event
+	 *
+	 * @return  Returns whether the event was handled, along with other possible actions
+	 */
+	virtual FReply OnMouseWheel(const FGeometry& MyGeometry, const FPointerEvent& MouseEvent) override;
+
+	/**
+	 * Called when a mouse button is double clicked.  Override this in derived classes.
+	 *
+	 * @param  InMyGeometry  Widget geometry
+	 * @param  InMouseEvent  Mouse button event
+	 *
+	 * @return  Returns whether the event was handled, along with other possible actions
+	 */
+	virtual FReply OnMouseButtonDoubleClick(const FGeometry& MyGeometry, const FPointerEvent& MouseEvent) override;
+
+	/**
+	 * Called when the user is dropping something onto a widget; terminates drag and drop.
+	 *
+	 * @param MyGeometry      The geometry of the widget receiving the event.
+	 * @param DragDropEvent   The drag and drop event.
+	 *
+	 * @return A reply that indicated whether this event was handled.
+	 */
+	virtual FReply OnDrop(const FGeometry& MyGeometry, const FDragDropEvent& DragDropEvent) override;
+
+	/**
+	 * Called during drag and drop when the drag enters a widget.
+	 *
+	 * Enter/Leave events in slate are meant as lightweight notifications.
+	 * So we do not want to capture mouse or set focus in response to these.
+	 * However, OnDragEnter must also support external APIs (e.g. OLE Drag/Drop)
+	 * Those require that we let them know whether we can handle the content
+	 * being dragged OnDragEnter.
+	 *
+	 * The concession is to return a can_handled/cannot_handle
+	 * boolean rather than a full FReply.
+	 *
+	 * @param MyGeometry      The geometry of the widget receiving the event.
+	 * @param DragDropEvent   The drag and drop event.
+	 *
+	 * @return A reply that indicated whether the contents of the DragDropEvent can potentially be processed by this widget.
+	 */
+	virtual void OnDragEnter(const FGeometry& MyGeometry, const FDragDropEvent& DragDropEvent) override;
+
+	/**
+	 * Called during drag and drop when the drag leaves a widget.
+	 *
+	 * @param DragDropEvent   The drag and drop event.
+	 */
+	virtual void OnDragLeave(const FDragDropEvent& DragDropEvent) override;
+
+	/**
+	 * Called during drag and drop when the the mouse is being dragged over a widget.
+	 *
+	 * @param MyGeometry      The geometry of the widget receiving the event.
+	 * @param DragDropEvent   The drag and drop event.
+	 *
+	 * @return A reply that indicated whether this event was handled.
+	 */
+	virtual FReply OnDragOver(const FGeometry& MyGeometry, const FDragDropEvent& DragDropEvent) override;
+
+	/**
+	 * Called when the system wants to know which cursor to display for this Widget.  This event is bubbled.
+	 *
+	 * @return  The cursor requested (can be None.)
+	 */
+	virtual FCursorReply OnCursorQuery(const FGeometry& MyGeometry, const FPointerEvent& CursorEvent) const override;
+
+protected:
+	virtual FVector2D ComputeDesiredSize(float) const override
+	{
+		return FVector2D(16.0f, 16.0f);
+	}
+
+	void UpdateGraphTrack();
+
+protected:
+	FTimeRulerTrack TimeRulerTrack;
+	FRandomGraphTrack GraphTrack;
+	FTimingTrackViewport Viewport;
+};
diff --git a/Engine/Source/Developer/TraceInsights/Private/Insights/Widgets/SIoProfilerWindow.cpp b/Engine/Source/Developer/TraceInsights/Private/Insights/Widgets/SIoProfilerWindow.cpp
new file mode 100644
--- /dev/null
+++ b/Engine/Source/Developer/TraceInsights/Private/Insights/Widgets/SIoProfilerWindow.cpp
@@ -0,0 +1,240 @@
+// Copyright 1998-2019 Epic Games, Inc. All Rights Reserved.
+
+#include "SIoProfilerWindow.h"
+
+#include "EditorStyleSet.h"
+#include "SlateOptMacros.h"
+#include "Widgets/Images/SImage.h"
+#include "Widgets/Layout/SBorder.h"
+#include "Widgets/Layout/SBox.h"
+#include "Widgets/Layout/SSpacer.h"
+#include "Widgets/Layout/SSplitter.h"
+#include "Widgets/Notifications/SNotificationList.h"
+#include "Widgets/SBoxPanel.h"
+#include "Widgets/SOverlay.h"
+#include "Widgets/Text/STextBlock.h"
+
+#if WITH_EDITOR
+	#include "EngineAnalytics.h"
+	#include "Runtime/Analytics/Analytics/Public/AnalyticsEventAttribute.h"
+	#include "Runtime/Analytics/Analytics/Public/Interfaces/IAnalyticsProvider.h"
+#endif // WITH_EDITOR
+
+// Insights
+#include "Insights/Version.h"
+#include "Insights/InsightsManager.h"
+#include "Insights/IoProfilerManager.h"
+#include "Insights/Widgets/STimingView.h"
+
+////////////////////////////////////////////////////////////////////////////////////////////////////
+
+#define LOCTEXT_NAMESPACE "SIoProfilerWindow"
+
+////////////////////////////////////////////////////////////////////////////////////////////////////
+
+SIoProfilerWindow::SIoProfilerWindow()
+	: DurationActive(0.0f)
+{
+}
+
+////////////////////////////////////////////////////////////////////////////////////////////////////
+
+SIoProfilerWindow::~SIoProfilerWindow()
+{
+#if WITH_EDITOR
+	if (DurationActive > 0.0f && FEngineAnalytics::IsAvailable())
+	{
+		FEngineAnalytics::GetProvider().RecordEvent(TEXT("Editor.Usage.Profiler"), FAnalyticsEventAttribute(TEXT("Duration"), DurationActive));
+	}
+#endif // WITH_EDITOR
+}
+
+////////////////////////////////////////////////////////////////////////////////////////////////////
+
+BEGIN_SLATE_FUNCTION_BUILD_OPTIMIZATION
+
+////////////////////////////////////////////////////////////////////////////////////////////////////
+
+void SIoProfilerWindow::Construct(const FArguments& InArgs)
+{
+	ChildSlot
+		[
+			SNew(SOverlay)
+
+			// Version
+			+ SOverlay::Slot()
+				.HAlign(HAlign_Right)
+				.VAlign(VAlign_Top)
+				.Padding(0.0f, -16.0f, 0.0f, 0.0f)
+				[
+					SNew(STextBlock)
+						.Clipping(EWidgetClipping::ClipToBoundsWithoutIntersecting)
+						.Text(LOCTEXT("UnrealInsightsVersion", UNREAL_INSIGHTS_VERSION_STRING_EX))
+						.ColorAndOpacity(FLinearColor(0.15f, 0.15f, 0.15f, 1.0f))
+				]
+
+			// Overlay slot for the main profiler window area
+			+ SOverlay::Slot()
+				.HAlign(HAlign_Fill)
+				.VAlign(VAlign_Fill)
+				[
+					SNew(SBox)
+						.Visibility(this, &SIoProfilerWindow::IsTimingViewVisible)
+						.IsEnabled(this, &SIoProfilerWindow::IsProfilerEnabled)
+						[
+							SAssignNew(TimingView, STimingView)
+						]
+				]
+
+			// Session hint overlay
+			+ SOverlay::Slot()
+				.HAlign(HAlign_Center)
+				.VAlign(VAlign_Center)
+				[
+					SNew(SBorder)
+						.Visibility(this, &SIoProfilerWindow::IsSessionOverlayVisible)
+						.BorderImage(FEditorStyle::GetBrush("NotificationList.ItemBackground"))
+						.Padding(8.0f)
+						[
+							SNew(STextBlock)
+								.Text(LOCTEXT("SelectTraceOverlayText", "Please select a trace."))
+						]
+				]
+		];
+
+	TimingView->EnableAssetLoadingMode();
+}
+
+END_SLATE_FUNCTION_BUILD_OPTIMIZATION
+
+////////////////////////////////////////////////////////////////////////////////////////////////////
+
+EVisibility SIoProfilerWindow::IsTimingViewVisible() const
+{
+	if (FInsightsManager::Get()->GetSession().IsValid())
+	{
+		return EVisibility::Visible;
+	}
+	else
+	{
+		return EVisibility::Collapsed;
+	}
+}
+
+////////////////////////////////////////////////////////////////////////////////////////////////////
+
+EVisibility SIoProfilerWindow::IsSessionOverlayVisible() const
+{
+	if (FInsightsManager::Get()->GetSession().IsValid())
+	{
+		return EVisibility::Hidden;
+	}
+	else
+	{
+		return EVisibility::Visible;
+	}
+}
+
+////////////////////////////////////////////////////////////////////////////////////////////////////
+
+bool SIoProfilerWindow::IsProfilerEnabled() const
+{
+	return FInsightsManager::Get()->GetSession().IsValid();
+}
+
+////////////////////////////////////////////////////////////////////////////////////////////////////
+
+EActiveTimerReturnType SIoProfilerWindow::UpdateActiveDuration(double InCurrentTime, float InDeltaTime)
+{
+	DurationActive += InDeltaTime;
+
+	// The profiler window will explicitly unregister this active timer when the mouse leaves.
+	return EActiveTimerReturnType::Continue;
+}
+
+////////////////////////////////////////////////////////////////////////////////////////////////////
+
+void SIoProfilerWindow::OnMouseEnter(const FGeometry& MyGeometry, const FPointerEvent& MouseEvent)
+{
+	SCompoundWidget::OnMouseEnter(MyGeometry, MouseEvent);
+
+	if (!ActiveTimerHandle.IsValid())
+	{
+		ActiveTimerHandle = RegisterActiveTimer(0.f, FWidgetActiveTimerDelegate::CreateSP(this, &SIoProfilerWindow::UpdateActiveDuration));
+	}
+}
+
+////////////////////////////////////////////////////////////////////////////////////////////////////
+
+void SIoProfilerWindow::OnMouseLeave(const FPointerEvent& MouseEvent)
+{
+	SCompoundWidget::OnMouseLeave(MouseEvent);
+
+	auto PinnedActiveTimerHandle = ActiveTimerHandle.Pin();
+	if (PinnedActiveTimerHandle.IsValid())
+	{
+		UnRegisterActiveTimer(PinnedActiveTimerHandle.ToSharedRef());
+	}
+}
+
+////////////////////////////////////////////////////////////////////////////////////////////////////
+
+FReply SIoProfilerWindow::OnKeyDown(const FGeometry& MyGeometry, const FKeyEvent& InKeyEvent)
+{
+	return FIoProfilerManager::Get()->GetCommandList()->ProcessCommandBindings(InKeyEvent) ? FReply::Handled() : FReply::Unhandled();
+}
+
+////////////////////////////////////////////////////////////////////////////////////////////////////
+
+FReply SIoProfilerWindow::OnDragOver(const FGeometry& MyGeometry, const FDragDropEvent& DragDropEvent)
+{
+	TSharedPtr<FExternalDragOperation> DragDropOp = DragDropEvent.GetOperationAs<FExternalDragOperation>();
+	if (DragDropOp.IsValid())
+	{
+		if (DragDropOp->HasFiles())
+		{
+			const TArray<FString>& Files = DragDropOp->GetFiles();
+			if (Files.Num() == 1)
+			{
+				const FString DraggedFileExtension = FPaths::GetExtension(Files[0], true);
+				if (DraggedFileExtension == TEXT(".utrace"))
+				{
+					return FReply::Handled();
+				}
+			}
+		}
+	}
+
+	return SCompoundWidget::OnDragOver(MyGeometry,DragDropEvent);
+}
+
+////////////////////////////////////////////////////////////////////////////////////////////////////
+
+FReply SIoProfilerWindow::OnDrop(const FGeometry& MyGeometry, const FDragDropEvent& DragDropEvent)
+{
+	TSharedPtr<FExternalDragOperation> DragDropOp = DragDropEvent.GetOperationAs<FExternalDragOperation>();
+	if (DragDropOp.IsValid())
+	{
+		if (DragDropOp->HasFiles())
+		{
+			// For now, only allow a single file.
+			const TArray<FString>& Files = DragDropOp->GetFiles();
+			if (Files.Num() == 1)
+			{
+				const FString DraggedFileExtension = FPaths::GetExtension(Files[0], true);
+				if (DraggedFileExtension == TEXT(".utrace"))
+				{
+					// Enqueue load operation.
+					FInsightsManager::Get()->LoadTraceFile(Files[0]);
+					return FReply::Handled();
+				}
+			}
+		}
+	}
+
+	return SCompoundWidget::OnDrop(MyGeometry,DragDropEvent);
+}
+
+////////////////////////////////////////////////////////////////////////////////////////////////////
+
+#undef LOCTEXT_NAMESPACE
diff --git a/Engine/Source/Developer/TraceInsights/Private/Insights/Widgets/SIoProfilerWindow.h b/Engine/Source/Developer/TraceInsights/Private/Insights/Widgets/SIoProfilerWindow.h
new file mode 100644
--- /dev/null
+++ b/Engine/Source/Developer/TraceInsights/Private/Insights/Widgets/SIoProfilerWindow.h
@@ -0,0 +1,113 @@
+// Copyright 1998-2019 Epic Games, Inc. All Rights Reserved.
+
+#pragma once
+
+#include "CoreMinimal.h"
+#include "Input/Reply.h"
+#include "Layout/Visibility.h"
+#include "Misc/Guid.h"
+#include "SlateFwd.h"
+#include "Widgets/DeclarativeSyntaxSupport.h"
+#include "Widgets/SCompoundWidget.h"
+
+// Insights
+#include "Insights/InsightsManager.h"
+
+class FActiveTimerHandle;
+class STimingView;
+
+/** Implements the timing profiler window. */
+class SIoProfilerWindow : public SCompoundWidget
+{
+public:
+	/** Default constructor. */
+	SIoProfilerWindow();
+
+	/** Virtual destructor. */
+	virtual ~SIoProfilerWindow();
+
+	SLATE_BEGIN_ARGS(SIoProfilerWindow){}
+	SLATE_END_ARGS()
+
+	/** Constructs this widget. */
+	void Construct(const FArguments& InArgs);
+
+	TSharedPtr<STimingView> GetTimingView() const { return TimingView; }
+
+private:
+	/** Callback for determining the visibility of the Timing view. */
+	EVisibility IsTimingViewVisible() const;
+
+	/** Callback for determining the visibility of the 'Select a session' overlay. */
+	EVisibility IsSessionOverlayVisible() const;
+
+	/** Callback for getting the enabled state of the profiler window. */
+	bool IsProfilerEnabled() const;
+
+	/** Updates the amount of time the profiler has been active. */
+	EActiveTimerReturnType UpdateActiveDuration(double InCurrentTime, float InDeltaTime);
+
+	/**
+	 * Ticks this widget. Override in derived classes, but always call the parent implementation.
+	 *
+	 * @param  AllottedGeometry The space allotted for this widget
+	 * @param  InCurrentTime  Current absolute real time
+	 * @param  InDeltaTime  Real time passed since last tick
+	 */
+	//virtual void Tick(const FGeometry& AllottedGeometry, const double InCurrentTime, const float InDeltaTime) override;
+
+	/**
+	 * The system will use this event to notify a widget that the cursor has entered it. This event is NOT bubbled.
+	 *
+	 * @param MyGeometry The Geometry of the widget receiving the event
+	 * @param MouseEvent Information about the input event
+	 */
+	virtual void OnMouseEnter(const FGeometry& MyGeometry, const FPointerEvent& MouseEvent) override;
+
+	/**
+	 * The system will use this event to notify a widget that the cursor has left it. This event is NOT bubbled.
+	 *
+	 * @param MouseEvent Information about the input event
+	 */
+	virtual void OnMouseLeave(const FPointerEvent& MouseEvent) override;
+
+	/**
+	 * Called after a key is pressed when this widget has focus
+	 *
+	 * @param MyGeometry The Geometry of the widget receiving the event
+	 * @param  InKeyEvent  Key event
+	 *
+	 * @return  Returns whether the event was handled, along with other possible actions
+	 */
+	virtual FReply OnKeyDown(const FGeometry& MyGeometry, const FKeyEvent& InKeyEvent) override;
+
+	/**
+	 * Called when the user is dropping something onto a widget; terminates drag and drop.
+	 *
+	 * @param MyGeometry      The geometry of the widget receiving the event.
+	 * @param DragDropEvent   The drag and drop event.
+	 *
+	 * @return A reply that indicated whether this event was handled.
+	 */
+	virtual FReply OnDrop(const FGeometry& MyGeometry, const FDragDropEvent& DragDropEvent) override;
+
+	/**
+	 * Called during drag and drop when the the mouse is being dragged over a widget.
+	 *
+	 * @param MyGeometry      The geometry of the widget receiving the event.
+	 * @param DragDropEvent   The drag and drop event.
+	 *
+	 * @return A reply that indicated whether this event was handled.
+	 */
+	virtual FReply OnDragOver(const FGeometry& MyGeometry, const FDragDropEvent& DragDropEvent)  override;
+
+private:
+	/** Widget for the timing track */
+	TSharedPtr<STimingView> TimingView;
+
+	/** The number of seconds the profiler has been active */
+	float DurationActive;
+
+	/** The handle to the active update duration tick */
+	TWeakPtr<FActiveTimerHandle> ActiveTimerHandle;
+};
diff --git a/Engine/Source/Developer/TraceInsights/Private/Insights/Widgets/SLogView.cpp b/Engine/Source/Developer/TraceInsights/Private/Insights/Widgets/SLogView.cpp
new file mode 100644
--- /dev/null
+++ b/Engine/Source/Developer/TraceInsights/Private/Insights/Widgets/SLogView.cpp
@@ -0,0 +1,1275 @@
+// Copyright 1998-2019 Epic Games, Inc. All Rights Reserved.
+
+#include "SLogView.h"
+
+#include "Algo/BinarySearch.h"
+#include "Async/AsyncWork.h"
+#include "EditorStyleSet.h"
+#include "Framework/Application/SlateApplication.h"
+#include "Framework/MultiBox/MultiBoxBuilder.h"
+#include "SlateOptMacros.h"
+#include "Widgets/Images/SImage.h"
+#include "Widgets/Input/SComboButton.h"
+#include "Widgets/Layout/SScrollBox.h"
+#include "Widgets/Text/STextBlock.h"
+
+// Insights
+#include "Insights/Common/TimeUtils.h"
+#include "Insights/InsightsManager.h"
+#include "Insights/TimingProfilerCommon.h"
+#include "Insights/TimingProfilerManager.h"
+#include "Insights/ViewModels/TimingViewDrawHelper.h"
+#include "Insights/Widgets/STimingProfilerWindow.h"
+#include "Insights/Widgets/STimingView.h"
+
+////////////////////////////////////////////////////////////////////////////////////////////////////
+
+#define LOCTEXT_NAMESPACE "SLogView"
+
+
+////////////////////////////////////////////////////////////////////////////////////////////////////
+// SLogMessageRow
+////////////////////////////////////////////////////////////////////////////////////////////////////
+
+class SLogMessageRow : public SMultiColumnTableRow<TSharedPtr<FLogMessage>>
+{
+	SLATE_BEGIN_ARGS(SLogMessageRow) {}
+	SLATE_END_ARGS()
+
+public:
+	/**
+	 * Constructs the widget.
+	 *
+	 * @param InArgs The construction arguments.
+	 * @param InLogMessage The log message displayed by this row.
+	 * @param InOwnerTableView The table to which the row must be added.
+	 */
+	void Construct(const FArguments& InArgs, TSharedPtr<FLogMessage> InLogMessage, TSharedRef<SLogView> InParentWidget, const TSharedRef<STableViewBase>& InOwnerTableView)
+	{
+		WeakLogMessage = MoveTemp(InLogMessage);
+		WeakParentWidget = InParentWidget;
+
+		SMultiColumnTableRow<TSharedPtr<FLogMessage>>::Construct(FSuperRowType::FArguments(), InOwnerTableView);
+
+		TSharedRef<SWidget> Row = ChildSlot.GetChildAt(0);
+
+		ChildSlot
+		[
+			SNew(SBorder)
+				.BorderImage(FCoreStyle::Get().GetBrush("WhiteBrush"))
+				.BorderBackgroundColor(this, &SLogMessageRow::GetBackgroundColor)
+				[
+					Row
+				]
+		];
+	}
+
+public:
+	virtual TSharedRef<SWidget> GenerateWidgetForColumn(const FName& ColumnName) override
+	{
+		if (ColumnName == LogViewColumns::IdColumnName)
+		{
+			return SNew(SBox)
+				.Padding(FMargin(4.0, 0.0))
+				[
+					SNew(STextBlock)
+					.Text(this, &SLogMessageRow::GetIndex)
+				];
+		}
+		else if (ColumnName == LogViewColumns::TimeColumnName)
+		{
+			return SNew(SBox)
+				.Padding(FMargin(4.0, 0.0))
+				[
+					SNew(STextBlock)
+					.Text(this, &SLogMessageRow::GetTime)
+					.ColorAndOpacity(this, &SLogMessageRow::GetColorAndOpacity)
+				];
+		}
+		else if (ColumnName == LogViewColumns::VerbosityColumnName)
+		{
+			return SNew(SBox)
+				.Padding(FMargin(4.0, 0.0))
+				[
+					SNew(STextBlock)
+					.Text(this, &SLogMessageRow::GetVerbosity)
+					.ColorAndOpacity(this, &SLogMessageRow::GetColorByVerbosity)
+				];
+		}
+		else if (ColumnName == LogViewColumns::CategoryColumnName)
+		{
+			return SNew(SBox)
+				.Padding(FMargin(4.0, 0.0))
+				[
+					SNew(STextBlock)
+					.Text(this, &SLogMessageRow::GetCategory)
+					.ColorAndOpacity(this, &SLogMessageRow::GetColorByCategory)
+				];
+		}
+		else if (ColumnName == LogViewColumns::MessageColumnName)
+		{
+			return SNew(SBox)
+				.Padding(FMargin(4.0, 0.0))
+				[
+					SNew(STextBlock)
+					.Text(this, &SLogMessageRow::GetMessage)
+					.HighlightText(this, &SLogMessageRow::GetMessageHighlightText)
+				];
+		}
+		else if (ColumnName == LogViewColumns::FileColumnName)
+		{
+			return SNew(SBox)
+				.Padding(FMargin(4.0, 0.0))
+				[
+					SNew(STextBlock)
+					.Text(this, &SLogMessageRow::GetFile)
+				];
+		}
+		else if (ColumnName == LogViewColumns::LineColumnName)
+		{
+			return SNew(SBox)
+				.Padding(FMargin(4.0, 0.0))
+				[
+					SNew(STextBlock)
+					.Text(this, &SLogMessageRow::GetLine)
+				];
+		}
+		else
+		{
+			return SNew(STextBlock).Text(LOCTEXT("UnknownColumn", "Unknown Column"));
+		}
+	}
+
+	FSlateColor GetBackgroundColor() const
+	{
+		TSharedPtr<SLogView> ParentWidgetPin = WeakParentWidget.Pin();
+		TSharedPtr<FLogMessage> LogMessagePin = WeakLogMessage.Pin();
+		if (ParentWidgetPin.IsValid() && LogMessagePin.IsValid())
+		{
+			TSharedPtr<FLogMessage> SelectedLogMessage = ParentWidgetPin->GetSelectedLogMessage();
+			if (!SelectedLogMessage || SelectedLogMessage->GetIndex() != LogMessagePin->GetIndex()) // if row is not selected
+			{
+				FLogMessageRecord& CacheEntry = ParentWidgetPin->GetCache().Get(LogMessagePin->GetIndex());
+				const double Time = CacheEntry.Time;
+
+				TSharedPtr<STimingProfilerWindow> Window = FTimingProfilerManager::Get()->GetProfilerWindow();
+				if (Window && Window->TimingView)
+				{
+					if (Window->TimingView->IsTimeSelectedInclusive(Time))
+					{
+						return FSlateColor(FLinearColor(0.25f, 0.5f, 1.0f, 0.25f));
+					}
+				}
+			}
+		}
+
+		return FSlateColor(FLinearColor(0.0f, 0.0f, 0.0f, 0.0f));
+	}
+
+	FSlateColor GetColorAndOpacity() const
+	{
+		bool IsSelected = false;
+
+		TSharedPtr<SLogView> ParentWidgetPin = WeakParentWidget.Pin();
+		TSharedPtr<FLogMessage> LogMessagePin = WeakLogMessage.Pin();
+		if (ParentWidgetPin.IsValid() && LogMessagePin.IsValid())
+		{
+			TSharedPtr<FLogMessage> SelectedLogMessage = ParentWidgetPin->GetSelectedLogMessage();
+			if (SelectedLogMessage && SelectedLogMessage->GetIndex() == LogMessagePin->GetIndex())
+			{
+				IsSelected = true;
+			}
+
+			FLogMessageRecord& CacheEntry = ParentWidgetPin->GetCache().Get(LogMessagePin->GetIndex());
+			const double Time = CacheEntry.Time;
+
+			TSharedPtr<STimingProfilerWindow> Window = FTimingProfilerManager::Get()->GetProfilerWindow();
+			if (Window && Window->TimingView)
+			{
+				if (Window->TimingView->IsTimeSelectedInclusive(Time))
+				{
+					if (IsSelected)
+					{
+						//return FSlateColor(FLinearColor(0.0f, 0.1f, 0.5f, 1.0f));
+						return FSlateColor(FLinearColor(0.0f, 0.05f, 0.2f, 1.0f));
+					}
+					else
+					{
+						//return FSlateColor(FLinearColor(0.2f, 0.4f, 0.8f, 1.0f));
+						return FSlateColor(FLinearColor(0.4f, 0.8f, 1.6f, 1.0f));
+					}
+				}
+			}
+		}
+
+		if (IsSelected)
+		{
+			return FSlateColor(FLinearColor::Black);
+		}
+		else
+		{
+			return FSlateColor(FLinearColor::White);
+		}
+	}
+
+	FText GetIndex() const
+	{
+		TSharedPtr<SLogView> ParentWidgetPin = WeakParentWidget.Pin();
+		TSharedPtr<FLogMessage> LogMessagePin = WeakLogMessage.Pin();
+		if (ParentWidgetPin.IsValid() && LogMessagePin.IsValid())
+		{
+			FLogMessageRecord& CacheEntry = ParentWidgetPin->GetCache().Get(LogMessagePin->GetIndex());
+			return CacheEntry.GetIndexAsText();
+		}
+		else
+		{
+			return FText();
+		}
+	}
+
+	FText GetTime() const
+	{
+		TSharedPtr<SLogView> ParentWidgetPin = WeakParentWidget.Pin();
+		TSharedPtr<FLogMessage> LogMessagePin = WeakLogMessage.Pin();
+		if (ParentWidgetPin.IsValid() && LogMessagePin.IsValid())
+		{
+			FLogMessageRecord& CacheEntry = ParentWidgetPin->GetCache().Get(LogMessagePin->GetIndex());
+			return CacheEntry.GetTimeAsText();
+		}
+		else
+		{
+			return FText();
+		}
+	}
+
+	FText GetVerbosity() const
+	{
+		TSharedPtr<SLogView> ParentWidgetPin = WeakParentWidget.Pin();
+		TSharedPtr<FLogMessage> LogMessagePin = WeakLogMessage.Pin();
+		if (ParentWidgetPin.IsValid() && LogMessagePin.IsValid())
+		{
+			FLogMessageRecord& CacheEntry = ParentWidgetPin->GetCache().Get(LogMessagePin->GetIndex());
+			return CacheEntry.GetVerbosityAsText();
+		}
+		else
+		{
+			return FText();
+		}
+	}
+
+	FSlateColor GetColorByVerbosity() const
+	{
+		if (IsSelected())
+			return FSlateColor(FLinearColor::Black);
+
+		TSharedPtr<SLogView> ParentWidgetPin = WeakParentWidget.Pin();
+		TSharedPtr<FLogMessage> LogMessagePin = WeakLogMessage.Pin();
+		if (ParentWidgetPin.IsValid() && LogMessagePin.IsValid())
+		{
+			FLogMessageRecord& CacheEntry = ParentWidgetPin->GetCache().Get(LogMessagePin->GetIndex());
+			return FSlateColor(FTimeMarkerTrackBuilder::GetColorByVerbosity(CacheEntry.Verbosity));
+		}
+		else
+		{
+			return FSlateColor(FLinearColor::White);
+		}
+	}
+
+	FText GetCategory() const
+	{
+		TSharedPtr<SLogView> ParentWidgetPin = WeakParentWidget.Pin();
+		TSharedPtr<FLogMessage> LogMessagePin = WeakLogMessage.Pin();
+		if (ParentWidgetPin.IsValid() && LogMessagePin.IsValid())
+		{
+			FLogMessageRecord& CacheEntry = ParentWidgetPin->GetCache().Get(LogMessagePin->GetIndex());
+			return CacheEntry.GetCategoryAsText();
+		}
+		else
+		{
+			return FText();
+		}
+	}
+
+	FSlateColor GetColorByCategory() const
+	{
+		if (IsSelected())
+			return FSlateColor(FLinearColor::Black);
+
+		TSharedPtr<SLogView> ParentWidgetPin = WeakParentWidget.Pin();
+		TSharedPtr<FLogMessage> LogMessagePin = WeakLogMessage.Pin();
+		if (ParentWidgetPin.IsValid() && LogMessagePin.IsValid())
+		{
+			FLogMessageRecord& CacheEntry = ParentWidgetPin->GetCache().Get(LogMessagePin->GetIndex());
+			return FSlateColor(FTimeMarkerTrackBuilder::GetColorByCategory(*CacheEntry.Category.ToString()));
+		}
+		else
+		{
+			return FSlateColor(FLinearColor::White);
+		}
+	}
+
+	FText GetMessage() const
+	{
+		TSharedPtr<SLogView> ParentWidgetPin = WeakParentWidget.Pin();
+		TSharedPtr<FLogMessage> LogMessagePin = WeakLogMessage.Pin();
+		if (ParentWidgetPin.IsValid() && LogMessagePin.IsValid())
+		{
+			FLogMessageRecord& CacheEntry = ParentWidgetPin->GetCache().Get(LogMessagePin->GetIndex());
+			return CacheEntry.GetMessageAsText();
+		}
+		else
+		{
+			return FText();
+		}
+	}
+
+	FText GetMessageHighlightText() const
+	{
+		TSharedPtr<SLogView> ParentWidgetPin = WeakParentWidget.Pin();
+		if (ParentWidgetPin.IsValid())
+		{
+			return ParentWidgetPin->GetFilterText();
+		}
+		else
+		{
+			return FText();
+		}
+	}
+
+	FText GetFile() const
+	{
+		TSharedPtr<SLogView> ParentWidgetPin = WeakParentWidget.Pin();
+		TSharedPtr<FLogMessage> LogMessagePin = WeakLogMessage.Pin();
+		if (ParentWidgetPin.IsValid() && LogMessagePin.IsValid())
+		{
+			FLogMessageRecord& CacheEntry = ParentWidgetPin->GetCache().Get(LogMessagePin->GetIndex());
+			return CacheEntry.GetFileAsText();
+		}
+		else
+		{
+			return FText();
+		}
+	}
+
+	FText GetLine() const
+	{
+		TSharedPtr<SLogView> ParentWidgetPin = WeakParentWidget.Pin();
+		TSharedPtr<FLogMessage> LogMessagePin = WeakLogMessage.Pin();
+		if (ParentWidgetPin.IsValid() && LogMessagePin.IsValid())
+		{
+			FLogMessageRecord& CacheEntry = ParentWidgetPin->GetCache().Get(LogMessagePin->GetIndex());
+			return CacheEntry.GetLineAsText();
+		}
+		else
+		{
+			return FText();
+		}
+	}
+
+	FText GetRowToolTip() const
+	{
+		TSharedPtr<SLogView> ParentWidgetPin = WeakParentWidget.Pin();
+		TSharedPtr<FLogMessage> LogMessagePin = WeakLogMessage.Pin();
+		if (ParentWidgetPin.IsValid() && LogMessagePin.IsValid())
+		{
+			FLogMessageRecord& CacheEntry = ParentWidgetPin->GetCache().Get(LogMessagePin->GetIndex());
+			return CacheEntry.ToDisplayString();
+		}
+		else
+		{
+			return FText();
+		}
+	}
+
+private:
+	TWeakPtr<FLogMessage> WeakLogMessage;
+	TWeakPtr<SLogView> WeakParentWidget;
+};
+
+////////////////////////////////////////////////////////////////////////////////////////////////////
+// SLogView
+////////////////////////////////////////////////////////////////////////////////////////////////////
+
+SLogView::SLogView()
+	: FilteringStartIndex(0)
+	, FilteringEndIndex(0)
+	, FilteringChangeNumber(0)
+	, bIsFilteringAsyncTaskCancelRequested(false)
+	, TotalNumCategories(0)
+	, TotalNumMessages(0)
+	, bIsDirty(false)
+{
+}
+
+////////////////////////////////////////////////////////////////////////////////////////////////////
+
+SLogView::~SLogView()
+{
+	// Remove ourselves from the profiler manager.
+	//if (FTimingProfilerManager::Get().IsValid())
+	//{
+	//	//TODO: FTimimgProfilerManager::Get()->OnRequestLogViewUpdate().RemoveAll(this);
+	//}
+
+	Reset();
+}
+
+////////////////////////////////////////////////////////////////////////////////////////////////////
+
+void SLogView::Reset()
+{
+	//ListView
+	//ExternalScrollbar
+	FilterTextBox->SetText(FText::GetEmpty());
+
+	Filter.Reset();
+	FilterChangeNumber = 0;
+
+	FilteringStartIndex = 0;
+	FilteringEndIndex = 0;
+	FilteringChangeNumber = 0;
+
+	// Clean up our async task if we're deleted before it is completed.
+	if (FilteringAsyncTask.IsValid())
+	{
+		if (!FilteringAsyncTask->Cancel())
+		{
+			bIsFilteringAsyncTaskCancelRequested = true;
+			FilteringAsyncTask->EnsureCompletion();
+		}
+		FilteringAsyncTask.Reset();
+	}
+
+	//bIsFilteringAsyncTaskCancelRequested = false;
+	//FilteringStopwatch.Stop();
+
+	TotalNumCategories = 0;
+	TotalNumMessages = 0;
+
+	bIsDirty = false;
+	DirtyStopwatch.Stop();
+
+	StatsText = FText::GetEmpty();
+
+	Cache.Reset();
+
+	Messages.Reset();
+
+	ListView->RebuildList();
+}
+
+////////////////////////////////////////////////////////////////////////////////////////////////////
+
+BEGIN_SLATE_FUNCTION_BUILD_OPTIMIZATION
+void SLogView::Construct(const FArguments& InArgs)
+{
+	SAssignNew(ExternalScrollbar, SScrollBar)
+	.AlwaysShowScrollbar(true);
+
+	ChildSlot
+	[
+		SNew(SVerticalBox)
+
+		// Toolbar
+		+ SVerticalBox::Slot()
+		.AutoHeight()
+		[
+			SNew(SHorizontalBox)
+
+			// Icon
+			//+ SHorizontalBox::Slot()
+			//.VAlign(VAlign_Center)
+			//.AutoWidth()
+			//[
+			//	SNew(SImage)
+			//	.Image(FEditorStyle::GetBrush(TEXT("Log.TabIcon")))
+			//]
+			//
+			//// Label
+			//+ SHorizontalBox::Slot()
+			//.VAlign(VAlign_Center)
+			//.Padding(2.0f, 0.0f)
+			//.AutoWidth()
+			//[
+			//	SNew(STextBlock)
+			//	.Text(LOCTEXT("LogViewLabel", "Log View"))
+			//]
+
+			// Verbosity Threshold Filter
+			+SHorizontalBox::Slot()
+			//.Padding(2.0f, 0.0f, 0.0f, 0.0f)
+			.VAlign(VAlign_Center)
+			.AutoWidth()
+			[
+				SNew(SComboButton)
+				.ComboButtonStyle(FEditorStyle::Get(), "GenericFilters.ComboButtonStyle")
+				.ForegroundColor(FLinearColor::White)
+				.ContentPadding(0)
+				.ToolTipText(LOCTEXT("VerbosityThresholdFilterToolTip", "Filter log messages by verbosity threshold."))
+				.OnGetMenuContent(this, &SLogView::MakeVerbosityThresholdMenu)
+				.HasDownArrow(true)
+				.ContentPadding(FMargin(1.0f, 1.0f, 1.0f, 0.0f))
+				.ButtonContent()
+				[
+					SNew(SHorizontalBox)
+
+					+SHorizontalBox::Slot()
+					.Padding(0.0f)
+					.AutoWidth()
+					[
+						SNew(STextBlock)
+						.TextStyle(FEditorStyle::Get(), "GenericFilters.TextStyle")
+						.Font(FEditorStyle::Get().GetFontStyle("FontAwesome.9"))
+						.Text(FText::FromString(FString(TEXT("\xf0b0"))) /*fa-filter*/)
+					]
+
+					+SHorizontalBox::Slot()
+					.Padding(2.0f, 0.0f, 0.0f, 0.0f)
+					.AutoWidth()
+					[
+						SNew(STextBlock)
+						.TextStyle(FEditorStyle::Get(), "GenericFilters.TextStyle")
+						.Text(LOCTEXT("VerbosityThresholdFilter", "Verbosity Threshold"))
+					]
+				]
+			]
+
+			// Category Filter
+			+SHorizontalBox::Slot()
+			.Padding(2.0f, 0.0f, 0.0f, 0.0f)
+			.VAlign(VAlign_Center)
+			.AutoWidth()
+			[
+				SNew(SComboButton)
+				.ComboButtonStyle(FEditorStyle::Get(), "GenericFilters.ComboButtonStyle")
+				.ForegroundColor(FLinearColor::White)
+				.ContentPadding(0)
+				.ToolTipText(LOCTEXT("CategoryFilterToolTip", "Filter log messages by category."))
+				.OnGetMenuContent(this, &SLogView::MakeCategoryFilterMenu)
+				.HasDownArrow(true)
+				.ContentPadding(FMargin(1.0f, 1.0f, 1.0f, 0.0f))
+				.ButtonContent()
+				[
+					SNew(SHorizontalBox)
+
+					+SHorizontalBox::Slot()
+					.Padding(0.0f)
+					.AutoWidth()
+					[
+						SNew(STextBlock)
+						.TextStyle(FEditorStyle::Get(), "GenericFilters.TextStyle")
+						.Font(FEditorStyle::Get().GetFontStyle("FontAwesome.9"))
+						.Text(FText::FromString(FString(TEXT("\xf0b0"))) /*fa-filter*/)
+					]
+
+					+SHorizontalBox::Slot()
+					.Padding(2.0f, 0.0f, 0.0f, 0.0f)
+					.AutoWidth()
+					[
+						SNew(STextBlock)
+						.TextStyle(FEditorStyle::Get(), "GenericFilters.TextStyle")
+						.Text(LOCTEXT("CategoryFilter", "Category Filter"))
+					]
+				]
+			]
+
+			// Text Filter (Search Box)
+			+ SHorizontalBox::Slot()
+			.VAlign(VAlign_Center)
+			.Padding(2.0f, 0.0f)
+			.AutoWidth()
+			[
+				SAssignNew(FilterTextBox, SSearchBox)
+				.HintText(LOCTEXT("FilterTextBoxHint", "Search log messages"))
+				.ToolTipText(LOCTEXT("FilterTextBoxToolTip", "Type here to filter the list of log messages."))
+				.OnTextChanged(this, &SLogView::FilterTextBox_OnTextChanged)
+			]
+
+			// Stats Text
+			+ SHorizontalBox::Slot()
+			.VAlign(VAlign_Center)
+			.Padding(2.0f, 0.0f)
+			.AutoWidth()
+			[
+				SNew(STextBlock)
+				.Text(this, &SLogView::GetStatsText)
+				.ColorAndOpacity(this, &SLogView::GetStatsTextColor)
+			]
+		]
+
+		+ SVerticalBox::Slot()
+		.FillHeight(1.0f)
+		[
+			SNew(SBox)
+			.VAlign(VAlign_Fill)
+			[
+				SNew(SHorizontalBox)
+
+				+ SHorizontalBox::Slot()
+				.FillWidth(1.0f)
+				.Padding(0.0f)
+				.VAlign(VAlign_Fill)
+				[
+					SNew(SScrollBox)
+					.Orientation(Orient_Horizontal)
+
+					+ SScrollBox::Slot()
+					.VAlign(VAlign_Fill)
+					[
+						SNew(SBorder)
+						.BorderImage(FEditorStyle::GetBrush("ToolPanel.GroupBorder"))
+						.Padding(0.0f)
+						[
+							SAssignNew(ListView, SListView<TSharedPtr<FLogMessage>>)
+							.ExternalScrollbar(ExternalScrollbar)
+							.ItemHeight(20.0f)
+							.SelectionMode(ESelectionMode::Single)
+							.OnSelectionChanged(this, &SLogView::OnSelectionChanged)
+							.ListItemsSource(&Messages)
+							.OnGenerateRow(this, &SLogView::OnGenerateRow)
+							.ConsumeMouseWheel(EConsumeMouseWheel::Always)
+							.OnContextMenuOpening(FOnContextMenuOpening::CreateSP(this, &SLogView::ListView_GetContextMenu))
+							.HeaderRow
+							(
+								SNew(SHeaderRow)
+
+								+ SHeaderRow::Column(LogViewColumns::IdColumnName)
+								.ManualWidth(60.0f)
+								.DefaultLabel(LOCTEXT("IdColumn", "Index"))
+
+								+ SHeaderRow::Column(LogViewColumns::TimeColumnName)
+								.ManualWidth(94.0f)
+								.DefaultLabel(LOCTEXT("TimeColumn", "Time"))
+
+								+ SHeaderRow::Column(LogViewColumns::VerbosityColumnName)
+								.ManualWidth(80.0f)
+								.DefaultLabel(LOCTEXT("VerbosityColumn", "Verbosity"))
+
+								+ SHeaderRow::Column(LogViewColumns::CategoryColumnName)
+								.ManualWidth(120.0f)
+								.DefaultLabel(LOCTEXT("CategoryColumn", "Category"))
+
+								+ SHeaderRow::Column(LogViewColumns::MessageColumnName)
+								.ManualWidth(880.0f)
+								.DefaultLabel(LOCTEXT("MessageColumn", "Message"))
+
+								+ SHeaderRow::Column(LogViewColumns::FileColumnName)
+								.ManualWidth(600.0f)
+								.DefaultLabel(LOCTEXT("FileColumn", "File"))
+
+								+ SHeaderRow::Column(LogViewColumns::LineColumnName)
+								.ManualWidth(60.0f)
+								.DefaultLabel(LOCTEXT("LineColumn", "Line"))
+							)
+						]
+					]
+				]
+
+				+ SHorizontalBox::Slot()
+				.AutoWidth()
+				.Padding(0.0f)
+				[
+					SNew(SBox)
+					.WidthOverride(FOptionalSize(13.0f))
+					[
+						ExternalScrollbar.ToSharedRef()
+					]
+				]
+			]
+		]
+	];
+
+	// Register ourselves with the profiler manager.
+	//TODO: FTimingProfilerManager::Get()->OnRequestLogViewUpdate().AddSP(this, &SLogView::ProfilerManager_OnRequestLogViewUpdate);
+}
+END_SLATE_FUNCTION_BUILD_OPTIMIZATION
+
+////////////////////////////////////////////////////////////////////////////////////////////////////
+
+TSharedRef<ITableRow> SLogView::OnGenerateRow(TSharedPtr<FLogMessage> InLogMessage, const TSharedRef<STableViewBase>& OwnerTable)
+{
+	// Generate a row for the log message corresponding to InLogMessage.
+	return SNew(SLogMessageRow, InLogMessage, SharedThis(this), OwnerTable);
+}
+
+////////////////////////////////////////////////////////////////////////////////////////////////////
+
+void SLogView::Tick(const FGeometry& AllottedGeometry, const double InCurrentTime, const float InDeltaTime)
+{
+	int32 NewMessageCount = 0;
+
+	TSharedPtr<const Trace::IAnalysisSession> Session = FInsightsManager::Get()->GetSession();
+
+	Cache.SetSession(Session);
+
+	if (Session.IsValid())
+	{
+		Trace::FAnalysisSessionReadScope SessionReadScope(*Session.Get());
+		const Trace::ILogProvider& LogProvider = Trace::ReadLogProvider(*Session.Get());
+
+		NewMessageCount = static_cast<int32>(LogProvider.GetMessageCount());
+
+		//TODO: show only categories that are used in current trace
+		//TODO: cause of duplicates: a) runtime, b) case insensitive, c) stripped "Log" prefix
+		//TODO: FString vs. FText vs. FName ?
+
+		//TODO: int32 NumCategories = static_cast<int32>(LogProvider.GetCategoriesCount());
+		int32 NumCategories = 0;
+		LogProvider.EnumerateCategories([&NumCategories](const Trace::FLogCategory& Category)
+		{
+			NumCategories++;
+		});
+
+		if (NumCategories != TotalNumCategories)
+		{
+			TotalNumCategories = NumCategories;
+			UE_LOG(TimingProfiler, Log, TEXT("[LogView] Total Log Categories: %d"), TotalNumCategories);
+
+			TSet<FName> Categories;
+			LogProvider.EnumerateCategories([&Categories](const Trace::FLogCategory& Category)
+			{
+				FString CategoryStr(Category.Name);
+				if (CategoryStr.StartsWith(TEXT("Log")))
+				{
+					CategoryStr = CategoryStr.RightChop(3);
+				}
+				if (Categories.Contains(FName(*CategoryStr)))
+				{
+					UE_LOG(TimingProfiler, Log, TEXT("[LogView] Duplicated Log Category: \"%s\""), Category.Name);
+				}
+				Categories.Add(FName(*CategoryStr));
+			});
+			Filter.SyncAvailableCategories(Categories);
+			UE_LOG(TimingProfiler, Log, TEXT("[LogView] Unique Log Categories: %d"), Filter.GetAvailableLogCategories().Num());
+
+			//Cache.Reset();
+			Messages.Reset();
+			TotalNumMessages = 0;
+			//ListView->RebuildList();
+			bIsDirty = true;
+			DirtyStopwatch.Start();
+		}
+	}
+
+	if (NewMessageCount != TotalNumMessages)
+	{
+		bIsDirty = true;
+		DirtyStopwatch.Start();
+
+		if (NewMessageCount > TotalNumMessages)
+		{
+			if (Filter.IsFilterSet())
+			{
+				// Filter messages async.
+				if (!FilteringAsyncTask.IsValid())
+				{
+					FilteringStopwatch.Restart();
+
+					bIsFilteringAsyncTaskCancelRequested = false;
+					FilteringStartIndex = TotalNumMessages;
+					FilteringEndIndex = NewMessageCount;
+					FilteringChangeNumber = Filter.GetChangeNumber();
+					FilteringAsyncTask = MakeUnique<FAsyncTask<FLogFilteringAsyncTask>>(FilteringStartIndex, FilteringEndIndex, Filter, SharedThis(this));
+					UE_LOG(TimingProfiler, Log, TEXT("[LogView] Start async task for filtering by%s%s%s (\"%s\") (%d to %d)"),
+						Filter.IsFilterSetByVerbosity() ? TEXT(" Verbosity,") : TEXT(""),
+						Filter.IsFilterSetByCategory() ? TEXT(" Category,") : TEXT(""),
+						Filter.IsFilterSetByText() ? TEXT(" Text") : TEXT(""),
+						*Filter.GetFilterText().ToString(),
+						FilteringStartIndex, FilteringEndIndex);
+					FilteringAsyncTask->StartBackgroundTask();
+				}
+				else
+				{
+					// A task is already in progress.
+					if (FilteringStartIndex == TotalNumMessages &&
+						FilteringEndIndex <= NewMessageCount &&
+						FilteringChangeNumber == Filter.GetChangeNumber())
+					{
+						// The filter is still valid. Just wait.
+					}
+					else
+					{
+						// The filter used by running task is obsolete. Cancel the task.
+						bIsFilteringAsyncTaskCancelRequested = true;
+					}
+				}
+			}
+			else // no filtering
+			{
+				FilteringStopwatch.Restart();
+
+				for (int32 Index = TotalNumMessages; Index < NewMessageCount; Index++)
+				{
+					FLogMessage LogMessage(Index);
+					Messages.Add(MakeShared<FLogMessage>(MoveTemp(LogMessage)));
+				}
+
+				const int32 NumAddedMessages = NewMessageCount - TotalNumMessages;
+				TotalNumMessages = NewMessageCount;
+				TSharedPtr<FLogMessage> SelectedLogMessage = GetSelectedLogMessage();
+				ListView->RebuildList();
+				if (SelectedLogMessage.IsValid())
+				{
+					SelectedLogMessageByLogIndex(SelectedLogMessage->GetIndex());
+				}
+				bIsDirty = false;
+				DirtyStopwatch.Reset();
+				UpdateStatsText();
+
+				FilteringStopwatch.Stop();
+				uint64 DurationMs = FilteringStopwatch.GetAccumulatedTimeMs();
+				if (DurationMs > 10) // avoids spams
+				{
+					UE_LOG(TimingProfiler, Log, TEXT("[LogView] Updated (no filter; %d added / %d total messages) in %llu ms."),
+						NumAddedMessages, NewMessageCount, DurationMs);
+				}
+			}
+		}
+		else // if (NewMessageCount < TotalNumMessages)
+		{
+			// Just reset. On next Tick() the list will grow if needed.
+			UE_LOG(TimingProfiler, Log, TEXT("[LogView] RESET"));
+			Cache.Reset();
+			Messages.Reset();
+			TotalNumMessages = 0;
+			ListView->RebuildList();
+			bIsDirty = (NewMessageCount != 0);
+			if (bIsDirty)
+			{
+				DirtyStopwatch.Start();
+			}
+			else
+			{
+				DirtyStopwatch.Reset();
+			}
+			UpdateStatsText();
+		}
+	}
+
+	if (FilteringAsyncTask.IsValid() &&
+		FilteringAsyncTask->IsDone())
+	{
+		// A filtering async task has completed. Check if filter used is still valid.
+		if (!bIsFilteringAsyncTaskCancelRequested &&
+			FilteringStartIndex == TotalNumMessages &&
+			FilteringEndIndex <= NewMessageCount &&
+			FilteringChangeNumber == Filter.GetChangeNumber())
+		{
+			FLogFilteringAsyncTask& Task = FilteringAsyncTask->GetTask();
+			const TArray<uint32>& FilteredMessages = Task.GetFilteredMessages();
+
+			// Add filtered messages to current Messages array.
+			const int32 NumFilteredMessages = FilteredMessages.Num();
+			for (int32 Index = 0; Index < NumFilteredMessages; Index++)
+			{
+				FLogMessage LogMessage(FilteredMessages[Index]);
+				Messages.Add(MakeShared<FLogMessage>(MoveTemp(LogMessage)));
+			}
+
+			TotalNumMessages = Task.GetEndIndex();
+			TSharedPtr<FLogMessage> SelectedLogMessage = GetSelectedLogMessage();
+			ListView->RebuildList();
+			if (SelectedLogMessage.IsValid())
+			{
+				SelectedLogMessageByLogIndex(SelectedLogMessage->GetIndex());
+			}
+			bIsDirty = false;
+			DirtyStopwatch.Reset();
+			UpdateStatsText();
+
+			FilteringStopwatch.Stop();
+			uint64 DurationMs = FilteringStopwatch.GetAccumulatedTimeMs();
+			if (DurationMs > 10) // avoids spams
+			{
+				int32 NumAsyncFilteredMessages = Task.GetEndIndex() - Task.GetStartIndex();
+				double Speed = static_cast<double>(NumAsyncFilteredMessages) / FilteringStopwatch.GetAccumulatedTime();
+				UE_LOG(TimingProfiler, Log, TEXT("[LogView] Updated (%d added / %d async filtered / %d total messages) in %llu ms (%.2f messages/second)."),
+					NumFilteredMessages, NumAsyncFilteredMessages, TotalNumMessages, DurationMs, Speed);
+			}
+		}
+
+		FilteringAsyncTask.Reset();
+	}
+
+	//SCompoundWidget::Tick(AllottedGeometry, InCurrentTime, InDeltaTime);
+}
+
+////////////////////////////////////////////////////////////////////////////////////////////////////
+
+TSharedPtr<FLogMessage> SLogView::GetSelectedLogMessage() const
+{
+	TArray<TSharedPtr<FLogMessage>> SelectedItems = ListView->GetSelectedItems();
+	return (SelectedItems.Num() == 1) ? SelectedItems[0] : nullptr;
+}
+
+////////////////////////////////////////////////////////////////////////////////////////////////////
+
+void SLogView::SelectedLogMessageByLogIndex(int32 LogIndex)
+{
+	// We are assuming the Messages list is sorted by log index...
+	int32 MessageIndex = Algo::BinarySearchBy(Messages, LogIndex, &FLogMessage::GetIndex);
+	if (MessageIndex != INDEX_NONE)
+	{
+		ListView->SetItemSelection(Messages[MessageIndex], true);
+		ListView->RequestScrollIntoView(Messages[MessageIndex]);
+	}
+}
+
+////////////////////////////////////////////////////////////////////////////////////////////////////
+
+void SLogView::OnSelectionChanged(TSharedPtr<FLogMessage> LogMessage, ESelectInfo::Type SelectInfo)
+{
+	TSharedPtr<STimingProfilerWindow> Window = FTimingProfilerManager::Get()->GetProfilerWindow();
+	if (Window && Window->TimingView)
+	{
+		// Single item selection.
+		if (LogMessage.IsValid())
+		{
+			const double Time = Cache.Get(LogMessage->GetIndex()).Time;
+
+			if (FSlateApplication::Get().GetModifierKeys().IsShiftDown())
+			{
+				Window->TimingView->SelectToTimeMarker(Time);
+			}
+			else
+			{
+				Window->TimingView->SetAndCenterOnTimeMarker(Time);
+			}
+		}
+	}
+}
+
+////////////////////////////////////////////////////////////////////////////////////////////////////
+
+void SLogView::FilterTextBox_OnTextChanged(const FText& InFilterText)
+{
+	if (Filter.GetFilterText().ToString().Equals(InFilterText.ToString(), ESearchCase::CaseSensitive))
+	{
+		// nothing to do
+		return;
+	}
+
+	// Set filter phrases.
+	Filter.SetFilterText(InFilterText);
+
+	// Report possible syntax errors back to the user.
+	FilterTextBox->SetError(Filter.GetSyntaxErrors());
+
+	OnFilterChanged();
+}
+
+////////////////////////////////////////////////////////////////////////////////////////////////////
+
+void SLogView::OnFilterChanged()
+{
+	const FString FilterText = FilterTextBox->GetText().ToString();
+	UE_LOG(TimingProfiler, Log, TEXT("[LogView] OnFilterChanged: \"%s\""), *FilterText);
+	Cache.Reset();
+	Messages.Reset();
+	TotalNumMessages = 0;
+	bIsDirty = true;
+	DirtyStopwatch.Start();
+	// The next Tick() will update the filtered list of messages.
+}
+
+////////////////////////////////////////////////////////////////////////////////////////////////////
+
+void SLogView::UpdateStatsText()
+{
+	if (Messages.Num() == TotalNumMessages)
+	{
+		StatsText = FText::Format(LOCTEXT("StatsText1", "{0} logs"), FText::AsNumber(TotalNumMessages));
+	}
+	else
+	{
+		StatsText = FText::Format(LOCTEXT("StatsText2", "{0} / {1} logs"), FText::AsNumber(Messages.Num()), FText::AsNumber(TotalNumMessages));
+	}
+}
+
+////////////////////////////////////////////////////////////////////////////////////////////////////
+
+FText SLogView::GetStatsText() const
+{
+	if (bIsDirty)
+	{
+		DirtyStopwatch.Update();
+		double DT = DirtyStopwatch.GetAccumulatedTime();
+		if (DT > 1.0)
+		{
+			return FText::Format(LOCTEXT("StatsTextEx", "{0} (filtering... please wait... {1}s)"), StatsText, FText::AsNumber(FMath::RoundToInt(DT)));
+		}
+	}
+
+	return StatsText;
+}
+
+////////////////////////////////////////////////////////////////////////////////////////////////////
+
+FSlateColor SLogView::GetStatsTextColor() const
+{
+	if (bIsDirty)
+	{
+		return FSlateColor(FLinearColor(1.0f, 0.5f, 0.5f, 1.0f));
+	}
+	else
+	{
+		return FSlateColor(FLinearColor(1.0f, 1.0f, 1.0f, 1.0f));
+	}
+}
+
+////////////////////////////////////////////////////////////////////////////////////////////////////
+
+TSharedPtr<SWidget> SLogView::ListView_GetContextMenu()
+{
+	TSharedPtr<FLogMessage> SelectedLogMessage = GetSelectedLogMessage();
+
+	FMenuBuilder MenuBuilder(/*bInShouldCloseWindowAfterMenuSelection=*/true, nullptr);
+
+	MenuBuilder.BeginSection("LogViewContextMenu");
+	{
+		if (SelectedLogMessage.IsValid())
+		{
+			FLogMessageRecord& Record = Cache.Get(SelectedLogMessage->GetIndex());
+			FName CategoryName(*Record.Category.ToString());
+
+			MenuBuilder.AddMenuEntry(
+				FText::Format(LOCTEXT("HideCategory", "Hide \"{0}\" Category"), Record.Category),
+				FText::Format(LOCTEXT("HideCategory_Tooltip", "Hide the \"{0}\" log category."), Record.Category),
+				FSlateIcon(),
+				FUIAction(FExecuteAction::CreateSP(this, &SLogView::ToggleCategory_Execute, CategoryName)),
+				NAME_None,
+				EUserInterfaceActionType::Button
+			);
+
+			MenuBuilder.AddMenuEntry(
+				FText::Format(LOCTEXT("ShowOnlyCategory", "Show Only \"{0}\" Category"), Record.Category),
+				FText::Format(LOCTEXT("ShowOnlyCategory_Tooltip", "Show only the \"{0}\" log category (hide all other log categories)."), Record.Category),
+				FSlateIcon(),
+				FUIAction(FExecuteAction::CreateSP(this, &SLogView::ShowOnlyCategory_Execute, CategoryName)),
+				NAME_None,
+				EUserInterfaceActionType::Button
+			);
+		}
+
+		MenuBuilder.AddMenuEntry(
+			LOCTEXT("ShowAllCategoriesCtxMenu", "Show All Categories"),
+			LOCTEXT("ShowAllCategoriesCtxMenu_Tooltip", "Reset category filter (show all log categories)."),
+			FSlateIcon(),
+			FUIAction(FExecuteAction::CreateSP(this, &SLogView::ShowAllCategories_Execute)),
+			NAME_None,
+			EUserInterfaceActionType::Button
+		);
+	}
+	MenuBuilder.EndSection();
+
+	return MenuBuilder.MakeWidget();
+}
+
+////////////////////////////////////////////////////////////////////////////////////////////////////
+
+TSharedRef<SWidget> SLogView::MakeVerbosityThresholdMenu()
+{
+	FMenuBuilder MenuBuilder(/*bInShouldCloseWindowAfterMenuSelection=*/true, nullptr);
+
+	MenuBuilder.BeginSection("LogViewVerbosityThreshold"/*, LOCTEXT("VerbosityThresholdHeading", "Verbosity Threshold")*/);
+	CreateVerbosityThresholdMenuSection(MenuBuilder);
+	MenuBuilder.EndSection();
+
+	return MenuBuilder.MakeWidget();
+}
+
+////////////////////////////////////////////////////////////////////////////////////////////////////
+
+void  SLogView::CreateVerbosityThresholdMenuSection(FMenuBuilder& MenuBuilder)
+{
+	struct FVerbosityThresholdInfo
+	{
+		ELogVerbosity::Type Verbosity;
+		FText Label;
+		FText ToolTip;
+	};
+
+	FVerbosityThresholdInfo VerbosityThresholds[] =
+	{
+		{ ELogVerbosity::VeryVerbose, LOCTEXT("VerbosityThresholdVeryVerbose", "VeryVerbose"), LOCTEXT("VerbosityThresholdVeryVerbose_Tooltip", "Show all log messages (any verbosity level)."), },
+		{ ELogVerbosity::Verbose,     LOCTEXT("VerbosityThresholdVerbose",     "Verbose"),     LOCTEXT("VerbosityThresholdVerbose_Tooltip",     "Show Verbose, Log, Display, Warning, Error and Fatal log messages (i.e. hide VeryVerbose log messages)."), },
+		{ ELogVerbosity::Log,         LOCTEXT("VerbosityThresholdLog",         "Log"),         LOCTEXT("VerbosityThresholdLog_Tooltip",         "Show Log, Display, Warning, Error and Fatal log messages (i.e. hide Verbose and VeryVerbose log messages)."), },
+		{ ELogVerbosity::Display,     LOCTEXT("VerbosityThresholdDisplay",     "Display"),     LOCTEXT("VerbosityThresholdDisplay_Tooltip",     "Show Display, Warning, Error and Fatal log messages (i.e. hide Log, Verbose and VeryVerbose log messages)."), },
+		{ ELogVerbosity::Warning,     LOCTEXT("VerbosityThresholdWarning",     "Warning"),     LOCTEXT("VerbosityThresholdWarning_Tooltip",     "Show only Warning, Error and Fatal log messages."), },
+		{ ELogVerbosity::Error,       LOCTEXT("VerbosityThresholdError",       "Error"),       LOCTEXT("VerbosityThresholdError_Tooltip",       "Show only Error and Fatal log messages."), },
+		{ ELogVerbosity::Fatal,       LOCTEXT("VerbosityThresholdFatal",       "Fatal"),       LOCTEXT("VerbosityThresholdFatal_Tooltip",       "Show only Fatal log messages."), },
+	};
+
+	for (int32 Index = 0; Index < sizeof(VerbosityThresholds) / sizeof(FVerbosityThresholdInfo); ++Index)
+	{
+		const FVerbosityThresholdInfo& Threshold = VerbosityThresholds[Index];
+
+		/*
+		MenuBuilder.AddMenuEntry(
+			Threshold.Label,
+			Threshold.ToolTip,
+			FSlateIcon(),
+			FUIAction(FExecuteAction::CreateSP(this, &SLogView::VerbosityThreshold_Execute, Threshold.Verbosity),
+				FCanExecuteAction::CreateLambda([] { return true; }),
+				FIsActionChecked::CreateSP(this, &SLogView::VerbosityThreshold_IsChecked, Threshold.Verbosity)),
+			NAME_None,
+			EUserInterfaceActionType::ToggleButton
+		);
+		*/
+
+		const TSharedRef<SWidget> TextBlock = SNew(STextBlock)
+			.Text(Threshold.Label)
+			.ShadowColorAndOpacity(FLinearColor(0.05f, 0.05f, 0.05f, 1.0f))
+			.ShadowOffset(FVector2D(1.0f, 1.0f))
+			.ColorAndOpacity(FSlateColor(FTimeMarkerTrackBuilder::GetColorByVerbosity(Threshold.Verbosity)));
+
+		MenuBuilder.AddMenuEntry(
+			FUIAction(FExecuteAction::CreateSP(this, &SLogView::VerbosityThreshold_Execute, Threshold.Verbosity),
+				FCanExecuteAction::CreateLambda([] { return true; }),
+				FIsActionChecked::CreateSP(this, &SLogView::VerbosityThreshold_IsChecked, Threshold.Verbosity)),
+			TextBlock,
+			NAME_None,
+			Threshold.ToolTip,
+			EUserInterfaceActionType::ToggleButton
+		);
+	}
+}
+
+////////////////////////////////////////////////////////////////////////////////////////////////////
+
+TSharedRef<SWidget> SLogView::MakeCategoryFilterMenu()
+{
+	FMenuBuilder MenuBuilder(/*bInShouldCloseWindowAfterMenuSelection=*/true, nullptr);
+
+	MenuBuilder.BeginSection("LogViewShowAllCategories", LOCTEXT("QuickFilterHeading", "Quick Filter"));
+	{
+		MenuBuilder.AddMenuEntry(
+			LOCTEXT("ShowAllCategories", "Show/Hide All"),
+			LOCTEXT("ShowAllCategories_Tooltip", "Change filtering to show/hide all categories"),
+			FSlateIcon(),
+			FUIAction(FExecuteAction::CreateSP(this, &SLogView::ShowHideAllCategories_Execute),
+				FCanExecuteAction::CreateLambda([] { return true; }),
+				FIsActionChecked::CreateSP(this, &SLogView::ShowHideAllCategories_IsChecked)),
+			NAME_None,
+			EUserInterfaceActionType::ToggleButton
+		);
+	}
+	MenuBuilder.EndSection();
+
+	MenuBuilder.BeginSection("LogViewCategoriesEntries", LOCTEXT("CategoriesHeading", "Categories"));
+	CreateCategoriesFilterMenuSection(MenuBuilder);
+	MenuBuilder.EndSection();
+
+	return MenuBuilder.MakeWidget();
+}
+
+////////////////////////////////////////////////////////////////////////////////////////////////////
+
+void SLogView::CreateCategoriesFilterMenuSection(FMenuBuilder& MenuBuilder)
+{
+	for (const FName CategoryName : Filter.GetAvailableLogCategories())
+	{
+		const FString CategoryString = CategoryName.ToString();
+		const FText CategoryText(FText::AsCultureInvariant(CategoryString));
+
+		//MenuBuilder.AddMenuEntry(
+		//	FText::AsCultureInvariant(CategoryString),
+		//	FText::Format(LOCTEXT("Category_Tooltip", "Filter the Log View to show/hide category: {0}"), CategoryText),
+		//	FSlateIcon(),
+		//	FUIAction(FExecuteAction::CreateSP(this, &SLogView::ToggleCategory_Execute, CategoryName),
+		//		FCanExecuteAction::CreateLambda([] { return true; }),
+		//		FIsActionChecked::CreateSP(this, &SLogView::ToggleCategory_IsChecked, CategoryName)),
+		//	NAME_None,
+		//	EUserInterfaceActionType::ToggleButton
+		//);
+
+		const TSharedRef<SWidget> TextBlock = SNew(STextBlock)
+			.Text(FText::AsCultureInvariant(CategoryString))
+			.ShadowColorAndOpacity(FLinearColor(0.05f, 0.05f, 0.05f, 1.0f))
+			.ShadowOffset(FVector2D(1.0f, 1.0f))
+			.ColorAndOpacity(FSlateColor(FTimeMarkerTrackBuilder::GetColorByCategory(*CategoryString)));
+
+		MenuBuilder.AddMenuEntry(
+			FUIAction(FExecuteAction::CreateSP(this, &SLogView::ToggleCategory_Execute, CategoryName),
+				FCanExecuteAction::CreateLambda([] { return true; }),
+				FIsActionChecked::CreateSP(this, &SLogView::ToggleCategory_IsChecked, CategoryName)),
+			TextBlock,
+			NAME_None,
+			FText::Format(LOCTEXT("Category_Tooltip", "Filter the Log View to show/hide category: {0}"), CategoryText),
+			EUserInterfaceActionType::ToggleButton
+		);
+	}
+}
+
+////////////////////////////////////////////////////////////////////////////////////////////////////
+
+bool SLogView::VerbosityThreshold_IsChecked(ELogVerbosity::Type Verbosity) const
+{
+	return Verbosity <= Filter.GetVerbosityThreshold();
+}
+
+////////////////////////////////////////////////////////////////////////////////////////////////////
+
+void SLogView::VerbosityThreshold_Execute(ELogVerbosity::Type Verbosity)
+{
+	Filter.SetVerbosityThreshold(Verbosity);
+	OnFilterChanged();
+}
+
+////////////////////////////////////////////////////////////////////////////////////////////////////
+
+bool SLogView::ShowHideAllCategories_IsChecked() const
+{
+	return Filter.IsShowAllCategoriesEnabled();
+}
+
+////////////////////////////////////////////////////////////////////////////////////////////////////
+
+void SLogView::ShowHideAllCategories_Execute()
+{
+	Filter.ToggleShowAllCategories();
+	OnFilterChanged();
+}
+
+////////////////////////////////////////////////////////////////////////////////////////////////////
+
+void SLogView::ShowAllCategories_Execute()
+{
+	if (Filter.IsShowAllCategoriesEnabled())
+	{
+		Filter.ToggleShowAllCategories(); // hide
+		Filter.ToggleShowAllCategories(); // show
+	}
+	else
+	{
+		Filter.ToggleShowAllCategories(); // show
+	}
+	OnFilterChanged();
+}
+
+////////////////////////////////////////////////////////////////////////////////////////////////////
+
+void SLogView::ShowOnlyCategory_Execute(FName InName)
+{
+	Filter.EnableOnlyCategory(InName);
+	OnFilterChanged();
+}
+
+////////////////////////////////////////////////////////////////////////////////////////////////////
+
+bool SLogView::ToggleCategory_IsChecked(FName InName) const
+{
+	return Filter.IsLogCategoryEnabled(InName);
+}
+
+////////////////////////////////////////////////////////////////////////////////////////////////////
+
+void SLogView::ToggleCategory_Execute(FName InName)
+{
+	Filter.ToggleLogCategory(InName);
+	OnFilterChanged();
+}
+
+////////////////////////////////////////////////////////////////////////////////////////////////////
+
+#undef LOCTEXT_NAMESPACE
