Copying //UE4/Dev-Core to Dev-Main (//UE4/Dev-Main)

diff --git a/Engine/Source/Developer/TraceInsights/Private/Insights/Widgets/SStartPageWindow.cpp b/Engine/Source/Developer/TraceInsights/Private/Insights/Widgets/SStartPageWindow.cpp
new file mode 100644
--- /dev/null
+++ b/Engine/Source/Developer/TraceInsights/Private/Insights/Widgets/SStartPageWindow.cpp
@@ -0,0 +1,1370 @@
+// Copyright 1998-2019 Epic Games, Inc. All Rights Reserved.
+
+#include "SStartPageWindow.h"
+
+#include "DesktopPlatformModule.h"
+#include "EditorStyleSet.h"
+#include "Framework/Application/SlateApplication.h"
+#include "Framework/Docking/WorkspaceItem.h"
+#include "Framework/MultiBox/MultiBoxBuilder.h"
+#include "HAL/FileManagerGeneric.h"
+#include "SlateOptMacros.h"
+#include "Widgets/Docking/SDockTab.h"
+#include "Widgets/Images/SImage.h"
+#include "Widgets/Layout/SBorder.h"
+#include "Widgets/Layout/SBox.h"
+#include "Widgets/Layout/SSpacer.h"
+#include "Widgets/Notifications/SNotificationList.h"
+#include "Widgets/SBoxPanel.h"
+#include "Widgets/Text/STextBlock.h"
+#include "Widgets/Input/SButton.h"
+#include "Widgets/Input/SEditableTextBox.h"
+#include "Widgets/Input/SComboButton.h"
+//#include "WorkspaceMenuStructure.h"
+//#include "WorkspaceMenuStructureModule.h"
+
+#if WITH_EDITOR
+	#include "EngineAnalytics.h"
+	#include "Runtime/Analytics/Analytics/Public/AnalyticsEventAttribute.h"
+	#include "Runtime/Analytics/Analytics/Public/Interfaces/IAnalyticsProvider.h"
+#endif // WITH_EDITOR
+
+// Insights
+#include "Insights/InsightsManager.h"
+#include "Insights/TimingProfilerManager.h"
+#include "Insights/InsightsStyle.h"
+#include "Insights/Version.h"
+#include "Insights/Widgets/SInsightsSettings.h"
+#include "Insights/Widgets/STimingProfilerWindow.h"
+
+////////////////////////////////////////////////////////////////////////////////////////////////////
+
+#define LOCTEXT_NAMESPACE "SStartPageWindow"
+
+////////////////////////////////////////////////////////////////////////////////////////////////////
+
+static FText GetTextForNotification(const EInsightsNotificationType NotificatonType, const ELoadingProgressState ProgressState, const FString& Filename, const float ProgressPercent = 0.0f)
+{
+	FText Result;
+
+	if (NotificatonType == EInsightsNotificationType::LoadingTraceFile)
+	{
+		if (ProgressState == ELoadingProgressState::Started)
+		{
+			FFormatNamedArguments Args;
+			Args.Add(TEXT("Filename"), FText::FromString(Filename));
+			Result = FText::Format(LOCTEXT("DescF_OfflineCapture_Started", "Started loading a file {Filename}"), Args);
+		}
+		else if (ProgressState == ELoadingProgressState::InProgress)
+		{
+			FFormatNamedArguments Args;
+			Args.Add(TEXT("Filename"), FText::FromString(Filename));
+			Args.Add(TEXT("LoadingProgressPercent"), FText::AsPercent(ProgressPercent));
+			Result = FText::Format(LOCTEXT("DescF_OfflineCapture_InProgress", "Loading a file {Filename} {LoadingProgressPercent}"), Args);
+		}
+		else if (ProgressState == ELoadingProgressState::Loaded)
+		{
+			FFormatNamedArguments Args;
+			Args.Add(TEXT("Filename"), FText::FromString(Filename));
+			Result = FText::Format(LOCTEXT("DescF_OfflineCapture_Loaded", "File {Filename} has been successfully loaded"), Args);
+		}
+		else if (ProgressState == ELoadingProgressState::Failed)
+		{
+			FFormatNamedArguments Args;
+			Args.Add(TEXT("Filename"), FText::FromString(Filename));
+			Result = FText::Format(LOCTEXT("DescF_OfflineCapture_Failed", "Failed to load file {Filename}"), Args);
+		}
+	}
+
+	return Result;
+}
+
+////////////////////////////////////////////////////////////////////////////////////////////////////
+// SConnectionRow
+////////////////////////////////////////////////////////////////////////////////////////////////////
+
+class SConnectionRow : public SMultiColumnTableRow<TSharedPtr<FRecorderConnection>>
+{
+	SLATE_BEGIN_ARGS(SConnectionRow) {}
+	SLATE_END_ARGS()
+
+public:
+	/**
+	 * Constructs the widget.
+	 *
+	 * @param InArgs The construction arguments.
+	 * @param InConnection The connection displayed by this row.
+	 * @param InOwnerTableView The table to which the row must be added.
+	 */
+	void Construct(const FArguments& InArgs, TSharedPtr<FRecorderConnection> InConnection, TSharedRef<SStartPageWindow> InParentWidget, const TSharedRef<STableViewBase>& InOwnerTableView)
+	{
+		WeakConnection = MoveTemp(InConnection);
+		WeakParentWidget = InParentWidget;
+
+		SMultiColumnTableRow<TSharedPtr<FRecorderConnection>>::Construct(FSuperRowType::FArguments(), InOwnerTableView);
+	}
+
+public:
+	virtual TSharedRef<SWidget> GenerateWidgetForColumn(const FName& ColumnName) override
+	{
+		if (ColumnName == FName(TEXT("Name")))
+		{
+			return SNew(SBox)
+				.Padding(FMargin(4.0, 0.0))
+				[
+					SNew(STextBlock)
+					.Text(this, &SConnectionRow::GetConnectionName)
+					.ToolTipText(this, &SConnectionRow::GetConnectionName)
+				];
+		}
+		else if (ColumnName == FName(TEXT("Uri")))
+		{
+			return SNew(SBox)
+				.Padding(FMargin(4.0, 0.0))
+				[
+					SNew(STextBlock)
+					.Text(this, &SConnectionRow::GetConnectionUri)
+					.ToolTipText(this, &SConnectionRow::GetConnectionUri)
+				];
+		}
+		else
+		{
+			return SNew(STextBlock).Text(LOCTEXT("UnknownColumn", "Unknown Column"));
+		}
+	}
+
+	FText GetConnectionName() const
+	{
+		TSharedPtr<FRecorderConnection> ConnectionPin = WeakConnection.Pin();
+		if (ConnectionPin.IsValid())
+		{
+			return ConnectionPin->Name;
+		}
+		else
+		{
+			return FText();
+		}
+	}
+
+	FText GetConnectionUri() const
+	{
+		TSharedPtr<FRecorderConnection> ConnectionPin = WeakConnection.Pin();
+		if (ConnectionPin.IsValid())
+		{
+			return ConnectionPin->Uri;
+		}
+		else
+		{
+			return FText();
+		}
+	}
+
+	FText GetRowToolTip() const
+	{
+		TSharedPtr<FRecorderConnection> ConnectionPin = WeakConnection.Pin();
+		if (ConnectionPin.IsValid())
+		{
+			return ConnectionPin->Uri;
+		}
+		else
+		{
+			return FText();
+		}
+	}
+
+private:
+	TWeakPtr<FRecorderConnection> WeakConnection;
+	TWeakPtr<SStartPageWindow> WeakParentWidget;
+};
+
+////////////////////////////////////////////////////////////////////////////////////////////////////
+// SStartPageWindow
+////////////////////////////////////////////////////////////////////////////////////////////////////
+
+SStartPageWindow::SStartPageWindow()
+	: DurationActive(0.0f)
+	, bIsAnyLiveSessionAvailable(false)
+	, LastLiveSessionHandle(static_cast<Trace::FSessionHandle>(0))
+	, bIsAnySessionAvailable(false)
+	, LastSessionHandle(static_cast<Trace::FSessionHandle>(0))
+{
+}
+
+////////////////////////////////////////////////////////////////////////////////////////////////////
+
+SStartPageWindow::~SStartPageWindow()
+{
+#if WITH_EDITOR
+	if (DurationActive > 0.0f && FEngineAnalytics::IsAvailable())
+	{
+		FEngineAnalytics::GetProvider().RecordEvent(TEXT("Editor.Usage.Profiler"), FAnalyticsEventAttribute(TEXT("Duration"), DurationActive));
+	}
+#endif // WITH_EDITOR
+}
+
+////////////////////////////////////////////////////////////////////////////////////////////////////
+
+BEGIN_SLATE_FUNCTION_BUILD_OPTIMIZATION
+
+void SStartPageWindow::Construct(const FArguments& InArgs)
+{
+	ChildSlot
+		[
+			SNew(SOverlay)
+
+			// Version
+			+ SOverlay::Slot()
+				.HAlign(HAlign_Right)
+				.VAlign(VAlign_Top)
+				.Padding(0.0f, -16.0f, 0.0f, 0.0f)
+				[
+					SNew(STextBlock)
+						.Clipping(EWidgetClipping::ClipToBoundsWithoutIntersecting)
+						.Text(LOCTEXT("UnrealInsightsVersion", UNREAL_INSIGHTS_VERSION_STRING_EX))
+						.ColorAndOpacity(FLinearColor(0.15f, 0.15f, 0.15f, 1.0f))
+				]
+
+			// Overlay slot for the main window area
+			+ SOverlay::Slot()
+				.HAlign(HAlign_Fill)
+				.VAlign(VAlign_Center)
+				[
+					SNew(SVerticalBox)
+
+					+ SVerticalBox::Slot()
+						.AutoHeight()
+						.HAlign(HAlign_Center)
+						.Padding(3.0f, 3.0f)
+						[
+							SNew(SBorder)
+							.Visibility(this, &SStartPageWindow::IsSessionOverlayVisible)
+							.BorderImage(FEditorStyle::GetBrush("NotificationList.ItemBackground"))
+							.Padding(8.0f)
+							[
+								SNew(STextBlock)
+								.Text(LOCTEXT("SelectTraceOverlayText", "Please select a trace..."))
+							]
+						]
+
+					+ SVerticalBox::Slot()
+						.AutoHeight()
+						.HAlign(HAlign_Center)
+						.Padding(3.0f, 3.0f)
+						[
+							SNew(SBorder)
+							.BorderImage(FEditorStyle::GetBrush("NotificationList.ItemBackground"))
+							.Padding(8.0f)
+							[
+								SNew(STextBlock)
+								.Text(LOCTEXT("SelectTraceHint", "\
+- Use \"Live\" button to start analysis for the latest available session that is live, if any.\n\
+- Use \"Latest\" button to start analysis for the latest available session, if any.\n\
+- Use \"Load\" button to load a trace file (*.utrace).\n\
+- Use the combo box to choose from available sessions.\n\
+- Drag and drop a *.utrace file over this window."))
+								//.Justification(ETextJustify::Center)
+							]
+						]
+
+					+ SVerticalBox::Slot()
+						.AutoHeight()
+						.HAlign(HAlign_Center)
+						.Padding(3.0f, 3.0f)
+						[
+							SNew(SHorizontalBox)
+
+							+ SHorizontalBox::Slot()
+								.AutoWidth()
+								.Padding(3.0f, 0.0f)
+								[
+									SNew(SButton)
+									.OnClicked(this, &SStartPageWindow::Live_OnClicked)
+									.IsEnabled(this, &SStartPageWindow::Live_IsEnabled)
+									.ToolTipText(LOCTEXT("LiveButtonTooltip", "Start analysis for the latest available session that is live, if any."))
+									.ContentPadding(8.0f)
+									.Content()
+									[
+										SNew(SHorizontalBox)
+
+										+SHorizontalBox::Slot()
+											.AutoWidth()
+											[
+												SNew(SImage)
+												.Image(FInsightsStyle::GetBrush("Live.Icon.Large"))
+											]
+
+										+SHorizontalBox::Slot()
+											.AutoWidth()
+											.VAlign(VAlign_Center)
+											[
+												SNew(SBox)
+												.WidthOverride(100.0f)
+												[
+													SNew(STextBlock)
+													.Font(FCoreStyle::GetDefaultFontStyle("Regular", 18))
+													.Text(LOCTEXT("LiveButtonText", "Live"))
+												]
+											]
+									]
+								]
+
+							+ SHorizontalBox::Slot()
+								.AutoWidth()
+								.Padding(3.0f, 0.0f)
+								[
+									SNew(SButton)
+									.OnClicked(this, &SStartPageWindow::Last_OnClicked)
+									.IsEnabled(this, &SStartPageWindow::Last_IsEnabled)
+									.ToolTipText(LOCTEXT("LastButtonTooltip", "Start analysis for the latest available session, if any."))
+									.ContentPadding(8.0f)
+									.Content()
+									[
+										SNew(SHorizontalBox)
+
+										+SHorizontalBox::Slot()
+											.AutoWidth()
+											[
+												SNew(SImage)
+												.Image(FInsightsStyle::GetBrush("Live.Icon.Large"))
+											]
+
+										+SHorizontalBox::Slot()
+											.AutoWidth()
+											.VAlign(VAlign_Center)
+											[
+												SNew(SBox)
+												.WidthOverride(100.0f)
+												[
+													SNew(STextBlock)
+													.Font(FCoreStyle::GetDefaultFontStyle("Regular", 18))
+													.Text(LOCTEXT("LastButtonText", "Latest"))
+												]
+											]
+									]
+								]
+
+							+ SHorizontalBox::Slot()
+								.AutoWidth()
+								.Padding(3.0f, 0.0f, 0.0f, 0.0f)
+								[
+									SNew(SButton)
+									.OnClicked(this, &SStartPageWindow::Load_OnClicked)
+									.ToolTipText(LOCTEXT("LoadButtonTooltip", "Load a trace file."))
+									.ContentPadding(8.0f)
+									.Content()
+									[
+										SNew(SHorizontalBox)
+
+										+ SHorizontalBox::Slot()
+											.AutoWidth()
+											[
+												SNew(SImage)
+												.Image(FInsightsStyle::GetBrush("Load.Icon.Large"))
+											]
+
+										+ SHorizontalBox::Slot()
+											.AutoWidth()
+											.VAlign(VAlign_Center)
+											[
+												SNew(SBox)
+												.WidthOverride(100.0f)
+												[
+													SNew(STextBlock)
+													.Font(FCoreStyle::GetDefaultFontStyle("Regular", 18))
+													.Text(LOCTEXT("LoadButtonText", "Load..."))
+												]
+											]
+									]
+								]
+
+							+ SHorizontalBox::Slot()
+								.AutoWidth()
+								.Padding(0.0f, 0.0f, 0.0f, 0.0f)
+								[
+									SNew(SComboButton)
+									.ToolTipText(LOCTEXT("SessionList_Tooltip", "Choose from Most Recently Used Sessions or from Available Sessions"))
+									.OnGetMenuContent(this, &SStartPageWindow::MakeSessionListMenu)
+									.HasDownArrow(true)
+									.ContentPadding(FMargin(1.0f, 1.0f, 1.0f, 1.0f))
+								]
+						]
+
+					+ SVerticalBox::Slot()
+						.AutoHeight()
+						.HAlign(HAlign_Center)
+						.Padding(3.0f, 3.0f)
+						[
+							SNew(SBorder)
+							.BorderImage(FEditorStyle::GetBrush("NotificationList.ItemBackground"))
+							.Padding(8.0f)
+							[
+								SNew(SVerticalBox)
+
+								+ SVerticalBox::Slot()
+									.AutoHeight()
+									.HAlign(HAlign_Center)
+									.Padding(0.0, 2.0f)
+									[
+										SNew(STextBlock)
+										.Text(LOCTEXT("LocalSessionDirectoryText", "Local Session Directory:"))
+										.ColorAndOpacity(FLinearColor::Gray)
+									]
+
+								+ SVerticalBox::Slot()
+									.AutoHeight()
+									.HAlign(HAlign_Center)
+									.Padding(0.0, 2.0f)
+									[
+										SNew(SHorizontalBox)
+
+										+ SHorizontalBox::Slot()
+											.AutoWidth()
+											.Padding(2.0f, 0.0f)
+											.VAlign(VAlign_Center)
+											[
+												SNew(STextBlock)
+												.Text(this, &SStartPageWindow::GetLocalSessionDirectory)
+											]
+
+										+ SHorizontalBox::Slot()
+											.AutoWidth()
+											.Padding(2.0f, 0.0f)
+											.VAlign(VAlign_Center)
+											[
+												SNew(SButton)
+												.Text(LOCTEXT("ExploreLocalSessionDirButton", "..."))
+												.ToolTipText(LOCTEXT("ExploreLocalSessionDirButtonToolTip", "Explore the Local Session Directory"))
+												.OnClicked(this, &SStartPageWindow::ExploreLocalSessionDirectory_OnClicked)
+											]
+									]
+							]
+						]
+
+					+ SVerticalBox::Slot()
+						.AutoHeight()
+						.HAlign(HAlign_Center)
+						.Padding(3.0f, 3.0f)
+						[
+							SNew(SBorder)
+							.BorderImage(FEditorStyle::GetBrush("NotificationList.ItemBackground"))
+							.Padding(8.0f)
+							[
+								ConstructRecoderPanel()
+							]
+						]
+				]
+
+			// Notification area overlay
+			+ SOverlay::Slot()
+				.HAlign(HAlign_Right)
+				.VAlign(VAlign_Bottom)
+				.Padding(16.0f)
+				[
+					SAssignNew(NotificationList, SNotificationList)
+				]
+
+			// Settings dialog overlay
+			+ SOverlay::Slot()
+				.HAlign(HAlign_Center)
+				.VAlign(VAlign_Center)
+				.Expose(OverlaySettingsSlot)
+		];
+}
+
+////////////////////////////////////////////////////////////////////////////////////////////////////
+
+TSharedRef<SWidget> SStartPageWindow::ConstructRecoderPanel()
+{
+	TSharedRef<SWidget> Widget = SNew(SVerticalBox)
+
+	+ SVerticalBox::Slot()
+		.AutoHeight()
+		.HAlign(HAlign_Center)
+		.Padding(0.0, 2.0f)
+		[
+			SNew(STextBlock)
+			.Text(LOCTEXT("RecorderPanelTitle", "Trace Recorder"))
+			.Font(FCoreStyle::GetDefaultFontStyle("Bold", 11))
+			.ColorAndOpacity(FLinearColor::Gray)
+		]
+
+	+ SVerticalBox::Slot()
+		.AutoHeight()
+		.HAlign(HAlign_Left)
+		.Padding(0.0, 2.0f)
+		[
+			SNew(SHorizontalBox)
+
+			+ SHorizontalBox::Slot()
+				.AutoWidth()
+				.Padding(0.0f, 0.0f, 2.0f, 0.0f)
+				.VAlign(VAlign_Center)
+				[
+					SNew(STextBlock)
+					.Text(LOCTEXT("RecorderStatusTitle", "Status:"))
+					.ColorAndOpacity(FLinearColor::Gray)
+				]
+
+			+ SHorizontalBox::Slot()
+				.AutoWidth()
+				.Padding(2.0f, 0.0f)
+				.VAlign(VAlign_Center)
+				[
+					SNew(STextBlock)
+					.Text(this, &SStartPageWindow::GetRecorderStatusText)
+				]
+
+			+ SHorizontalBox::Slot()
+				.AutoWidth()
+				.Padding(2.0f, 0.0f)
+				.VAlign(VAlign_Center)
+				[
+					SNew(SButton)
+					.Text(LOCTEXT("StartRecorder", "Start"))
+					.ToolTipText(LOCTEXT("StartRecorderToolTip", "Start the Trace Recorder"))
+					.OnClicked(this, &SStartPageWindow::StartTraceRecorder_OnClicked)
+					.Visibility(this, &SStartPageWindow::StartTraceRecorder_Visibility)
+				]
+
+			+ SHorizontalBox::Slot()
+				.AutoWidth()
+				.Padding(2.0f, 0.0f)
+				.VAlign(VAlign_Center)
+				[
+					SNew(SButton)
+					.Text(LOCTEXT("StopRecorder", "Stop"))
+					.ToolTipText(LOCTEXT("StopRecorderToolTip", "Stop the Trace Recorder"))
+					.OnClicked(this, &SStartPageWindow::StopTraceRecorder_OnClicked)
+					.Visibility(this, &SStartPageWindow::StopTraceRecorder_Visibility)
+				]
+		]
+
+	+ SVerticalBox::Slot()
+		.AutoHeight()
+		.HAlign(HAlign_Fill)
+		.Padding(0.0, 2.0f)
+		[
+			SNew(SHorizontalBox)
+			.Visibility(this, &SStartPageWindow::StopTraceRecorder_Visibility)
+
+			+ SHorizontalBox::Slot()
+				.AutoWidth()
+				.Padding(0.0f, 0.0f, 2.0f, 0.0f)
+				.VAlign(VAlign_Center)
+				[
+					SNew(STextBlock)
+					.Text(LOCTEXT("HostTitle", "Host:"))
+					.ColorAndOpacity(FLinearColor::Gray)
+				]
+
+			+ SHorizontalBox::Slot()
+				.FillWidth(1.0)
+				.VAlign(VAlign_Center)
+				[
+					SAssignNew(HostTextBox, SEditableTextBox)
+				]
+
+			+ SHorizontalBox::Slot()
+				.AutoWidth()
+				.Padding(2.0f, 0.0f)
+				.VAlign(VAlign_Center)
+				[
+					SNew(SButton)
+					.Text(LOCTEXT("Connect", "Connect"))
+					.ToolTipText(LOCTEXT("ConnectToolTip", "Try connecting to host."))
+					.OnClicked(this, &SStartPageWindow::Connect_OnClicked)
+				]
+		]
+
+	+ SVerticalBox::Slot()
+		.AutoHeight()
+		.HAlign(HAlign_Left)
+		.Padding(0.0, 2.0f, 0.0f, 1.0f)
+		[
+			SNew(SHorizontalBox)
+			.Visibility(this, &SStartPageWindow::StopTraceRecorder_Visibility)
+
+			+ SHorizontalBox::Slot()
+				.AutoWidth()
+				.Padding(0.0f, 0.0f, 2.0f, 0.0f)
+				.VAlign(VAlign_Center)
+				[
+					SNew(STextBlock)
+					.Text_Lambda([this]() -> FText
+					{
+						return FText::Format(LOCTEXT("ConnectionListTitle", "Connections (live sessions): {0}"), FText::AsNumber(Connections.Num()));
+					})
+					.ColorAndOpacity(FLinearColor::Gray)
+				]
+
+			+ SHorizontalBox::Slot()
+				.AutoWidth()
+				.Padding(2.0f, 0.0f)
+				.VAlign(VAlign_Center)
+				[
+					SNew(SButton)
+					.Text(LOCTEXT("Refresh", "Refresh"))
+					.ToolTipText(LOCTEXT("RefreshToolTip", "Refresh the connection list (live sessions)."))
+					.OnClicked(this, &SStartPageWindow::RefreshConnections_OnClicked)
+				]
+		]
+
+	+ SVerticalBox::Slot()
+		.AutoHeight()
+		.HAlign(HAlign_Left)
+		.Padding(0.0, 1.0f, 0.0f, 2.0f)
+		[
+			SNew(SBox)
+			.WidthOverride(640.0f)
+			.MaxDesiredHeight(78.0f)
+			.Visibility(this, &SStartPageWindow::Connections_Visibility)
+			[
+				SAssignNew(ConnectionsListView, SListView<TSharedPtr<FRecorderConnection>>)
+				.ItemHeight(20.0f)
+				.SelectionMode(ESelectionMode::Single)
+				.OnSelectionChanged(this, &SStartPageWindow::Connections_OnSelectionChanged)
+				.ListItemsSource(&Connections)
+				.OnGenerateRow(this, &SStartPageWindow::Connections_OnGenerateRow)
+				.ConsumeMouseWheel(EConsumeMouseWheel::Always)
+				//.OnContextMenuOpening(FOnContextMenuOpening::CreateSP(this, &SStartPageWindow::Connections_GetContextMenu))
+				.HeaderRow
+				(
+					SNew(SHeaderRow)
+
+					+ SHeaderRow::Column(FName(TEXT("Name")))
+					.FillWidth(0.2f)
+					.DefaultLabel(LOCTEXT("NameColumn", "Name"))
+
+					+ SHeaderRow::Column(FName(TEXT("Uri")))
+					.FillWidth(0.8f)
+					.DefaultLabel(LOCTEXT("UriColumn", "URI"))
+				)
+			]
+		]
+
+	+ SVerticalBox::Slot()
+		.AutoHeight()
+		.HAlign(HAlign_Left)
+		.Padding(0.0, 2.0f, 0.0f, 1.0f)
+		[
+			SNew(STextBlock)
+			.Visibility(this, &SStartPageWindow::Modules_Visibility)
+			.Text(LOCTEXT("ModulesListTitle", "Modules for selected connection:"))
+			.ColorAndOpacity(FLinearColor::Gray)
+		]
+
+	+ SVerticalBox::Slot()
+		.AutoHeight()
+		.HAlign(HAlign_Left)
+		.Padding(0.0, 1.0f, 0.0f, 2.0f)
+		[
+			ConstructModuleList()
+		]
+	;
+
+	RefreshConnectionList();
+
+	return Widget;
+}
+
+////////////////////////////////////////////////////////////////////////////////////////////////////
+
+TSharedRef<SWidget> SStartPageWindow::ConstructModuleList()
+{
+	TSharedRef<SVerticalBox> VerticalBox = SNew(SVerticalBox)
+		.Visibility(this, &SStartPageWindow::Modules_Visibility);
+
+	TSharedRef<Trace::IModuleService> ModuleService = FInsightsManager::Get()->GetModuleService();
+	ModuleService->GetAvailableModules(AvailableModules);
+
+	constexpr bool bDefaultModuleEnableState = false;
+	AvailableModulesEnabledState.Reset();
+
+	for (int32 ModuleIndex = 0; ModuleIndex < AvailableModules.Num(); ++ModuleIndex)
+	{
+		const Trace::FModuleInfo& Module = AvailableModules[ModuleIndex];
+
+		AvailableModulesEnabledState.Add(bDefaultModuleEnableState);
+		ModuleService->SetModuleEnabled(Module.Name, bDefaultModuleEnableState);
+
+		VerticalBox->AddSlot()
+			.AutoHeight()
+			.Padding(0.0f, 1.0f)
+			[
+				SNew(SCheckBox)
+				.IsChecked(this, &SStartPageWindow::Module_IsChecked, ModuleIndex)
+				.OnCheckStateChanged(this, &SStartPageWindow::Module_OnCheckStateChanged, ModuleIndex)
+				[
+					SNew(STextBlock)
+					.Text(FText::FromString(Module.DisplayName))
+				]
+			];
+	}
+
+	return VerticalBox;
+}
+
+////////////////////////////////////////////////////////////////////////////////////////////////////
+
+TSharedRef<ITableRow> SStartPageWindow::Connections_OnGenerateRow(TSharedPtr<FRecorderConnection> InConnection, const TSharedRef<STableViewBase>& OwnerTable)
+{
+	return SNew(SConnectionRow, InConnection, SharedThis(this), OwnerTable);
+}
+
+END_SLATE_FUNCTION_BUILD_OPTIMIZATION
+
+////////////////////////////////////////////////////////////////////////////////////////////////////
+
+EVisibility SStartPageWindow::Modules_Visibility() const
+{
+	TSharedRef<Trace::ISessionService> SessionService = FInsightsManager::Get()->GetSessionService();
+	return (SessionService->IsRecorderServerRunning() && Connections.Num() > 0 && SelectedConnection.IsValid()) ? EVisibility::Visible : EVisibility::Collapsed;
+}
+
+////////////////////////////////////////////////////////////////////////////////////////////////////
+
+ECheckBoxState SStartPageWindow::Module_IsChecked(int32 ModuleIndex) const
+{
+	if (ModuleIndex >= 0 && ModuleIndex < AvailableModulesEnabledState.Num())
+	{
+		const bool bIsEnabled = AvailableModulesEnabledState[ModuleIndex];
+		return bIsEnabled ? ECheckBoxState::Checked : ECheckBoxState::Unchecked;
+	}
+	else
+	{
+		return ECheckBoxState::Undetermined;
+	}
+}
+
+////////////////////////////////////////////////////////////////////////////////////////////////////
+
+void SStartPageWindow::Module_OnCheckStateChanged(ECheckBoxState NewRadioState, int32 ModuleIndex)
+{
+	ensure(AvailableModulesEnabledState.Num() == AvailableModules.Num());
+
+	if (SelectedConnection.IsValid() &&
+		ModuleIndex >= 0 && ModuleIndex < AvailableModules.Num())
+	{
+		bool bIsEnabled = NewRadioState == ECheckBoxState::Checked;
+
+		TSharedRef<Trace::ISessionService> SessionService = FInsightsManager::Get()->GetSessionService();
+		
+		SessionService->SetModuleEnabled(SelectedConnection->SessionHandle, AvailableModules[ModuleIndex].Name, bIsEnabled);
+		//AvailableModulesEnabledState[ModuleIndex] = bIsEnabled;
+		AvailableModulesEnabledState[ModuleIndex] = SessionService->IsModuleEnabled(SelectedConnection->SessionHandle, AvailableModules[ModuleIndex].Name);
+	}
+}
+
+////////////////////////////////////////////////////////////////////////////////////////////////////
+
+FReply SStartPageWindow::RefreshConnections_OnClicked()
+{
+	RefreshConnectionList();
+	return FReply::Handled();
+}
+
+////////////////////////////////////////////////////////////////////////////////////////////////////
+
+FReply SStartPageWindow::Connect_OnClicked()
+{
+	TSharedRef<Trace::ISessionService> SessionService = FInsightsManager::Get()->GetSessionService();
+
+	FText HostText = HostTextBox->GetText();
+	if (HostText.IsEmptyOrWhitespace())
+	{
+		//...
+	}
+	else if (SessionService->ConnectSession(*HostText.ToString()))
+	{
+		FNotificationInfo NotificationInfo(FText::Format(LOCTEXT("ConnectSuccess", "Successfully connected to \"{0}\"!"), HostText));
+		NotificationInfo.bFireAndForget = false;
+		NotificationInfo.bUseLargeFont = false;
+		NotificationInfo.bUseSuccessFailIcons = true;
+		NotificationInfo.ExpireDuration = 10.0f;
+		SNotificationItemWeak NotificationItem = NotificationList->AddNotification(NotificationInfo);
+		NotificationItem.Pin()->SetCompletionState(SNotificationItem::CS_Success);
+		NotificationItem.Pin()->ExpireAndFadeout();
+		ActiveNotifications.Add(TEXT("ConnectSuccess"), NotificationItem);
+	}
+	else
+	{
+		FNotificationInfo NotificationInfo(FText::Format(LOCTEXT("ConnectFailed", "Failed to connect to \"{0}\"!"), HostText));
+		NotificationInfo.bFireAndForget = false;
+		NotificationInfo.bUseLargeFont = false;
+		NotificationInfo.bUseSuccessFailIcons = true;
+		NotificationInfo.ExpireDuration = 10.0f;
+		SNotificationItemWeak NotificationItem = NotificationList->AddNotification(NotificationInfo);
+		NotificationItem.Pin()->SetCompletionState(SNotificationItem::CS_Fail);
+		NotificationItem.Pin()->ExpireAndFadeout();
+		ActiveNotifications.Add(TEXT("ConnectFailed"), NotificationItem);
+	}
+
+	RefreshConnectionList();
+	return FReply::Handled();
+}
+
+////////////////////////////////////////////////////////////////////////////////////////////////////
+
+void SStartPageWindow::RefreshConnectionList()
+{
+	TSharedRef<Trace::ISessionService> SessionService = FInsightsManager::Get()->GetSessionService();
+
+	TArray<Trace::FSessionHandle> LiveSessions;
+	SessionService->GetLiveSessions(LiveSessions);
+
+	TSharedPtr<FRecorderConnection> NewSelectedConnection;
+
+	Connections.Reset();
+
+	for (Trace::FSessionHandle SessionHandle : LiveSessions)
+	{
+		Trace::FSessionInfo SessionInfo;
+		SessionService->GetSessionInfo(SessionHandle, SessionInfo);
+
+		Connections.Add(MakeShareable(new FRecorderConnection(SessionHandle, SessionInfo)));
+
+		if (SelectedConnection && SelectedConnection->SessionHandle == SessionHandle)
+		{
+			NewSelectedConnection = Connections.Last();
+		}
+	}
+
+	ConnectionsListView->RebuildList();
+
+	if (NewSelectedConnection.IsValid())
+	{
+		ConnectionsListView->SetItemSelection(NewSelectedConnection, true);
+		ConnectionsListView->RequestScrollIntoView(NewSelectedConnection);
+	}
+}
+
+////////////////////////////////////////////////////////////////////////////////////////////////////
+
+void SStartPageWindow::Connections_OnSelectionChanged(TSharedPtr<FRecorderConnection> Connection, ESelectInfo::Type SelectInfo)
+{
+	SelectedConnection = Connection;
+
+	if (Connection.IsValid())
+	{
+		TSharedRef<Trace::ISessionService> SessionService = FInsightsManager::Get()->GetSessionService();
+		for (int32 Index = 0; Index < AvailableModulesEnabledState.Num(); ++Index)
+		{
+			AvailableModulesEnabledState[Index] = SessionService->IsModuleEnabled(SelectedConnection->SessionHandle, AvailableModules[Index].Name);
+		}
+	}
+	else
+	{
+		for (int32 Index = 0; Index < AvailableModulesEnabledState.Num(); ++Index)
+		{
+			AvailableModulesEnabledState[Index] = false;
+		}
+	}
+}
+
+////////////////////////////////////////////////////////////////////////////////////////////////////
+
+EVisibility SStartPageWindow::Connections_Visibility() const
+{
+	TSharedRef<Trace::ISessionService> SessionService = FInsightsManager::Get()->GetSessionService();
+	return (SessionService->IsRecorderServerRunning() && Connections.Num() > 0) ? EVisibility::Visible : EVisibility::Collapsed;
+}
+
+////////////////////////////////////////////////////////////////////////////////////////////////////
+
+void SStartPageWindow::Tick(const FGeometry& AllottedGeometry, const double InCurrentTime, const float InDeltaTime)
+{
+	// We need to check if there is any available session, but not too often.
+	static uint64 NextTimestamp = 0;
+	uint64 Time = FPlatformTime::Cycles64();
+	if (Time > NextTimestamp)
+	{
+		const uint64 WaitTime = static_cast<uint64>(0.5 / FPlatformTime::GetSecondsPerCycle64()); // 500ms
+		NextTimestamp = Time + WaitTime;
+
+		bIsAnySessionAvailable = FInsightsManager::Get()->IsAnySessionAvailable(LastSessionHandle);
+		//bIsAnyLiveSessionAvailable = FInsightsManager::Get()->IsAnyLiveSessionAvailable(LastLiveSessionHandle);
+
+		TSharedRef<Trace::ISessionService> SessionService = FInsightsManager::Get()->GetSessionService();
+		TArray<Trace::FSessionHandle> LiveSessions;
+		SessionService->GetLiveSessions(LiveSessions);
+
+		if (LiveSessions.Num() > 0)
+		{
+			LastLiveSessionHandle = LiveSessions.Last();
+			bIsAnyLiveSessionAvailable = true;
+		}
+		else
+		{
+			bIsAnyLiveSessionAvailable = false;
+		}
+
+		if (LiveSessions.Num() != Connections.Num())
+		{
+			RefreshConnectionList();
+		}
+	}
+}
+
+////////////////////////////////////////////////////////////////////////////////////////////////////
+
+EVisibility SStartPageWindow::IsSessionOverlayVisible() const
+{
+	if (FInsightsManager::Get()->GetSession().IsValid())
+	{
+		return EVisibility::Hidden;
+	}
+
+	return EVisibility::Visible;
+}
+
+////////////////////////////////////////////////////////////////////////////////////////////////////
+
+bool SStartPageWindow::IsSessionValid() const
+{
+	const bool bIsActive = FInsightsManager::Get()->GetSession().IsValid();
+	return bIsActive;
+}
+
+////////////////////////////////////////////////////////////////////////////////////////////////////
+
+void SStartPageWindow::ManageLoadingProgressNotificationState(const FString& Filename, const EInsightsNotificationType NotificatonType, const ELoadingProgressState ProgressState, const float LoadingProgress)
+{
+	const FString BaseFilename = FPaths::GetBaseFilename(Filename);
+
+	if (ProgressState == ELoadingProgressState::Started)
+	{
+		const bool bContains = ActiveNotifications.Contains(Filename);
+		if (!bContains)
+		{
+			FNotificationInfo NotificationInfo(GetTextForNotification(NotificatonType, ProgressState, BaseFilename));
+			NotificationInfo.bFireAndForget = false;
+			NotificationInfo.bUseLargeFont = false;
+
+			// Add two buttons, one for cancel, one for loading the received file.
+			if (NotificatonType == EInsightsNotificationType::LoadingTraceFile)
+			{
+				const FText CancelButtonText = LOCTEXT("CancelButton_Text", "Cancel");
+				const FText CancelButtonTT = LOCTEXT("CancelButton_TTText", "Hides this notification");
+				const FText LoadButtonText = LOCTEXT("LoadButton_Text", "Load file");
+				const FText LoadButtonTT = LOCTEXT("LoadButton_TTText", "Loads the received file and hides this notification");
+
+				NotificationInfo.ButtonDetails.Add(FNotificationButtonInfo(CancelButtonText, CancelButtonTT,
+					FSimpleDelegate::CreateSP(this, &SStartPageWindow::SendingServiceSideCapture_Cancel, Filename), SNotificationItem::CS_Success));
+				NotificationInfo.ButtonDetails.Add(FNotificationButtonInfo(LoadButtonText, LoadButtonTT,
+					FSimpleDelegate::CreateSP(this, &SStartPageWindow::SendingServiceSideCapture_Load, Filename), SNotificationItem::CS_Success));
+			}
+
+			SNotificationItemWeak NotificationItem = NotificationList->AddNotification(NotificationInfo);
+			NotificationItem.Pin()->SetCompletionState(SNotificationItem::CS_Pending);
+			ActiveNotifications.Add(Filename, NotificationItem);
+		}
+	}
+	else if (ProgressState == ELoadingProgressState::InProgress)
+	{
+		const SNotificationItemWeak* LoadingProgressPtr = ActiveNotifications.Find(Filename);
+		if (LoadingProgressPtr)
+		{
+			TSharedPtr<SNotificationItem> LoadingProcessPinned = LoadingProgressPtr->Pin();
+			LoadingProcessPinned->SetText(GetTextForNotification(NotificatonType, ProgressState, BaseFilename, LoadingProgress));
+			LoadingProcessPinned->SetCompletionState(SNotificationItem::CS_Pending);
+		}
+	}
+	else if (ProgressState == ELoadingProgressState::Loaded)
+	{
+		const SNotificationItemWeak* LoadingProgressPtr = ActiveNotifications.Find(Filename);
+		if (LoadingProgressPtr)
+		{
+			TSharedPtr<SNotificationItem> LoadingProcessPinned = LoadingProgressPtr->Pin();
+			LoadingProcessPinned->SetText(GetTextForNotification(NotificatonType, ProgressState, BaseFilename));
+			LoadingProcessPinned->SetCompletionState(SNotificationItem::CS_Success);
+
+			// Notifications for received files are handled by the user.
+			if (NotificatonType == EInsightsNotificationType::LoadingTraceFile)
+			{
+				LoadingProcessPinned->ExpireAndFadeout();
+				ActiveNotifications.Remove(Filename);
+			}
+		}
+	}
+	else if (ProgressState == ELoadingProgressState::Failed)
+	{
+		const SNotificationItemWeak* LoadingProgressPtr = ActiveNotifications.Find(Filename);
+		if (LoadingProgressPtr)
+		{
+			TSharedPtr<SNotificationItem> LoadingProcessPinned = LoadingProgressPtr->Pin();
+			LoadingProcessPinned->SetText(GetTextForNotification(NotificatonType, ProgressState, BaseFilename));
+			LoadingProcessPinned->SetCompletionState(SNotificationItem::CS_Fail);
+
+			LoadingProcessPinned->ExpireAndFadeout();
+			ActiveNotifications.Remove(Filename);
+		}
+	}
+	else if (ProgressState == ELoadingProgressState::Cancelled)
+	{
+		const SNotificationItemWeak* LoadingProgressPtr = ActiveNotifications.Find(Filename);
+		if (LoadingProgressPtr)
+		{
+			TSharedPtr<SNotificationItem> LoadingProcessPinned = LoadingProgressPtr->Pin();
+			LoadingProcessPinned->ExpireAndFadeout();
+			ActiveNotifications.Remove(Filename);
+		}
+	}
+}
+
+////////////////////////////////////////////////////////////////////////////////////////////////////
+
+void SStartPageWindow::SendingServiceSideCapture_Cancel(const FString Filename)
+{
+	const SNotificationItemWeak* LoadingProgressPtr = ActiveNotifications.Find(Filename);
+	if (LoadingProgressPtr)
+	{
+		TSharedPtr<SNotificationItem> LoadingProcessPinned = LoadingProgressPtr->Pin();
+		LoadingProcessPinned->ExpireAndFadeout();
+		ActiveNotifications.Remove(Filename);
+	}
+}
+
+////////////////////////////////////////////////////////////////////////////////////////////////////
+
+void SStartPageWindow::SendingServiceSideCapture_Load(const FString Filename)
+{
+	const SNotificationItemWeak* LoadingProgressPtr = ActiveNotifications.Find(Filename);
+	if (LoadingProgressPtr)
+	{
+		TSharedPtr<SNotificationItem> LoadingProcessPinned = LoadingProgressPtr->Pin();
+		LoadingProcessPinned->ExpireAndFadeout();
+		ActiveNotifications.Remove(Filename);
+
+		const FString PathName = FPaths::ProfilingDir() + TEXT("UnrealStats/Received/");
+		const FString TraceFilepath = PathName + Filename;
+		FInsightsManager::Get()->LoadTraceFile(TraceFilepath);
+	}
+}
+
+////////////////////////////////////////////////////////////////////////////////////////////////////
+
+EActiveTimerReturnType SStartPageWindow::UpdateActiveDuration(double InCurrentTime, float InDeltaTime)
+{
+	DurationActive += InDeltaTime;
+
+	// The window will explicitly unregister this active timer when the mouse leaves.
+	return EActiveTimerReturnType::Continue;
+}
+
+////////////////////////////////////////////////////////////////////////////////////////////////////
+
+void SStartPageWindow::OnMouseEnter(const FGeometry& MyGeometry, const FPointerEvent& MouseEvent)
+{
+	SCompoundWidget::OnMouseEnter(MyGeometry, MouseEvent);
+
+	if (!ActiveTimerHandle.IsValid())
+	{
+		ActiveTimerHandle = RegisterActiveTimer(0.f, FWidgetActiveTimerDelegate::CreateSP(this, &SStartPageWindow::UpdateActiveDuration));
+	}
+}
+
+////////////////////////////////////////////////////////////////////////////////////////////////////
+
+void SStartPageWindow::OnMouseLeave(const FPointerEvent& MouseEvent)
+{
+	SCompoundWidget::OnMouseLeave(MouseEvent);
+
+	auto PinnedActiveTimerHandle = ActiveTimerHandle.Pin();
+	if (PinnedActiveTimerHandle.IsValid())
+	{
+		UnRegisterActiveTimer(PinnedActiveTimerHandle.ToSharedRef());
+	}
+}
+
+////////////////////////////////////////////////////////////////////////////////////////////////////
+
+FReply SStartPageWindow::OnKeyDown(const FGeometry& MyGeometry, const FKeyEvent& InKeyEvent)
+{
+	return FReply::Unhandled();
+}
+
+////////////////////////////////////////////////////////////////////////////////////////////////////
+
+FReply SStartPageWindow::OnDragOver(const FGeometry& MyGeometry, const FDragDropEvent& DragDropEvent)
+{
+	TSharedPtr<FExternalDragOperation> DragDropOp = DragDropEvent.GetOperationAs<FExternalDragOperation>();
+	if (DragDropOp.IsValid())
+	{
+		if (DragDropOp->HasFiles())
+		{
+			const TArray<FString>& Files = DragDropOp->GetFiles();
+			if (Files.Num() == 1)
+			{
+				const FString DraggedFileExtension = FPaths::GetExtension(Files[0], true);
+				if (DraggedFileExtension == TEXT(".utrace"))
+				{
+					return FReply::Handled();
+				}
+			}
+		}
+	}
+
+	return SCompoundWidget::OnDragOver(MyGeometry, DragDropEvent);
+}
+
+////////////////////////////////////////////////////////////////////////////////////////////////////
+
+FReply SStartPageWindow::OnDrop(const FGeometry& MyGeometry, const FDragDropEvent& DragDropEvent)
+{
+	TSharedPtr<FExternalDragOperation> DragDropOp = DragDropEvent.GetOperationAs<FExternalDragOperation>();
+	if (DragDropOp.IsValid())
+	{
+		if (DragDropOp->HasFiles())
+		{
+			// For now, only allow a single file.
+			const TArray<FString>& Files = DragDropOp->GetFiles();
+			if (Files.Num() == 1)
+			{
+				const FString DraggedFileExtension = FPaths::GetExtension(Files[0], true);
+				if (DraggedFileExtension == TEXT(".utrace"))
+				{
+					FInsightsManager::Get()->LoadTraceFile(Files[0]);
+					return FReply::Handled();
+				}
+			}
+		}
+	}
+
+	return SCompoundWidget::OnDrop(MyGeometry, DragDropEvent);
+}
+
+////////////////////////////////////////////////////////////////////////////////////////////////////
+
+bool SStartPageWindow::Live_IsEnabled() const
+{
+	return bIsAnyLiveSessionAvailable;
+}
+
+////////////////////////////////////////////////////////////////////////////////////////////////////
+
+bool SStartPageWindow::Last_IsEnabled() const
+{
+	return bIsAnySessionAvailable;
+}
+
+////////////////////////////////////////////////////////////////////////////////////////////////////
+
+FReply SStartPageWindow::Live_OnClicked()
+{
+	FInsightsManager::Get()->LoadLastLiveSession();
+	return FReply::Handled();
+}
+
+////////////////////////////////////////////////////////////////////////////////////////////////////
+
+FReply SStartPageWindow::Last_OnClicked()
+{
+	FInsightsManager::Get()->LoadLastSession();
+	return FReply::Handled();
+}
+
+////////////////////////////////////////////////////////////////////////////////////////////////////
+
+FReply SStartPageWindow::Load_OnClicked()
+{
+	//const FString ProfilingDirectory(FPaths::ConvertRelativePathToFull(*FPaths::ProfilingDir()));
+	TSharedRef<Trace::ISessionService> SessionService = FInsightsManager::Get()->GetSessionService();
+	const FString ProfilingDirectory(FPaths::ConvertRelativePathToFull(SessionService->GetLocalSessionDirectory()));
+
+	TArray<FString> OutFiles;
+	bool bOpened = false;
+
+	IDesktopPlatform* DesktopPlatform = FDesktopPlatformModule::Get();
+	if (DesktopPlatform != nullptr)
+	{
+		FSlateApplication::Get().CloseToolTip();
+
+		bOpened = DesktopPlatform->OpenFileDialog
+		(
+			FSlateApplication::Get().FindBestParentWindowHandleForDialogs(nullptr),
+			LOCTEXT("LoadTrace_FileDesc", "Open trace file...").ToString(),
+			ProfilingDirectory,
+			TEXT(""),
+			LOCTEXT("LoadTrace_FileFilter", "Trace files (*.utrace)|*.utrace|All files (*.*)|*.*").ToString(),
+			EFileDialogFlags::None,
+			OutFiles
+		);
+	}
+
+	if (bOpened == true)
+	{
+		if (OutFiles.Num() == 1)
+		{
+			FInsightsManager::Get()->LoadTraceFile(OutFiles[0]);
+		}
+	}
+
+	return FReply::Handled();
+}
+
+////////////////////////////////////////////////////////////////////////////////////////////////////
+
+void SStartPageWindow::LoadTraceFile(const TCHAR* TraceFile)
+{
+	FInsightsManager::Get()->LoadTraceFile(FString(TraceFile));
+}
+
+////////////////////////////////////////////////////////////////////////////////////////////////////
+
+void SStartPageWindow::LoadSession(Trace::FSessionHandle SessionHandle)
+{
+	FInsightsManager::Get()->LoadSession(SessionHandle);
+}
+
+////////////////////////////////////////////////////////////////////////////////////////////////////
+
+TSharedRef<SWidget> SStartPageWindow::MakeSessionListMenu()
+{
+	RefreshConnectionList();
+
+	FMenuBuilder MenuBuilder(/*bInShouldCloseWindowAfterMenuSelection=*/true, nullptr);
+
+	MenuBuilder.BeginSection("MostRecentlyUsedSessions", LOCTEXT("MostRecentlyUsedSessionsHeading", "Most Recently Used Sessions"));
+	{
+		//TODO: MRU
+	}
+	MenuBuilder.EndSection();
+
+	MenuBuilder.BeginSection("AvailableSessions", LOCTEXT("AvailableSessionsHeading", "Available Sessions"));
+	{
+		TSharedRef<Trace::ISessionService> SessionService = FInsightsManager::Get()->GetSessionService();
+
+		TArray<Trace::FSessionHandle> AvailableSessions;
+		SessionService->GetAvailableSessions(AvailableSessions);
+
+		// Iterate in reverse order as we want most recent sessions first.
+		for (int32 SessionIndex = AvailableSessions.Num() - 1; SessionIndex >= 0; --SessionIndex)
+		{
+			Trace::FSessionHandle SessionHandle = AvailableSessions[SessionIndex];
+
+			Trace::FSessionInfo SessionInfo;
+			SessionService->GetSessionInfo(SessionHandle, SessionInfo);
+
+			FText Label = FText::FromString(SessionInfo.Name);
+			if (SessionInfo.bIsLive)
+			{
+				Label = FText::Format(LOCTEXT("LiveSessionTextFmt", "{0} (LIVE!)"), Label);
+			}
+
+			MenuBuilder.AddMenuEntry(
+				Label,
+				TAttribute<FText>(), // no tooltip
+				FSlateIcon(),
+				FUIAction(FExecuteAction::CreateSP(this, &SStartPageWindow::LoadSession, SessionHandle)),
+				NAME_None,
+				EUserInterfaceActionType::Button
+			);
+		}
+	}
+	MenuBuilder.EndSection();
+
+	return MenuBuilder.MakeWidget();
+}
+
+////////////////////////////////////////////////////////////////////////////////////////////////////
+
+FText SStartPageWindow::GetLocalSessionDirectory() const
+{
+	TSharedRef<Trace::ISessionService> SessionService = FInsightsManager::Get()->GetSessionService();
+	return FText::FromString(FPaths::ConvertRelativePathToFull(SessionService->GetLocalSessionDirectory()));
+}
+
+////////////////////////////////////////////////////////////////////////////////////////////////////
+
+FReply SStartPageWindow::ExploreLocalSessionDirectory_OnClicked()
+{
+	TSharedRef<Trace::ISessionService> SessionService = FInsightsManager::Get()->GetSessionService();
+	FString FullPath(FPaths::ConvertRelativePathToFull(SessionService->GetLocalSessionDirectory()));
+	FPlatformProcess::ExploreFolder(*FullPath);
+	return FReply::Handled();
+}
+
+////////////////////////////////////////////////////////////////////////////////////////////////////
+
+FText SStartPageWindow::GetRecorderStatusText() const
+{
+	TSharedRef<Trace::ISessionService> SessionService = FInsightsManager::Get()->GetSessionService();
+	if (SessionService->IsRecorderServerRunning())
+	{
+		return FText(LOCTEXT("RecorderServerRunning", "Running"));
+	}
+	else
+	{
+		return FText(LOCTEXT("RecorderServerRunning", "Stopped"));
+	}
+}
+
+////////////////////////////////////////////////////////////////////////////////////////////////////
+
+EVisibility SStartPageWindow::StartTraceRecorder_Visibility() const
+{
+	TSharedRef<Trace::ISessionService> SessionService = FInsightsManager::Get()->GetSessionService();
+	return SessionService->IsRecorderServerRunning() ? EVisibility::Collapsed : EVisibility::Visible;
+}
+
+////////////////////////////////////////////////////////////////////////////////////////////////////
+
+EVisibility SStartPageWindow::StopTraceRecorder_Visibility() const
+{
+	TSharedRef<Trace::ISessionService> SessionService = FInsightsManager::Get()->GetSessionService();
+	return SessionService->IsRecorderServerRunning() ? EVisibility::Visible : EVisibility::Collapsed;
+}
+
+////////////////////////////////////////////////////////////////////////////////////////////////////
+
+FReply SStartPageWindow::StartTraceRecorder_OnClicked()
+{
+	TSharedRef<Trace::ISessionService> SessionService = FInsightsManager::Get()->GetSessionService();
+	SessionService->StartRecorderServer();
+	RefreshConnectionList();
+	return FReply::Handled();
+}
+
+////////////////////////////////////////////////////////////////////////////////////////////////////
+
+FReply SStartPageWindow::StopTraceRecorder_OnClicked()
+{
+	TSharedRef<Trace::ISessionService> SessionService = FInsightsManager::Get()->GetSessionService();
+	SessionService->StopRecorderServer();
+	RefreshConnectionList();
+	return FReply::Handled();
+}
+
+////////////////////////////////////////////////////////////////////////////////////////////////////
+
+void SStartPageWindow::OpenSettings()
+{
+	MainContentPanel->SetEnabled(false);
+	(*OverlaySettingsSlot)
+	[
+		SNew(SBorder)
+		.BorderImage(FEditorStyle::GetBrush("NotificationList.ItemBackground"))
+		.Padding(8.0f)
+		[
+			SNew(SInsightsSettings)
+			.OnClose(this, &SStartPageWindow::CloseSettings)
+			.SettingPtr(&FInsightsManager::GetSettings())
+		]
+	];
+}
+
+////////////////////////////////////////////////////////////////////////////////////////////////////
+
+void SStartPageWindow::CloseSettings()
+{
+	// Close the profiler settings by simply replacing widget with a null one.
+	(*OverlaySettingsSlot)
+	[
+		SNullWidget::NullWidget
+	];
+	MainContentPanel->SetEnabled(true);
+}
+
+////////////////////////////////////////////////////////////////////////////////////////////////////
+
+#undef LOCTEXT_NAMESPACE
diff --git a/Engine/Source/Developer/TraceInsights/Private/Insights/Widgets/SStartPageWindow.h b/Engine/Source/Developer/TraceInsights/Private/Insights/Widgets/SStartPageWindow.h
new file mode 100644
--- /dev/null
+++ b/Engine/Source/Developer/TraceInsights/Private/Insights/Widgets/SStartPageWindow.h
@@ -0,0 +1,211 @@
+// Copyright 1998-2019 Epic Games, Inc. All Rights Reserved.
+
+#pragma once
+
+#include "CoreMinimal.h"
+#include "Framework/Docking/TabManager.h"
+#include "Input/Reply.h"
+#include "Layout/Visibility.h"
+#include "Misc/Guid.h"
+#include "SlateFwd.h"
+#include "TraceServices/ModuleService.h"
+#include "TraceServices/SessionService.h"
+#include "Widgets/DeclarativeSyntaxSupport.h"
+#include "Widgets/Docking/SDockTab.h"
+#include "Widgets/Layout/SSplitter.h"
+#include "Widgets/SCompoundWidget.h"
+#include "Widgets/SOverlay.h"
+#include "Widgets/Views/SListView.h"
+
+// Insights
+#include "Insights/InsightsManager.h"
+
+////////////////////////////////////////////////////////////////////////////////////////////////////
+
+class FActiveTimerHandle;
+class SVerticalBox;
+class SEditableTextBox;
+
+////////////////////////////////////////////////////////////////////////////////////////////////////
+
+/** Type definition for shared pointers to instances of SNotificationItem. */
+typedef TSharedPtr<class SNotificationItem> SNotificationItemPtr;
+
+/** Type definition for shared references to instances of SNotificationItem. */
+typedef TSharedRef<class SNotificationItem> SNotificationItemRef;
+
+/** Type definition for weak references to instances of SNotificationItem. */
+typedef TWeakPtr<class SNotificationItem> SNotificationItemWeak;
+
+////////////////////////////////////////////////////////////////////////////////////////////////////
+
+struct FRecorderConnection
+{
+	Trace::FSessionHandle SessionHandle;
+	FText Name;
+	FText Uri;
+
+	FRecorderConnection(const Trace::FSessionHandle InSessionHandle, const Trace::FSessionInfo& InSessionInfo)
+		: SessionHandle(InSessionHandle)
+		, Name(FText::FromString(InSessionInfo.Name))
+		, Uri(FText::FromString(InSessionInfo.Uri))
+	{}
+};
+
+/** Implements the Start Page window. */
+class SStartPageWindow : public SCompoundWidget
+{
+public:
+	/** Default constructor. */
+	SStartPageWindow();
+
+	/** Virtual destructor. */
+	virtual ~SStartPageWindow();
+
+	SLATE_BEGIN_ARGS(SStartPageWindow){}
+	SLATE_END_ARGS()
+
+	/** Constructs this widget. */
+	void Construct(const FArguments& InArgs);
+
+	void ManageLoadingProgressNotificationState(const FString& Filename, const EInsightsNotificationType NotificatonType, const ELoadingProgressState ProgressState, const float LoadingProgress);
+
+	void OpenSettings();
+	void CloseSettings();
+
+private:
+	TSharedRef<SWidget> ConstructRecoderPanel();
+	TSharedRef<SWidget> ConstructModuleList();
+
+	/** Generate a new row for the Connections list view. */
+	TSharedRef<ITableRow> Connections_OnGenerateRow(TSharedPtr<FRecorderConnection> InConnection, const TSharedRef<STableViewBase>& OwnerTable);
+
+	/** Callback for determining the visibility of the 'Please select a trace' overlay. */
+	EVisibility IsSessionOverlayVisible() const;
+
+	bool IsSessionValid() const;
+
+	void SendingServiceSideCapture_Cancel(const FString Filename);
+	void SendingServiceSideCapture_Load(const FString Filename);
+
+	bool Live_IsEnabled() const;
+	bool Last_IsEnabled() const;
+
+	FReply Live_OnClicked();
+	FReply Last_OnClicked();
+	FReply Load_OnClicked();
+
+	void LoadTraceFile(const TCHAR* TraceFile);
+	void LoadSession(Trace::FSessionHandle SessionHandle);
+
+	TSharedRef<SWidget> MakeSessionListMenu();
+
+	FText GetLocalSessionDirectory() const;
+	FReply ExploreLocalSessionDirectory_OnClicked();
+
+	FText GetRecorderStatusText() const;
+	EVisibility StartTraceRecorder_Visibility() const;
+	EVisibility StopTraceRecorder_Visibility() const;
+	FReply StartTraceRecorder_OnClicked();
+	FReply StopTraceRecorder_OnClicked();
+
+	EVisibility Modules_Visibility() const;
+	ECheckBoxState Module_IsChecked(int32 ModuleIndex) const;
+	void Module_OnCheckStateChanged(ECheckBoxState NewRadioState, int32 ModuleIndex);
+
+	FReply RefreshConnections_OnClicked();
+	FReply Connect_OnClicked();
+	void RefreshConnectionList();
+	void Connections_OnSelectionChanged(TSharedPtr<FRecorderConnection> Connection, ESelectInfo::Type SelectInfo);
+	EVisibility Connections_Visibility() const;
+
+	/** Updates the amount of time the profiler has been active. */
+	EActiveTimerReturnType UpdateActiveDuration(double InCurrentTime, float InDeltaTime);
+
+	/**
+	 * Ticks this widget.  Override in derived classes, but always call the parent implementation.
+	 *
+	 * @param  AllottedGeometry The space allotted for this widget
+	 * @param  InCurrentTime  Current absolute real time
+	 * @param  InDeltaTime  Real time passed since last tick
+	 */
+	virtual void Tick(const FGeometry& AllottedGeometry, const double InCurrentTime, const float InDeltaTime) override;
+
+	/**
+	 * The system will use this event to notify a widget that the cursor has entered it. This event is NOT bubbled.
+	 *
+	 * @param MyGeometry The Geometry of the widget receiving the event
+	 * @param MouseEvent Information about the input event
+	 */
+	virtual void OnMouseEnter(const FGeometry& MyGeometry, const FPointerEvent& MouseEvent) override;
+
+	/**
+	 * The system will use this event to notify a widget that the cursor has left it. This event is NOT bubbled.
+	 *
+	 * @param MouseEvent Information about the input event
+	 */
+	virtual void OnMouseLeave(const FPointerEvent& MouseEvent) override;
+
+	/**
+	 * Called after a key is pressed when this widget has focus
+	 *
+	 * @param MyGeometry The Geometry of the widget receiving the event
+	 * @param  InKeyEvent  Key event
+	 *
+	 * @return  Returns whether the event was handled, along with other possible actions
+	 */
+	virtual FReply OnKeyDown(const FGeometry& MyGeometry, const FKeyEvent& InKeyEvent) override;
+
+	/**
+	 * Called when the user is dropping something onto a widget; terminates drag and drop.
+	 *
+	 * @param MyGeometry      The geometry of the widget receiving the event.
+	 * @param DragDropEvent   The drag and drop event.
+	 *
+	 * @return A reply that indicated whether this event was handled.
+	 */
+	virtual FReply OnDrop(const FGeometry& MyGeometry, const FDragDropEvent& DragDropEvent) override;
+
+	/**
+	 * Called during drag and drop when the the mouse is being dragged over a widget.
+	 *
+	 * @param MyGeometry      The geometry of the widget receiving the event.
+	 * @param DragDropEvent   The drag and drop event.
+	 *
+	 * @return A reply that indicated whether this event was handled.
+	 */
+	virtual FReply OnDragOver(const FGeometry& MyGeometry, const FDragDropEvent& DragDropEvent)  override;
+
+public:
+	/** Widget for the non-intrusive notifications. */
+	TSharedPtr<SNotificationList> NotificationList;
+
+	/** Holds all active and visible notifications, stored as FGuid -> SNotificationItemWeak. */
+	TMap<FString, SNotificationItemWeak> ActiveNotifications;
+
+	/** Overlay slot which contains the profiler settings widget. */
+	SOverlay::FOverlaySlot* OverlaySettingsSlot;
+
+	/** The number of seconds the profiler has been active */
+	float DurationActive;
+
+private:
+	/** The handle to the active update duration tick */
+	TWeakPtr<FActiveTimerHandle> ActiveTimerHandle;
+
+	/** Holds all widgets for the profiler window like menu bar, toolbar and tabs. */
+	TSharedPtr<SVerticalBox> MainContentPanel;
+
+	bool bIsAnyLiveSessionAvailable;
+	Trace::FSessionHandle LastLiveSessionHandle;
+
+	bool bIsAnySessionAvailable;
+	Trace::FSessionHandle LastSessionHandle;
+
+	TSharedPtr<SListView<TSharedPtr<FRecorderConnection>>> ConnectionsListView;
+	TArray<TSharedPtr<FRecorderConnection>> Connections;
+	TSharedPtr<SEditableTextBox> HostTextBox;
+	TSharedPtr<FRecorderConnection> SelectedConnection;
+	TArray<Trace::FModuleInfo> AvailableModules;
+	TArray<bool> AvailableModulesEnabledState;
+};
diff --git a/Engine/Source/Developer/TraceInsights/Private/Insights/Widgets/SStatsTableCell.cpp b/Engine/Source/Developer/TraceInsights/Private/Insights/Widgets/SStatsTableCell.cpp
new file mode 100644
--- /dev/null
+++ b/Engine/Source/Developer/TraceInsights/Private/Insights/Widgets/SStatsTableCell.cpp
@@ -0,0 +1,146 @@
+// Copyright 1998-2019 Epic Games, Inc. All Rights Reserved.
+
+#include "SStatsTableCell.h"
+
+#include "EditorStyleSet.h"
+#include "SlateOptMacros.h"
+#include "Widgets/Images/SImage.h"
+#include "Widgets/Input/SButton.h"
+#include "Widgets/SBoxPanel.h"
+#include "Widgets/Text/STextBlock.h"
+#include "Widgets/Views/SExpanderArrow.h"
+
+// Insights
+#include "Insights/ViewModels/StatsViewColumn.h"
+#include "Insights/ViewModels/StatsViewColumnFactory.h"
+#include "Insights/Widgets/SStatsViewTooltip.h"
+
+#define LOCTEXT_NAMESPACE "SStatsView"
+
+////////////////////////////////////////////////////////////////////////////////////////////////////
+
+BEGIN_SLATE_FUNCTION_BUILD_OPTIMIZATION
+
+void SStatsTableCell::Construct(const FArguments& InArgs, const TSharedRef<class ITableRow>& TableRow)
+{
+	SetHoveredTableCellDelegate = InArgs._OnSetHoveredTableCell;
+	StatsNodePtr = InArgs._StatsNodePtr;
+	ColumnId = InArgs._ColumnId;
+	//HighlightText = InArgs._HighlightText;
+
+	ChildSlot
+		[
+			GenerateWidgetForColumn(InArgs, TableRow)
+		];
+}
+
+////////////////////////////////////////////////////////////////////////////////////////////////////
+
+TSharedRef<SWidget> SStatsTableCell::GenerateWidgetForColumn(const FArguments& InArgs, const TSharedRef<class ITableRow>& TableRow)
+{
+	if (InArgs._IsStatsNameColumn)
+	{
+		return GenerateWidgetForNameColumn(InArgs, TableRow);
+	}
+	else
+	{
+		return GenerateWidgetForStatsColumn(InArgs, TableRow);
+	}
+}
+
+////////////////////////////////////////////////////////////////////////////////////////////////////
+
+TSharedRef<SWidget> SStatsTableCell::GenerateWidgetForNameColumn(const FArguments& InArgs, const TSharedRef<class ITableRow>& TableRow)
+{
+	const FStatsViewColumn& Column = *FStatsViewColumnFactory::Get().ColumnIdToPtrMapping.FindChecked(ColumnId);
+
+	return
+		SNew(SHorizontalBox)
+
+		+ SHorizontalBox::Slot()
+		.AutoWidth()
+		.HAlign(HAlign_Right)
+		.VAlign(VAlign_Center)
+		[
+			SNew(SExpanderArrow, TableRow)
+		]
+
+		// Event info icon + tooltip.
+		+ SHorizontalBox::Slot()
+		.AutoWidth()
+		.HAlign(HAlign_Center)
+		.VAlign(VAlign_Center)
+		[
+			SNew(SImage)
+			.Visibility(this, &SStatsTableCell::GetHintIconVisibility)
+			.Image(FEditorStyle::GetBrush("Profiler.Tooltip.HintIcon10"))
+			.ToolTip(SStatsViewTooltip::GetTableCellTooltip(StatsNodePtr))
+		]
+
+		// Name
+		+ SHorizontalBox::Slot()
+		.AutoWidth()
+		.VAlign(VAlign_Center)
+		.HAlign(Column.HorizontalAlignment)
+		.Padding(FMargin(2.0f, 0.0f))
+		[
+			SNew(STextBlock)
+			//.Text(StatsNodePtr->GetNameEx())
+			.Text(this, &SStatsTableCell::GetNameEx)
+			.HighlightText(InArgs._HighlightText)
+			.TextStyle(FEditorStyle::Get(), TEXT("Profiler.Tooltip"))
+			.ColorAndOpacity(this, &SStatsTableCell::GetColorAndOpacity)
+			.ShadowColorAndOpacity(this, &SStatsTableCell::GetShadowColorAndOpacity)
+		]
+
+		/*
+		// Culled children warning icon + tooltip.
+		+ SHorizontalBox::Slot()
+		.AutoWidth()
+		.HAlign(HAlign_Center)
+		.VAlign(VAlign_Center)
+		[
+			SNew(SButton)
+			.ButtonStyle(FEditorStyle::Get(), TEXT("HoverHintOnly"))
+			.ContentPadding(0.0f)
+			.IsFocusable(false)
+			//.OnClicked(this, &SStatsTableCell::ExpandCulledEvents_OnClicked)
+			[
+				SNew(SImage)
+				.Visibility(this, &SStatsTableCell::GetCulledEventsIconVisibility)
+				.Image(FEditorStyle::GetBrush("Profiler.EventGraph.HasCulledEventsSmall"))
+				.ToolTipText(LOCTEXT("HasCulledEvents_TT", "This event contains culled children, if you want to see all children, please disable culling or use function details, or press this icon"))
+			]
+		]*/
+		;
+}
+
+////////////////////////////////////////////////////////////////////////////////////////////////////
+
+TSharedRef<SWidget> SStatsTableCell::GenerateWidgetForStatsColumn(const FArguments& InArgs, const TSharedRef<class ITableRow>& TableRow)
+{
+	const FStatsViewColumn& Column = *FStatsViewColumnFactory::Get().ColumnIdToPtrMapping.FindChecked(ColumnId);
+	const FText FormattedValue = Column.GetFormattedValue(*StatsNodePtr);
+
+	return
+		SNew(SHorizontalBox)
+
+		+ SHorizontalBox::Slot()
+		.FillWidth(1.0f)
+		.VAlign(VAlign_Center)
+		.HAlign(Column.HorizontalAlignment)
+		.Padding(FMargin(2.0f, 0.0f))
+		[
+			SNew(STextBlock)
+			.Text(FormattedValue)
+			.TextStyle(FEditorStyle::Get(), TEXT("Profiler.Tooltip"))
+			.ColorAndOpacity(this, &SStatsTableCell::GetStatsColorAndOpacity)
+			.ShadowColorAndOpacity(this, &SStatsTableCell::GetShadowColorAndOpacity)
+		];
+}
+
+END_SLATE_FUNCTION_BUILD_OPTIMIZATION
+
+////////////////////////////////////////////////////////////////////////////////////////////////////
+
+#undef LOCTEXT_NAMESPACE
diff --git a/Engine/Source/Developer/TraceInsights/Private/Insights/Widgets/SStatsTableCell.h b/Engine/Source/Developer/TraceInsights/Private/Insights/Widgets/SStatsTableCell.h
new file mode 100644
--- /dev/null
+++ b/Engine/Source/Developer/TraceInsights/Private/Insights/Widgets/SStatsTableCell.h
@@ -0,0 +1,140 @@
+// Copyright 1998-2019 Epic Games, Inc. All Rights Reserved.
+
+#pragma once
+
+#include "CoreMinimal.h"
+#include "Widgets/DeclarativeSyntaxSupport.h"
+#include "Widgets/SCompoundWidget.h"
+
+// Insights
+#include "Insights/ViewModels/StatsNode.h"
+
+DECLARE_DELEGATE_TwoParams(FSetHoveredStatsTableCell, const FName /*ColumnId*/, const FStatsNodePtr /*SamplePtr*/);
+DECLARE_DELEGATE_RetVal_OneParam(bool, FIsColumnVisibleDelegate, const FName /*ColumnId*/);
+DECLARE_DELEGATE_RetVal_OneParam(EHorizontalAlignment, FGetColumnOutlineHAlignmentDelegate, const FName /*ColumnId*/);
+
+class SStatsTableCell : public SCompoundWidget
+{
+public:
+	SLATE_BEGIN_ARGS(SStatsTableCell) {}
+		SLATE_EVENT(FSetHoveredStatsTableCell, OnSetHoveredTableCell)
+		SLATE_ATTRIBUTE(FText, HighlightText)
+		SLATE_ARGUMENT(FStatsNodePtr, StatsNodePtr)
+		SLATE_ARGUMENT(FName, ColumnId)
+		SLATE_ARGUMENT(bool, IsStatsNameColumn)
+	SLATE_END_ARGS()
+
+	void Construct(const FArguments& InArgs, const TSharedRef<class ITableRow>& TableRow);
+
+protected:
+	TSharedRef<SWidget> GenerateWidgetForColumn(const FArguments& InArgs, const TSharedRef<class ITableRow>& TableRow);
+	TSharedRef<SWidget> GenerateWidgetForNameColumn(const FArguments& InArgs, const TSharedRef<class ITableRow>& TableRow);
+	TSharedRef<SWidget> GenerateWidgetForStatsColumn(const FArguments& InArgs, const TSharedRef<class ITableRow>& TableRow);
+
+	/**
+	 * The system will use this event to notify a widget that the cursor has entered it. This event is NOT bubbled.
+	 *
+	 * @param MyGeometry The Geometry of the widget receiving the event
+	 * @param MouseEvent Information about the input event
+	 */
+	virtual void OnMouseEnter(const FGeometry& MyGeometry, const FPointerEvent& MouseEvent) override
+	{
+		SCompoundWidget::OnMouseEnter(MyGeometry, MouseEvent);
+		SetHoveredTableCellDelegate.ExecuteIfBound(ColumnId, StatsNodePtr);
+	}
+
+	/**
+	 * The system will use this event to notify a widget that the cursor has left it. This event is NOT bubbled.
+	 *
+	 * @param MouseEvent Information about the input event
+	 */
+	virtual void OnMouseLeave(const FPointerEvent& MouseEvent) override
+	{
+		SCompoundWidget::OnMouseLeave(MouseEvent);
+		SetHoveredTableCellDelegate.ExecuteIfBound(NAME_None, nullptr);
+	}
+
+	/**
+	 * Called during drag and drop when the drag enters a widget.
+	 *
+	 * Enter/Leave events in slate are meant as lightweight notifications.
+	 * So we do not want to capture mouse or set focus in response to these.
+	 * However, OnDragEnter must also support external APIs (e.g. OLE Drag/Drop)
+	 * Those require that we let them know whether we can handle the content
+	 * being dragged OnDragEnter.
+	 *
+	 * The concession is to return a can_handled/cannot_handle
+	 * boolean rather than a full FReply.
+	 *
+	 * @param MyGeometry      The geometry of the widget receiving the event.
+	 * @param DragDropEvent   The drag and drop event.
+	 *
+	 * @return A reply that indicated whether the contents of the DragDropEvent can potentially be processed by this widget.
+	 */
+	virtual void OnDragEnter(const FGeometry& MyGeometry, const FDragDropEvent& DragDropEvent) override
+	{
+		SCompoundWidget::OnDragEnter(MyGeometry, DragDropEvent);
+		SetHoveredTableCellDelegate.ExecuteIfBound(ColumnId, StatsNodePtr);
+	}
+
+	/**
+	 * Called during drag and drop when the drag leaves a widget.
+	 *
+	 * @param DragDropEvent   The drag and drop event.
+	 */
+	virtual void OnDragLeave(const FDragDropEvent& DragDropEvent)  override
+	{
+		SCompoundWidget::OnDragLeave(DragDropEvent);
+		SetHoveredTableCellDelegate.ExecuteIfBound(NAME_None, nullptr);
+	}
+
+	EVisibility GetHintIconVisibility() const
+	{
+		return IsHovered() ? EVisibility::Visible : EVisibility::Hidden;
+	}
+
+	EVisibility GetCulledEventsIconVisibility() const
+	{
+		//return StatsNodePtr->HasCulledChildren() ? EVisibility::Visible : EVisibility::Collapsed;
+		return EVisibility::Collapsed;
+	}
+
+	FText GetNameEx() const
+	{
+		return StatsNodePtr->GetNameEx();
+	}
+
+	FSlateColor GetColorAndOpacity() const
+	{
+		const FLinearColor TextColor =
+			StatsNodePtr->IsFiltered() ? FLinearColor(1.0f, 1.0f, 1.0f, 0.5f) :
+			FLinearColor::White;
+		return TextColor;
+	}
+
+	FSlateColor GetStatsColorAndOpacity() const
+	{
+		const FLinearColor TextColor =
+			StatsNodePtr->IsFiltered() ? FLinearColor(1.0f, 1.0f, 1.0f, 0.5f) :
+			StatsNodePtr->GetAggregatedStats().Count == 0 ? FLinearColor(1.0f, 1.0f, 1.0f, 0.6f) :
+			FLinearColor::White;
+		return TextColor;
+	}
+
+	FLinearColor GetShadowColorAndOpacity() const
+	{
+		const FLinearColor ShadowColor =
+			StatsNodePtr->IsFiltered() ? FLinearColor(0.f, 0.f, 0.f, 0.25f) :
+			FLinearColor(0.0f, 0.0f, 0.0f, 0.5f);
+		return ShadowColor;
+	}
+
+protected:
+	/** A shared pointer to the stats node. */
+	FStatsNodePtr StatsNodePtr;
+
+	/** The Id of the column where this stats belongs. */
+	FName ColumnId;
+
+	FSetHoveredStatsTableCell SetHoveredTableCellDelegate;
+};
diff --git a/Engine/Source/Developer/TraceInsights/Private/Insights/Widgets/SStatsTableRow.cpp b/Engine/Source/Developer/TraceInsights/Private/Insights/Widgets/SStatsTableRow.cpp
new file mode 100644
--- /dev/null
+++ b/Engine/Source/Developer/TraceInsights/Private/Insights/Widgets/SStatsTableRow.cpp
@@ -0,0 +1,273 @@
+// Copyright 1998-2019 Epic Games, Inc. All Rights Reserved.
+
+#include "SStatsTableRow.h"
+
+#include "Widgets/SOverlay.h"
+#include "Widgets/Images/SImage.h"
+
+// Insights
+#include "Insights/Common/TimeUtils.h"
+#include "Insights/ViewModels/StatsViewColumnFactory.h"
+#include "Insights/Widgets/SStatsViewTooltip.h"
+#include "Insights/Widgets/SStatsTableCell.h"
+
+#define LOCTEXT_NAMESPACE "SStatsView"
+
+BEGIN_SLATE_FUNCTION_BUILD_OPTIMIZATION
+
+void SStatsTableRow::Construct(const FArguments& InArgs, const TSharedRef<STableViewBase>& InOwnerTableView)
+{
+	OnShouldBeEnabled = InArgs._OnShouldBeEnabled;
+	IsColumnVisibleDelegate = InArgs._OnIsColumnVisible;
+	SetHoveredTableCellDelegate = InArgs._OnSetHoveredTableCell;
+	GetColumnOutlineHAlignmentDelegate = InArgs._OnGetColumnOutlineHAlignmentDelegate;
+
+	HighlightText = InArgs._HighlightText;
+	HighlightedStatsName = InArgs._HighlightedStatsName;
+
+	StatsNodePtr = InArgs._StatsNodePtr;
+
+	SetEnabled(TAttribute<bool>(this, &SStatsTableRow::HandleShouldBeEnabled));
+
+	SMultiColumnTableRow<FStatsNodePtr>::Construct(SMultiColumnTableRow<FStatsNodePtr>::FArguments(), InOwnerTableView);
+
+	/*
+	const FSlateBrush* const NodeIcon = StatsNodeTypeHelper::GetIconForStatsNodeType(InStatsNode->GetType());
+	const TSharedRef<SToolTip> Tooltip = InStatsNode->IsGroup() ? SNew(SToolTip) : SStatsViewTooltip(InStatsNode->GetId()).GetTooltip();
+
+	ChildSlot
+		[
+			SNew(SHorizontalBox)
+
+			// Expander arrow.
+			+ SHorizontalBox::Slot()
+			.AutoWidth()
+			.HAlign(HAlign_Right)
+			.VAlign(VAlign_Center)
+			[
+				SNew(SExpanderArrow, SharedThis(this))
+			]
+
+			// Icon to visualize group or timer type.
+			+ SHorizontalBox::Slot()
+			.AutoWidth()
+			.Padding(0.0f, 0.0f, 8.0f, 0.0f)
+			[
+				SNew(SImage)
+				.Image(NodeIcon)
+				.ToolTip(Tooltip)
+			]
+
+			// Description text.
+			+ SHorizontalBox::Slot()
+			.AutoWidth()
+			.VAlign(VAlign_Center)
+			.HAlign(HAlign_Left)
+			.Padding(FMargin(2.0f, 0.0f))
+			[
+				SNew(STextBlock)
+				.Text(this, &SStatsTableRow::GetText)
+				.HighlightText(InArgs._HighlightText)
+				.TextStyle(FEditorStyle::Get(), TEXT("Profiler.Tooltip"))
+				.ColorAndOpacity(this, &SStatsTableRow::GetColorAndOpacity)
+			]
+
+			+ SHorizontalBox::Slot()
+			.AutoWidth()
+			.HAlign(HAlign_Center)
+			.VAlign(VAlign_Top)
+			.Padding(0.0f, 1.0f, 0.0f, 0.0f)
+			[
+				SNew(SImage)
+				.Visibility(!InStatsNode->IsGroup() ? EVisibility::Visible : EVisibility::Collapsed)
+				.Image(FEditorStyle::GetBrush("Profiler.Tooltip.HintIcon10"))
+				.ToolTip(Tooltip)
+			]
+		];
+
+	STableRow<FStatsNodePtr>::ConstructInternal(STableRow::FArguments().ShowSelection(true), InOwnerTableView);
+	*/
+}
+
+TSharedRef<SWidget> SStatsTableRow::GenerateWidgetForColumn(const FName& ColumnId)
+{
+	return
+
+	SNew(SOverlay)
+	.Visibility(EVisibility::SelfHitTestInvisible)
+
+	+SOverlay::Slot()
+	.Padding(0.0f)
+	[
+		SNew(SImage)
+		.Image(FEditorStyle::GetBrush("Profiler.LineGraphArea"))
+		.ColorAndOpacity(this, &SStatsTableRow::GetBackgroundColorAndOpacity)
+	]
+
+	+SOverlay::Slot()
+	.Padding(0.0f)
+	[
+		SNew(SImage)
+		.Image(this, &SStatsTableRow::GetOutlineBrush, ColumnId)
+		.ColorAndOpacity(this, &SStatsTableRow::GetOutlineColorAndOpacity)
+	]
+
+	+SOverlay::Slot()
+	[
+		SNew(SStatsTableCell, SharedThis(this))
+		.Visibility(this, &SStatsTableRow::IsColumnVisible, ColumnId)
+		.StatsNodePtr(StatsNodePtr)
+		.ColumnId(ColumnId)
+		.HighlightText(HighlightText)
+		.IsStatsNameColumn(ColumnId == FStatsViewColumnFactory::Get().Collection[0]->Id) // name column
+		.OnSetHoveredTableCell(this, &SStatsTableRow::OnSetHoveredTableCell)
+	];
+}
+
+END_SLATE_FUNCTION_BUILD_OPTIMIZATION
+
+FReply SStatsTableRow::OnDragDetected(const FGeometry& MyGeometry, const FPointerEvent& MouseEvent)
+{
+	//im:TODO
+	//if (MouseEvent.IsMouseButtonDown(EKeys::LeftMouseButton))
+	//{
+	//	if (StatsNode->IsGroup())
+	//	{
+	//		// Add all timer Ids for the group.
+	//		TArray<int32> StatsIds;
+	//		const TArray<FStatsNodePtr>& FilteredChildren = StatsNode->GetFilteredChildren();
+	//		const int32 NumFilteredChildren = FilteredChildren.Num();
+	//
+	//		StatsIds.Reserve(NumFilteredChildren);
+	//		for (int32 Nx = 0; Nx < NumFilteredChildren; ++Nx)
+	//		{
+	//			StatsIds.Add(FilteredChildren[Nx]->GetId());
+	//		}
+	//
+	//		return FReply::Handled().BeginDragDrop(FStatIDDragDropOp::NewGroup(StatsIds, StatsNode->GetName().GetPlainNameString()));
+	//	}
+	//	else
+	//	{
+	//		return FReply::Handled().BeginDragDrop(FStatIDDragDropOp::NewSingle(StatsNode->GetId(), StatsNode->GetName().GetPlainNameString()));
+	//	}
+	//}
+
+	return SMultiColumnTableRow<FStatsNodePtr>::OnDragDetected(MyGeometry, MouseEvent);
+}
+
+FText SStatsTableRow::GetText() const
+{
+	FText Text = FText::GetEmpty();
+
+	if (StatsNodePtr->IsGroup())
+	{
+		Text = FText::Format(LOCTEXT("StatsNode_GroupNodeTextFmt", "{0} ({1})"), FText::FromName(StatsNodePtr->GetName()), FText::AsNumber(StatsNodePtr->GetChildren().Num()));
+	}
+	else
+	{
+		Text = FText::FromName(StatsNodePtr->GetName());
+	}
+
+	return Text;
+}
+
+FSlateFontInfo SStatsTableRow::GetFont() const
+{
+	const bool bIsStatTracked = false;//im:TODO: FTimingProfilerManager::Get()->IsStatTracked(StatsNodePtr->GetStatsId());
+	const FSlateFontInfo FontInfo = bIsStatTracked ? FEditorStyle::GetFontStyle("BoldFont") : FEditorStyle::GetFontStyle("NormalFont");
+	return FontInfo;
+}
+
+FSlateColor SStatsTableRow::GetColorAndOpacity() const
+{
+	//const bool bIsStatTracked = FTimingProfilerManager::Get()->IsStatTracked(StatsNodePtr->GetId());
+	//const FSlateColor Color = bIsStatTracked ? FTimingProfilerManager::Get()->GetColorForStatsId(StatsNodePtr->GetId()) : FLinearColor::White;
+	//return Color;
+	//FTimingProfilerManager::GetSettings().GetColorForStat(StatName)
+	return FLinearColor::White;
+}
+
+FSlateColor SStatsTableRow::GetBackgroundColorAndOpacity() const
+{
+	//return GetBackgroundColorAndOpacity(StatsNodePtr->GetAggregatedStats().Sum);
+	return FLinearColor(0.0f, 0.0f, 0.0f, 1.0f);
+}
+
+FSlateColor SStatsTableRow::GetBackgroundColorAndOpacity(double Time) const
+{
+	const FLinearColor Color =	Time > TimeUtils::Second      ? FLinearColor(0.3f, 0.0f, 0.0f, 1.0f) :
+								Time > TimeUtils::Milisecond  ? FLinearColor(0.3f, 0.1f, 0.0f, 1.0f) :
+								Time > TimeUtils::Microsecond ? FLinearColor(0.0f, 0.1f, 0.0f, 1.0f) :
+																FLinearColor(0.0f, 0.0f, 0.0f, 1.0f);
+	return Color;
+}
+
+FSlateColor SStatsTableRow::GetOutlineColorAndOpacity() const
+{
+	const FLinearColor NoColor(0.0f, 0.0f, 0.0f, 0.0f);
+	const bool bShouldBeHighlighted = StatsNodePtr->GetName() == HighlightedStatsName.Get();
+	const FLinearColor OutlineColorAndOpacity = bShouldBeHighlighted ? FLinearColor(FColorList::SlateBlue) : NoColor;
+	return OutlineColorAndOpacity;
+}
+
+const FSlateBrush* SStatsTableRow::GetOutlineBrush(const FName ColumnId) const
+{
+	EHorizontalAlignment Result = HAlign_Center;
+	if (IsColumnVisibleDelegate.IsBound())
+	{
+		Result = GetColumnOutlineHAlignmentDelegate.Execute(ColumnId);
+	}
+
+	const FSlateBrush* Brush = nullptr;
+	if (Result == HAlign_Left)
+	{
+		Brush = FEditorStyle::GetBrush("Profiler.EventGraph.Border.L");
+	}
+	else if(Result == HAlign_Right)
+	{
+		Brush = FEditorStyle::GetBrush("Profiler.EventGraph.Border.R");
+	}
+	else
+	{
+		Brush = FEditorStyle::GetBrush("Profiler.EventGraph.Border.TB");
+	}
+	return Brush;
+}
+
+bool SStatsTableRow::HandleShouldBeEnabled() const
+{
+	bool bResult = false;
+
+	if (StatsNodePtr->IsGroup())
+	{
+		bResult = true;
+	}
+	else
+	{
+		if (OnShouldBeEnabled.IsBound())
+		{
+			bResult = OnShouldBeEnabled.Execute(StatsNodePtr->GetId());
+		}
+	}
+
+	return bResult;
+}
+
+EVisibility SStatsTableRow::IsColumnVisible(const FName ColumnId) const
+{
+	EVisibility Result = EVisibility::Collapsed;
+
+	if (IsColumnVisibleDelegate.IsBound())
+	{
+		Result = IsColumnVisibleDelegate.Execute(ColumnId) ? EVisibility::Visible : EVisibility::Collapsed;
+	}
+
+	return Result;
+}
+
+void SStatsTableRow::OnSetHoveredTableCell(const FName InColumnId, const FStatsNodePtr InSamplePtr)
+{
+	SetHoveredTableCellDelegate.ExecuteIfBound(InColumnId, InSamplePtr);
+}
+
+#undef LOCTEXT_NAMESPACE
diff --git a/Engine/Source/Developer/TraceInsights/Private/Insights/Widgets/SStatsTableRow.h b/Engine/Source/Developer/TraceInsights/Private/Insights/Widgets/SStatsTableRow.h
new file mode 100644
--- /dev/null
+++ b/Engine/Source/Developer/TraceInsights/Private/Insights/Widgets/SStatsTableRow.h
@@ -0,0 +1,93 @@
+// Copyright 1998-2019 Epic Games, Inc. All Rights Reserved.
+
+#pragma once
+
+#include "CoreMinimal.h"
+#include "Widgets/DeclarativeSyntaxSupport.h"
+#include "Widgets/Views/SHeaderRow.h"
+#include "Widgets/Views/STableViewBase.h"
+#include "Widgets/Views/STableRow.h"
+#include "Widgets/Views/STreeView.h"
+
+// Insights
+#include "Insights/ViewModels/StatsNodeHelper.h"
+
+DECLARE_DELEGATE_RetVal_OneParam(bool, FShouldBeEnabledDelegate, const uint32 /*StatsId*/);
+DECLARE_DELEGATE_RetVal_OneParam(bool, FIsColumnVisibleDelegate, const FName /*ColumnId*/);
+DECLARE_DELEGATE_TwoParams(FSetHoveredStatsTableCell, const FName /*ColumnId*/, const FStatsNodePtr /*StatsNodePtr*/);
+DECLARE_DELEGATE_RetVal_OneParam(EHorizontalAlignment, FGetColumnOutlineHAlignmentDelegate, const FName /*ColumnId*/);
+
+/** Widget that represents a table row in the stats' tree control. Generates widgets for each column on demand. */
+class SStatsTableRow : public SMultiColumnTableRow<FStatsNodePtr>
+{
+public:
+	SLATE_BEGIN_ARGS(SStatsTableRow) {}
+		SLATE_EVENT(FShouldBeEnabledDelegate, OnShouldBeEnabled)
+		SLATE_EVENT(FIsColumnVisibleDelegate, OnIsColumnVisible)
+		SLATE_EVENT(FSetHoveredStatsTableCell, OnSetHoveredTableCell)
+		SLATE_EVENT(FGetColumnOutlineHAlignmentDelegate, OnGetColumnOutlineHAlignmentDelegate)
+		SLATE_ATTRIBUTE(FText, HighlightText)
+		SLATE_ATTRIBUTE(FName, HighlightedStatsName)
+		SLATE_ARGUMENT(FStatsNodePtr, StatsNodePtr)
+	SLATE_END_ARGS()
+
+public:
+	void Construct(const FArguments& InArgs, const TSharedRef<STableViewBase>& InOwnerTableView);
+
+	virtual TSharedRef<SWidget> GenerateWidgetForColumn(const FName& ColumnId) override;
+
+	/**
+	 * Called when Slate detects that a widget started to be dragged.
+	 * Usage:
+	 * A widget can ask Slate to detect a drag.
+	 * OnMouseDown() reply with FReply::Handled().DetectDrag(SharedThis(this)).
+	 * Slate will either send an OnDragDetected() event or do nothing.
+	 * If the user releases a mouse button or leaves the widget before
+	 * a drag is triggered (maybe user started at the very edge) then no event will be
+	 * sent.
+	 *
+	 * @param  InMyGeometry  Widget geometry
+	 * @param  InMouseEvent  MouseMove that triggered the drag
+	 *
+	 */
+	virtual FReply OnDragDetected(const FGeometry& MyGeometry, const FPointerEvent& MouseEvent) override;
+
+protected:
+	/**
+	 * @return a text which describes this table row, refers to both groups and stats
+	 */
+	FText GetText() const;
+
+	/**
+	 * @return a font style which is used to draw this table row, refers to both groups and stats
+	 */
+	FSlateFontInfo GetFont() const;
+
+	/**
+	 * @return a color and opacity value used to draw this table row, refers to both groups and stats
+	 */
+	FSlateColor GetColorAndOpacity() const;
+
+	FSlateColor GetBackgroundColorAndOpacity() const;
+	FSlateColor GetBackgroundColorAndOpacity(double Time) const;
+	FSlateColor GetOutlineColorAndOpacity() const;
+	const FSlateBrush* GetOutlineBrush(const FName ColumnId) const;
+	bool HandleShouldBeEnabled() const;
+	EVisibility IsColumnVisible(const FName ColumnId) const;
+	void OnSetHoveredTableCell(const FName ColumnId, const FStatsNodePtr SamplePtr);
+
+protected:
+	/** Data context for this table row. */
+	FStatsNodePtr StatsNodePtr;
+
+	FShouldBeEnabledDelegate OnShouldBeEnabled;
+	FIsColumnVisibleDelegate IsColumnVisibleDelegate;
+	FSetHoveredStatsTableCell SetHoveredTableCellDelegate;
+	FGetColumnOutlineHAlignmentDelegate GetColumnOutlineHAlignmentDelegate;
+
+	/** Text to be highlighted on stats name. */
+	TAttribute<FText> HighlightText;
+
+	/** Name of the stats that should be drawn as highlighted. */
+	TAttribute<FName> HighlightedStatsName;
+};
