Copying //UE4/Dev-Core to Dev-Main (//UE4/Dev-Main)

diff --git a/Engine/Source/Developer/TraceInsights/Private/Insights/Widgets/SStatsView.cpp b/Engine/Source/Developer/TraceInsights/Private/Insights/Widgets/SStatsView.cpp
new file mode 100644
--- /dev/null
+++ b/Engine/Source/Developer/TraceInsights/Private/Insights/Widgets/SStatsView.cpp
@@ -0,0 +1,1794 @@
+// Copyright 1998-2019 Epic Games, Inc. All Rights Reserved.
+
+#include "SStatsView.h"
+
+#include "EditorStyleSet.h"
+#include "Framework/MultiBox/MultiBoxBuilder.h"
+#include "SlateOptMacros.h"
+#include "TraceServices/AnalysisService.h"
+#include "Widgets/Input/SCheckBox.h"
+#include "Widgets/Layout/SScrollBox.h"
+#include "Widgets/Input/SSearchBox.h"
+#include "Widgets/Images/SImage.h"
+#include "Widgets/Layout/SGridPanel.h"
+#include "Widgets/Layout/SSeparator.h"
+#include "Widgets/SToolTip.h"
+#include "Widgets/Views/STableViewBase.h"
+
+// Insights
+#include "Insights/TimingProfilerCommon.h"
+#include "Insights/TimingProfilerManager.h"
+#include "Insights/ViewModels/StatsNodeHelper.h"
+#include "Insights/ViewModels/StatsViewColumnFactory.h"
+#include "Insights/Widgets/SStatsViewTooltip.h"
+#include "Insights/Widgets/SStatsTableRow.h"
+
+#define LOCTEXT_NAMESPACE "SStatsView"
+
+////////////////////////////////////////////////////////////////////////////////////////////////////
+
+const EColumnSortMode::Type SStatsView::DefaultColumnSortMode(EColumnSortMode::Ascending);
+const FName SStatsView::DefaultColumnBeingSorted(FStatsViewColumns::NameColumnID);
+
+////////////////////////////////////////////////////////////////////////////////////////////////////
+
+SStatsView::SStatsView()
+	: bExpansionSaved(false)
+	, GroupingMode(EStatsGroupingMode::Flat)
+	, ColumnSortMode(DefaultColumnSortMode)
+	, ColumnBeingSorted(DefaultColumnBeingSorted)
+{
+	FMemory::Memset(bStatsNodeIsVisible, 1);
+}
+
+////////////////////////////////////////////////////////////////////////////////////////////////////
+
+SStatsView::~SStatsView()
+{
+	// Remove ourselves from the Insights manager.
+	if (FInsightsManager::Get().IsValid())
+	{
+		FInsightsManager::Get()->GetSessionChangedEvent().RemoveAll(this);
+	}
+}
+
+////////////////////////////////////////////////////////////////////////////////////////////////////
+
+BEGIN_SLATE_FUNCTION_BUILD_OPTIMIZATION
+void SStatsView::Construct(const FArguments& InArgs)
+{
+	SAssignNew(ExternalScrollbar, SScrollBar)
+	.AlwaysShowScrollbar(true);
+
+	ChildSlot
+	[
+		SNew(SVerticalBox)
+
+		+SVerticalBox::Slot()
+		.VAlign(VAlign_Center)
+		.AutoHeight()
+		[
+			SNew(SBorder)
+			.BorderImage(FEditorStyle::GetBrush("ToolPanel.GroupBorder"))
+			.Padding(2.0f)
+			[
+				SNew(SVerticalBox)
+
+				// Search box
+				+SVerticalBox::Slot()
+				.VAlign(VAlign_Center)
+				.Padding(2.0f)
+				.AutoHeight()
+				[
+					SAssignNew(SearchBox, SSearchBox)
+					.HintText(LOCTEXT("SearchBoxHint", "Search timers or groups"))
+					.OnTextChanged(this, &SStatsView::SearchBox_OnTextChanged)
+					.IsEnabled(this, &SStatsView::SearchBox_IsEnabled)
+					.ToolTipText(LOCTEXT("FilterSearchHint", "Type here to search timer or group"))
+				]
+
+				// Group by
+				+SVerticalBox::Slot()
+				.VAlign(VAlign_Center)
+				.Padding(2.0f)
+				.AutoHeight()
+				[
+					SNew(SHorizontalBox)
+
+					+SHorizontalBox::Slot()
+					.FillWidth(1.0f)
+					.VAlign(VAlign_Center)
+					[
+						SNew(STextBlock)
+						.Text(LOCTEXT("GroupByText", "Group by"))
+					]
+
+					+SHorizontalBox::Slot()
+					.FillWidth(2.0f)
+					.VAlign(VAlign_Center)
+					[
+						SAssignNew(GroupByComboBox, SComboBox<TSharedPtr<EStatsGroupingMode>>)
+						.ToolTipText(this, &SStatsView::GroupBy_GetSelectedTooltipText)
+						.OptionsSource(&GroupByOptionsSource)
+						.OnSelectionChanged(this, &SStatsView::GroupBy_OnSelectionChanged)
+						.OnGenerateWidget(this, &SStatsView::GroupBy_OnGenerateWidget)
+						[
+							SNew(STextBlock)
+							.Text(this, &SStatsView::GroupBy_GetSelectedText)
+						]
+					]
+				]
+
+				// Check boxes for: GpuScope, ComputeScope, CpuScope
+				+ SVerticalBox::Slot()
+				.VAlign(VAlign_Center)
+				.Padding(2.0f)
+				.AutoHeight()
+				[
+					SNew(SHorizontalBox)
+
+					+SHorizontalBox::Slot()
+					.Padding(FMargin(0.0f,0.0f,1.0f,0.0f))
+					.FillWidth(1.0f)
+					[
+						GetToggleButtonForStatsType(EStatsNodeType::Int64)
+					]
+
+					+SHorizontalBox::Slot()
+					.Padding(FMargin(1.0f,0.0f,1.0f,0.0f))
+					.FillWidth(1.0f)
+					[
+						GetToggleButtonForStatsType(EStatsNodeType::Float)
+					]
+				]
+			]
+		]
+
+		// Tree view
+		+SVerticalBox::Slot()
+		.FillHeight(1.0f)
+		.Padding(0.0f, 6.0f, 0.0f, 0.0f)
+		[
+			SNew(SHorizontalBox)
+
+			+ SHorizontalBox::Slot()
+			.FillWidth(1.0f)
+			.Padding(0.0f)
+			[
+				SNew(SScrollBox)
+				.Orientation(Orient_Horizontal)
+
+				+ SScrollBox::Slot()
+				[
+					SNew(SBorder)
+					.BorderImage(FEditorStyle::GetBrush("ToolPanel.GroupBorder"))
+					.Padding(0.0f)
+					[
+						SAssignNew(TreeView, STreeView<FStatsNodePtr>)
+						.ExternalScrollbar(ExternalScrollbar)
+						.SelectionMode(ESelectionMode::Multi)
+						.TreeItemsSource(&FilteredGroupNodes)
+						.OnGetChildren(this, &SStatsView::TreeView_OnGetChildren)
+						.OnGenerateRow(this, &SStatsView::TreeView_OnGenerateRow)
+						.OnSelectionChanged(this, &SStatsView::TreeView_OnSelectionChanged)
+						.OnMouseButtonDoubleClick(this, &SStatsView::TreeView_OnMouseButtonDoubleClick)
+						.OnContextMenuOpening(FOnContextMenuOpening::CreateSP(this, &SStatsView::TreeView_GetMenuContent))
+						.ItemHeight(12.0f)
+						.HeaderRow
+						(
+							SAssignNew(TreeViewHeaderRow, SHeaderRow)
+							.Visibility(EVisibility::Visible)
+						)
+					]
+				]
+			]
+
+			+ SHorizontalBox::Slot()
+			.AutoWidth()
+			.Padding(0.0f)
+			[
+				SNew(SBox)
+				.WidthOverride(FOptionalSize(13.0f))
+				[
+					ExternalScrollbar.ToSharedRef()
+				]
+			]
+		]
+	];
+
+	InitializeAndShowHeaderColumns();
+	//BindCommands();
+
+	// Create the search filters: text based, type based etc.
+	TextFilter = MakeShareable(new FStatsNodeTextFilter(FStatsNodeTextFilter::FItemToStringArray::CreateSP(this, &SStatsView::HandleItemToStringArray)));
+	Filters = MakeShareable(new FStatsNodeFilterCollection());
+	Filters->Add(TextFilter);
+
+	CreateGroupByOptionsSources();
+
+	// Register ourselves with the Insights manager.
+	FInsightsManager::Get()->GetSessionChangedEvent().AddSP(this, &SStatsView::InsightsManager_OnSessionChanged);
+}
+END_SLATE_FUNCTION_BUILD_OPTIMIZATION
+
+////////////////////////////////////////////////////////////////////////////////////////////////////
+
+TSharedPtr<SWidget> SStatsView::TreeView_GetMenuContent()
+{
+	const TArray<FStatsNodePtr> SelectedStatsNodes = TreeView->GetSelectedItems();
+	const int32 NumSelectedStatsNodes = SelectedStatsNodes.Num();
+	FStatsNodePtr SelectedStatsNode = NumSelectedStatsNodes ? SelectedStatsNodes[0] : nullptr;
+
+	const FStatsViewColumn* const * ColumnPtrPtr = FStatsViewColumnFactory::Get().ColumnIdToPtrMapping.Find(HoveredColumnId);
+	const FStatsViewColumn* const ColumnPtr = (ColumnPtrPtr != nullptr) ? *ColumnPtrPtr : nullptr;
+
+	FText SelectionStr;
+	FText PropertyName;
+	FText PropertyValue;
+
+	if (NumSelectedStatsNodes == 0)
+	{
+		SelectionStr = LOCTEXT("NothingSelected", "Nothing selected");
+	}
+	else if (NumSelectedStatsNodes == 1)
+	{
+		if (ColumnPtr != nullptr)
+		{
+			PropertyName = ColumnPtr->ShortName;
+			PropertyValue = ColumnPtr->GetFormattedValue(*SelectedStatsNode);
+		}
+		SelectionStr = FText::FromName(SelectedStatsNode->GetName());
+	}
+	else
+	{
+		SelectionStr = LOCTEXT("MultipleSelection", "Multiple selection");
+	}
+
+	const bool bShouldCloseWindowAfterMenuSelection = true;
+	FMenuBuilder MenuBuilder(bShouldCloseWindowAfterMenuSelection, NULL);
+
+	// Selection menu
+	MenuBuilder.BeginSection("Selection", LOCTEXT("ContextMenu_Header_Selection", "Selection"));
+	{
+		struct FLocal
+		{
+			static bool ReturnFalse()
+			{
+				return false;
+			}
+		};
+
+		FUIAction DummyUIAction;
+		DummyUIAction.CanExecuteAction = FCanExecuteAction::CreateStatic(&FLocal::ReturnFalse);
+		MenuBuilder.AddMenuEntry
+		(
+			SelectionStr,
+			LOCTEXT("ContextMenu_Selection", "Currently selected items"),
+			FSlateIcon(FEditorStyle::GetStyleSetName(), "@missing.icon"), DummyUIAction, NAME_None, EUserInterfaceActionType::Button
+		);
+	}
+	MenuBuilder.EndSection();
+
+	MenuBuilder.BeginSection("Misc", LOCTEXT("ContextMenu_Header_Misc", "Miscellaneous"));
+	{
+		/*TODO
+		FUIAction Action_CopySelectedToClipboard
+		(
+			FExecuteAction::CreateSP(this, &SStatsView::ContextMenu_CopySelectedToClipboard_Execute),
+			FCanExecuteAction::CreateSP(this, &SStatsView::ContextMenu_CopySelectedToClipboard_CanExecute)
+		);
+		MenuBuilder.AddMenuEntry
+		(
+			LOCTEXT("ContextMenu_Header_Misc_CopySelectedToClipboard", "Copy To Clipboard"),
+			LOCTEXT("ContextMenu_Header_Misc_CopySelectedToClipboard_Desc", "Copies selection to clipboard"),
+			FSlateIcon(FEditorStyle::GetStyleSetName(), "Profiler.Misc.CopyToClipboard"), Action_CopySelectedToClipboard, NAME_None, EUserInterfaceActionType::Button
+		);
+		*/
+
+		MenuBuilder.AddSubMenu
+		(
+			LOCTEXT("ContextMenu_Header_Misc_Sort", "Sort By"),
+			LOCTEXT("ContextMenu_Header_Misc_Sort_Desc", "Sort by column"),
+			FNewMenuDelegate::CreateSP(this, &SStatsView::TreeView_BuildSortByMenu),
+			false,
+			FSlateIcon(FEditorStyle::GetStyleSetName(), "Profiler.Misc.SortBy")
+		);
+	}
+	MenuBuilder.EndSection();
+
+	MenuBuilder.BeginSection("Columns", LOCTEXT("ContextMenu_Header_Columns", "Columns"));
+	{
+		MenuBuilder.AddSubMenu
+		(
+			LOCTEXT("ContextMenu_Header_Columns_View", "View Column"),
+			LOCTEXT("ContextMenu_Header_Columns_View_Desc", "Hides or shows columns"),
+			FNewMenuDelegate::CreateSP(this, &SStatsView::TreeView_BuildViewColumnMenu),
+			false,
+			FSlateIcon(FEditorStyle::GetStyleSetName(), "Profiler.EventGraph.ViewColumn")
+		);
+
+		FUIAction Action_ShowAllColumns
+		(
+			FExecuteAction::CreateSP(this, &SStatsView::ContextMenu_ShowAllColumns_Execute),
+			FCanExecuteAction::CreateSP(this, &SStatsView::ContextMenu_ShowAllColumns_CanExecute)
+		);
+		MenuBuilder.AddMenuEntry
+		(
+			LOCTEXT("ContextMenu_Header_Columns_ShowAllColumns", "Show All Columns"),
+			LOCTEXT("ContextMenu_Header_Columns_ShowAllColumns_Desc", "Resets tree view to show all columns"),
+			FSlateIcon(FEditorStyle::GetStyleSetName(), "Profiler.EventGraph.ResetColumn"), Action_ShowAllColumns, NAME_None, EUserInterfaceActionType::Button
+		);
+
+		FUIAction Action_ShowMinMaxMedColumns
+		(
+			FExecuteAction::CreateSP(this, &SStatsView::ContextMenu_ShowMinMaxMedColumns_Execute),
+			FCanExecuteAction::CreateSP(this, &SStatsView::ContextMenu_ShowMinMaxMedColumns_CanExecute)
+		);
+		MenuBuilder.AddMenuEntry
+		(
+			LOCTEXT("ContextMenu_Header_Columns_ShowMinMaxMedColumns", "Reset Columns to Min/Max/Median Preset"),
+			LOCTEXT("ContextMenu_Header_Columns_ShowMinMaxMedColumns_Desc", "Resets columns to Min/Max/Median preset"),
+			FSlateIcon(FEditorStyle::GetStyleSetName(), "Profiler.EventGraph.ResetColumn"), Action_ShowMinMaxMedColumns, NAME_None, EUserInterfaceActionType::Button
+		);
+
+		FUIAction Action_ResetColumns
+		(
+			FExecuteAction::CreateSP(this, &SStatsView::ContextMenu_ResetColumns_Execute),
+			FCanExecuteAction::CreateSP(this, &SStatsView::ContextMenu_ResetColumns_CanExecute)
+		);
+		MenuBuilder.AddMenuEntry
+		(
+			LOCTEXT("ContextMenu_Header_Columns_ResetColumns", "Reset Columns to Default"),
+			LOCTEXT("ContextMenu_Header_Columns_ResetColumns_Desc", "Resets columns to default"),
+			FSlateIcon(FEditorStyle::GetStyleSetName(), "Profiler.EventGraph.ResetColumn"), Action_ResetColumns, NAME_None, EUserInterfaceActionType::Button
+		);
+	}
+	MenuBuilder.EndSection();
+
+	return MenuBuilder.MakeWidget();
+}
+
+////////////////////////////////////////////////////////////////////////////////////////////////////
+
+void SStatsView::TreeView_BuildSortByMenu(FMenuBuilder& MenuBuilder)
+{
+	// TODO: Refactor later @see TSharedPtr<SWidget> SCascadePreviewViewportToolBar::GenerateViewMenu() const
+
+	MenuBuilder.BeginSection("ColumnName", LOCTEXT("ContextMenu_Header_Misc_ColumnName", "Column Name"));
+
+	for (auto It = TreeViewHeaderColumns.CreateConstIterator(); It; ++It)
+	{
+		const FStatsViewColumn& Column = It.Value();
+
+		if (Column.bIsVisible && Column.bCanBeSorted())
+		{
+			FUIAction Action_SortByColumn
+			(
+				FExecuteAction::CreateSP(this, &SStatsView::ContextMenu_SortByColumn_Execute, Column.Id),
+				FCanExecuteAction::CreateSP(this, &SStatsView::ContextMenu_SortByColumn_CanExecute, Column.Id),
+				FIsActionChecked::CreateSP(this, &SStatsView::ContextMenu_SortByColumn_IsChecked, Column.Id)
+			);
+			MenuBuilder.AddMenuEntry
+			(
+				Column.TitleName,
+				Column.Description,
+				FSlateIcon(), Action_SortByColumn, NAME_None, EUserInterfaceActionType::RadioButton
+			);
+		}
+	}
+
+	MenuBuilder.EndSection();
+
+	//-----------------------------------------------------------------------------
+
+	MenuBuilder.BeginSection("SortMode", LOCTEXT("ContextMenu_Header_Misc_Sort_SortMode", "Sort Mode"));
+	{
+		FUIAction Action_SortAscending
+		(
+			FExecuteAction::CreateSP(this, &SStatsView::ContextMenu_SortMode_Execute, EColumnSortMode::Ascending),
+			FCanExecuteAction::CreateSP(this, &SStatsView::ContextMenu_SortMode_CanExecute, EColumnSortMode::Ascending),
+			FIsActionChecked::CreateSP(this, &SStatsView::ContextMenu_SortMode_IsChecked, EColumnSortMode::Ascending)
+		);
+		MenuBuilder.AddMenuEntry
+		(
+			LOCTEXT("ContextMenu_Header_Misc_Sort_SortAscending", "Sort Ascending"),
+			LOCTEXT("ContextMenu_Header_Misc_Sort_SortAscending_Desc", "Sorts ascending"),
+			FSlateIcon(FEditorStyle::GetStyleSetName(), "Profiler.Misc.SortAscending"), Action_SortAscending, NAME_None, EUserInterfaceActionType::RadioButton
+		);
+
+		FUIAction Action_SortDescending
+		(
+			FExecuteAction::CreateSP(this, &SStatsView::ContextMenu_SortMode_Execute, EColumnSortMode::Descending),
+			FCanExecuteAction::CreateSP(this, &SStatsView::ContextMenu_SortMode_CanExecute, EColumnSortMode::Descending),
+			FIsActionChecked::CreateSP(this, &SStatsView::ContextMenu_SortMode_IsChecked, EColumnSortMode::Descending)
+		);
+		MenuBuilder.AddMenuEntry
+		(
+			LOCTEXT("ContextMenu_Header_Misc_Sort_SortDescending", "Sort Descending"),
+			LOCTEXT("ContextMenu_Header_Misc_Sort_SortDescending_Desc", "Sorts descending"),
+			FSlateIcon(FEditorStyle::GetStyleSetName(), "Profiler.Misc.SortDescending"), Action_SortDescending, NAME_None, EUserInterfaceActionType::RadioButton
+		);
+	}
+	MenuBuilder.EndSection();
+}
+
+////////////////////////////////////////////////////////////////////////////////////////////////////
+
+void SStatsView::TreeView_BuildViewColumnMenu(FMenuBuilder& MenuBuilder)
+{
+	MenuBuilder.BeginSection("ViewColumn", LOCTEXT("ContextMenu_Header_Columns_View", "View Column"));
+
+	for (auto It = TreeViewHeaderColumns.CreateConstIterator(); It; ++It)
+	{
+		const FStatsViewColumn& Column = It.Value();
+
+		FUIAction Action_ToggleColumn
+		(
+			FExecuteAction::CreateSP(this, &SStatsView::ContextMenu_ToggleColumn_Execute, Column.Id),
+			FCanExecuteAction::CreateSP(this, &SStatsView::ContextMenu_ToggleColumn_CanExecute, Column.Id),
+			FIsActionChecked::CreateSP(this, &SStatsView::ContextMenu_ToggleColumn_IsChecked, Column.Id)
+		);
+		MenuBuilder.AddMenuEntry
+		(
+			Column.TitleName,
+			Column.Description,
+			FSlateIcon(), Action_ToggleColumn, NAME_None, EUserInterfaceActionType::ToggleButton
+		);
+	}
+
+	MenuBuilder.EndSection();
+}
+
+////////////////////////////////////////////////////////////////////////////////////////////////////
+
+void SStatsView::InitializeAndShowHeaderColumns()
+{
+	const int32 NumColumns = FStatsViewColumnFactory::Get().Collection.Num();
+	for (int32 ColumnIndex = 0; ColumnIndex < NumColumns; ColumnIndex++)
+	{
+		TreeViewHeaderRow_CreateColumnArgs(ColumnIndex);
+	}
+
+	for (auto It = TreeViewHeaderColumns.CreateConstIterator(); It; ++It)
+	{
+		const FStatsViewColumn& Column = It.Value();
+
+		if (Column.bIsVisible)
+		{
+			TreeViewHeaderRow_ShowColumn(Column.Id);
+		}
+	}
+}
+
+////////////////////////////////////////////////////////////////////////////////////////////////////
+
+void SStatsView::TreeViewHeaderRow_CreateColumnArgs(const int32 ColumnIndex)
+{
+	const FStatsViewColumn& Column = *FStatsViewColumnFactory::Get().Collection[ColumnIndex];
+	SHeaderRow::FColumn::FArguments ColumnArgs;
+
+	ColumnArgs
+		.ColumnId(Column.Id)
+		.DefaultLabel(Column.ShortName)
+		.SortMode(EColumnSortMode::None)
+		.HAlignHeader(HAlign_Fill)
+		.VAlignHeader(VAlign_Fill)
+		.HeaderContentPadding(TOptional<FMargin>(2.0f))
+		.HAlignCell(HAlign_Fill)
+		.VAlignCell(VAlign_Fill)
+		.SortMode(this, &SStatsView::GetSortModeForColumn, Column.Id)
+		.OnSort(this, &SStatsView::OnSortModeChanged)
+		.ManualWidth(Column.InitialColumnWidth)
+		.FixedWidth(Column.bIsFixedColumnWidth() ? Column.InitialColumnWidth : TOptional<float>())
+		.HeaderContent()
+		[
+			SNew(SHorizontalBox)
+			.ToolTip(SStatsViewTooltip::GetColumnTooltip(Column))
+
+			+ SHorizontalBox::Slot()
+			.HAlign(Column.HorizontalAlignment)
+			.VAlign(VAlign_Center)
+			[
+				SNew(STextBlock)
+				.Text(Column.ShortName)
+			]
+		]
+		.MenuContent()
+		[
+			TreeViewHeaderRow_GenerateColumnMenu(Column)
+		];
+
+	TreeViewHeaderColumnArgs.Add(Column.Id, ColumnArgs);
+	TreeViewHeaderColumns.Add(Column.Id, Column);
+}
+
+////////////////////////////////////////////////////////////////////////////////////////////////////
+
+void SStatsView::TreeViewHeaderRow_ShowColumn(const FName ColumnId)
+{
+	FStatsViewColumn& Column = TreeViewHeaderColumns.FindChecked(ColumnId);
+	Column.bIsVisible = true;
+	SHeaderRow::FColumn::FArguments& ColumnArgs = TreeViewHeaderColumnArgs.FindChecked(ColumnId);
+
+	const int32 NumColumns = TreeViewHeaderRow->GetColumns().Num();
+	int32 ColumnIndex = 0;
+	for (; ColumnIndex < NumColumns; ColumnIndex++)
+	{
+		const SHeaderRow::FColumn& CurrentColumn = TreeViewHeaderRow->GetColumns()[ColumnIndex];
+		const FStatsViewColumn& CurrentStatsViewColumn = TreeViewHeaderColumns.FindChecked(CurrentColumn.ColumnId);
+		if (Column.Order < CurrentStatsViewColumn.Order)
+			break;
+	}
+
+	TreeViewHeaderRow->InsertColumn(ColumnArgs, ColumnIndex);
+}
+
+////////////////////////////////////////////////////////////////////////////////////////////////////
+
+TSharedRef<SWidget> SStatsView::TreeViewHeaderRow_GenerateColumnMenu(const FStatsViewColumn& Column)
+{
+	bool bIsMenuVisible = false;
+
+	const bool bShouldCloseWindowAfterMenuSelection = true;
+	FMenuBuilder MenuBuilder(bShouldCloseWindowAfterMenuSelection, NULL);
+	{
+		if (Column.bCanBeHidden())
+		{
+			MenuBuilder.BeginSection("Column", LOCTEXT("TreeViewHeaderRow_Header_Column", "Column"));
+
+			FUIAction Action_HideColumn
+			(
+				FExecuteAction::CreateSP(this, &SStatsView::HeaderMenu_HideColumn_Execute, Column.Id),
+				FCanExecuteAction::CreateSP(this, &SStatsView::HeaderMenu_HideColumn_CanExecute, Column.Id)
+			);
+
+			MenuBuilder.AddMenuEntry
+			(
+				LOCTEXT("TreeViewHeaderRow_HideColumn", "Hide"),
+				LOCTEXT("TreeViewHeaderRow_HideColumn_Desc", "Hides the selected column"),
+				FSlateIcon(), Action_HideColumn, NAME_None, EUserInterfaceActionType::Button
+			);
+			bIsMenuVisible = true;
+
+			MenuBuilder.EndSection();
+		}
+
+		if (Column.bCanBeSorted())
+		{
+			MenuBuilder.BeginSection("SortMode", LOCTEXT("ContextMenu_Header_Misc_Sort_SortMode", "Sort Mode"));
+
+			FUIAction Action_SortAscending
+			(
+				FExecuteAction::CreateSP(this, &SStatsView::HeaderMenu_SortMode_Execute, Column.Id, EColumnSortMode::Ascending),
+				FCanExecuteAction::CreateSP(this, &SStatsView::HeaderMenu_SortMode_CanExecute, Column.Id, EColumnSortMode::Ascending),
+				FIsActionChecked::CreateSP(this, &SStatsView::HeaderMenu_SortMode_IsChecked, Column.Id, EColumnSortMode::Ascending)
+			);
+			MenuBuilder.AddMenuEntry
+			(
+				LOCTEXT("ContextMenu_Header_Misc_Sort_SortAscending", "Sort Ascending"),
+				LOCTEXT("ContextMenu_Header_Misc_Sort_SortAscending_Desc", "Sorts ascending"),
+				FSlateIcon(FEditorStyle::GetStyleSetName(), "Profiler.Misc.SortAscending"), Action_SortAscending, NAME_None, EUserInterfaceActionType::RadioButton
+			);
+
+			FUIAction Action_SortDescending
+			(
+				FExecuteAction::CreateSP(this, &SStatsView::HeaderMenu_SortMode_Execute, Column.Id, EColumnSortMode::Descending),
+				FCanExecuteAction::CreateSP(this, &SStatsView::HeaderMenu_SortMode_CanExecute, Column.Id, EColumnSortMode::Descending),
+				FIsActionChecked::CreateSP(this, &SStatsView::HeaderMenu_SortMode_IsChecked, Column.Id, EColumnSortMode::Descending)
+			);
+			MenuBuilder.AddMenuEntry
+			(
+				LOCTEXT("ContextMenu_Header_Misc_Sort_SortDescending", "Sort Descending"),
+				LOCTEXT("ContextMenu_Header_Misc_Sort_SortDescending_Desc", "Sorts descending"),
+				FSlateIcon(FEditorStyle::GetStyleSetName(), "Profiler.Misc.SortDescending"), Action_SortDescending, NAME_None, EUserInterfaceActionType::RadioButton
+			);
+			bIsMenuVisible = true;
+
+			MenuBuilder.EndSection();
+		}
+
+		if (Column.bCanBeFiltered())
+		{
+			MenuBuilder.BeginSection("FilterMode", LOCTEXT("ContextMenu_Header_Misc_Filter_FilterMode", "Filter Mode"));
+			bIsMenuVisible = true;
+			MenuBuilder.EndSection();
+		}
+	}
+
+	/*
+	TODO:
+	- Show top ten
+	- Show top bottom
+	- Filter by list (avg, median, 10%, 90%, etc.)
+	- Text box for filtering for each column instead of one text box used for filtering
+	- Grouping button for flat view modes (show at most X groups, show all groups for names)
+	*/
+
+	return bIsMenuVisible ? MenuBuilder.MakeWidget() : (TSharedRef<SWidget>)SNullWidget::NullWidget;
+}
+
+////////////////////////////////////////////////////////////////////////////////////////////////////
+
+void SStatsView::InsightsManager_OnSessionChanged()
+{
+	TSharedPtr<const Trace::IAnalysisSession> NewSession = FInsightsManager::Get()->GetSession();
+
+	if (NewSession != Session)
+	{
+		Session = NewSession;
+		RebuildTree();
+	}
+	else
+	{
+		UpdateTree();
+	}
+}
+
+////////////////////////////////////////////////////////////////////////////////////////////////////
+
+void SStatsView::UpdateTree()
+{
+	// Create groups, sort timers within the group and apply filtering.
+	CreateGroups();
+	SortStats();
+	ApplyFiltering();
+}
+
+////////////////////////////////////////////////////////////////////////////////////////////////////
+
+void SStatsView::ApplyFiltering()
+{
+	FilteredGroupNodes.Reset();
+
+	// Apply filter to all groups and its children.
+	const int32 NumGroups = GroupNodes.Num();
+	for (int32 GroupIndex = 0; GroupIndex < NumGroups; ++GroupIndex)
+	{
+		FStatsNodePtr& GroupPtr = GroupNodes[GroupIndex];
+		GroupPtr->ClearFilteredChildren();
+		const bool bIsGroupVisible = Filters->PassesAllFilters(GroupPtr);
+
+		const TArray<FStatsNodePtr>& GroupChildren = GroupPtr->GetChildren();
+		const int32 NumChildren = GroupChildren.Num();
+		int32 NumVisibleChildren = 0;
+		for (int32 Cx = 0; Cx < NumChildren; ++Cx)
+		{
+			// Add a child.
+			const FStatsNodePtr& NodePtr = GroupChildren[Cx];
+			const bool bIsChildVisible = Filters->PassesAllFilters(NodePtr) && bStatsNodeIsVisible[static_cast<int>(NodePtr->GetType())];
+			if (bIsChildVisible)
+			{
+				GroupPtr->AddFilteredChild(NodePtr);
+				NumVisibleChildren++;
+			}
+		}
+
+		if (bIsGroupVisible || NumVisibleChildren > 0)
+		{
+			// Add a group.
+			FilteredGroupNodes.Add(GroupPtr);
+			GroupPtr->bForceExpandGroupNode = true;
+		}
+		else
+		{
+			GroupPtr->bForceExpandGroupNode = false;
+		}
+	}
+
+	// Only expand timer nodes if we have a text filter.
+	const bool bNonEmptyTextFilter = !TextFilter->GetRawFilterText().IsEmpty();
+	if (bNonEmptyTextFilter)
+	{
+		if (!bExpansionSaved)
+		{
+			ExpandedNodes.Empty();
+			TreeView->GetExpandedItems(ExpandedNodes);
+			bExpansionSaved = true;
+		}
+
+		for (int32 Fx = 0; Fx < FilteredGroupNodes.Num(); Fx++)
+		{
+			const FStatsNodePtr& GroupPtr = FilteredGroupNodes[Fx];
+			TreeView->SetItemExpansion(GroupPtr, GroupPtr->bForceExpandGroupNode);
+		}
+	}
+	else
+	{
+		if (bExpansionSaved)
+		{
+			// Restore previously expanded nodes when the text filter is disabled.
+			TreeView->ClearExpandedItems();
+			for (auto It = ExpandedNodes.CreateConstIterator(); It; ++It)
+			{
+				TreeView->SetItemExpansion(*It, true);
+			}
+			bExpansionSaved = false;
+		}
+	}
+
+	// Request tree refresh
+	TreeView->RequestTreeRefresh();
+}
+
+////////////////////////////////////////////////////////////////////////////////////////////////////
+
+void SStatsView::HandleItemToStringArray(const FStatsNodePtr& FStatsNodePtr, TArray<FString>& OutSearchStrings) const
+{
+	OutSearchStrings.Add(FStatsNodePtr->GetName().GetPlainNameString());
+}
+
+////////////////////////////////////////////////////////////////////////////////////////////////////
+
+TSharedRef<SWidget> SStatsView::GetToggleButtonForStatsType(const EStatsNodeType NodeType)
+{
+	return SNew(SCheckBox)
+		.Style(FEditorStyle::Get(), "ToggleButtonCheckbox")
+		.HAlign(HAlign_Center)
+		.Padding(2.0f)
+		.OnCheckStateChanged(this, &SStatsView::FilterByStatsType_OnCheckStateChanged, NodeType)
+		.IsChecked(this, &SStatsView::FilterByStatsType_IsChecked, NodeType)
+		.ToolTipText(StatsNodeTypeHelper::ToDescription(NodeType))
+		[
+			SNew(SHorizontalBox)
+
+			+ SHorizontalBox::Slot()
+				.AutoWidth()
+				.VAlign(VAlign_Center)
+				[
+					SNew(SImage)
+						.Image(StatsNodeTypeHelper::GetIconForStatsNodeType(NodeType))
+				]
+
+			+SHorizontalBox::Slot()
+				.Padding(2.0f, 0.0f, 0.0f, 0.0f)
+				.VAlign(VAlign_Center)
+				[
+					SNew(STextBlock)
+						.Text(StatsNodeTypeHelper::ToName(NodeType))
+						.TextStyle(FEditorStyle::Get(), TEXT("Profiler.Caption"))
+				]
+		];
+}
+
+////////////////////////////////////////////////////////////////////////////////////////////////////
+
+void SStatsView::FilterByStatsType_OnCheckStateChanged(ECheckBoxState NewRadioState, const EStatsNodeType InStatType)
+{
+	bStatsNodeIsVisible[static_cast<int>(InStatType)] = (NewRadioState == ECheckBoxState::Checked);
+	ApplyFiltering();
+}
+
+////////////////////////////////////////////////////////////////////////////////////////////////////
+
+ECheckBoxState SStatsView::FilterByStatsType_IsChecked(const EStatsNodeType InStatType) const
+{
+	return bStatsNodeIsVisible[static_cast<int>(InStatType)] ? ECheckBoxState::Checked : ECheckBoxState::Unchecked;
+}
+
+////////////////////////////////////////////////////////////////////////////////////////////////////
+// TreeView
+////////////////////////////////////////////////////////////////////////////////////////////////////
+
+void SStatsView::TreeView_Refresh()
+{
+	if (TreeView.IsValid())
+	{
+		TreeView->RequestTreeRefresh();
+	}
+}
+
+////////////////////////////////////////////////////////////////////////////////////////////////////
+
+void SStatsView::TreeView_OnSelectionChanged(FStatsNodePtr SelectedItem, ESelectInfo::Type SelectInfo)
+{
+	if (SelectInfo != ESelectInfo::Direct)
+	{
+		TArray<FStatsNodePtr> SelectedItems = TreeView->GetSelectedItems();
+		if (SelectedItems.Num() == 1)
+		{
+			//HighlightedStatsName = SelectedItems[0]->GetName();
+		}
+		else
+		{
+			//HighlightedStatsName = NAME_None;
+		}
+	}
+}
+
+////////////////////////////////////////////////////////////////////////////////////////////////////
+
+void SStatsView::TreeView_OnGetChildren(FStatsNodePtr InParent, TArray<FStatsNodePtr>& OutChildren)
+{
+	OutChildren = InParent->GetFilteredChildren();
+}
+
+////////////////////////////////////////////////////////////////////////////////////////////////////
+
+void SStatsView::TreeView_OnMouseButtonDoubleClick(FStatsNodePtr StatsNode)
+{
+	if (!StatsNode->IsGroup())
+	{
+		//im:TODO: const bool bIsStatTracked = FTimingProfilerManager::Get()->IsStatTracked(StatsNode->GetId());
+	//	if (!bIsStatTracked)
+	//	{
+	//		// Add a new graph.
+	//		FTimingProfilerManager::Get()->TrackStat(StatsNode->GetId());
+	//	}
+	//	else
+	//	{
+	//		// Remove a graph
+	//		FTimingProfilerManager::Get()->UntrackStat(StatsNode->GetId());
+	//	}
+	}
+	else
+	{
+		const bool bIsGroupExpanded = TreeView->IsItemExpanded(StatsNode);
+		TreeView->SetItemExpansion(StatsNode, !bIsGroupExpanded);
+	}
+}
+
+////////////////////////////////////////////////////////////////////////////////////////////////////
+// Tree View's Table Row
+////////////////////////////////////////////////////////////////////////////////////////////////////
+
+TSharedRef<ITableRow> SStatsView::TreeView_OnGenerateRow(FStatsNodePtr StatsNodePtr, const TSharedRef<STableViewBase>& OwnerTable)
+{
+	TSharedRef<ITableRow> TableRow =
+		SNew(SStatsTableRow, OwnerTable)
+		.OnShouldBeEnabled(this, &SStatsView::TableRow_ShouldBeEnabled)
+		.OnIsColumnVisible(this, &SStatsView::TableRow_IsColumnVisible)
+		.OnSetHoveredTableCell(this, &SStatsView::TableRow_SetHoveredTableCell)
+		.OnGetColumnOutlineHAlignmentDelegate(this, &SStatsView::TableRow_GetColumnOutlineHAlignment)
+		.HighlightText(this, &SStatsView::TableRow_GetHighlightText)
+		.HighlightedStatsName(this, &SStatsView::TableRow_GetHighlightedStatsName)
+		.StatsNodePtr(StatsNodePtr);
+
+	return TableRow;
+}
+
+////////////////////////////////////////////////////////////////////////////////////////////////////
+
+bool SStatsView::TableRow_IsColumnVisible(const FName ColumnId) const
+{
+	bool bResult = false;
+	const FStatsViewColumn& ColumnPtr = TreeViewHeaderColumns.FindChecked(ColumnId);
+	return ColumnPtr.bIsVisible;
+}
+
+////////////////////////////////////////////////////////////////////////////////////////////////////
+
+void SStatsView::TableRow_SetHoveredTableCell(const FName ColumnId, const FStatsNodePtr StatsNodePtr)
+{
+	HoveredColumnId = ColumnId;
+
+	const bool bIsAnyMenusVisible = FSlateApplication::Get().AnyMenusVisible();
+	if (!HasMouseCapture() && !bIsAnyMenusVisible)
+	{
+		HoveredStatsNodePtr = StatsNodePtr;
+	}
+
+	//UE_LOG(TimingProfiler, Log, TEXT("%s -> %s"), *HoveredColumnId.GetPlainNameString(), StatsNodePtr.IsValid() ? *StatsNodePtr->GetName().GetPlainNameString() : TEXT("nullptr"));
+}
+
+////////////////////////////////////////////////////////////////////////////////////////////////////
+
+EHorizontalAlignment SStatsView::TableRow_GetColumnOutlineHAlignment(const FName ColumnId) const
+{
+	const TIndirectArray<SHeaderRow::FColumn>& Columns = TreeViewHeaderRow->GetColumns();
+	const int32 LastColumnIdx = Columns.Num() - 1;
+
+	// First column
+	if (Columns[0].ColumnId == ColumnId)
+	{
+		return HAlign_Left;
+	}
+	// Last column
+	else if (Columns[LastColumnIdx].ColumnId == ColumnId)
+	{
+		return HAlign_Right;
+	}
+	// Middle columns
+	{
+		return HAlign_Center;
+	}
+}
+
+////////////////////////////////////////////////////////////////////////////////////////////////////
+
+FText SStatsView::TableRow_GetHighlightText() const
+{
+	return SearchBox->GetText();
+}
+
+////////////////////////////////////////////////////////////////////////////////////////////////////
+
+FName SStatsView::TableRow_GetHighlightedStatsName() const
+{
+	return HighlightedStatsName;
+}
+
+////////////////////////////////////////////////////////////////////////////////////////////////////
+
+bool SStatsView::TableRow_ShouldBeEnabled(const uint32 StatsId) const
+{
+	return true;//im:TODO: Session->GetAggregatedStat(StatsId) != nullptr;
+}
+
+////////////////////////////////////////////////////////////////////////////////////////////////////
+// SearchBox
+////////////////////////////////////////////////////////////////////////////////////////////////////
+
+void SStatsView::SearchBox_OnTextChanged(const FText& InFilterText)
+{
+	TextFilter->SetRawFilterText(InFilterText);
+	SearchBox->SetError(TextFilter->GetFilterErrorText());
+	ApplyFiltering();
+}
+
+////////////////////////////////////////////////////////////////////////////////////////////////////
+
+bool SStatsView::SearchBox_IsEnabled() const
+{
+	return StatsNodes.Num() > 0;
+}
+
+////////////////////////////////////////////////////////////////////////////////////////////////////
+// GroupBy
+////////////////////////////////////////////////////////////////////////////////////////////////////
+
+void SStatsView::CreateGroups()
+{
+	TMap<FName, FStatsNodePtr> GroupNodeSet;
+
+	if (GroupingMode == EStatsGroupingMode::Flat)
+	{
+		const FName GroupName(TEXT("All"));
+		FStatsNodePtr* GroupPtr = GroupNodeSet.Find(GroupName);
+		if (!GroupPtr)
+		{
+			GroupPtr = &GroupNodeSet.Add(GroupName, MakeShareable(new FStatsNode(GroupName)));
+		}
+
+		for (const FStatsNodePtr& StatsNodePtr : StatsNodes)
+		{
+			(*GroupPtr)->AddChildAndSetGroupPtr(StatsNodePtr);
+		}
+
+		TreeView->SetItemExpansion(*GroupPtr, true);
+	}
+	// Creates groups based on stat metadata groups.
+	else if (GroupingMode == EStatsGroupingMode::ByMetaGroupName)
+	{
+		for (const FStatsNodePtr& StatsNodePtr : StatsNodes)
+		{
+			const FName GroupName = StatsNodePtr->GetMetaGroupName();
+
+			FStatsNodePtr* GroupPtr = GroupNodeSet.Find(GroupName);
+			if (!GroupPtr)
+			{
+				GroupPtr = &GroupNodeSet.Add(GroupName, MakeShareable(new FStatsNode(GroupName)));
+			}
+
+			(*GroupPtr)->AddChildAndSetGroupPtr(StatsNodePtr);
+			TreeView->SetItemExpansion(*GroupPtr, true);
+		}
+	}
+	// Creates one group for each stat type.
+	else if (GroupingMode == EStatsGroupingMode::ByType)
+	{
+		for (const FStatsNodePtr& StatsNodePtr : StatsNodes)
+		{
+			const FName GroupName = *StatsNodeTypeHelper::ToName(StatsNodePtr->GetType()).ToString();
+
+			FStatsNodePtr* GroupPtr = GroupNodeSet.Find(GroupName);
+			if (!GroupPtr)
+			{
+				GroupPtr = &GroupNodeSet.Add(GroupName, MakeShareable(new FStatsNode(GroupName)));
+			}
+
+			(*GroupPtr)->AddChildAndSetGroupPtr(StatsNodePtr);
+			TreeView->SetItemExpansion(*GroupPtr, true);
+		}
+	}
+	// Creates one group for one letter.
+	else if (GroupingMode == EStatsGroupingMode::ByName)
+	{
+		for (const FStatsNodePtr& StatsNodePtr : StatsNodes)
+		{
+			const FName GroupName = *StatsNodePtr->GetName().GetPlainNameString().Left(1);
+
+			FStatsNodePtr* GroupPtr = GroupNodeSet.Find(GroupName);
+			if (!GroupPtr)
+			{
+				GroupPtr = &GroupNodeSet.Add(GroupName, MakeShareable(new FStatsNode(GroupName)));
+			}
+
+			(*GroupPtr)->AddChildAndSetGroupPtr(StatsNodePtr);
+		}
+	}
+
+	GroupNodeSet.GenerateValueArray(GroupNodes);
+
+	// Sort by a fake group name.
+	GroupNodes.Sort(StatsNodeSortingHelper::ByNameAscending());
+}
+
+////////////////////////////////////////////////////////////////////////////////////////////////////
+
+void SStatsView::CreateGroupByOptionsSources()
+{
+	GroupByOptionsSource.Reset(3);
+
+	// Must be added in order of elements in the EStatsGroupingMode.
+	GroupByOptionsSource.Add(MakeShareable(new EStatsGroupingMode(EStatsGroupingMode::Flat)));
+	GroupByOptionsSource.Add(MakeShareable(new EStatsGroupingMode(EStatsGroupingMode::ByName)));
+	GroupByOptionsSource.Add(MakeShareable(new EStatsGroupingMode(EStatsGroupingMode::ByMetaGroupName)));
+	GroupByOptionsSource.Add(MakeShareable(new EStatsGroupingMode(EStatsGroupingMode::ByType)));
+
+	EStatsGroupingModePtr* GroupingModePtrPtr = GroupByOptionsSource.FindByPredicate([&](const EStatsGroupingModePtr InGroupingModePtr) { return *InGroupingModePtr == GroupingMode; });
+	if (GroupingModePtrPtr != nullptr)
+	{
+		GroupByComboBox->SetSelectedItem(*GroupingModePtrPtr);
+	}
+
+	GroupByComboBox->RefreshOptions();
+}
+
+////////////////////////////////////////////////////////////////////////////////////////////////////
+
+void SStatsView::GroupBy_OnSelectionChanged(TSharedPtr<EStatsGroupingMode> NewGroupingMode, ESelectInfo::Type SelectInfo)
+{
+	if (SelectInfo != ESelectInfo::Direct)
+	{
+		GroupingMode = *NewGroupingMode;
+
+		// Create groups, sort timers within the group and apply filtering.
+		CreateGroups();
+		SortStats();
+		ApplyFiltering();
+	}
+}
+
+////////////////////////////////////////////////////////////////////////////////////////////////////
+
+TSharedRef<SWidget> SStatsView::GroupBy_OnGenerateWidget(TSharedPtr<EStatsGroupingMode> InGroupingMode) const
+{
+	return SNew(STextBlock)
+		.Text(StatsNodeGroupingHelper::ToName(*InGroupingMode))
+		.ToolTipText(StatsNodeGroupingHelper::ToDescription(*InGroupingMode));
+}
+
+////////////////////////////////////////////////////////////////////////////////////////////////////
+
+FText SStatsView::GroupBy_GetSelectedText() const
+{
+	return StatsNodeGroupingHelper::ToName(GroupingMode);
+}
+
+////////////////////////////////////////////////////////////////////////////////////////////////////
+
+FText SStatsView::GroupBy_GetSelectedTooltipText() const
+{
+	return StatsNodeGroupingHelper::ToDescription(GroupingMode);
+}
+
+////////////////////////////////////////////////////////////////////////////////////////////////////
+// SortBy
+////////////////////////////////////////////////////////////////////////////////////////////////////
+
+void SStatsView::SortStats()
+{
+	const int32 NumGroups = GroupNodes.Num();
+
+	#define CHECK_AND_SORT_COLUMN(ColumnId, SortTypeName) \
+		if (ColumnBeingSorted == ColumnId) \
+		{ \
+			if (ColumnSortMode == EColumnSortMode::Type::Descending) \
+			{ \
+				for (int32 ID = 0; ID < NumGroups; ++ID) \
+				{ \
+					GroupNodes[ID]->SortChildren(StatsNodeSortingHelper::SortTypeName##Descending()); \
+				} \
+			} \
+			else /*if (ColumnSortMode == EColumnSortMode::Type::Ascending)*/ \
+			{ \
+				for (int32 ID = 0; ID < NumGroups; ++ID) \
+				{ \
+					GroupNodes[ID]->SortChildren(StatsNodeSortingHelper::SortTypeName##Ascending()); \
+				} \
+			} \
+		}
+
+		 CHECK_AND_SORT_COLUMN(FStatsViewColumns::NameColumnID,          ByName)
+	else CHECK_AND_SORT_COLUMN(FStatsViewColumns::MetaGroupNameColumnID, ByMetaGroupName)
+	else CHECK_AND_SORT_COLUMN(FStatsViewColumns::TypeColumnID,          ByType)
+	else CHECK_AND_SORT_COLUMN(FStatsViewColumns::CountColumnID,         ByCount)
+	else CHECK_AND_SORT_COLUMN(FStatsViewColumns::SumColumnID,           BySum)
+	else CHECK_AND_SORT_COLUMN(FStatsViewColumns::MaxColumnID,           ByMax)
+	else CHECK_AND_SORT_COLUMN(FStatsViewColumns::UpperQuartileColumnID, ByUpperQuartile)
+	else CHECK_AND_SORT_COLUMN(FStatsViewColumns::AverageColumnID,       ByAverage)
+	else CHECK_AND_SORT_COLUMN(FStatsViewColumns::MedianColumnID,        ByMedian)
+	else CHECK_AND_SORT_COLUMN(FStatsViewColumns::LowerQuartileColumnID, ByLowerQuartile)
+	else CHECK_AND_SORT_COLUMN(FStatsViewColumns::MinColumnID,           ByMin)
+
+	#undef CHECK_AND_SORT_COLUMN
+}
+
+////////////////////////////////////////////////////////////////////////////////////////////////////
+
+EColumnSortMode::Type SStatsView::GetSortModeForColumn(const FName ColumnId) const
+{
+	if (ColumnBeingSorted != ColumnId)
+	{
+		return EColumnSortMode::None;
+	}
+
+	return ColumnSortMode;
+}
+
+////////////////////////////////////////////////////////////////////////////////////////////////////
+
+void SStatsView::SetSortModeForColumn(const FName& ColumnId, const EColumnSortMode::Type SortMode)
+{
+	ColumnBeingSorted = ColumnId;
+	ColumnSortMode = SortMode;
+
+	// Sort timers and apply filtering.
+	SortStats();
+	ApplyFiltering();
+}
+
+////////////////////////////////////////////////////////////////////////////////////////////////////
+
+void SStatsView::OnSortModeChanged(const EColumnSortPriority::Type SortPriority, const FName& ColumnId, const EColumnSortMode::Type SortMode)
+{
+	SetSortModeForColumn(ColumnId, SortMode);
+	TreeView_Refresh();
+}
+
+////////////////////////////////////////////////////////////////////////////////////////////////////
+// SortMode action (HeaderMenu)
+////////////////////////////////////////////////////////////////////////////////////////////////////
+
+bool SStatsView::HeaderMenu_SortMode_IsChecked(const FName ColumnId, const EColumnSortMode::Type InSortMode)
+{
+	return ColumnBeingSorted == ColumnId && ColumnSortMode == InSortMode;
+}
+
+////////////////////////////////////////////////////////////////////////////////////////////////////
+
+bool SStatsView::HeaderMenu_SortMode_CanExecute(const FName ColumnId, const EColumnSortMode::Type InSortMode) const
+{
+	const FStatsViewColumn& Column = TreeViewHeaderColumns.FindChecked(ColumnId);
+	const bool bIsValid = Column.bCanBeSorted();
+
+	bool bCanExecute = ColumnBeingSorted != ColumnId ? true : ColumnSortMode != InSortMode;
+	bCanExecute = bCanExecute && bIsValid;
+
+	return bCanExecute;
+}
+
+////////////////////////////////////////////////////////////////////////////////////////////////////
+
+void SStatsView::HeaderMenu_SortMode_Execute(const FName ColumnId, const EColumnSortMode::Type InSortMode)
+{
+	SetSortModeForColumn(ColumnId, InSortMode);
+	TreeView_Refresh();
+}
+
+////////////////////////////////////////////////////////////////////////////////////////////////////
+// SortMode action (ContextMenu)
+////////////////////////////////////////////////////////////////////////////////////////////////////
+
+bool SStatsView::ContextMenu_SortMode_IsChecked(const EColumnSortMode::Type InSortMode)
+{
+	return ColumnSortMode == InSortMode;
+}
+
+////////////////////////////////////////////////////////////////////////////////////////////////////
+
+bool SStatsView::ContextMenu_SortMode_CanExecute(const EColumnSortMode::Type InSortMode) const
+{
+	return ColumnSortMode != InSortMode;
+}
+
+////////////////////////////////////////////////////////////////////////////////////////////////////
+
+void SStatsView::ContextMenu_SortMode_Execute(const EColumnSortMode::Type InSortMode)
+{
+	SetSortModeForColumn(ColumnBeingSorted, InSortMode);
+	TreeView_Refresh();
+}
+
+////////////////////////////////////////////////////////////////////////////////////////////////////
+// SortByColumn action (ContextMenu)
+////////////////////////////////////////////////////////////////////////////////////////////////////
+
+bool SStatsView::ContextMenu_SortByColumn_IsChecked(const FName ColumnId)
+{
+	return ColumnId == ColumnBeingSorted;
+}
+
+////////////////////////////////////////////////////////////////////////////////////////////////////
+
+bool SStatsView::ContextMenu_SortByColumn_CanExecute(const FName ColumnId) const
+{
+	return ColumnId != ColumnBeingSorted;
+}
+
+////////////////////////////////////////////////////////////////////////////////////////////////////
+
+void SStatsView::ContextMenu_SortByColumn_Execute(const FName ColumnId)
+{
+	SetSortModeForColumn(ColumnId, EColumnSortMode::Descending);
+	TreeView_Refresh();
+}
+
+////////////////////////////////////////////////////////////////////////////////////////////////////
+// HideColumn action (HeaderMenu)
+////////////////////////////////////////////////////////////////////////////////////////////////////
+
+bool SStatsView::HeaderMenu_HideColumn_CanExecute(const FName ColumnId) const
+{
+	const FStatsViewColumn& Column = TreeViewHeaderColumns.FindChecked(ColumnId);
+	return Column.bCanBeHidden();
+}
+
+////////////////////////////////////////////////////////////////////////////////////////////////////
+
+void SStatsView::HeaderMenu_HideColumn_Execute(const FName ColumnId)
+{
+	FStatsViewColumn& Column = TreeViewHeaderColumns.FindChecked(ColumnId);
+	Column.bIsVisible = false;
+	TreeViewHeaderRow->RemoveColumn(ColumnId);
+}
+
+////////////////////////////////////////////////////////////////////////////////////////////////////
+// ToggleColumn action (ContextMenu)
+////////////////////////////////////////////////////////////////////////////////////////////////////
+
+bool SStatsView::ContextMenu_ToggleColumn_IsChecked(const FName ColumnId)
+{
+	const FStatsViewColumn& Column = TreeViewHeaderColumns.FindChecked(ColumnId);
+	return Column.bIsVisible;
+}
+
+////////////////////////////////////////////////////////////////////////////////////////////////////
+
+bool SStatsView::ContextMenu_ToggleColumn_CanExecute(const FName ColumnId) const
+{
+	const FStatsViewColumn& Column = TreeViewHeaderColumns.FindChecked(ColumnId);
+	return Column.bCanBeHidden();
+}
+
+////////////////////////////////////////////////////////////////////////////////////////////////////
+
+void SStatsView::ContextMenu_ToggleColumn_Execute(const FName ColumnId)
+{
+	FStatsViewColumn& Column = TreeViewHeaderColumns.FindChecked(ColumnId);
+	if (Column.bIsVisible)
+	{
+		HeaderMenu_HideColumn_Execute(ColumnId);
+	}
+	else
+	{
+		TreeViewHeaderRow_ShowColumn(ColumnId);
+	}
+}
+
+////////////////////////////////////////////////////////////////////////////////////////////////////
+// "Show All Columns" action (ContextMenu)
+////////////////////////////////////////////////////////////////////////////////////////////////////
+
+bool SStatsView::ContextMenu_ShowAllColumns_CanExecute() const
+{
+	return true;
+}
+
+////////////////////////////////////////////////////////////////////////////////////////////////////
+
+void SStatsView::ContextMenu_ShowAllColumns_Execute()
+{
+	ColumnSortMode = DefaultColumnSortMode;
+	ColumnBeingSorted = DefaultColumnBeingSorted;
+
+	const int32 NumColumns = FStatsViewColumnFactory::Get().Collection.Num();
+	for (int32 ColumnIndex = 0; ColumnIndex < NumColumns; ColumnIndex++)
+	{
+		const FStatsViewColumn& DefaultColumn = *FStatsViewColumnFactory::Get().Collection[ColumnIndex];
+		const FStatsViewColumn& CurrentColumn = TreeViewHeaderColumns.FindChecked(DefaultColumn.Id);
+
+		if (!CurrentColumn.bIsVisible)
+		{
+			TreeViewHeaderRow_ShowColumn(DefaultColumn.Id);
+		}
+	}
+}
+
+////////////////////////////////////////////////////////////////////////////////////////////////////
+// "Show Min/Max/Median Columns" action (ContextMenu)
+////////////////////////////////////////////////////////////////////////////////////////////////////
+
+bool SStatsView::ContextMenu_ShowMinMaxMedColumns_CanExecute() const
+{
+	return true;
+}
+
+////////////////////////////////////////////////////////////////////////////////////////////////////
+
+void SStatsView::ContextMenu_ShowMinMaxMedColumns_Execute()
+{
+	TSet<FName> Preset =
+	{
+		FStatsViewColumns::NameColumnID,
+		//FStatsViewColumns::MetaGroupNameColumnID,
+		//FStatsViewColumns::TypeColumnID,
+		FStatsViewColumns::CountColumnID,
+		FStatsViewColumns::SumColumnID,
+		FStatsViewColumns::MaxColumnID,
+		FStatsViewColumns::UpperQuartileColumnID,
+		//FStatsViewColumns::AverageColumnID,
+		FStatsViewColumns::MedianColumnID,
+		FStatsViewColumns::LowerQuartileColumnID,
+		FStatsViewColumns::MinColumnID,
+	};
+
+	ColumnSortMode = EColumnSortMode::Descending;
+	ColumnBeingSorted = FStatsViewColumns::CountColumnID;
+
+	const int32 NumColumns = FStatsViewColumnFactory::Get().Collection.Num();
+	for (int32 ColumnIndex = 0; ColumnIndex < NumColumns; ColumnIndex++)
+	{
+		const FStatsViewColumn& DefaultColumn = *FStatsViewColumnFactory::Get().Collection[ColumnIndex];
+		const FStatsViewColumn& CurrentColumn = TreeViewHeaderColumns.FindChecked(DefaultColumn.Id);
+
+		bool bIsVisible = Preset.Contains(DefaultColumn.Id);
+		if (bIsVisible && !CurrentColumn.bIsVisible)
+		{
+			TreeViewHeaderRow_ShowColumn(DefaultColumn.Id);
+		}
+		else if (!bIsVisible && CurrentColumn.bIsVisible)
+		{
+			HeaderMenu_HideColumn_Execute(DefaultColumn.Id);
+		}
+	}
+}
+
+////////////////////////////////////////////////////////////////////////////////////////////////////
+// ResetColumns action (ContextMenu)
+////////////////////////////////////////////////////////////////////////////////////////////////////
+
+bool SStatsView::ContextMenu_ResetColumns_CanExecute() const
+{
+	return true;
+}
+
+////////////////////////////////////////////////////////////////////////////////////////////////////
+
+void SStatsView::ContextMenu_ResetColumns_Execute()
+{
+	ColumnSortMode = DefaultColumnSortMode;
+	ColumnBeingSorted = DefaultColumnBeingSorted;
+
+	const int32 NumColumns = FStatsViewColumnFactory::Get().Collection.Num();
+	for (int32 ColumnIndex = 0; ColumnIndex < NumColumns; ColumnIndex++)
+	{
+		const FStatsViewColumn& DefaultColumn = *FStatsViewColumnFactory::Get().Collection[ColumnIndex];
+		const FStatsViewColumn& CurrentColumn = TreeViewHeaderColumns.FindChecked(DefaultColumn.Id);
+
+		if (DefaultColumn.bIsVisible && !CurrentColumn.bIsVisible)
+		{
+			TreeViewHeaderRow_ShowColumn(DefaultColumn.Id);
+		}
+		else if (!DefaultColumn.bIsVisible && CurrentColumn.bIsVisible)
+		{
+			HeaderMenu_HideColumn_Execute(DefaultColumn.Id);
+		}
+	}
+}
+
+////////////////////////////////////////////////////////////////////////////////////////////////////
+
+void SStatsView::RebuildTree(bool bResync)
+{
+	bool bListHasChanged = false;
+
+	if (bResync)
+	{
+		StatsNodes.Empty(StatsNodes.Num());
+		//StatsNodesMap.Empty(StatsNodesMap.Num());
+		StatsNodesIdMap.Empty(StatsNodesIdMap.Num());
+		bListHasChanged = true;
+	}
+
+	if (Session.IsValid() && Trace::ReadCounterProvider(*Session.Get()))
+	{
+		Trace::FAnalysisSessionReadScope SessionReadScope(*Session.Get());
+
+		const Trace::ICounterProvider& CountersProvider = *Trace::ReadCounterProvider(*Session.Get());
+
+		if (!bResync)
+		{
+			bResync = (CountersProvider.GetCounterCount() != StatsNodes.Num());
+		}
+
+		if (bResync)
+		{
+			StatsNodes.Empty(StatsNodes.Num());
+			//StatsNodesMap.Empty(StatsNodesMap.Num());
+			StatsNodesIdMap.Empty(StatsNodesIdMap.Num());
+			bListHasChanged = true;
+
+			CountersProvider.EnumerateCounters([this](const Trace::ICounter& Counter)
+			{
+				FName Name(Counter.GetName());
+				FName Group(Counter.GetDisplayHint() == Trace::CounterDisplayHint_Memory ? TEXT("Memory") :
+							Counter.GetDisplayHint() == Trace::CounterDisplayHint_FloatingPoint ? TEXT("float") : TEXT("int64"));
+				EStatsNodeType Type = Counter.GetDisplayHint() == Trace::CounterDisplayHint_FloatingPoint ? EStatsNodeType::Float : EStatsNodeType::Int64;
+				FStatsNodePtr StatsNodePtr = MakeShareable(new FStatsNode(Counter.GetId(), Name, Group, Type));
+				StatsNodes.Add(StatsNodePtr);
+				//StatsNodesMap.Add(Name, StatsNodePtr);
+				StatsNodesIdMap.Add(Counter.GetId(), StatsNodePtr);
+			});
+		}
+	}
+
+	if (bListHasChanged)
+	{
+		UpdateTree();
+		UpdateStats(StatsStartTime, StatsEndTime);
+	}
+}
+
+////////////////////////////////////////////////////////////////////////////////////////////////////
+
+template<typename Type>
+class TTimeCalculationHelper
+{
+public:
+	TTimeCalculationHelper<Type>(double InIntervalStartTime, double InIntervalEndTime)
+		: IntervalStartTime(InIntervalStartTime)
+		, IntervalEndTime(InIntervalEndTime)
+	{
+	}
+
+	void Update(const Trace::ICounter& Counter)
+	{
+		EnumerateValues(Counter, UpdateMinMax);
+	}
+
+	void PrecomputeHistograms();
+
+	void UpdateHistograms(const Trace::ICounter& Counter)
+	{
+		EnumerateValues(Counter, UpdateHistogram);
+	}
+
+	void PostProcess(TMap<uint64, FStatsNodePtr>& StatsNodesIdMap, bool bComputeMedian);
+
+private:
+	template<typename CallbackType>
+	void EnumerateValues(const Trace::ICounter& Counter, CallbackType Callback);
+
+	static void UpdateMinMax(TAggregatedStatsEx<Type>& Stats, Type Value);
+	static void UpdateHistogram(TAggregatedStatsEx<Type>& StatsEx, Type Value);
+	static void PostProcess(TAggregatedStatsEx<Type>& StatsEx, bool bComputeMedian);
+
+	double IntervalStartTime;
+	double IntervalEndTime;
+	TMap<uint64, TAggregatedStatsEx<Type>> StatsMap;
+};
+
+////////////////////////////////////////////////////////////////////////////////////////////////////
+// specialization for Type = double
+
+template<>
+template<typename CallbackType>
+void TTimeCalculationHelper<double>::EnumerateValues(const Trace::ICounter& Counter, CallbackType Callback)
+{
+	TAggregatedStatsEx<double>* StatsExPtr = StatsMap.Find(Counter.GetId());
+	if (!StatsExPtr)
+	{
+		StatsExPtr = &StatsMap.Add(Counter.GetId());
+		StatsExPtr->BaseStats.Min = +MAX_dbl;
+		StatsExPtr->BaseStats.Max = -MAX_dbl;
+	}
+	TAggregatedStatsEx<double>& StatsEx = *StatsExPtr;
+
+	Counter.EnumerateFloatValues(IntervalStartTime, IntervalEndTime, [this, &StatsEx, Callback](double Time, double Value)
+	{
+		Callback(StatsEx, Value);
+	});
+}
+
+////////////////////////////////////////////////////////////////////////////////////////////////////
+// specialization for Type = int64
+
+template<>
+template<typename CallbackType>
+void TTimeCalculationHelper<int64>::EnumerateValues(const Trace::ICounter& Counter, CallbackType Callback)
+{
+	TAggregatedStatsEx<int64>* StatsExPtr = StatsMap.Find(Counter.GetId());
+	if (!StatsExPtr)
+	{
+		StatsExPtr = &StatsMap.Add(Counter.GetId());
+		StatsExPtr->BaseStats.Min = +MAX_int64;
+		StatsExPtr->BaseStats.Max = -MAX_int64;
+	}
+	TAggregatedStatsEx<int64>& StatsEx = *StatsExPtr;
+
+	Counter.EnumerateValues(IntervalStartTime, IntervalEndTime, [this, &StatsEx, Callback](double Time, int64 Value)
+	{
+		Callback(StatsEx, Value);
+	});
+}
+
+////////////////////////////////////////////////////////////////////////////////////////////////////
+
+template<typename Type>
+void TTimeCalculationHelper<Type>::UpdateMinMax(TAggregatedStatsEx<Type>& StatsEx, Type Value)
+{
+	TAggregatedStats<Type>& Stats = StatsEx.BaseStats;
+
+	Stats.Sum += Value;
+
+	if (Value < Stats.Min)
+	{
+		Stats.Min = Value;
+	}
+
+	if (Value > Stats.Max)
+	{
+		Stats.Max = Value;
+	}
+
+	Stats.Count++;
+}
+
+////////////////////////////////////////////////////////////////////////////////////////////////////
+// specialization for Type = double
+
+template<>
+void TTimeCalculationHelper<double>::PrecomputeHistograms()
+{
+	for (auto& KV : StatsMap)
+	{
+		TAggregatedStatsEx<double>& StatsEx = KV.Value;
+		const TAggregatedStats<double>& Stats = StatsEx.BaseStats;
+
+		// Each bucket (Histogram[i]) will be centered on a value.
+		// I.e. First bucket (bucket 0) is centered on Min value: [Min-DT/2, Min+DT/2)
+		// and last bucket (bucket N-1) is centered on Max value: [Max-DT/2, Max+DT/2).
+
+		if (Stats.Max == Stats.Min)
+		{
+			StatsEx.DT = 1.0; // single large bucket
+		}
+		else
+		{
+			StatsEx.DT = (Stats.Max - Stats.Min) / (TAggregatedStatsEx<double>::HistogramLen - 1);
+			if (StatsEx.DT == 0.0)
+			{
+				StatsEx.DT = 1.0;
+			}
+		}
+	}
+}
+
+////////////////////////////////////////////////////////////////////////////////////////////////////
+// specialization for Type = int64
+
+template<>
+void TTimeCalculationHelper<int64>::PrecomputeHistograms()
+{
+	for (auto& KV : StatsMap)
+	{
+		TAggregatedStatsEx<int64>& StatsEx = KV.Value;
+		const TAggregatedStats<int64>& Stats = StatsEx.BaseStats;
+
+		// Each bucket (Histogram[i]) will be centered on a value.
+		// I.e. First bucket (bucket 0) is centered on Min value: [Min-DT/2, Min+DT/2)
+		// and last bucket (bucket N-1) is centered on Max value: [Max-DT/2, Max+DT/2).
+
+		if (Stats.Max == Stats.Min)
+		{
+			StatsEx.DT = 1; // single bucket
+		}
+		else
+		{
+			// DT = Ceil[(Max - Min) / (N - 1)]
+			StatsEx.DT = (Stats.Max - Stats.Min + TAggregatedStatsEx<int64>::HistogramLen - 2) / (TAggregatedStatsEx<int64>::HistogramLen - 1);
+			if (StatsEx.DT == 0)
+			{
+				StatsEx.DT = 1;
+			}
+		}
+	}
+}
+
+////////////////////////////////////////////////////////////////////////////////////////////////////
+
+template<typename Type>
+void TTimeCalculationHelper<Type>::UpdateHistogram(TAggregatedStatsEx<Type>& StatsEx, Type Value)
+{
+	const TAggregatedStats<Type>& Stats = StatsEx.BaseStats;
+
+	// Index = (Value - Min + DT/2) / DT
+	int32 Index = static_cast<int32>((Value - Stats.Min + StatsEx.DT/2) / StatsEx.DT);
+	ensure(Index >= 0);
+	if (Index < 0)
+	{
+		Index = 0;
+	}
+	ensure(Index < TAggregatedStatsEx<Type>::HistogramLen);
+	if (Index >= TAggregatedStatsEx<Type>::HistogramLen)
+	{
+		Index = TAggregatedStatsEx<Type>::HistogramLen - 1;
+	}
+	StatsEx.Histogram[Index]++;
+}
+
+////////////////////////////////////////////////////////////////////////////////////////////////////
+
+template<typename Type>
+void TTimeCalculationHelper<Type>::PostProcess(TAggregatedStatsEx<Type>& StatsEx, bool bComputeMedian)
+{
+	TAggregatedStats<Type>& Stats = StatsEx.BaseStats;
+
+	// Compute average value.
+	if (Stats.Count > 0)
+	{
+		Stats.Average = Stats.Sum / static_cast<Type>(Stats.Count);
+
+		if (bComputeMedian)
+		{
+			const int32 HalfCount = Stats.Count / 2;
+
+			// Compute median value.
+			int32 Count = 0;
+			for (int32 HistogramIndex = 0; HistogramIndex < TAggregatedStatsEx<Type>::HistogramLen; HistogramIndex++)
+			{
+				Count += StatsEx.Histogram[HistogramIndex];
+				if (Count > HalfCount)
+				{
+					Stats.Median = Stats.Min + HistogramIndex * StatsEx.DT;
+
+					if (HistogramIndex > 0 &&
+						Stats.Count % 2 == 0 &&
+						Count - StatsEx.Histogram[HistogramIndex] == HalfCount)
+					{
+						const Type PrevMedian = Stats.Min + (HistogramIndex - 1) * StatsEx.DT;
+						Stats.Median = (Stats.Median + PrevMedian) / 2;
+					}
+
+					break;
+				}
+			}
+		}
+	}
+}
+
+////////////////////////////////////////////////////////////////////////////////////////////////////
+// specialization for Type = double
+
+template<>
+void TTimeCalculationHelper<double>::PostProcess(TMap<uint64, FStatsNodePtr>& StatsNodesIdMap, bool bComputeMedian)
+{
+	for (auto& KV : StatsMap)
+	{
+		PostProcess(KV.Value, bComputeMedian);
+
+		// Update the stats node.
+		FStatsNodePtr* StatsNodePtrPtr = StatsNodesIdMap.Find(KV.Key);
+		if (StatsNodePtrPtr != nullptr)
+		{
+			(*StatsNodePtrPtr)->SetAggregatedStats(KV.Value.BaseStats);
+		}
+	}
+}
+
+////////////////////////////////////////////////////////////////////////////////////////////////////
+// specialization for Type = int64
+
+template<>
+void TTimeCalculationHelper<int64>::PostProcess(TMap<uint64, FStatsNodePtr>& StatsNodesIdMap, bool bComputeMedian)
+{
+	for (auto& KV : StatsMap)
+	{
+		PostProcess(KV.Value, bComputeMedian);
+
+		// Update the stats node.
+		FStatsNodePtr* StatsNodePtrPtr = StatsNodesIdMap.Find(KV.Key);
+		if (StatsNodePtrPtr != nullptr)
+		{
+			(*StatsNodePtrPtr)->SetAggregatedIntegerStats(KV.Value.BaseStats);
+		}
+	}
+}
+
+////////////////////////////////////////////////////////////////////////////////////////////////////
+
+void SStatsView::UpdateStats(double StartTime, double EndTime)
+{
+	for (const FStatsNodePtr& StatsNodePtr : StatsNodes)
+	{
+		StatsNodePtr->ResetAggregatedStats();
+	}
+
+	if (Session.IsValid() && StartTime < EndTime && Trace::ReadCounterProvider(*Session.Get()))
+	{
+		const bool bComputeMedian = true;
+
+		TTimeCalculationHelper<double> CalculationHelperDbl(StartTime, EndTime);
+		TTimeCalculationHelper<int64>  CalculationHelperInt(StartTime, EndTime);
+
+		{
+			Trace::FAnalysisSessionReadScope SessionReadScope(*Session.Get());
+
+			const Trace::ICounterProvider& CountersProvider = *Trace::ReadCounterProvider(*Session.Get());
+
+			// Compute instance count and total/min/max inclusive/exclusive times for each counter.
+			// Iterate through all counters.
+			CountersProvider.EnumerateCounters([&CalculationHelperDbl, &CalculationHelperInt](const Trace::ICounter& Counter)
+			{
+				if (Counter.GetDisplayHint() == Trace::CounterDisplayHint_FloatingPoint)
+				{
+					CalculationHelperDbl.Update(Counter);
+				}
+				else
+				{
+					CalculationHelperInt.Update(Counter);
+				}
+			});
+
+			// Now, as we know min/max inclusive/exclusive times for counter, we can compute histogram and median values.
+			if (bComputeMedian)
+			{
+				// Update bucket size (DT) for computing histogram.
+				CalculationHelperDbl.PrecomputeHistograms();
+				CalculationHelperInt.PrecomputeHistograms();
+
+				// Compute histogram.
+				// Iterate again through all counters.
+				CountersProvider.EnumerateCounters([&CalculationHelperDbl, &CalculationHelperInt](const Trace::ICounter& Counter)
+				{
+					if (Counter.GetDisplayHint() == Trace::CounterDisplayHint_FloatingPoint)
+					{
+						CalculationHelperDbl.UpdateHistograms(Counter);
+					}
+					else
+					{
+						CalculationHelperInt.UpdateHistograms(Counter);
+					}
+				});
+			}
+		}
+
+		// Compute average and median inclusive/exclusive times.
+		CalculationHelperDbl.PostProcess(StatsNodesIdMap, bComputeMedian);
+		CalculationHelperInt.PostProcess(StatsNodesIdMap, bComputeMedian);
+
+		StatsStartTime = StartTime;
+		StatsEndTime = EndTime;
+
+		UpdateTree();
+		TreeView->RebuildList();
+	}
+}
+
+////////////////////////////////////////////////////////////////////////////////////////////////////
+
+void SStatsView::SelectStatsNode(uint64 Id)
+{
+	FStatsNodePtr* StatsNodePtrPtr = StatsNodesIdMap.Find(Id);
+	if (StatsNodePtrPtr != nullptr)
+	{
+		//UE_LOG(TimingProfiler, Log, TEXT("Select and RequestScrollIntoView %s"), *(*StatsNodePtrPtr)->GetName().ToString());
+		TreeView->SetSelection(*StatsNodePtrPtr);
+		TreeView->RequestScrollIntoView(*StatsNodePtrPtr);
+	}
+}
+
+////////////////////////////////////////////////////////////////////////////////////////////////////
+
+#undef LOCTEXT_NAMESPACE
diff --git a/Engine/Source/Developer/TraceInsights/Private/Insights/Widgets/SStatsView.h b/Engine/Source/Developer/TraceInsights/Private/Insights/Widgets/SStatsView.h
new file mode 100644
--- /dev/null
+++ b/Engine/Source/Developer/TraceInsights/Private/Insights/Widgets/SStatsView.h
@@ -0,0 +1,341 @@
+// Copyright 1998-2019 Epic Games, Inc. All Rights Reserved.
+
+#pragma once
+
+#include "CoreMinimal.h"
+#include "Misc/FilterCollection.h"
+#include "Misc/TextFilter.h"
+#include "SlateFwd.h"
+#include "Widgets/DeclarativeSyntaxSupport.h"
+#include "Widgets/Input/SComboBox.h"
+#include "Widgets/SCompoundWidget.h"
+#include "Widgets/SWidget.h"
+#include "Widgets/Views/SHeaderRow.h"
+#include "Widgets/Views/STableRow.h"
+#include "Widgets/Views/STableViewBase.h"
+#include "Widgets/Views/STreeView.h"
+
+// Insights
+#include "Insights/ViewModels/StatsNode.h"
+#include "Insights/ViewModels/StatsNodeHelper.h"
+#include "Insights/ViewModels/StatsViewColumn.h"
+
+////////////////////////////////////////////////////////////////////////////////////////////////////
+
+class FMenuBuilder;
+
+namespace Trace
+{
+	class IAnalysisSession;
+	struct FTimelineEvent;
+}
+
+////////////////////////////////////////////////////////////////////////////////////////////////////
+
+template<typename Type>
+struct TAggregatedStatsEx
+{
+	static constexpr int32 HistogramLen = 100; // number of buckets per histogram
+
+	TAggregatedStats<Type> BaseStats;
+
+	// Histogram for computing median and lower/upper quartiles.
+	int32 Histogram[HistogramLen];
+	Type DT; // bucket size
+
+	TAggregatedStatsEx()
+	{
+		FMemory::Memzero(Histogram, sizeof(int32) * HistogramLen);
+		DT = Type(1);
+	}
+};
+
+////////////////////////////////////////////////////////////////////////////////////////////////////
+
+/** The filter collection - used for updating the list of stats nodes. */
+typedef TFilterCollection<const FStatsNodePtr&> FStatsNodeFilterCollection;
+
+/** The text based filter - used for updating the list of stats nodes. */
+typedef TTextFilter<const FStatsNodePtr&> FStatsNodeTextFilter;
+
+////////////////////////////////////////////////////////////////////////////////////////////////////
+/**
+ * A custom widget used to display the list of stats.
+ */
+class SStatsView : public SCompoundWidget
+{
+public:
+	/** Default constructor. */
+	SStatsView();
+
+	/** Virtual destructor. */
+	virtual ~SStatsView();
+
+	SLATE_BEGIN_ARGS(SStatsView){}
+	SLATE_END_ARGS()
+
+	/**
+	 * Construct this widget
+	 * @param InArgs - The declaration data for this widget
+	 */
+	void Construct(const FArguments& InArgs);
+
+	/**
+	 * Rebuilds the tree (if necessary).
+	 * @param bResync - If true, it forces a resync with list of stats counters from Analysis, even if the list did not changed since last sync.
+	 */
+	void RebuildTree(bool bResync = true);
+	void UpdateStats(double StartTime, double EndTime);
+
+	void SelectStatsNode(uint64 Id);
+
+	//const TSet<FStatsNodePtr>& GetStatsNodes() const { return StatsNodes; }
+	//const TMap<uint64, FStatsNodePtr> GetStatsNodesIdMap() const { return StatsNodesIdMap; }
+	const FStatsNodePtr* GetStatsNode(uint64 Id) const { return StatsNodesIdMap.Find(Id); }
+
+protected:
+
+	void UpdateTree();
+
+	/** Called when the analysis session has changed. */
+	void InsightsManager_OnSessionChanged();
+
+	/**
+	 * Populates OutSearchStrings with the strings that should be used in searching.
+	 *
+	 * @param GroupOrStatNodePtr - the group and stat node to get a text description from.
+	 * @param OutSearchStrings   - an array of strings to use in searching.
+	 *
+	 */
+	void HandleItemToStringArray(const FStatsNodePtr& GroupOrStatNodePtr, TArray<FString>& OutSearchStrings) const;
+
+	////////////////////////////////////////////////////////////////////////////////////////////////////
+	// Tree View - Context Menu
+
+	TSharedPtr<SWidget> TreeView_GetMenuContent();
+	void TreeView_BuildSortByMenu(FMenuBuilder& MenuBuilder);
+	void TreeView_BuildViewColumnMenu(FMenuBuilder& MenuBuilder);
+
+	////////////////////////////////////////////////////////////////////////////////////////////////////
+	// Tree View - Columns' Header
+
+	void InitializeAndShowHeaderColumns();
+	void TreeViewHeaderRow_CreateColumnArgs(const int32 ColumnIndex);
+	void TreeViewHeaderRow_ShowColumn(const FName ColumnId);
+	TSharedRef<SWidget> TreeViewHeaderRow_GenerateColumnMenu(const FStatsViewColumn& Column);
+
+	////////////////////////////////////////////////////////////////////////////////////////////////////
+	// Tree View - Misc
+
+	void TreeView_Refresh();
+
+	/**
+	 * Called by STreeView to retrieves the children for the specified parent item.
+	 * @param InParent    - The parent node to retrieve the children from.
+	 * @param OutChildren - List of children for the parent node.
+	 */
+	void TreeView_OnGetChildren(FStatsNodePtr InParent, TArray<FStatsNodePtr>& OutChildren);
+
+	/** Called by STreeView when selection has changed. */
+	void TreeView_OnSelectionChanged(FStatsNodePtr SelectedItem, ESelectInfo::Type SelectInfo);
+
+	/** Called by STreeView when a tree item is double clicked. */
+	void TreeView_OnMouseButtonDoubleClick(FStatsNodePtr TreeNode);
+
+	////////////////////////////////////////////////////////////////////////////////////////////////////
+	// Tree View - Table Row
+
+	/** Called by STreeView to generate a table row for the specified item. */
+	TSharedRef<ITableRow> TreeView_OnGenerateRow(FStatsNodePtr TreeNode, const TSharedRef<STableViewBase>& OwnerTable);
+
+	bool TableRow_IsColumnVisible(const FName ColumnId) const;
+	void TableRow_SetHoveredTableCell(const FName ColumnId, const FStatsNodePtr StatsNodePtr);
+	EHorizontalAlignment TableRow_GetColumnOutlineHAlignment(const FName ColumnId) const;
+
+	FText TableRow_GetHighlightText() const;
+	FName TableRow_GetHighlightedStatsName() const;
+
+	bool TableRow_ShouldBeEnabled(const uint32 StatsId) const;
+
+	////////////////////////////////////////////////////////////////////////////////////////////////////
+	// Filtering
+
+	/** Populates the group and stat tree with items based on the current data. */
+	void ApplyFiltering();
+
+	TSharedRef<SWidget> GetToggleButtonForStatsType(const EStatsNodeType StatType);
+	void FilterByStatsType_OnCheckStateChanged(ECheckBoxState NewRadioState, const EStatsNodeType InStatType);
+	ECheckBoxState FilterByStatsType_IsChecked(const EStatsNodeType InStatType) const;
+
+	bool SearchBox_IsEnabled() const;
+	void SearchBox_OnTextChanged(const FText& InFilterText);
+
+	////////////////////////////////////////////////////////////////////////////////////////////////////
+	// GroupBy
+
+	void CreateGroups();
+	void CreateGroupByOptionsSources();
+
+	void GroupBy_OnSelectionChanged(TSharedPtr<EStatsGroupingMode> NewGroupingMode, ESelectInfo::Type SelectInfo);
+
+	TSharedRef<SWidget> GroupBy_OnGenerateWidget(TSharedPtr<EStatsGroupingMode> InGroupingMode) const;
+
+	FText GroupBy_GetSelectedText() const;
+
+	FText GroupBy_GetSelectedTooltipText() const;
+
+	////////////////////////////////////////////////////////////////////////////////////////////////////
+	// Sorting
+
+	void SortStats();
+
+	EColumnSortMode::Type GetSortModeForColumn(const FName ColumnId) const;
+	void SetSortModeForColumn(const FName& ColumnId, EColumnSortMode::Type SortMode);
+	void OnSortModeChanged(const EColumnSortPriority::Type SortPriority, const FName& ColumnId, const EColumnSortMode::Type SortMode);
+
+	//void TreeView_BuildSortByMenu(FMenuBuilder& MenuBuilder);
+
+	////////////////////////////////////////////////////////////////////////////////////////////////////
+	// Sorting actions
+
+	// SortMode (HeaderMenu)
+	bool HeaderMenu_SortMode_IsChecked(const FName ColumnId, const EColumnSortMode::Type InSortMode);
+	bool HeaderMenu_SortMode_CanExecute(const FName ColumnId, const EColumnSortMode::Type InSortMode) const;
+	void HeaderMenu_SortMode_Execute(const FName ColumnId, const EColumnSortMode::Type InSortMode);
+
+	// SortMode (ContextMenu)
+	bool ContextMenu_SortMode_IsChecked(const EColumnSortMode::Type InSortMode);
+	bool ContextMenu_SortMode_CanExecute(const EColumnSortMode::Type InSortMode) const;
+	void ContextMenu_SortMode_Execute(const EColumnSortMode::Type InSortMode);
+
+	// SortByColumn (ContextMenu)
+	bool ContextMenu_SortByColumn_IsChecked(const FName ColumnId);
+	bool ContextMenu_SortByColumn_CanExecute(const FName ColumnId) const;
+	void ContextMenu_SortByColumn_Execute(const FName ColumnId);
+
+	////////////////////////////////////////////////////////////////////////////////////////////////////
+	// Column visibility actions
+
+	// HideColumn (HeaderMenu)
+	bool HeaderMenu_HideColumn_CanExecute(const FName ColumnId) const;
+	void HeaderMenu_HideColumn_Execute(const FName ColumnId);
+
+	// ToggleColumn (ContextMenu)
+	bool ContextMenu_ToggleColumn_IsChecked(const FName ColumnId);
+	bool ContextMenu_ToggleColumn_CanExecute(const FName ColumnId) const;
+	void ContextMenu_ToggleColumn_Execute(const FName ColumnId);
+
+	// ShowAllColumns (ContextMenu)
+	bool ContextMenu_ShowAllColumns_CanExecute() const;
+	void ContextMenu_ShowAllColumns_Execute();
+
+	// MinMaxMedColumns (ContextMenu)
+	bool ContextMenu_ShowMinMaxMedColumns_CanExecute() const;
+	void ContextMenu_ShowMinMaxMedColumns_Execute();
+
+	// ResetColumns (ContextMenu)
+	bool ContextMenu_ResetColumns_CanExecute() const;
+	void ContextMenu_ResetColumns_Execute();
+
+protected:
+	/** A weak pointer to the profiler session used to populate this widget. */
+	TSharedPtr<const Trace::IAnalysisSession>/*Weak*/ Session;
+
+	//////////////////////////////////////////////////
+	// Tree View, Columns
+
+	/** The tree widget which holds the list of groups and stats corresponding with each group. */
+	TSharedPtr<STreeView<FStatsNodePtr>> TreeView;
+
+	/** Column metadata used to initialize column arguments, stored as PropertyName -> FEventGraphColumn. */
+	TMap<FName, FStatsViewColumn> TreeViewHeaderColumns;
+
+	/** Column arguments used to initialize a new header column in the tree view, stored as column name to column arguments mapping. */
+	TMap<FName, SHeaderRow::FColumn::FArguments> TreeViewHeaderColumnArgs;
+
+	/** Holds the tree view header row widget which display all columns in the tree view. */
+	TSharedPtr<SHeaderRow> TreeViewHeaderRow;
+
+	/** External scrollbar used to synchronize tree view position. */
+	TSharedPtr<SScrollBar> ExternalScrollbar;
+
+	//////////////////////////////////////////////////
+	// Hovered Column, Hovered Stats Node
+
+	/** Name of the column currently being hovered by the mouse. */
+	FName HoveredColumnId;
+
+	/** A shared pointer to the stats node currently being hovered by the mouse. */
+	FStatsNodePtr HoveredStatsNodePtr;
+
+	/** Name of the stats that should be drawn as highlighted. */
+	FName HighlightedStatsName;
+
+	//////////////////////////////////////////////////
+	// Stats Nodes
+
+	/** An array of group and stats nodes generated from the metadata. */
+	TArray<FStatsNodePtr> GroupNodes;
+
+	/** A filtered array of group and stats nodes to be displayed in the tree widget. */
+	TArray<FStatsNodePtr> FilteredGroupNodes;
+
+	/** All stats nodes. */
+	TSet<FStatsNodePtr> StatsNodes;
+
+	/** All stats nodes, stored as StatsName -> FStatsNodePtr. */
+	//TMap<FName, FStatsNodePtr> StatsNodesMap;
+
+	/** All stats nodes, stored as StatsId -> FStatsNodePtr. */
+	TMap<uint64, FStatsNodePtr> StatsNodesIdMap;
+
+	/** Currently expanded group nodes. */
+	TSet<FStatsNodePtr> ExpandedNodes;
+
+	/** If true, the expanded nodes have been saved before applying a text filter. */
+	bool bExpansionSaved;
+
+	//////////////////////////////////////////////////
+	// Search box and filters
+
+	/** The search box widget used to filter items displayed in the stats and groups tree. */
+	TSharedPtr<SSearchBox> SearchBox;
+
+	/** The text based filter. */
+	TSharedPtr<FStatsNodeTextFilter> TextFilter;
+
+	/** The filter collection. */
+	TSharedPtr<FStatsNodeFilterCollection> Filters;
+
+	/** Holds the visibility of each stats type. */
+	bool bStatsNodeIsVisible[static_cast<int>(EStatsNodeType::InvalidOrMax)];
+
+	//////////////////////////////////////////////////
+	// Grouping
+
+	TArray<TSharedPtr<EStatsGroupingMode>> GroupByOptionsSource;
+
+	TSharedPtr<SComboBox<TSharedPtr<EStatsGroupingMode>>> GroupByComboBox;
+
+	/** How we group the stats? */
+	EStatsGroupingMode GroupingMode;
+
+	//////////////////////////////////////////////////
+	// Sorting
+
+	/** How we sort the stats? */
+	EColumnSortMode::Type ColumnSortMode;
+
+	/** Name of the column currently being sorted, NAME_None if sorting is disabled. */
+	FName ColumnBeingSorted;
+
+	static const EColumnSortMode::Type DefaultColumnSortMode;
+	static const FName DefaultColumnBeingSorted;
+
+	//////////////////////////////////////////////////
+
+	double StatsStartTime;
+	double StatsEndTime;
+};
+
+////////////////////////////////////////////////////////////////////////////////////////////////////
diff --git a/Engine/Source/Developer/TraceInsights/Private/Insights/Widgets/SStatsViewTooltip.cpp b/Engine/Source/Developer/TraceInsights/Private/Insights/Widgets/SStatsViewTooltip.cpp
new file mode 100644
--- /dev/null
+++ b/Engine/Source/Developer/TraceInsights/Private/Insights/Widgets/SStatsViewTooltip.cpp
@@ -0,0 +1,526 @@
+// Copyright 1998-2019 Epic Games, Inc. All Rights Reserved.
+
+#include "SStatsViewTooltip.h"
+
+#include "EditorStyleSet.h"
+#include "SlateOptMacros.h"
+#include "Widgets/Images/SImage.h"
+#include "Widgets/Layout/SGridPanel.h"
+#include "Widgets/Layout/SSeparator.h"
+#include "Widgets/SBoxPanel.h"
+#include "Widgets/SToolTip.h"
+#include "Widgets/Text/STextBlock.h"
+
+// Insights
+#include "Insights/Common/TimeUtils.h"
+#include "Insights/ViewModels/StatsNode.h"
+#include "Insights/ViewModels/StatsViewColumn.h"
+#include "Insights/ViewModels/StatsNodeHelper.h"
+
+////////////////////////////////////////////////////////////////////////////////////////////////////
+
+#define LOCTEXT_NAMESPACE "SStatsViewTooltip"
+
+////////////////////////////////////////////////////////////////////////////////////////////////////
+
+BEGIN_SLATE_FUNCTION_BUILD_OPTIMIZATION
+
+TSharedPtr<SToolTip> SStatsViewTooltip::GetColumnTooltip(const FStatsViewColumn& Column)
+{
+	TSharedPtr<SToolTip> ColumnTooltip =
+		SNew(SToolTip)
+		[
+			SNew(SVerticalBox)
+
+			+ SVerticalBox::Slot()
+			.AutoHeight()
+			.Padding(2.0f)
+			[
+				SNew(STextBlock)
+				.Text(Column.TitleName)
+				.TextStyle(FEditorStyle::Get(), TEXT("Profiler.TooltipBold"))
+			]
+
+			+ SVerticalBox::Slot()
+			.AutoHeight()
+			.Padding(2.0f)
+			[
+				SNew(STextBlock)
+				.Text(Column.Description)
+				.TextStyle(FEditorStyle::Get(), TEXT("Profiler.Tooltip"))
+			]
+		];
+
+	return ColumnTooltip;
+}
+
+////////////////////////////////////////////////////////////////////////////////////////////////////
+
+TSharedPtr<SToolTip> SStatsViewTooltip::GetTableCellTooltip(const TSharedPtr<FStatsNode> StatsNodePtr)
+{
+	//const FSlateFontInfo TitleFont = FCoreStyle::GetDefaultFontStyle("Bold", 10);
+	//const FSlateFontInfo DescriptionFont = FCoreStyle::GetDefaultFontStyle("Regular", 8);
+	//const FSlateFontInfo DescriptionFontB = FCoreStyle::GetDefaultFontStyle("Bold", 8);
+
+	const FLinearColor DefaultColor(1.0f,1.0f,1.0f,1.0f);
+	const FLinearColor ThreadColor(5.0f, 0.0f, 0.0f, 1.0f);
+	const float Alpha = 0.0f;//StatsNodePtr->_FramePct * 0.01f;
+	const FLinearColor ColorAndOpacity = FMath::Lerp(DefaultColor, ThreadColor,Alpha);
+
+	const FText InstanceCountText = FText::AsNumber(StatsNodePtr->GetAggregatedStats().Count);
+
+	FText SumText = StatsNodePtr->GetTextForAggregatedStatsSum();
+	FText MinText = StatsNodePtr->GetTextForAggregatedStatsMin();
+	FText MaxText = StatsNodePtr->GetTextForAggregatedStatsMax();
+	FText AvgText = StatsNodePtr->GetTextForAggregatedStatsAverage();
+	FText MedText = StatsNodePtr->GetTextForAggregatedStatsMedian();
+	FText LowText = StatsNodePtr->GetTextForAggregatedStatsLowerQuartile();
+	FText UppText = StatsNodePtr->GetTextForAggregatedStatsUpperQuartile();
+
+	//TSharedPtr<SHorizontalBox> HBoxCaption;
+	TSharedPtr<SGridPanel> GridPanel;
+	TSharedPtr<SHorizontalBox> HBox;
+
+	TSharedPtr<SToolTip> TableCellTooltip =
+		SNew(SToolTip)
+		[
+			SAssignNew(HBox, SHorizontalBox)
+
+			+SHorizontalBox::Slot()
+			.AutoWidth()
+			[
+				SNew(SVerticalBox)
+
+				//+SVerticalBox::Slot()
+				//.AutoHeight()
+				//.Padding(2.0f)
+				//[
+				//	SAssignNew(HBoxCaption, SHorizontalBox)
+				//]
+
+				+SVerticalBox::Slot()
+				.AutoHeight()
+				.Padding(2.0f)
+				[
+					SNew(SSeparator)
+					.Orientation(Orient_Horizontal)
+				]
+
+				+SVerticalBox::Slot()
+				.AutoHeight()
+				.Padding(2.0f)
+				[
+					SNew(SGridPanel)
+
+					// Id: [Id]
+					+SGridPanel::Slot(0, 0)
+					.Padding(2.0f)
+					[
+						SNew(STextBlock)
+						.Text(LOCTEXT("TT_Id", "Id:"))
+						.TextStyle(FEditorStyle::Get(), TEXT("Profiler.TooltipBold"))
+					]
+					+SGridPanel::Slot(1, 0)
+					.Padding(2.0f)
+					[
+						SNew(STextBlock)
+						.Text(FText::AsNumber(StatsNodePtr->GetId()))
+						.TextStyle(FEditorStyle::Get(), TEXT("Profiler.Tooltip"))
+					]
+
+					// Name: [Name]
+					+SGridPanel::Slot(0, 1)
+					.Padding(2.0f)
+					[
+						SNew(STextBlock)
+						.Text(LOCTEXT("TT_Name", "Name:"))
+						.TextStyle(FEditorStyle::Get(), TEXT("Profiler.TooltipBold"))
+					]
+					+SGridPanel::Slot(1, 1)
+					.Padding(2.0f)
+					[
+						SNew(STextBlock)
+						.Text(FText::FromName(StatsNodePtr->GetName()))
+						.TextStyle(FEditorStyle::Get(), TEXT("Profiler.Tooltip"))
+					]
+
+					//// Group: [MetaGroupName]
+					//+SGridPanel::Slot(0, 2)
+					//.Padding(2.0f)
+					//[
+					//	SNew(STextBlock)
+					//	.Text(LOCTEXT("TT_Group", "Group:"))
+					//	.TextStyle(FEditorStyle::Get(), TEXT("Profiler.TooltipBold"))
+					//]
+					//+SGridPanel::Slot(1, 2)
+					//.Padding(2.0f)
+					//[
+					//	SNew(STextBlock)
+					//	.Text(FText::FromName(StatsNodePtr->GetMetaGroupName()))
+					//	.TextStyle(FEditorStyle::Get(), TEXT("Profiler.Tooltip"))
+					//]
+
+					// Type: [Type]
+					+ SGridPanel::Slot(0, 3)
+					.Padding(2.0f)
+					[
+						SNew(STextBlock)
+						.Text(LOCTEXT("TT_Type", "Type:"))
+						.TextStyle(FEditorStyle::Get(), TEXT("Profiler.TooltipBold"))
+					]
+					+ SGridPanel::Slot(1, 3)
+					.Padding(2.0f)
+					[
+						SNew(STextBlock)
+						.Text(StatsNodeTypeHelper::ToName(StatsNodePtr->GetType()))
+						.TextStyle(FEditorStyle::Get(), TEXT("Profiler.Tooltip"))
+					]
+				]
+
+				+SVerticalBox::Slot()
+				.AutoHeight()
+				.Padding(2.0f)
+				[
+					SNew(SSeparator)
+					.Orientation(Orient_Horizontal)
+				]
+
+				+ SVerticalBox::Slot()
+				.AutoHeight()
+				.Padding(2.0f)
+				[
+					SNew(SGridPanel)
+
+					+ SGridPanel::Slot(0, 0)
+					.Padding(2.0f)
+					[
+						SNew(STextBlock)
+						.Text(LOCTEXT("TT_NumInstances", "Num Instances:"))
+						.TextStyle(FEditorStyle::Get(), TEXT("Profiler.TooltipBold"))
+					]
+					+ SGridPanel::Slot(1, 0)
+					.Padding(2.0f)
+					[
+						SNew(STextBlock)
+						.Text(InstanceCountText)
+						.TextStyle(FEditorStyle::Get(), TEXT("Profiler.Tooltip"))
+					]
+				]
+
+				+SVerticalBox::Slot()
+				.AutoHeight()
+				.Padding(2.0f)
+				[
+					SNew(SSeparator)
+					.Orientation(Orient_Horizontal)
+				]
+
+				+ SVerticalBox::Slot()
+				.AutoHeight()
+				.Padding(2.0f)
+				[
+					SAssignNew(GridPanel, SGridPanel)
+
+					// Aggregated stats are added here.
+				]
+
+				+SVerticalBox::Slot()
+				.AutoHeight()
+				.Padding(2.0f)
+				[
+					SNew(SSeparator)
+					.Orientation(Orient_Horizontal)
+				]
+			]
+		];
+
+	int32 Row = 1;
+	AddAggregatedStatsRow(GridPanel, Row, LOCTEXT("TT_Sum",     "Sum:"),            SumText);
+	AddAggregatedStatsRow(GridPanel, Row, LOCTEXT("TT_Max",     "Max:"),            MaxText);
+	AddAggregatedStatsRow(GridPanel, Row, LOCTEXT("TT_UpperQ",  "Upper Quartile:"), UppText);
+	AddAggregatedStatsRow(GridPanel, Row, LOCTEXT("TT_Average", "Average:"),        AvgText);
+	AddAggregatedStatsRow(GridPanel, Row, LOCTEXT("TT_Median",  "Median:"),         MedText);
+	AddAggregatedStatsRow(GridPanel, Row, LOCTEXT("TT_LowerQ",  "Lower Quartile:"), LowText);
+	AddAggregatedStatsRow(GridPanel, Row, LOCTEXT("TT_Min",     "Min:"),            MinText);
+
+	/*
+	//TODO: We need Stats hierarchy (not the grouping hierarchy)!
+	const bool bHasParent = StatsNodePtr->GetStats()->GetParent().IsValid();
+	const bool bHasChildren = StatsNodePtr->GetStats()->GetChildren().Num() > 0;
+
+	if (bHasParent)
+	{
+		const FText ParentName = FText::FromName(StatsNodePtr->GetGroupPtr()->GetName());
+		HBoxCaption->AddSlot()
+			.AutoWidth()
+			[
+				SNew(STextBlock)
+				.Text(ParentName)
+				.TextStyle(FEditorStyle::Get(), TEXT("Profiler.Caption"))
+			];
+
+		HBoxCaption->AddSlot()
+			.AutoWidth()
+			[
+				SNew(SImage)
+				.Image(FEditorStyle::GetBrush("BreadcrumbTrail.Delimiter"))
+			];
+	}
+
+	const FText StatsName = FText::FromName(StatsNodePtr->GetName());
+	HBoxCaption->AddSlot()
+		.AutoWidth()
+		[
+			SNew(STextBlock)
+			.Text(StatsName)
+			.TextStyle(FEditorStyle::Get(), TEXT("Profiler.CaptionBold"))
+		];
+
+	if (bHasChildren)
+	{
+		typedef TKeyValuePair<float, FName> FEventNameAndPct;
+		TArray<FEventNameAndPct> MinimalChildren;
+
+		const TArray<FStatsNodePtr>& Children = StatsNodePtr->GetChildren();
+		for(int32 ChildIndex = 0; ChildIndex < Children.Num(); ChildIndex++)
+		{
+			const FStatsNodePtr Child = Children[ChildIndex];
+			float Percent = static_cast<float>(Child->GetStats().TotalInclusiveTime / StatsNodePtr->GetStats().TotalInclusiveTime * 100.0);
+			MinimalChildren.Add(FEventNameAndPct(Child->GetStats().TotalInclusiveTime, Child->GetName()));
+		}
+
+		struct FCompareByFloatDescending
+		{
+			FORCEINLINE bool operator()(const FEventNameAndPct& A, const FEventNameAndPct& B) const
+			{
+				return A.Key > B.Key;
+			}
+		};
+		MinimalChildren.Sort(FCompareByFloatDescending());
+
+		FString ChildrenNames;
+		const int32 NumChildrenToDisplay = FMath::Min(MinimalChildren.Num(), 3);
+		for(int32 SortedChildIndex = 0; SortedChildIndex < NumChildrenToDisplay; SortedChildIndex++)
+		{
+			const FEventNameAndPct& MinimalChild = MinimalChildren[SortedChildIndex];
+			ChildrenNames += FString::Printf(TEXT("%s (%.1f %%)"), *MinimalChild.Value.ToString(), MinimalChild.Key);
+
+			const bool bAddDelimiter = SortedChildIndex < NumChildrenToDisplay - 1;
+			if (bAddDelimiter)
+			{
+				ChildrenNames += TEXT(", ");
+			}
+		}
+
+		HBoxCaption->AddSlot()
+			.AutoWidth()
+			[
+				SNew(SImage)
+				.Image(FEditorStyle::GetBrush("BreadcrumbTrail.Delimiter"))
+			];
+
+		HBoxCaption->AddSlot()
+			.AutoWidth()
+			[
+				SNew(STextBlock)
+				.Text(FText::FromString(ChildrenNames))
+				.TextStyle(FEditorStyle::Get(), TEXT("Profiler.Caption"))
+			];
+	}
+	*/
+
+	return TableCellTooltip;
+}
+
+////////////////////////////////////////////////////////////////////////////////////////////////////
+
+TSharedRef<SToolTip> SStatsViewTooltip::GetTooltip()
+{
+	if (Session.IsValid())
+	{
+		const TSharedRef<SGridPanel> ToolTipGrid = SNew(SGridPanel);
+		int32 CurrentRowPos = 0;
+
+		AddHeader(ToolTipGrid, CurrentRowPos);
+		AddDescription(ToolTipGrid, CurrentRowPos);
+
+		AddNoDataInformation(ToolTipGrid, CurrentRowPos);
+
+		return SNew(SToolTip)
+			[
+				ToolTipGrid
+			];
+	}
+	else
+	{
+		return SNew(SToolTip)
+			.Text(LOCTEXT("NotImplemented", "Tooltip for multiple profiler instances has not been implemented yet"));
+	}
+}
+
+////////////////////////////////////////////////////////////////////////////////////////////////////
+
+void SStatsViewTooltip::AddNoDataInformation(const TSharedRef<SGridPanel>& Grid, int32& RowPos)
+{
+	Grid->AddSlot(0, RowPos)
+		.Padding(2.0f)
+		.ColumnSpan(3)
+		[
+			SNew(STextBlock)
+			.TextStyle(FEditorStyle::Get(), TEXT("Profiler.TooltipBold"))
+			.Text(LOCTEXT("NoDataAvailable", "N/A"))
+		];
+	RowPos++;
+}
+
+////////////////////////////////////////////////////////////////////////////////////////////////////
+
+void SStatsViewTooltip::AddHeader(const TSharedRef<SGridPanel>& Grid, int32& RowPos)
+{
+	const FString InstanceName = Session->GetName();
+
+	Grid->AddSlot(0, RowPos++)
+		.Padding(2.0f)
+		.ColumnSpan(3)
+		[
+			SNew(STextBlock)
+			.TextStyle(FEditorStyle::Get(), TEXT("Profiler.TooltipBold"))
+			.Text(LOCTEXT("StatInstance", "Stat information for profiler instance"))
+		];
+
+	Grid->AddSlot(0, RowPos++)
+		.Padding(2.0f)
+		.ColumnSpan(3)
+		[
+			SNew(STextBlock)
+			.TextStyle(FEditorStyle::Get(), TEXT("Profiler.Tooltip"))
+			.Text(FText::FromString(InstanceName))
+		];
+
+	AddSeparator(Grid, RowPos);
+}
+
+////////////////////////////////////////////////////////////////////////////////////////////////////
+
+void SStatsViewTooltip::AddDescription(const TSharedRef<SGridPanel>& Grid, int32& RowPos)
+{
+	/*
+	const FProfilerStat& ProfilerStat = Session->GetMetaData()->GetStatsById(StatsId);
+	const EStatsNodeType SampleType = Session->GetMetaData()->GetSampleTypeForStatsId(StatsId);
+	const FSlateBrush* const StatIcon = SStatsViewHelper::GetIconForStatType(SampleType);
+
+	Grid->AddSlot(0, RowPos)
+	.Padding(2.0f)
+	[
+		SNew(STextBlock)
+		.TextStyle(FEditorStyle::Get(), TEXT("Profiler.TooltipBold"))
+		.Text(LOCTEXT("GroupDesc","Group:"))
+	];
+
+	Grid->AddSlot(1, RowPos)
+	.Padding(2.0f)
+	.ColumnSpan(2)
+	[
+		SNew(STextBlock)
+		.TextStyle(FEditorStyle::Get(), TEXT("Profiler.Tooltip"))
+		.Text(FText::FromName(ProfilerStat.OwningGroup().Name()))
+	];
+	RowPos++;
+
+	Grid->AddSlot(0, RowPos)
+	.Padding(2.0f)
+	[
+		SNew(STextBlock)
+		.TextStyle(FEditorStyle::Get(), TEXT("Profiler.TooltipBold"))
+		.Text(LOCTEXT("NameDesc","Name:"))
+	];
+
+	Grid->AddSlot(1, RowPos)
+	.Padding(2.0f)
+	.ColumnSpan(2)
+	[
+		SNew(STextBlock)
+		.TextStyle(FEditorStyle::Get(), TEXT("Profiler.Tooltip"))
+		.Text(FText::FromName(ProfilerStat.Name()))
+	];
+	RowPos++;
+
+	Grid->AddSlot(0, RowPos)
+	.Padding(2.0f)
+	[
+		SNew(STextBlock)
+		.TextStyle(FEditorStyle::Get(), TEXT("Profiler.TooltipBold"))
+		.Text(LOCTEXT("TypeDesc","Type:"))
+	];
+
+	Grid->AddSlot(1, RowPos)
+	.Padding(2.0f)
+	.ColumnSpan(2)
+	[
+		SNew(SHorizontalBox)
+
+		+SHorizontalBox::Slot()
+		.AutoWidth()
+		[
+			SNew(SImage)
+			.Image(StatIcon)
+		]
+
+		+SHorizontalBox::Slot()
+		.AutoWidth()
+		.HAlign(HAlign_Left)
+		.VAlign(VAlign_Center)
+		[
+			SNew(STextBlock)
+			.Text(FText::FromString(EStatsNodeType::ToDescription(SampleType)))
+			.TextStyle(FEditorStyle::Get(), TEXT("Profiler.Tooltip"))
+		]
+	];
+	RowPos++;
+
+	AddSeparator(Grid, RowPos);
+	*/
+}
+
+////////////////////////////////////////////////////////////////////////////////////////////////////
+
+void SStatsViewTooltip::AddSeparator(const TSharedRef<SGridPanel>& Grid, int32& RowPos)
+{
+	Grid->AddSlot(0, RowPos++)
+		.Padding(2.0f)
+		.ColumnSpan(3)
+		[
+			SNew(SSeparator)
+			.Orientation(Orient_Horizontal)
+		];
+}
+
+////////////////////////////////////////////////////////////////////////////////////////////////////
+
+void SStatsViewTooltip::AddAggregatedStatsRow(TSharedPtr<SGridPanel> Grid, int32& Row, const FText& Name, const FText& Value)
+{
+	Grid->AddSlot(0, Row)
+		.Padding(2.0f)
+		[
+			SNew(STextBlock)
+			.Text(Name)
+			.TextStyle(FEditorStyle::Get(), TEXT("Profiler.TooltipBold"))
+		];
+
+	Grid->AddSlot(1, Row)
+		.Padding(2.0f)
+		.HAlign(HAlign_Right)
+		[
+			SNew(STextBlock)
+			.Text(Value)
+			.TextStyle(FEditorStyle::Get(), TEXT("Profiler.Tooltip"))
+		];
+
+	Row++;
+}
+
+END_SLATE_FUNCTION_BUILD_OPTIMIZATION
+
+////////////////////////////////////////////////////////////////////////////////////////////////////
+
+#undef LOCTEXT_NAMESPACE
diff --git a/Engine/Source/Developer/TraceInsights/Private/Insights/Widgets/SStatsViewTooltip.h b/Engine/Source/Developer/TraceInsights/Private/Insights/Widgets/SStatsViewTooltip.h
new file mode 100644
--- /dev/null
+++ b/Engine/Source/Developer/TraceInsights/Private/Insights/Widgets/SStatsViewTooltip.h
@@ -0,0 +1,53 @@
+// Copyright 1998-2019 Epic Games, Inc. All Rights Reserved.
+
+#pragma once
+
+#include "CoreMinimal.h"
+#include "EditorStyleSet.h"
+#include "SlateOptMacros.h"
+#include "TraceServices/AnalysisService.h"
+#include "Widgets/Images/SImage.h"
+#include "Widgets/Layout/SGridPanel.h"
+#include "Widgets/Layout/SSeparator.h"
+#include "Widgets/SToolTip.h"
+#include "Widgets/Text/STextBlock.h"
+
+// Insights
+#include "Insights/InsightsManager.h"
+
+// Insights
+class FStatsNode;
+class FStatsViewColumn;
+
+#define LOCTEXT_NAMESPACE "SStatsView"
+
+/** Stats View Tooltip */
+class SStatsViewTooltip
+{
+	const uint32 StatsId;
+	TSharedPtr<const Trace::IAnalysisSession> Session;
+
+public:
+	SStatsViewTooltip(const uint32 InStatsId)
+		: StatsId(InStatsId)
+	{
+		Session = FInsightsManager::Get()->GetSession();
+	}
+
+	TSharedRef<SToolTip> GetTooltip();
+
+protected:
+	void AddNoDataInformation(const TSharedRef<SGridPanel>& Grid, int32& RowPos);
+	void AddHeader(const TSharedRef<SGridPanel>& Grid, int32& RowPos);
+	void AddDescription(const TSharedRef<SGridPanel>& Grid, int32& RowPos);
+	void AddSeparator(const TSharedRef<SGridPanel>& Grid, int32& RowPos);
+
+public:
+	static TSharedPtr<SToolTip> GetColumnTooltip(const FStatsViewColumn& Column);
+	static TSharedPtr<SToolTip> GetTableCellTooltip(const TSharedPtr<FStatsNode> StatsNodePtr);
+
+protected:
+	static void AddAggregatedStatsRow(TSharedPtr<SGridPanel> Grid, int32& Row, const FText& Name, const FText& Value);
+};
+
+#undef LOCTEXT_NAMESPACE
diff --git a/Engine/Source/Developer/TraceInsights/Private/Insights/Widgets/STimerTableCell.cpp b/Engine/Source/Developer/TraceInsights/Private/Insights/Widgets/STimerTableCell.cpp
new file mode 100644
--- /dev/null
+++ b/Engine/Source/Developer/TraceInsights/Private/Insights/Widgets/STimerTableCell.cpp
@@ -0,0 +1,146 @@
+// Copyright 1998-2019 Epic Games, Inc. All Rights Reserved.
+
+#include "STimerTableCell.h"
+
+#include "EditorStyleSet.h"
+#include "SlateOptMacros.h"
+#include "Widgets/Images/SImage.h"
+#include "Widgets/Input/SButton.h"
+#include "Widgets/SBoxPanel.h"
+#include "Widgets/Text/STextBlock.h"
+#include "Widgets/Views/SExpanderArrow.h"
+
+// Insights
+#include "Insights/ViewModels/TimersViewColumn.h"
+#include "Insights/ViewModels/TimersViewColumnFactory.h"
+#include "Insights/Widgets/STimersViewTooltip.h"
+
+#define LOCTEXT_NAMESPACE "STimersView"
+
+////////////////////////////////////////////////////////////////////////////////////////////////////
+
+BEGIN_SLATE_FUNCTION_BUILD_OPTIMIZATION
+
+void STimerTableCell::Construct(const FArguments& InArgs, const TSharedRef<class ITableRow>& TableRow)
+{
+	SetHoveredTableCellDelegate = InArgs._OnSetHoveredTableCell;
+	TimerNodePtr = InArgs._TimerNodePtr;
+	ColumnId = InArgs._ColumnId;
+	//HighlightText = InArgs._HighlightText;
+
+	ChildSlot
+		[
+			GenerateWidgetForColumn(InArgs, TableRow)
+		];
+}
+
+////////////////////////////////////////////////////////////////////////////////////////////////////
+
+TSharedRef<SWidget> STimerTableCell::GenerateWidgetForColumn(const FArguments& InArgs, const TSharedRef<class ITableRow>& TableRow)
+{
+	if (InArgs._IsTimerNameColumn)
+	{
+		return GenerateWidgetForNameColumn(InArgs, TableRow);
+	}
+	else
+	{
+		return GenerateWidgetForStatsColumn(InArgs, TableRow);
+	}
+}
+
+////////////////////////////////////////////////////////////////////////////////////////////////////
+
+TSharedRef<SWidget> STimerTableCell::GenerateWidgetForNameColumn(const FArguments& InArgs, const TSharedRef<class ITableRow>& TableRow)
+{
+	const FTimersViewColumn& Column = *FTimersViewColumnFactory::Get().ColumnIdToPtrMapping.FindChecked(ColumnId);
+
+	return
+		SNew(SHorizontalBox)
+
+		+ SHorizontalBox::Slot()
+		.AutoWidth()
+		.HAlign(HAlign_Right)
+		.VAlign(VAlign_Center)
+		[
+			SNew(SExpanderArrow, TableRow)
+		]
+
+		// Event info icon + tooltip.
+		+ SHorizontalBox::Slot()
+		.AutoWidth()
+		.HAlign(HAlign_Center)
+		.VAlign(VAlign_Center)
+		[
+			SNew(SImage)
+			.Visibility(this, &STimerTableCell::GetHintIconVisibility)
+			.Image(FEditorStyle::GetBrush("Profiler.Tooltip.HintIcon10"))
+			.ToolTip(STimersViewTooltip::GetTableCellTooltip(TimerNodePtr))
+		]
+
+		// Name
+		+ SHorizontalBox::Slot()
+		.AutoWidth()
+		.VAlign(VAlign_Center)
+		.HAlign(Column.HorizontalAlignment)
+		.Padding(FMargin(2.0f, 0.0f))
+		[
+			SNew(STextBlock)
+			//.Text(TimerNodePtr->GetNameEx())
+			.Text(this, &STimerTableCell::GetNameEx)
+			.HighlightText(InArgs._HighlightText)
+			.TextStyle(FEditorStyle::Get(), TEXT("Profiler.Tooltip"))
+			.ColorAndOpacity(this, &STimerTableCell::GetColorAndOpacity)
+			.ShadowColorAndOpacity(this, &STimerTableCell::GetShadowColorAndOpacity)
+		]
+
+		/*
+		// Culled children warning icon + tooltip.
+		+ SHorizontalBox::Slot()
+		.AutoWidth()
+		.HAlign(HAlign_Center)
+		.VAlign(VAlign_Center)
+		[
+			SNew(SButton)
+			.ButtonStyle(FEditorStyle::Get(), TEXT("HoverHintOnly"))
+			.ContentPadding(0.0f)
+			.IsFocusable(false)
+			//.OnClicked(this, &STimerTableCell::ExpandCulledEvents_OnClicked)
+			[
+				SNew(SImage)
+				.Visibility(this, &STimerTableCell::GetCulledEventsIconVisibility)
+				.Image(FEditorStyle::GetBrush("Profiler.EventGraph.HasCulledEventsSmall"))
+				.ToolTipText(LOCTEXT("HasCulledEvents_TT", "This event contains culled children, if you want to see all children, please disable culling or use function details, or press this icon"))
+			]
+		]*/
+		;
+}
+
+////////////////////////////////////////////////////////////////////////////////////////////////////
+
+TSharedRef<SWidget> STimerTableCell::GenerateWidgetForStatsColumn(const FArguments& InArgs, const TSharedRef<class ITableRow>& TableRow)
+{
+	const FTimersViewColumn& Column = *FTimersViewColumnFactory::Get().ColumnIdToPtrMapping.FindChecked(ColumnId);
+	const FText FormattedValue = Column.GetFormattedValue(*TimerNodePtr);
+
+	return
+		SNew(SHorizontalBox)
+
+		+ SHorizontalBox::Slot()
+		.FillWidth(1.0f)
+		.VAlign(VAlign_Center)
+		.HAlign(Column.HorizontalAlignment)
+		.Padding(FMargin(2.0f, 0.0f))
+		[
+			SNew(STextBlock)
+			.Text(FormattedValue)
+			.TextStyle(FEditorStyle::Get(), TEXT("Profiler.Tooltip"))
+			.ColorAndOpacity(this, &STimerTableCell::GetStatsColorAndOpacity)
+			.ShadowColorAndOpacity(this, &STimerTableCell::GetShadowColorAndOpacity)
+		];
+}
+
+END_SLATE_FUNCTION_BUILD_OPTIMIZATION
+
+////////////////////////////////////////////////////////////////////////////////////////////////////
+
+#undef LOCTEXT_NAMESPACE
