Copying //UE4/Dev-Core to Dev-Main (//UE4/Dev-Main)

diff --git a/Engine/Source/Developer/TraceInsights/Private/Insights/Widgets/STimerTableCell.h b/Engine/Source/Developer/TraceInsights/Private/Insights/Widgets/STimerTableCell.h
new file mode 100644
--- /dev/null
+++ b/Engine/Source/Developer/TraceInsights/Private/Insights/Widgets/STimerTableCell.h
@@ -0,0 +1,140 @@
+// Copyright 1998-2019 Epic Games, Inc. All Rights Reserved.
+
+#pragma once
+
+#include "CoreMinimal.h"
+#include "Widgets/DeclarativeSyntaxSupport.h"
+#include "Widgets/SCompoundWidget.h"
+
+// Insights
+#include "Insights/ViewModels/TimerNode.h"
+
+DECLARE_DELEGATE_TwoParams(FSetHoveredTimerTableCell, const FName /*ColumnId*/, const FTimerNodePtr /*SamplePtr*/);
+DECLARE_DELEGATE_RetVal_OneParam(bool, FIsColumnVisibleDelegate, const FName /*ColumnId*/);
+DECLARE_DELEGATE_RetVal_OneParam(EHorizontalAlignment, FGetColumnOutlineHAlignmentDelegate, const FName /*ColumnId*/);
+
+class STimerTableCell : public SCompoundWidget
+{
+public:
+	SLATE_BEGIN_ARGS(STimerTableCell) {}
+		SLATE_EVENT(FSetHoveredTimerTableCell, OnSetHoveredTableCell)
+		SLATE_ATTRIBUTE(FText, HighlightText)
+		SLATE_ARGUMENT(FTimerNodePtr, TimerNodePtr)
+		SLATE_ARGUMENT(FName, ColumnId)
+		SLATE_ARGUMENT(bool, IsTimerNameColumn)
+	SLATE_END_ARGS()
+
+	void Construct(const FArguments& InArgs, const TSharedRef<class ITableRow>& TableRow);
+
+protected:
+	TSharedRef<SWidget> GenerateWidgetForColumn(const FArguments& InArgs, const TSharedRef<class ITableRow>& TableRow);
+	TSharedRef<SWidget> GenerateWidgetForNameColumn(const FArguments& InArgs, const TSharedRef<class ITableRow>& TableRow);
+	TSharedRef<SWidget> GenerateWidgetForStatsColumn(const FArguments& InArgs, const TSharedRef<class ITableRow>& TableRow);
+
+	/**
+	 * The system will use this event to notify a widget that the cursor has entered it. This event is NOT bubbled.
+	 *
+	 * @param MyGeometry The Geometry of the widget receiving the event
+	 * @param MouseEvent Information about the input event
+	 */
+	virtual void OnMouseEnter(const FGeometry& MyGeometry, const FPointerEvent& MouseEvent) override
+	{
+		SCompoundWidget::OnMouseEnter(MyGeometry, MouseEvent);
+		SetHoveredTableCellDelegate.ExecuteIfBound(ColumnId, TimerNodePtr);
+	}
+
+	/**
+	 * The system will use this event to notify a widget that the cursor has left it. This event is NOT bubbled.
+	 *
+	 * @param MouseEvent Information about the input event
+	 */
+	virtual void OnMouseLeave(const FPointerEvent& MouseEvent) override
+	{
+		SCompoundWidget::OnMouseLeave(MouseEvent);
+		SetHoveredTableCellDelegate.ExecuteIfBound(NAME_None, nullptr);
+	}
+
+	/**
+	 * Called during drag and drop when the drag enters a widget.
+	 *
+	 * Enter/Leave events in slate are meant as lightweight notifications.
+	 * So we do not want to capture mouse or set focus in response to these.
+	 * However, OnDragEnter must also support external APIs (e.g. OLE Drag/Drop)
+	 * Those require that we let them know whether we can handle the content
+	 * being dragged OnDragEnter.
+	 *
+	 * The concession is to return a can_handled/cannot_handle
+	 * boolean rather than a full FReply.
+	 *
+	 * @param MyGeometry      The geometry of the widget receiving the event.
+	 * @param DragDropEvent   The drag and drop event.
+	 *
+	 * @return A reply that indicated whether the contents of the DragDropEvent can potentially be processed by this widget.
+	 */
+	virtual void OnDragEnter(const FGeometry& MyGeometry, const FDragDropEvent& DragDropEvent) override
+	{
+		SCompoundWidget::OnDragEnter(MyGeometry, DragDropEvent);
+		SetHoveredTableCellDelegate.ExecuteIfBound(ColumnId, TimerNodePtr);
+	}
+
+	/**
+	 * Called during drag and drop when the drag leaves a widget.
+	 *
+	 * @param DragDropEvent   The drag and drop event.
+	 */
+	virtual void OnDragLeave(const FDragDropEvent& DragDropEvent)  override
+	{
+		SCompoundWidget::OnDragLeave(DragDropEvent);
+		SetHoveredTableCellDelegate.ExecuteIfBound(NAME_None, nullptr);
+	}
+
+	EVisibility GetHintIconVisibility() const
+	{
+		return IsHovered() ? EVisibility::Visible : EVisibility::Hidden;
+	}
+
+	EVisibility GetCulledEventsIconVisibility() const
+	{
+		//return TimerNodePtr->HasCulledChildren() ? EVisibility::Visible : EVisibility::Collapsed;
+		return EVisibility::Collapsed;
+	}
+
+	FText GetNameEx() const
+	{
+		return TimerNodePtr->GetNameEx();
+	}
+
+	FSlateColor GetColorAndOpacity() const
+	{
+		const FLinearColor TextColor =
+			TimerNodePtr->IsFiltered() ? FLinearColor(1.0f, 1.0f, 1.0f, 0.5f) :
+			FLinearColor::White;
+		return TextColor;
+	}
+
+	FSlateColor GetStatsColorAndOpacity() const
+	{
+		const FLinearColor TextColor =
+			TimerNodePtr->IsFiltered() ? FLinearColor(1.0f, 1.0f, 1.0f, 0.5f) :
+			TimerNodePtr->GetAggregatedStats().InstanceCount == 0 ? FLinearColor(1.0f, 1.0f, 1.0f, 0.6f) :
+			FLinearColor::White;
+		return TextColor;
+	}
+
+	FLinearColor GetShadowColorAndOpacity() const
+	{
+		const FLinearColor ShadowColor =
+			TimerNodePtr->IsFiltered() ? FLinearColor(0.f, 0.f, 0.f, 0.25f) :
+			FLinearColor(0.0f, 0.0f, 0.0f, 0.5f);
+		return ShadowColor;
+	}
+
+protected:
+	/** A shared pointer to the timer node. */
+	FTimerNodePtr TimerNodePtr;
+
+	/** The Id of the column where this timer belongs. */
+	FName ColumnId;
+
+	FSetHoveredTimerTableCell SetHoveredTableCellDelegate;
+};
diff --git a/Engine/Source/Developer/TraceInsights/Private/Insights/Widgets/STimerTableRow.cpp b/Engine/Source/Developer/TraceInsights/Private/Insights/Widgets/STimerTableRow.cpp
new file mode 100644
--- /dev/null
+++ b/Engine/Source/Developer/TraceInsights/Private/Insights/Widgets/STimerTableRow.cpp
@@ -0,0 +1,273 @@
+// Copyright 1998-2019 Epic Games, Inc. All Rights Reserved.
+
+#include "STimerTableRow.h"
+
+#include "Widgets/SOverlay.h"
+#include "Widgets/Images/SImage.h"
+
+// Insights
+#include "Insights/Common/TimeUtils.h"
+#include "Insights/ViewModels/TimersViewColumnFactory.h"
+#include "Insights/Widgets/STimersViewTooltip.h"
+#include "Insights/Widgets/STimerTableCell.h"
+
+#define LOCTEXT_NAMESPACE "STimersView"
+
+BEGIN_SLATE_FUNCTION_BUILD_OPTIMIZATION
+
+void STimerTableRow::Construct(const FArguments& InArgs, const TSharedRef<STableViewBase>& InOwnerTableView)
+{
+	OnShouldBeEnabled = InArgs._OnShouldBeEnabled;
+	IsColumnVisibleDelegate = InArgs._OnIsColumnVisible;
+	SetHoveredTableCellDelegate = InArgs._OnSetHoveredTableCell;
+	GetColumnOutlineHAlignmentDelegate = InArgs._OnGetColumnOutlineHAlignmentDelegate;
+
+	HighlightText = InArgs._HighlightText;
+	HighlightedTimerName = InArgs._HighlightedTimerName;
+
+	TimerNodePtr = InArgs._TimerNodePtr;
+
+	SetEnabled(TAttribute<bool>(this, &STimerTableRow::HandleShouldBeEnabled));
+
+	SMultiColumnTableRow<FTimerNodePtr>::Construct(SMultiColumnTableRow<FTimerNodePtr>::FArguments(), InOwnerTableView);
+
+	/*
+	const FSlateBrush* const NodeIcon = TimerNodeTypeHelper::GetIconForTimerNodeType(InTimerNode->GetType());
+	const TSharedRef<SToolTip> Tooltip = InTimerNode->IsGroup() ? SNew(SToolTip) : STimersViewTooltip(InTimerNode->GetId()).GetTooltip();
+
+	ChildSlot
+		[
+			SNew(SHorizontalBox)
+
+			// Expander arrow.
+			+ SHorizontalBox::Slot()
+			.AutoWidth()
+			.HAlign(HAlign_Right)
+			.VAlign(VAlign_Center)
+			[
+				SNew(SExpanderArrow, SharedThis(this))
+			]
+
+			// Icon to visualize group or timer type.
+			+ SHorizontalBox::Slot()
+			.AutoWidth()
+			.Padding(0.0f, 0.0f, 8.0f, 0.0f)
+			[
+				SNew(SImage)
+				.Image(NodeIcon)
+				.ToolTip(Tooltip)
+			]
+
+			// Description text.
+			+ SHorizontalBox::Slot()
+			.AutoWidth()
+			.VAlign(VAlign_Center)
+			.HAlign(HAlign_Left)
+			.Padding(FMargin(2.0f, 0.0f))
+			[
+				SNew(STextBlock)
+				.Text(this, &STimerTableRow::GetText)
+				.HighlightText(InArgs._HighlightText)
+				.TextStyle(FEditorStyle::Get(), TEXT("Profiler.Tooltip"))
+				.ColorAndOpacity(this, &STimerTableRow::GetColorAndOpacity)
+			]
+
+			+ SHorizontalBox::Slot()
+			.AutoWidth()
+			.HAlign(HAlign_Center)
+			.VAlign(VAlign_Top)
+			.Padding(0.0f, 1.0f, 0.0f, 0.0f)
+			[
+				SNew(SImage)
+				.Visibility(!InTimerNode->IsGroup() ? EVisibility::Visible : EVisibility::Collapsed)
+				.Image(FEditorStyle::GetBrush("Profiler.Tooltip.HintIcon10"))
+				.ToolTip(Tooltip)
+			]
+		];
+
+	STableRow<FTimerNodePtr>::ConstructInternal(STableRow::FArguments().ShowSelection(true), InOwnerTableView);
+	*/
+}
+
+TSharedRef<SWidget> STimerTableRow::GenerateWidgetForColumn(const FName& ColumnId)
+{
+	return
+
+	SNew(SOverlay)
+	.Visibility(EVisibility::SelfHitTestInvisible)
+
+	+SOverlay::Slot()
+	.Padding(0.0f)
+	[
+		SNew(SImage)
+		.Image(FEditorStyle::GetBrush("Profiler.LineGraphArea"))
+		.ColorAndOpacity(this, &STimerTableRow::GetBackgroundColorAndOpacity)
+	]
+
+	+SOverlay::Slot()
+	.Padding(0.0f)
+	[
+		SNew(SImage)
+		.Image(this, &STimerTableRow::GetOutlineBrush, ColumnId)
+		.ColorAndOpacity(this, &STimerTableRow::GetOutlineColorAndOpacity)
+	]
+
+	+SOverlay::Slot()
+	[
+		SNew(STimerTableCell, SharedThis(this))
+		.Visibility(this, &STimerTableRow::IsColumnVisible, ColumnId)
+		.TimerNodePtr(TimerNodePtr)
+		.ColumnId(ColumnId)
+		.HighlightText(HighlightText)
+		.IsTimerNameColumn(ColumnId == FTimersViewColumnFactory::Get().Collection[0]->Id) // name column
+		.OnSetHoveredTableCell(this, &STimerTableRow::OnSetHoveredTableCell)
+	];
+}
+
+END_SLATE_FUNCTION_BUILD_OPTIMIZATION
+
+FReply STimerTableRow::OnDragDetected(const FGeometry& MyGeometry, const FPointerEvent& MouseEvent)
+{
+	//im:TODO
+	//if (MouseEvent.IsMouseButtonDown(EKeys::LeftMouseButton))
+	//{
+	//	if (TimerNode->IsGroup())
+	//	{
+	//		// Add all timer Ids for the group.
+	//		TArray<int32> TimerIds;
+	//		const TArray<FTimerNodePtr>& FilteredChildren = TimerNode->GetFilteredChildren();
+	//		const int32 NumFilteredChildren = FilteredChildren.Num();
+	//
+	//		TimerIds.Reserve(NumFilteredChildren);
+	//		for (int32 Nx = 0; Nx < NumFilteredChildren; ++Nx)
+	//		{
+	//			TimerIds.Add(FilteredChildren[Nx]->GetId());
+	//		}
+	//
+	//		return FReply::Handled().BeginDragDrop(FStatIDDragDropOp::NewGroup(TimerIds, TimerNode->GetName().GetPlainNameString()));
+	//	}
+	//	else
+	//	{
+	//		return FReply::Handled().BeginDragDrop(FStatIDDragDropOp::NewSingle(TimerNode->GetId(), TimerNode->GetName().GetPlainNameString()));
+	//	}
+	//}
+
+	return SMultiColumnTableRow<FTimerNodePtr>::OnDragDetected(MyGeometry, MouseEvent);
+}
+
+FText STimerTableRow::GetText() const
+{
+	FText Text = FText::GetEmpty();
+
+	if (TimerNodePtr->IsGroup())
+	{
+		Text = FText::Format(LOCTEXT("TimerNode_GroupNodeTextFmt", "{0} ({1})"), FText::FromName(TimerNodePtr->GetName()), FText::AsNumber(TimerNodePtr->GetChildren().Num()));
+	}
+	else
+	{
+		Text = FText::FromName(TimerNodePtr->GetName());
+	}
+
+	return Text;
+}
+
+FSlateFontInfo STimerTableRow::GetFont() const
+{
+	const bool bIsStatTracked = false;//im:TODO: FTimingProfilerManager::Get()->IsStatTracked(TimerNodePtr->GetTimerId());
+	const FSlateFontInfo FontInfo = bIsStatTracked ? FEditorStyle::GetFontStyle("BoldFont") : FEditorStyle::GetFontStyle("NormalFont");
+	return FontInfo;
+}
+
+FSlateColor STimerTableRow::GetColorAndOpacity() const
+{
+	//const bool bIsStatTracked = FTimingProfilerManager::Get()->IsStatTracked(TimerNodePtr->GetId());
+	//const FSlateColor Color = bIsStatTracked ? FTimingProfilerManager::Get()->GetColorForTimerId(TimerNodePtr->GetId()) : FLinearColor::White;
+	//return Color;
+	//FTimingProfilerManager::GetSettings().GetColorForStat(StatName)
+	return FLinearColor::White;
+}
+
+FSlateColor STimerTableRow::GetBackgroundColorAndOpacity() const
+{
+	return GetBackgroundColorAndOpacity(TimerNodePtr->GetAggregatedStats().TotalInclusiveTime);
+	//return FLinearColor(0.0f, 0.0f, 0.0f, 1.0f)
+}
+
+FSlateColor STimerTableRow::GetBackgroundColorAndOpacity(double Time) const
+{
+	const FLinearColor Color =	Time > TimeUtils::Second      ? FLinearColor(0.3f, 0.0f, 0.0f, 1.0f) :
+								Time > TimeUtils::Milisecond  ? FLinearColor(0.3f, 0.1f, 0.0f, 1.0f) :
+								Time > TimeUtils::Microsecond ? FLinearColor(0.0f, 0.1f, 0.0f, 1.0f) :
+																FLinearColor(0.0f, 0.0f, 0.0f, 1.0f);
+	return Color;
+}
+
+FSlateColor STimerTableRow::GetOutlineColorAndOpacity() const
+{
+	const FLinearColor NoColor(0.0f, 0.0f, 0.0f, 0.0f);
+	const bool bShouldBeHighlighted = TimerNodePtr->GetName() == HighlightedTimerName.Get();
+	const FLinearColor OutlineColorAndOpacity = bShouldBeHighlighted ? FLinearColor(FColorList::SlateBlue) : NoColor;
+	return OutlineColorAndOpacity;
+}
+
+const FSlateBrush* STimerTableRow::GetOutlineBrush(const FName ColumnId) const
+{
+	EHorizontalAlignment Result = HAlign_Center;
+	if (IsColumnVisibleDelegate.IsBound())
+	{
+		Result = GetColumnOutlineHAlignmentDelegate.Execute(ColumnId);
+	}
+
+	const FSlateBrush* Brush = nullptr;
+	if (Result == HAlign_Left)
+	{
+		Brush = FEditorStyle::GetBrush("Profiler.EventGraph.Border.L");
+	}
+	else if(Result == HAlign_Right)
+	{
+		Brush = FEditorStyle::GetBrush("Profiler.EventGraph.Border.R");
+	}
+	else
+	{
+		Brush = FEditorStyle::GetBrush("Profiler.EventGraph.Border.TB");
+	}
+	return Brush;
+}
+
+bool STimerTableRow::HandleShouldBeEnabled() const
+{
+	bool bResult = false;
+
+	if (TimerNodePtr->IsGroup())
+	{
+		bResult = true;
+	}
+	else
+	{
+		if (OnShouldBeEnabled.IsBound())
+		{
+			bResult = OnShouldBeEnabled.Execute(TimerNodePtr->GetId());
+		}
+	}
+
+	return bResult;
+}
+
+EVisibility STimerTableRow::IsColumnVisible(const FName ColumnId) const
+{
+	EVisibility Result = EVisibility::Collapsed;
+
+	if (IsColumnVisibleDelegate.IsBound())
+	{
+		Result = IsColumnVisibleDelegate.Execute(ColumnId) ? EVisibility::Visible : EVisibility::Collapsed;
+	}
+
+	return Result;
+}
+
+void STimerTableRow::OnSetHoveredTableCell(const FName InColumnId, const FTimerNodePtr InSamplePtr)
+{
+	SetHoveredTableCellDelegate.ExecuteIfBound(InColumnId, InSamplePtr);
+}
+
+#undef LOCTEXT_NAMESPACE
diff --git a/Engine/Source/Developer/TraceInsights/Private/Insights/Widgets/STimerTableRow.h b/Engine/Source/Developer/TraceInsights/Private/Insights/Widgets/STimerTableRow.h
new file mode 100644
--- /dev/null
+++ b/Engine/Source/Developer/TraceInsights/Private/Insights/Widgets/STimerTableRow.h
@@ -0,0 +1,93 @@
+// Copyright 1998-2019 Epic Games, Inc. All Rights Reserved.
+
+#pragma once
+
+#include "CoreMinimal.h"
+#include "Widgets/DeclarativeSyntaxSupport.h"
+#include "Widgets/Views/SHeaderRow.h"
+#include "Widgets/Views/STableViewBase.h"
+#include "Widgets/Views/STableRow.h"
+#include "Widgets/Views/STreeView.h"
+
+// Insights
+#include "Insights/ViewModels/TimerNodeHelper.h"
+
+DECLARE_DELEGATE_RetVal_OneParam(bool, FShouldBeEnabledDelegate, const uint32 /*TimerId*/);
+DECLARE_DELEGATE_RetVal_OneParam(bool, FIsColumnVisibleDelegate, const FName /*ColumnId*/);
+DECLARE_DELEGATE_TwoParams(FSetHoveredTimerTableCell, const FName /*ColumnId*/, const FTimerNodePtr /*TimerNodePtr*/);
+DECLARE_DELEGATE_RetVal_OneParam(EHorizontalAlignment, FGetColumnOutlineHAlignmentDelegate, const FName /*ColumnId*/);
+
+/** Widget that represents a table row in the timers' tree control. Generates widgets for each column on demand. */
+class STimerTableRow : public SMultiColumnTableRow<FTimerNodePtr>
+{
+public:
+	SLATE_BEGIN_ARGS(STimerTableRow) {}
+		SLATE_EVENT(FShouldBeEnabledDelegate, OnShouldBeEnabled)
+		SLATE_EVENT(FIsColumnVisibleDelegate, OnIsColumnVisible)
+		SLATE_EVENT(FSetHoveredTimerTableCell, OnSetHoveredTableCell)
+		SLATE_EVENT(FGetColumnOutlineHAlignmentDelegate, OnGetColumnOutlineHAlignmentDelegate)
+		SLATE_ATTRIBUTE(FText, HighlightText)
+		SLATE_ATTRIBUTE(FName, HighlightedTimerName)
+		SLATE_ARGUMENT(FTimerNodePtr, TimerNodePtr)
+	SLATE_END_ARGS()
+
+public:
+	void Construct(const FArguments& InArgs, const TSharedRef<STableViewBase>& InOwnerTableView);
+
+	virtual TSharedRef<SWidget> GenerateWidgetForColumn(const FName& ColumnId) override;
+
+	/**
+	 * Called when Slate detects that a widget started to be dragged.
+	 * Usage:
+	 * A widget can ask Slate to detect a drag.
+	 * OnMouseDown() reply with FReply::Handled().DetectDrag(SharedThis(this)).
+	 * Slate will either send an OnDragDetected() event or do nothing.
+	 * If the user releases a mouse button or leaves the widget before
+	 * a drag is triggered (maybe user started at the very edge) then no event will be
+	 * sent.
+	 *
+	 * @param  InMyGeometry  Widget geometry
+	 * @param  InMouseEvent  MouseMove that triggered the drag
+	 *
+	 */
+	virtual FReply OnDragDetected(const FGeometry& MyGeometry, const FPointerEvent& MouseEvent) override;
+
+protected:
+	/**
+	 * @return a text which describes this table row, refers to both groups and timers
+	 */
+	FText GetText() const;
+
+	/**
+	 * @return a font style which is used to draw this table row, refers to both groups and timers
+	 */
+	FSlateFontInfo GetFont() const;
+
+	/**
+	 * @return a color and opacity value used to draw this table row, refers to both groups and timers
+	 */
+	FSlateColor GetColorAndOpacity() const;
+
+	FSlateColor GetBackgroundColorAndOpacity() const;
+	FSlateColor GetBackgroundColorAndOpacity(double Time) const;
+	FSlateColor GetOutlineColorAndOpacity() const;
+	const FSlateBrush* GetOutlineBrush(const FName ColumnId) const;
+	bool HandleShouldBeEnabled() const;
+	EVisibility IsColumnVisible(const FName ColumnId) const;
+	void OnSetHoveredTableCell(const FName ColumnId, const FTimerNodePtr SamplePtr);
+
+protected:
+	/** Data context for this table row. */
+	FTimerNodePtr TimerNodePtr;
+
+	FShouldBeEnabledDelegate OnShouldBeEnabled;
+	FIsColumnVisibleDelegate IsColumnVisibleDelegate;
+	FSetHoveredTimerTableCell SetHoveredTableCellDelegate;
+	FGetColumnOutlineHAlignmentDelegate GetColumnOutlineHAlignmentDelegate;
+
+	/** Text to be highlighted on timer name. */
+	TAttribute<FText> HighlightText;
+
+	/** Name of the timer that should be drawn as highlighted. */
+	TAttribute<FName> HighlightedTimerName;
+};
diff --git a/Engine/Source/Developer/TraceInsights/Private/Insights/Widgets/STimersView.cpp b/Engine/Source/Developer/TraceInsights/Private/Insights/Widgets/STimersView.cpp
new file mode 100644
--- /dev/null
+++ b/Engine/Source/Developer/TraceInsights/Private/Insights/Widgets/STimersView.cpp
@@ -0,0 +1,1542 @@
+// Copyright 1998-2019 Epic Games, Inc. All Rights Reserved.
+
+#include "STimersView.h"
+
+#include "EditorStyleSet.h"
+#include "Framework/MultiBox/MultiBoxBuilder.h"
+#include "SlateOptMacros.h"
+#include "TraceServices/AnalysisService.h"
+#include "Widgets/Input/SCheckBox.h"
+#include "Widgets/Layout/SScrollBox.h"
+#include "Widgets/Input/SSearchBox.h"
+#include "Widgets/Images/SImage.h"
+#include "Widgets/Layout/SGridPanel.h"
+#include "Widgets/Layout/SSeparator.h"
+#include "Widgets/SToolTip.h"
+#include "Widgets/Views/STableViewBase.h"
+
+// Insights
+#include "Insights/TimingProfilerCommon.h"
+#include "Insights/TimingProfilerManager.h"
+#include "Insights/ViewModels/TimerNodeHelper.h"
+#include "Insights/ViewModels/TimersViewColumnFactory.h"
+#include "Insights/Widgets/STimersViewTooltip.h"
+#include "Insights/Widgets/STimerTableRow.h"
+#include "Insights/Widgets/STimingProfilerWindow.h"
+#include "Insights/Widgets/STimingView.h"
+
+#define LOCTEXT_NAMESPACE "STimersView"
+
+////////////////////////////////////////////////////////////////////////////////////////////////////
+
+const EColumnSortMode::Type STimersView::DefaultColumnSortMode(EColumnSortMode::Descending);
+const FName STimersView::DefaultColumnBeingSorted(FTimersViewColumns::TotalInclusiveTimeColumnID);
+
+////////////////////////////////////////////////////////////////////////////////////////////////////
+
+STimersView::STimersView()
+	: bExpansionSaved(false)
+	, GroupingMode(ETimerGroupingMode::ByType)
+	, ColumnSortMode(DefaultColumnSortMode)
+	, ColumnBeingSorted(DefaultColumnBeingSorted)
+{
+	FMemory::Memset(bTimerNodeIsVisible, 1);
+}
+
+////////////////////////////////////////////////////////////////////////////////////////////////////
+
+STimersView::~STimersView()
+{
+	// Remove ourselves from the Insights manager.
+	if (FInsightsManager::Get().IsValid())
+	{
+		FInsightsManager::Get()->GetSessionChangedEvent().RemoveAll(this);
+	}
+}
+
+////////////////////////////////////////////////////////////////////////////////////////////////////
+
+BEGIN_SLATE_FUNCTION_BUILD_OPTIMIZATION
+void STimersView::Construct(const FArguments& InArgs)
+{
+	SAssignNew(ExternalScrollbar, SScrollBar)
+	.AlwaysShowScrollbar(true);
+
+	ChildSlot
+	[
+		SNew(SVerticalBox)
+
+		+SVerticalBox::Slot()
+		.VAlign(VAlign_Center)
+		.AutoHeight()
+		[
+			SNew(SBorder)
+			.BorderImage(FEditorStyle::GetBrush("ToolPanel.GroupBorder"))
+			.Padding(2.0f)
+			[
+				SNew(SVerticalBox)
+
+				// Search box
+				+SVerticalBox::Slot()
+				.VAlign(VAlign_Center)
+				.Padding(2.0f)
+				.AutoHeight()
+				[
+					SAssignNew(SearchBox, SSearchBox)
+					.HintText(LOCTEXT("SearchBoxHint", "Search timers or groups"))
+					.OnTextChanged(this, &STimersView::SearchBox_OnTextChanged)
+					.IsEnabled(this, &STimersView::SearchBox_IsEnabled)
+					.ToolTipText(LOCTEXT("FilterSearchHint", "Type here to search timer or group"))
+				]
+
+				// Group by
+				+SVerticalBox::Slot()
+				.VAlign(VAlign_Center)
+				.Padding(2.0f)
+				.AutoHeight()
+				[
+					SNew(SHorizontalBox)
+
+					+SHorizontalBox::Slot()
+					.FillWidth(1.0f)
+					.VAlign(VAlign_Center)
+					[
+						SNew(STextBlock)
+						.Text(LOCTEXT("GroupByText", "Group by"))
+					]
+
+					+SHorizontalBox::Slot()
+					.FillWidth(2.0f)
+					.VAlign(VAlign_Center)
+					[
+						SAssignNew(GroupByComboBox, SComboBox<TSharedPtr<ETimerGroupingMode>>)
+						.ToolTipText(this, &STimersView::GroupBy_GetSelectedTooltipText)
+						.OptionsSource(&GroupByOptionsSource)
+						.OnSelectionChanged(this, &STimersView::GroupBy_OnSelectionChanged)
+						.OnGenerateWidget(this, &STimersView::GroupBy_OnGenerateWidget)
+						[
+							SNew(STextBlock)
+							.Text(this, &STimersView::GroupBy_GetSelectedText)
+						]
+					]
+				]
+
+				// Check boxes for: GpuScope, ComputeScope, CpuScope
+				+ SVerticalBox::Slot()
+				.VAlign(VAlign_Center)
+				.Padding(2.0f)
+				.AutoHeight()
+				[
+					SNew(SHorizontalBox)
+
+					+SHorizontalBox::Slot()
+					.Padding(FMargin(0.0f,0.0f,1.0f,0.0f))
+					.FillWidth(1.0f)
+					[
+						GetToggleButtonForTimerType(ETimerNodeType::GpuScope)
+					]
+
+					//+SHorizontalBox::Slot()
+					//.Padding(FMargin(1.0f,0.0f,1.0f,0.0f))
+					//.FillWidth(1.0f)
+					//[
+					//	GetToggleButtonForTimerType(ETimerNodeType::ComputeScope)
+					//]
+
+					+SHorizontalBox::Slot()
+					.Padding(FMargin(1.0f,0.0f,1.0f,0.0f))
+					.FillWidth(1.0f)
+					[
+						GetToggleButtonForTimerType(ETimerNodeType::CpuScope)
+					]
+				]
+			]
+		]
+
+		// Tree view
+		+SVerticalBox::Slot()
+		.FillHeight(1.0f)
+		.Padding(0.0f, 6.0f, 0.0f, 0.0f)
+		[
+			SNew(SHorizontalBox)
+
+			+ SHorizontalBox::Slot()
+			.FillWidth(1.0f)
+			.Padding(0.0f)
+			[
+				SNew(SScrollBox)
+				.Orientation(Orient_Horizontal)
+
+				+ SScrollBox::Slot()
+				[
+					SNew(SBorder)
+					.BorderImage(FEditorStyle::GetBrush("ToolPanel.GroupBorder"))
+					.Padding(0.0f)
+					[
+						SAssignNew(TreeView, STreeView<FTimerNodePtr>)
+						.ExternalScrollbar(ExternalScrollbar)
+						.SelectionMode(ESelectionMode::Multi)
+						.TreeItemsSource(&FilteredGroupNodes)
+						.OnGetChildren(this, &STimersView::TreeView_OnGetChildren)
+						.OnGenerateRow(this, &STimersView::TreeView_OnGenerateRow)
+						.OnSelectionChanged(this, &STimersView::TreeView_OnSelectionChanged)
+						.OnMouseButtonDoubleClick(this, &STimersView::TreeView_OnMouseButtonDoubleClick)
+						.OnContextMenuOpening(FOnContextMenuOpening::CreateSP(this, &STimersView::TreeView_GetMenuContent))
+						.ItemHeight(12.0f)
+						.HeaderRow
+						(
+							SAssignNew(TreeViewHeaderRow, SHeaderRow)
+							.Visibility(EVisibility::Visible)
+						)
+					]
+				]
+			]
+
+			+ SHorizontalBox::Slot()
+			.AutoWidth()
+			.Padding(0.0f)
+			[
+				SNew(SBox)
+				.WidthOverride(FOptionalSize(13.0f))
+				[
+					ExternalScrollbar.ToSharedRef()
+				]
+			]
+		]
+	];
+
+	InitializeAndShowHeaderColumns();
+	//BindCommands();
+
+	// Create the search filters: text based, type based etc.
+	TextFilter = MakeShareable(new FTimerNodeTextFilter(FTimerNodeTextFilter::FItemToStringArray::CreateSP(this, &STimersView::HandleItemToStringArray)));
+	Filters = MakeShareable(new FTimerNodeFilterCollection());
+	Filters->Add(TextFilter);
+
+	CreateGroupByOptionsSources();
+
+	// Register ourselves with the Insights manager.
+	FInsightsManager::Get()->GetSessionChangedEvent().AddSP(this, &STimersView::InsightsManager_OnSessionChanged);
+}
+END_SLATE_FUNCTION_BUILD_OPTIMIZATION
+
+////////////////////////////////////////////////////////////////////////////////////////////////////
+
+TSharedPtr<SWidget> STimersView::TreeView_GetMenuContent()
+{
+	const TArray<FTimerNodePtr> SelectedTimerNodes = TreeView->GetSelectedItems();
+	const int32 NumSelectedTimerNodes = SelectedTimerNodes.Num();
+	FTimerNodePtr SelectedTimerNode = NumSelectedTimerNodes ? SelectedTimerNodes[0] : nullptr;
+
+	const FTimersViewColumn* const * ColumnPtrPtr = FTimersViewColumnFactory::Get().ColumnIdToPtrMapping.Find(HoveredColumnId);
+	const FTimersViewColumn* const ColumnPtr = (ColumnPtrPtr != nullptr) ? *ColumnPtrPtr : nullptr;
+
+	FText SelectionStr;
+	FText PropertyName;
+	FText PropertyValue;
+
+	if (NumSelectedTimerNodes == 0)
+	{
+		SelectionStr = LOCTEXT("NothingSelected", "Nothing selected");
+	}
+	else if (NumSelectedTimerNodes == 1)
+	{
+		if (ColumnPtr != nullptr)
+		{
+			PropertyName = ColumnPtr->ShortName;
+			PropertyValue = ColumnPtr->GetFormattedValue(*SelectedTimerNode);
+		}
+		SelectionStr = FText::FromName(SelectedTimerNode->GetName());
+	}
+	else
+	{
+		SelectionStr = LOCTEXT("MultipleSelection", "Multiple selection");
+	}
+
+	const bool bShouldCloseWindowAfterMenuSelection = true;
+	FMenuBuilder MenuBuilder(bShouldCloseWindowAfterMenuSelection, NULL);
+
+	// Selection menu
+	MenuBuilder.BeginSection("Selection", LOCTEXT("ContextMenu_Header_Selection", "Selection"));
+	{
+		struct FLocal
+		{
+			static bool ReturnFalse()
+			{
+				return false;
+			}
+		};
+
+		FUIAction DummyUIAction;
+		DummyUIAction.CanExecuteAction = FCanExecuteAction::CreateStatic(&FLocal::ReturnFalse);
+		MenuBuilder.AddMenuEntry
+		(
+			SelectionStr,
+			LOCTEXT("ContextMenu_Selection", "Currently selected items"),
+			FSlateIcon(FEditorStyle::GetStyleSetName(), "@missing.icon"), DummyUIAction, NAME_None, EUserInterfaceActionType::Button
+		);
+	}
+	MenuBuilder.EndSection();
+
+	MenuBuilder.BeginSection("Misc", LOCTEXT("ContextMenu_Header_Misc", "Miscellaneous"));
+	{
+		/*TODO
+		FUIAction Action_CopySelectedToClipboard
+		(
+			FExecuteAction::CreateSP(this, &STimersView::ContextMenu_CopySelectedToClipboard_Execute),
+			FCanExecuteAction::CreateSP(this, &STimersView::ContextMenu_CopySelectedToClipboard_CanExecute)
+		);
+		MenuBuilder.AddMenuEntry
+		(
+			LOCTEXT("ContextMenu_Header_Misc_CopySelectedToClipboard", "Copy To Clipboard"),
+			LOCTEXT("ContextMenu_Header_Misc_CopySelectedToClipboard_Desc", "Copies selection to clipboard"),
+			FSlateIcon(FEditorStyle::GetStyleSetName(), "Profiler.Misc.CopyToClipboard"), Action_CopySelectedToClipboard, NAME_None, EUserInterfaceActionType::Button
+		);
+		*/
+
+		MenuBuilder.AddSubMenu
+		(
+			LOCTEXT("ContextMenu_Header_Misc_Sort", "Sort By"),
+			LOCTEXT("ContextMenu_Header_Misc_Sort_Desc", "Sort by column"),
+			FNewMenuDelegate::CreateSP(this, &STimersView::TreeView_BuildSortByMenu),
+			false,
+			FSlateIcon(FEditorStyle::GetStyleSetName(), "Profiler.Misc.SortBy")
+		);
+	}
+	MenuBuilder.EndSection();
+
+	MenuBuilder.BeginSection("Columns", LOCTEXT("ContextMenu_Header_Columns", "Columns"));
+	{
+		MenuBuilder.AddSubMenu
+		(
+			LOCTEXT("ContextMenu_Header_Columns_View", "View Column"),
+			LOCTEXT("ContextMenu_Header_Columns_View_Desc", "Hides or shows columns"),
+			FNewMenuDelegate::CreateSP(this, &STimersView::TreeView_BuildViewColumnMenu),
+			false,
+			FSlateIcon(FEditorStyle::GetStyleSetName(), "Profiler.EventGraph.ViewColumn")
+		);
+
+		FUIAction Action_ShowAllColumns
+		(
+			FExecuteAction::CreateSP(this, &STimersView::ContextMenu_ShowAllColumns_Execute),
+			FCanExecuteAction::CreateSP(this, &STimersView::ContextMenu_ShowAllColumns_CanExecute)
+		);
+		MenuBuilder.AddMenuEntry
+		(
+			LOCTEXT("ContextMenu_Header_Columns_ShowAllColumns", "Show All Columns"),
+			LOCTEXT("ContextMenu_Header_Columns_ShowAllColumns_Desc", "Resets tree view to show all columns"),
+			FSlateIcon(FEditorStyle::GetStyleSetName(), "Profiler.EventGraph.ResetColumn"), Action_ShowAllColumns, NAME_None, EUserInterfaceActionType::Button
+		);
+
+		FUIAction Action_ShowMinMaxMedColumns
+		(
+			FExecuteAction::CreateSP(this, &STimersView::ContextMenu_ShowMinMaxMedColumns_Execute),
+			FCanExecuteAction::CreateSP(this, &STimersView::ContextMenu_ShowMinMaxMedColumns_CanExecute)
+		);
+		MenuBuilder.AddMenuEntry
+		(
+			LOCTEXT("ContextMenu_Header_Columns_ShowMinMaxMedColumns", "Reset Columns to Min/Max/Median Preset"),
+			LOCTEXT("ContextMenu_Header_Columns_ShowMinMaxMedColumns_Desc", "Resets columns to Min/Max/Median preset"),
+			FSlateIcon(FEditorStyle::GetStyleSetName(), "Profiler.EventGraph.ResetColumn"), Action_ShowMinMaxMedColumns, NAME_None, EUserInterfaceActionType::Button
+		);
+
+		FUIAction Action_ResetColumns
+		(
+			FExecuteAction::CreateSP(this, &STimersView::ContextMenu_ResetColumns_Execute),
+			FCanExecuteAction::CreateSP(this, &STimersView::ContextMenu_ResetColumns_CanExecute)
+		);
+		MenuBuilder.AddMenuEntry
+		(
+			LOCTEXT("ContextMenu_Header_Columns_ResetColumns", "Reset Columns to Default"),
+			LOCTEXT("ContextMenu_Header_Columns_ResetColumns_Desc", "Resets columns to default"),
+			FSlateIcon(FEditorStyle::GetStyleSetName(), "Profiler.EventGraph.ResetColumn"), Action_ResetColumns, NAME_None, EUserInterfaceActionType::Button
+		);
+	}
+	MenuBuilder.EndSection();
+
+	return MenuBuilder.MakeWidget();
+}
+
+////////////////////////////////////////////////////////////////////////////////////////////////////
+
+void STimersView::TreeView_BuildSortByMenu(FMenuBuilder& MenuBuilder)
+{
+	// TODO: Refactor later @see TSharedPtr<SWidget> SCascadePreviewViewportToolBar::GenerateViewMenu() const
+
+	MenuBuilder.BeginSection("ColumnName", LOCTEXT("ContextMenu_Header_Misc_ColumnName", "Column Name"));
+
+	for (auto It = TreeViewHeaderColumns.CreateConstIterator(); It; ++It)
+	{
+		const FTimersViewColumn& Column = It.Value();
+
+		if (Column.bIsVisible && Column.bCanBeSorted())
+		{
+			FUIAction Action_SortByColumn
+			(
+				FExecuteAction::CreateSP(this, &STimersView::ContextMenu_SortByColumn_Execute, Column.Id),
+				FCanExecuteAction::CreateSP(this, &STimersView::ContextMenu_SortByColumn_CanExecute, Column.Id),
+				FIsActionChecked::CreateSP(this, &STimersView::ContextMenu_SortByColumn_IsChecked, Column.Id)
+			);
+			MenuBuilder.AddMenuEntry
+			(
+				Column.TitleName,
+				Column.Description,
+				FSlateIcon(), Action_SortByColumn, NAME_None, EUserInterfaceActionType::RadioButton
+			);
+		}
+	}
+
+	MenuBuilder.EndSection();
+
+	//-----------------------------------------------------------------------------
+
+	MenuBuilder.BeginSection("SortMode", LOCTEXT("ContextMenu_Header_Misc_Sort_SortMode", "Sort Mode"));
+	{
+		FUIAction Action_SortAscending
+		(
+			FExecuteAction::CreateSP(this, &STimersView::ContextMenu_SortMode_Execute, EColumnSortMode::Ascending),
+			FCanExecuteAction::CreateSP(this, &STimersView::ContextMenu_SortMode_CanExecute, EColumnSortMode::Ascending),
+			FIsActionChecked::CreateSP(this, &STimersView::ContextMenu_SortMode_IsChecked, EColumnSortMode::Ascending)
+		);
+		MenuBuilder.AddMenuEntry
+		(
+			LOCTEXT("ContextMenu_Header_Misc_Sort_SortAscending", "Sort Ascending"),
+			LOCTEXT("ContextMenu_Header_Misc_Sort_SortAscending_Desc", "Sorts ascending"),
+			FSlateIcon(FEditorStyle::GetStyleSetName(), "Profiler.Misc.SortAscending"), Action_SortAscending, NAME_None, EUserInterfaceActionType::RadioButton
+		);
+
+		FUIAction Action_SortDescending
+		(
+			FExecuteAction::CreateSP(this, &STimersView::ContextMenu_SortMode_Execute, EColumnSortMode::Descending),
+			FCanExecuteAction::CreateSP(this, &STimersView::ContextMenu_SortMode_CanExecute, EColumnSortMode::Descending),
+			FIsActionChecked::CreateSP(this, &STimersView::ContextMenu_SortMode_IsChecked, EColumnSortMode::Descending)
+		);
+		MenuBuilder.AddMenuEntry
+		(
+			LOCTEXT("ContextMenu_Header_Misc_Sort_SortDescending", "Sort Descending"),
+			LOCTEXT("ContextMenu_Header_Misc_Sort_SortDescending_Desc", "Sorts descending"),
+			FSlateIcon(FEditorStyle::GetStyleSetName(), "Profiler.Misc.SortDescending"), Action_SortDescending, NAME_None, EUserInterfaceActionType::RadioButton
+		);
+	}
+	MenuBuilder.EndSection();
+}
+
+////////////////////////////////////////////////////////////////////////////////////////////////////
+
+void STimersView::TreeView_BuildViewColumnMenu(FMenuBuilder& MenuBuilder)
+{
+	MenuBuilder.BeginSection("ViewColumn", LOCTEXT("ContextMenu_Header_Columns_View", "View Column"));
+
+	for (auto It = TreeViewHeaderColumns.CreateConstIterator(); It; ++It)
+	{
+		const FTimersViewColumn& Column = It.Value();
+
+		FUIAction Action_ToggleColumn
+		(
+			FExecuteAction::CreateSP(this, &STimersView::ContextMenu_ToggleColumn_Execute, Column.Id),
+			FCanExecuteAction::CreateSP(this, &STimersView::ContextMenu_ToggleColumn_CanExecute, Column.Id),
+			FIsActionChecked::CreateSP(this, &STimersView::ContextMenu_ToggleColumn_IsChecked, Column.Id)
+		);
+		MenuBuilder.AddMenuEntry
+		(
+			Column.TitleName,
+			Column.Description,
+			FSlateIcon(), Action_ToggleColumn, NAME_None, EUserInterfaceActionType::ToggleButton
+		);
+	}
+
+	MenuBuilder.EndSection();
+}
+
+////////////////////////////////////////////////////////////////////////////////////////////////////
+
+void STimersView::InitializeAndShowHeaderColumns()
+{
+	const int32 NumColumns = FTimersViewColumnFactory::Get().Collection.Num();
+	for (int32 ColumnIndex = 0; ColumnIndex < NumColumns; ColumnIndex++)
+	{
+		TreeViewHeaderRow_CreateColumnArgs(ColumnIndex);
+	}
+
+	for (auto It = TreeViewHeaderColumns.CreateConstIterator(); It; ++It)
+	{
+		const FTimersViewColumn& Column = It.Value();
+
+		if (Column.bIsVisible)
+		{
+			TreeViewHeaderRow_ShowColumn(Column.Id);
+		}
+	}
+}
+
+////////////////////////////////////////////////////////////////////////////////////////////////////
+
+void STimersView::TreeViewHeaderRow_CreateColumnArgs(const int32 ColumnIndex)
+{
+	const FTimersViewColumn& Column = *FTimersViewColumnFactory::Get().Collection[ColumnIndex];
+	SHeaderRow::FColumn::FArguments ColumnArgs;
+
+	ColumnArgs
+		.ColumnId(Column.Id)
+		.DefaultLabel(Column.ShortName)
+		.SortMode(EColumnSortMode::None)
+		.HAlignHeader(HAlign_Fill)
+		.VAlignHeader(VAlign_Fill)
+		.HeaderContentPadding(TOptional<FMargin>(2.0f))
+		.HAlignCell(HAlign_Fill)
+		.VAlignCell(VAlign_Fill)
+		.SortMode(this, &STimersView::GetSortModeForColumn, Column.Id)
+		.OnSort(this, &STimersView::OnSortModeChanged)
+		.ManualWidth(Column.InitialColumnWidth)
+		.FixedWidth(Column.bIsFixedColumnWidth() ? Column.InitialColumnWidth : TOptional<float>())
+		.HeaderContent()
+		[
+			SNew(SHorizontalBox)
+			.ToolTip(STimersViewTooltip::GetColumnTooltip(Column))
+
+			+ SHorizontalBox::Slot()
+			.HAlign(Column.HorizontalAlignment)
+			.VAlign(VAlign_Center)
+			[
+				SNew(STextBlock)
+				.Text(Column.ShortName)
+			]
+		]
+		.MenuContent()
+		[
+			TreeViewHeaderRow_GenerateColumnMenu(Column)
+		];
+
+	TreeViewHeaderColumnArgs.Add(Column.Id, ColumnArgs);
+	TreeViewHeaderColumns.Add(Column.Id, Column);
+}
+
+////////////////////////////////////////////////////////////////////////////////////////////////////
+
+void STimersView::TreeViewHeaderRow_ShowColumn(const FName ColumnId)
+{
+	FTimersViewColumn& Column = TreeViewHeaderColumns.FindChecked(ColumnId);
+	Column.bIsVisible = true;
+	SHeaderRow::FColumn::FArguments& ColumnArgs = TreeViewHeaderColumnArgs.FindChecked(ColumnId);
+
+	const int32 NumColumns = TreeViewHeaderRow->GetColumns().Num();
+	int32 ColumnIndex = 0;
+	for (; ColumnIndex < NumColumns; ColumnIndex++)
+	{
+		const SHeaderRow::FColumn& CurrentColumn = TreeViewHeaderRow->GetColumns()[ColumnIndex];
+		const FTimersViewColumn& CurrentTimersViewColumn = TreeViewHeaderColumns.FindChecked(CurrentColumn.ColumnId);
+		if (Column.Order < CurrentTimersViewColumn.Order)
+			break;
+	}
+
+	TreeViewHeaderRow->InsertColumn(ColumnArgs, ColumnIndex);
+}
+
+////////////////////////////////////////////////////////////////////////////////////////////////////
+
+TSharedRef<SWidget> STimersView::TreeViewHeaderRow_GenerateColumnMenu(const FTimersViewColumn& Column)
+{
+	bool bIsMenuVisible = false;
+
+	const bool bShouldCloseWindowAfterMenuSelection = true;
+	FMenuBuilder MenuBuilder(bShouldCloseWindowAfterMenuSelection, NULL);
+	{
+		if (Column.bCanBeHidden())
+		{
+			MenuBuilder.BeginSection("Column", LOCTEXT("TreeViewHeaderRow_Header_Column", "Column"));
+
+			FUIAction Action_HideColumn
+			(
+				FExecuteAction::CreateSP(this, &STimersView::HeaderMenu_HideColumn_Execute, Column.Id),
+				FCanExecuteAction::CreateSP(this, &STimersView::HeaderMenu_HideColumn_CanExecute, Column.Id)
+			);
+
+			MenuBuilder.AddMenuEntry
+			(
+				LOCTEXT("TreeViewHeaderRow_HideColumn", "Hide"),
+				LOCTEXT("TreeViewHeaderRow_HideColumn_Desc", "Hides the selected column"),
+				FSlateIcon(), Action_HideColumn, NAME_None, EUserInterfaceActionType::Button
+			);
+			bIsMenuVisible = true;
+
+			MenuBuilder.EndSection();
+		}
+
+		if (Column.bCanBeSorted())
+		{
+			MenuBuilder.BeginSection("SortMode", LOCTEXT("ContextMenu_Header_Misc_Sort_SortMode", "Sort Mode"));
+
+			FUIAction Action_SortAscending
+			(
+				FExecuteAction::CreateSP(this, &STimersView::HeaderMenu_SortMode_Execute, Column.Id, EColumnSortMode::Ascending),
+				FCanExecuteAction::CreateSP(this, &STimersView::HeaderMenu_SortMode_CanExecute, Column.Id, EColumnSortMode::Ascending),
+				FIsActionChecked::CreateSP(this, &STimersView::HeaderMenu_SortMode_IsChecked, Column.Id, EColumnSortMode::Ascending)
+			);
+			MenuBuilder.AddMenuEntry
+			(
+				LOCTEXT("ContextMenu_Header_Misc_Sort_SortAscending", "Sort Ascending"),
+				LOCTEXT("ContextMenu_Header_Misc_Sort_SortAscending_Desc", "Sorts ascending"),
+				FSlateIcon(FEditorStyle::GetStyleSetName(), "Profiler.Misc.SortAscending"), Action_SortAscending, NAME_None, EUserInterfaceActionType::RadioButton
+			);
+
+			FUIAction Action_SortDescending
+			(
+				FExecuteAction::CreateSP(this, &STimersView::HeaderMenu_SortMode_Execute, Column.Id, EColumnSortMode::Descending),
+				FCanExecuteAction::CreateSP(this, &STimersView::HeaderMenu_SortMode_CanExecute, Column.Id, EColumnSortMode::Descending),
+				FIsActionChecked::CreateSP(this, &STimersView::HeaderMenu_SortMode_IsChecked, Column.Id, EColumnSortMode::Descending)
+			);
+			MenuBuilder.AddMenuEntry
+			(
+				LOCTEXT("ContextMenu_Header_Misc_Sort_SortDescending", "Sort Descending"),
+				LOCTEXT("ContextMenu_Header_Misc_Sort_SortDescending_Desc", "Sorts descending"),
+				FSlateIcon(FEditorStyle::GetStyleSetName(), "Profiler.Misc.SortDescending"), Action_SortDescending, NAME_None, EUserInterfaceActionType::RadioButton
+			);
+			bIsMenuVisible = true;
+
+			MenuBuilder.EndSection();
+		}
+
+		if (Column.bCanBeFiltered())
+		{
+			MenuBuilder.BeginSection("FilterMode", LOCTEXT("ContextMenu_Header_Misc_Filter_FilterMode", "Filter Mode"));
+			bIsMenuVisible = true;
+			MenuBuilder.EndSection();
+		}
+	}
+
+	/*
+	TODO:
+	- Show top ten
+	- Show top bottom
+	- Filter by list (avg, median, 10%, 90%, etc.)
+	- Text box for filtering for each column instead of one text box used for filtering
+	- Grouping button for flat view modes (show at most X groups, show all groups for names)
+	*/
+
+	return bIsMenuVisible ? MenuBuilder.MakeWidget() : (TSharedRef<SWidget>)SNullWidget::NullWidget;
+}
+
+////////////////////////////////////////////////////////////////////////////////////////////////////
+
+void STimersView::InsightsManager_OnSessionChanged()
+{
+	TSharedPtr<const Trace::IAnalysisSession> NewSession = FInsightsManager::Get()->GetSession();
+
+	if (NewSession != Session)
+	{
+		Session = NewSession;
+		RebuildTree();
+	}
+	else
+	{
+		UpdateTree();
+	}
+}
+
+////////////////////////////////////////////////////////////////////////////////////////////////////
+
+void STimersView::UpdateTree()
+{
+	// Create groups, sort timers within the group and apply filtering.
+	CreateGroups();
+	SortTimers();
+	ApplyFiltering();
+}
+
+////////////////////////////////////////////////////////////////////////////////////////////////////
+
+void STimersView::ApplyFiltering()
+{
+	FilteredGroupNodes.Reset();
+
+	// Apply filter to all groups and its children.
+	const int32 NumGroups = GroupNodes.Num();
+	for (int32 GroupIndex = 0; GroupIndex < NumGroups; ++GroupIndex)
+	{
+		FTimerNodePtr& GroupPtr = GroupNodes[GroupIndex];
+		GroupPtr->ClearFilteredChildren();
+		const bool bIsGroupVisible = Filters->PassesAllFilters(GroupPtr);
+
+		const TArray<FTimerNodePtr>& GroupChildren = GroupPtr->GetChildren();
+		const int32 NumChildren = GroupChildren.Num();
+		int32 NumVisibleChildren = 0;
+		for (int32 Cx = 0; Cx < NumChildren; ++Cx)
+		{
+			// Add a child.
+			const FTimerNodePtr& NodePtr = GroupChildren[Cx];
+			const bool bIsChildVisible = Filters->PassesAllFilters(NodePtr) && bTimerNodeIsVisible[static_cast<int>(NodePtr->GetType())];
+			if (bIsChildVisible)
+			{
+				GroupPtr->AddFilteredChild(NodePtr);
+				NumVisibleChildren++;
+			}
+		}
+
+		if (bIsGroupVisible || NumVisibleChildren > 0)
+		{
+			// Add a group.
+			FilteredGroupNodes.Add(GroupPtr);
+			GroupPtr->bForceExpandGroupNode = true;
+		}
+		else
+		{
+			GroupPtr->bForceExpandGroupNode = false;
+		}
+	}
+
+	// Only expand timer nodes if we have a text filter.
+	const bool bNonEmptyTextFilter = !TextFilter->GetRawFilterText().IsEmpty();
+	if (bNonEmptyTextFilter)
+	{
+		if (!bExpansionSaved)
+		{
+			ExpandedNodes.Empty();
+			TreeView->GetExpandedItems(ExpandedNodes);
+			bExpansionSaved = true;
+		}
+
+		for (int32 Fx = 0; Fx < FilteredGroupNodes.Num(); Fx++)
+		{
+			const FTimerNodePtr& GroupPtr = FilteredGroupNodes[Fx];
+			TreeView->SetItemExpansion(GroupPtr, GroupPtr->bForceExpandGroupNode);
+		}
+	}
+	else
+	{
+		if (bExpansionSaved)
+		{
+			// Restore previously expanded nodes when the text filter is disabled.
+			TreeView->ClearExpandedItems();
+			for (auto It = ExpandedNodes.CreateConstIterator(); It; ++It)
+			{
+				TreeView->SetItemExpansion(*It, true);
+			}
+			bExpansionSaved = false;
+		}
+	}
+
+	// Request tree refresh
+	TreeView->RequestTreeRefresh();
+}
+
+////////////////////////////////////////////////////////////////////////////////////////////////////
+
+void STimersView::HandleItemToStringArray(const FTimerNodePtr& FTimerNodePtr, TArray<FString>& OutSearchStrings) const
+{
+	OutSearchStrings.Add(FTimerNodePtr->GetName().GetPlainNameString());
+}
+
+////////////////////////////////////////////////////////////////////////////////////////////////////
+
+TSharedRef<SWidget> STimersView::GetToggleButtonForTimerType(const ETimerNodeType NodeType)
+{
+	return SNew(SCheckBox)
+		.Style(FEditorStyle::Get(), "ToggleButtonCheckbox")
+		.HAlign(HAlign_Center)
+		.Padding(2.0f)
+		.OnCheckStateChanged(this, &STimersView::FilterByTimerType_OnCheckStateChanged, NodeType)
+		.IsChecked(this, &STimersView::FilterByTimerType_IsChecked, NodeType)
+		.ToolTipText(TimerNodeTypeHelper::ToDescription(NodeType))
+		[
+			SNew(SHorizontalBox)
+
+			+ SHorizontalBox::Slot()
+				.AutoWidth()
+				.VAlign(VAlign_Center)
+				[
+					SNew(SImage)
+						.Image(TimerNodeTypeHelper::GetIconForTimerNodeType(NodeType))
+				]
+
+			+SHorizontalBox::Slot()
+				.Padding(2.0f, 0.0f, 0.0f, 0.0f)
+				.VAlign(VAlign_Center)
+				[
+					SNew(STextBlock)
+						.Text(TimerNodeTypeHelper::ToName(NodeType))
+						.TextStyle(FEditorStyle::Get(), TEXT("Profiler.Caption"))
+				]
+		];
+}
+
+////////////////////////////////////////////////////////////////////////////////////////////////////
+
+void STimersView::FilterByTimerType_OnCheckStateChanged(ECheckBoxState NewRadioState, const ETimerNodeType InStatType)
+{
+	bTimerNodeIsVisible[static_cast<int>(InStatType)] = (NewRadioState == ECheckBoxState::Checked);
+	ApplyFiltering();
+}
+
+////////////////////////////////////////////////////////////////////////////////////////////////////
+
+ECheckBoxState STimersView::FilterByTimerType_IsChecked(const ETimerNodeType InStatType) const
+{
+	return bTimerNodeIsVisible[static_cast<int>(InStatType)] ? ECheckBoxState::Checked : ECheckBoxState::Unchecked;
+}
+
+////////////////////////////////////////////////////////////////////////////////////////////////////
+// TreeView
+////////////////////////////////////////////////////////////////////////////////////////////////////
+
+void STimersView::TreeView_Refresh()
+{
+	if (TreeView.IsValid())
+	{
+		TreeView->RequestTreeRefresh();
+	}
+}
+
+////////////////////////////////////////////////////////////////////////////////////////////////////
+
+void STimersView::TreeView_OnSelectionChanged(FTimerNodePtr SelectedItem, ESelectInfo::Type SelectInfo)
+{
+	if (SelectInfo != ESelectInfo::Direct)
+	{
+		TArray<FTimerNodePtr> SelectedItems = TreeView->GetSelectedItems();
+		if (SelectedItems.Num() == 1)
+		{
+			//HighlightedTimerName = SelectedItems[0]->GetName();
+		}
+		else
+		{
+			//HighlightedTimerName = NAME_None;
+		}
+	}
+}
+
+////////////////////////////////////////////////////////////////////////////////////////////////////
+
+void STimersView::TreeView_OnGetChildren(FTimerNodePtr InParent, TArray<FTimerNodePtr>& OutChildren)
+{
+	OutChildren = InParent->GetFilteredChildren();
+}
+
+////////////////////////////////////////////////////////////////////////////////////////////////////
+
+void STimersView::TreeView_OnMouseButtonDoubleClick(FTimerNodePtr TimerNode)
+{
+	if (!TimerNode->IsGroup())
+	{
+		//im:TODO: const bool bIsStatTracked = FTimingProfilerManager::Get()->IsStatTracked(TimerNode->GetId());
+	//	if (!bIsStatTracked)
+	//	{
+	//		// Add a new graph.
+	//		FTimingProfilerManager::Get()->TrackStat(TimerNode->GetId());
+	//	}
+	//	else
+	//	{
+	//		// Remove a graph
+	//		FTimingProfilerManager::Get()->UntrackStat(TimerNode->GetId());
+	//	}
+	}
+	else
+	{
+		const bool bIsGroupExpanded = TreeView->IsItemExpanded(TimerNode);
+		TreeView->SetItemExpansion(TimerNode, !bIsGroupExpanded);
+	}
+}
+
+////////////////////////////////////////////////////////////////////////////////////////////////////
+// Tree View's Table Row
+////////////////////////////////////////////////////////////////////////////////////////////////////
+
+TSharedRef<ITableRow> STimersView::TreeView_OnGenerateRow(FTimerNodePtr TimerNodePtr, const TSharedRef<STableViewBase>& OwnerTable)
+{
+	TSharedRef<ITableRow> TableRow =
+		SNew(STimerTableRow, OwnerTable)
+		.OnShouldBeEnabled(this, &STimersView::TableRow_ShouldBeEnabled)
+		.OnIsColumnVisible(this, &STimersView::TableRow_IsColumnVisible)
+		.OnSetHoveredTableCell(this, &STimersView::TableRow_SetHoveredTableCell)
+		.OnGetColumnOutlineHAlignmentDelegate(this, &STimersView::TableRow_GetColumnOutlineHAlignment)
+		.HighlightText(this, &STimersView::TableRow_GetHighlightText)
+		.HighlightedTimerName(this, &STimersView::TableRow_GetHighlightedTimerName)
+		.TimerNodePtr(TimerNodePtr);
+
+	return TableRow;
+}
+
+////////////////////////////////////////////////////////////////////////////////////////////////////
+
+bool STimersView::TableRow_IsColumnVisible(const FName ColumnId) const
+{
+	bool bResult = false;
+	const FTimersViewColumn& ColumnPtr = TreeViewHeaderColumns.FindChecked(ColumnId);
+	return ColumnPtr.bIsVisible;
+}
+
+////////////////////////////////////////////////////////////////////////////////////////////////////
+
+void STimersView::TableRow_SetHoveredTableCell(const FName ColumnId, const FTimerNodePtr TimerNodePtr)
+{
+	HoveredColumnId = ColumnId;
+
+	const bool bIsAnyMenusVisible = FSlateApplication::Get().AnyMenusVisible();
+	if (!HasMouseCapture() && !bIsAnyMenusVisible)
+	{
+		HoveredTimerNodePtr = TimerNodePtr;
+	}
+
+	//UE_LOG(TimingProfiler, Log, TEXT("%s -> %s"), *HoveredColumnId.GetPlainNameString(), TimerNodePtr.IsValid() ? *TimerNodePtr->GetName().GetPlainNameString() : TEXT("nullptr"));
+}
+
+////////////////////////////////////////////////////////////////////////////////////////////////////
+
+EHorizontalAlignment STimersView::TableRow_GetColumnOutlineHAlignment(const FName ColumnId) const
+{
+	const TIndirectArray<SHeaderRow::FColumn>& Columns = TreeViewHeaderRow->GetColumns();
+	const int32 LastColumnIdx = Columns.Num() - 1;
+
+	// First column
+	if (Columns[0].ColumnId == ColumnId)
+	{
+		return HAlign_Left;
+	}
+	// Last column
+	else if (Columns[LastColumnIdx].ColumnId == ColumnId)
+	{
+		return HAlign_Right;
+	}
+	// Middle columns
+	{
+		return HAlign_Center;
+	}
+}
+
+////////////////////////////////////////////////////////////////////////////////////////////////////
+
+FText STimersView::TableRow_GetHighlightText() const
+{
+	return SearchBox->GetText();
+}
+
+////////////////////////////////////////////////////////////////////////////////////////////////////
+
+FName STimersView::TableRow_GetHighlightedTimerName() const
+{
+	return HighlightedTimerName;
+}
+
+////////////////////////////////////////////////////////////////////////////////////////////////////
+
+bool STimersView::TableRow_ShouldBeEnabled(const uint32 TimerId) const
+{
+	return true;//im:TODO: Session->GetAggregatedStat(TimerId) != nullptr;
+}
+
+////////////////////////////////////////////////////////////////////////////////////////////////////
+// SearchBox
+////////////////////////////////////////////////////////////////////////////////////////////////////
+
+void STimersView::SearchBox_OnTextChanged(const FText& InFilterText)
+{
+	TextFilter->SetRawFilterText(InFilterText);
+	SearchBox->SetError(TextFilter->GetFilterErrorText());
+	ApplyFiltering();
+}
+
+////////////////////////////////////////////////////////////////////////////////////////////////////
+
+bool STimersView::SearchBox_IsEnabled() const
+{
+	return TimerNodes.Num() > 0;
+}
+
+////////////////////////////////////////////////////////////////////////////////////////////////////
+// GroupBy
+////////////////////////////////////////////////////////////////////////////////////////////////////
+
+void STimersView::CreateGroups()
+{
+	TMap<FName, FTimerNodePtr> GroupNodeSet;
+
+	if (GroupingMode == ETimerGroupingMode::Flat)
+	{
+		const FName GroupName(TEXT("All"));
+		FTimerNodePtr* GroupPtr = GroupNodeSet.Find(GroupName);
+		if (!GroupPtr)
+		{
+			GroupPtr = &GroupNodeSet.Add(GroupName, MakeShareable(new FTimerNode(GroupName)));
+		}
+
+		for (const FTimerNodePtr& TimerNodePtr : TimerNodes)
+		{
+			(*GroupPtr)->AddChildAndSetGroupPtr(TimerNodePtr);
+		}
+
+		TreeView->SetItemExpansion(*GroupPtr, true);
+	}
+	// Creates groups based on stat metadata groups.
+	else if (GroupingMode == ETimerGroupingMode::ByMetaGroupName)
+	{
+		for (const FTimerNodePtr& TimerNodePtr : TimerNodes)
+		{
+			const FName GroupName = TimerNodePtr->GetMetaGroupName();
+
+			FTimerNodePtr* GroupPtr = GroupNodeSet.Find(GroupName);
+			if (!GroupPtr)
+			{
+				GroupPtr = &GroupNodeSet.Add(GroupName, MakeShareable(new FTimerNode(GroupName)));
+			}
+
+			(*GroupPtr)->AddChildAndSetGroupPtr(TimerNodePtr);
+			TreeView->SetItemExpansion(*GroupPtr, true);
+		}
+	}
+	// Creates one group for each stat type.
+	else if (GroupingMode == ETimerGroupingMode::ByType)
+	{
+		for (const FTimerNodePtr& TimerNodePtr : TimerNodes)
+		{
+			const FName GroupName = *TimerNodeTypeHelper::ToName(TimerNodePtr->GetType()).ToString();
+
+			FTimerNodePtr* GroupPtr = GroupNodeSet.Find(GroupName);
+			if (!GroupPtr)
+			{
+				GroupPtr = &GroupNodeSet.Add(GroupName, MakeShareable(new FTimerNode(GroupName)));
+			}
+
+			(*GroupPtr)->AddChildAndSetGroupPtr(TimerNodePtr);
+			TreeView->SetItemExpansion(*GroupPtr, true);
+		}
+	}
+	// Creates one group for one letter.
+	else if (GroupingMode == ETimerGroupingMode::ByName)
+	{
+		for (const FTimerNodePtr& TimerNodePtr : TimerNodes)
+		{
+			const FName GroupName = *TimerNodePtr->GetName().GetPlainNameString().Left(1);
+
+			FTimerNodePtr* GroupPtr = GroupNodeSet.Find(GroupName);
+			if (!GroupPtr)
+			{
+				GroupPtr = &GroupNodeSet.Add(GroupName, MakeShareable(new FTimerNode(GroupName)));
+			}
+
+			(*GroupPtr)->AddChildAndSetGroupPtr(TimerNodePtr);
+		}
+	}
+	// Creates one group for each logarithmic range ie. 0.001 - 0.01, 0.01 - 0.1, 0.1 - 1.0, 1.0 - 10.0, etc.
+	else if (GroupingMode == ETimerGroupingMode::ByTotalInclusiveTime)
+	{
+		//im:TODO:
+	}
+	// Creates one group for each logarithmic range ie. 0.001 - 0.01, 0.01 - 0.1, 0.1 - 1.0, 1.0 - 10.0, etc.
+	else if (GroupingMode == ETimerGroupingMode::ByTotalExclusiveTime)
+	{
+		//im:TODO:
+	}
+	// Creates one group for each logarithmic range ie. 0, 1 - 10, 10 - 100, 100 - 1000, etc.
+	else if (GroupingMode == ETimerGroupingMode::ByInstanceCount)
+	{
+		//im:TODO:
+	}
+
+	GroupNodeSet.GenerateValueArray(GroupNodes);
+
+	// Sort by a fake group name.
+	GroupNodes.Sort(TimerNodeSortingHelper::ByNameAscending());
+}
+
+////////////////////////////////////////////////////////////////////////////////////////////////////
+
+void STimersView::CreateGroupByOptionsSources()
+{
+	GroupByOptionsSource.Reset(3);
+
+	// Must be added in order of elements in the ETimerGroupingMode.
+	GroupByOptionsSource.Add(MakeShareable(new ETimerGroupingMode(ETimerGroupingMode::Flat)));
+	GroupByOptionsSource.Add(MakeShareable(new ETimerGroupingMode(ETimerGroupingMode::ByName)));
+	//GroupByOptionsSource.Add(MakeShareable(new ETimerGroupingMode(ETimerGroupingMode::ByMetaGroupName)));
+	GroupByOptionsSource.Add(MakeShareable(new ETimerGroupingMode(ETimerGroupingMode::ByType)));
+
+	ETimerGroupingModePtr* GroupingModePtrPtr = GroupByOptionsSource.FindByPredicate([&](const ETimerGroupingModePtr InGroupingModePtr) { return *InGroupingModePtr == GroupingMode; });
+	if (GroupingModePtrPtr != nullptr)
+	{
+		GroupByComboBox->SetSelectedItem(*GroupingModePtrPtr);
+	}
+
+	GroupByComboBox->RefreshOptions();
+}
+
+////////////////////////////////////////////////////////////////////////////////////////////////////
+
+void STimersView::GroupBy_OnSelectionChanged(TSharedPtr<ETimerGroupingMode> NewGroupingMode, ESelectInfo::Type SelectInfo)
+{
+	if (SelectInfo != ESelectInfo::Direct)
+	{
+		GroupingMode = *NewGroupingMode;
+
+		// Create groups, sort timers within the group and apply filtering.
+		CreateGroups();
+		SortTimers();
+		ApplyFiltering();
+	}
+}
+
+////////////////////////////////////////////////////////////////////////////////////////////////////
+
+TSharedRef<SWidget> STimersView::GroupBy_OnGenerateWidget(TSharedPtr<ETimerGroupingMode> InGroupingMode) const
+{
+	return SNew(STextBlock)
+		.Text(TimerNodeGroupingHelper::ToName(*InGroupingMode))
+		.ToolTipText(TimerNodeGroupingHelper::ToDescription(*InGroupingMode));
+}
+
+////////////////////////////////////////////////////////////////////////////////////////////////////
+
+FText STimersView::GroupBy_GetSelectedText() const
+{
+	return TimerNodeGroupingHelper::ToName(GroupingMode);
+}
+
+////////////////////////////////////////////////////////////////////////////////////////////////////
+
+FText STimersView::GroupBy_GetSelectedTooltipText() const
+{
+	return TimerNodeGroupingHelper::ToDescription(GroupingMode);
+}
+
+////////////////////////////////////////////////////////////////////////////////////////////////////
+// SortBy
+////////////////////////////////////////////////////////////////////////////////////////////////////
+
+void STimersView::SortTimers()
+{
+	const int32 NumGroups = GroupNodes.Num();
+
+	#define CHECK_AND_SORT_COLUMN(ColumnId, SortTypeName) \
+		if (ColumnBeingSorted == ColumnId) \
+		{ \
+			if (ColumnSortMode == EColumnSortMode::Type::Descending) \
+			{ \
+				for (int32 ID = 0; ID < NumGroups; ++ID) \
+				{ \
+					GroupNodes[ID]->SortChildren(TimerNodeSortingHelper::SortTypeName##Descending()); \
+				} \
+			} \
+			else /*if (ColumnSortMode == EColumnSortMode::Type::Ascending)*/ \
+			{ \
+				for (int32 ID = 0; ID < NumGroups; ++ID) \
+				{ \
+					GroupNodes[ID]->SortChildren(TimerNodeSortingHelper::SortTypeName##Ascending()); \
+				} \
+			} \
+		}
+
+		 CHECK_AND_SORT_COLUMN(FTimersViewColumns::NameColumnID,          ByName)
+	else CHECK_AND_SORT_COLUMN(FTimersViewColumns::MetaGroupNameColumnID, ByMetaGroupName)
+	else CHECK_AND_SORT_COLUMN(FTimersViewColumns::TypeColumnID,          ByType)
+	else CHECK_AND_SORT_COLUMN(FTimersViewColumns::InstanceCountColumnID, ByInstanceCount)
+	// Inclusive Time columns
+	else CHECK_AND_SORT_COLUMN(FTimersViewColumns::TotalInclusiveTimeColumnID, ByTotalInclusiveTime)
+	else CHECK_AND_SORT_COLUMN(FTimersViewColumns::MaxInclusiveTimeColumnID,     ByMaxInclusiveTime)
+	else CHECK_AND_SORT_COLUMN(FTimersViewColumns::AverageInclusiveTimeColumnID, ByAverageInclusiveTime)
+	else CHECK_AND_SORT_COLUMN(FTimersViewColumns::MedianInclusiveTimeColumnID,  ByMedianInclusiveTime)
+	else CHECK_AND_SORT_COLUMN(FTimersViewColumns::MinInclusiveTimeColumnID,     ByMinInclusiveTime)
+	// Exclusive Time columns
+	else CHECK_AND_SORT_COLUMN(FTimersViewColumns::TotalExclusiveTimeColumnID,   ByTotalExclusiveTime)
+	else CHECK_AND_SORT_COLUMN(FTimersViewColumns::MaxExclusiveTimeColumnID,     ByMaxExclusiveTime)
+	else CHECK_AND_SORT_COLUMN(FTimersViewColumns::AverageExclusiveTimeColumnID, ByAverageExclusiveTime)
+	else CHECK_AND_SORT_COLUMN(FTimersViewColumns::MedianExclusiveTimeColumnID,  ByMedianExclusiveTime)
+	else CHECK_AND_SORT_COLUMN(FTimersViewColumns::MinExclusiveTimeColumnID,     ByMinExclusiveTime)
+
+	#undef CHECK_AND_SORT_COLUMN
+}
+
+////////////////////////////////////////////////////////////////////////////////////////////////////
+
+EColumnSortMode::Type STimersView::GetSortModeForColumn(const FName ColumnId) const
+{
+	if (ColumnBeingSorted != ColumnId)
+	{
+		return EColumnSortMode::None;
+	}
+
+	return ColumnSortMode;
+}
+
+////////////////////////////////////////////////////////////////////////////////////////////////////
+
+void STimersView::SetSortModeForColumn(const FName& ColumnId, const EColumnSortMode::Type SortMode)
+{
+	ColumnBeingSorted = ColumnId;
+	ColumnSortMode = SortMode;
+
+	// Sort timers and apply filtering.
+	SortTimers();
+	ApplyFiltering();
+}
+
+////////////////////////////////////////////////////////////////////////////////////////////////////
+
+void STimersView::OnSortModeChanged(const EColumnSortPriority::Type SortPriority, const FName& ColumnId, const EColumnSortMode::Type SortMode)
+{
+	SetSortModeForColumn(ColumnId, SortMode);
+	TreeView_Refresh();
+}
+
+////////////////////////////////////////////////////////////////////////////////////////////////////
+// SortMode action (HeaderMenu)
+////////////////////////////////////////////////////////////////////////////////////////////////////
+
+bool STimersView::HeaderMenu_SortMode_IsChecked(const FName ColumnId, const EColumnSortMode::Type InSortMode)
+{
+	return ColumnBeingSorted == ColumnId && ColumnSortMode == InSortMode;
+}
+
+////////////////////////////////////////////////////////////////////////////////////////////////////
+
+bool STimersView::HeaderMenu_SortMode_CanExecute(const FName ColumnId, const EColumnSortMode::Type InSortMode) const
+{
+	const FTimersViewColumn& Column = TreeViewHeaderColumns.FindChecked(ColumnId);
+	const bool bIsValid = Column.bCanBeSorted();
+
+	bool bCanExecute = ColumnBeingSorted != ColumnId ? true : ColumnSortMode != InSortMode;
+	bCanExecute = bCanExecute && bIsValid;
+
+	return bCanExecute;
+}
+
+////////////////////////////////////////////////////////////////////////////////////////////////////
+
+void STimersView::HeaderMenu_SortMode_Execute(const FName ColumnId, const EColumnSortMode::Type InSortMode)
+{
+	SetSortModeForColumn(ColumnId, InSortMode);
+	TreeView_Refresh();
+}
+
+////////////////////////////////////////////////////////////////////////////////////////////////////
+// SortMode action (ContextMenu)
+////////////////////////////////////////////////////////////////////////////////////////////////////
+
+bool STimersView::ContextMenu_SortMode_IsChecked(const EColumnSortMode::Type InSortMode)
+{
+	return ColumnSortMode == InSortMode;
+}
+
+////////////////////////////////////////////////////////////////////////////////////////////////////
+
+bool STimersView::ContextMenu_SortMode_CanExecute(const EColumnSortMode::Type InSortMode) const
+{
+	return ColumnSortMode != InSortMode;
+}
+
+////////////////////////////////////////////////////////////////////////////////////////////////////
+
+void STimersView::ContextMenu_SortMode_Execute(const EColumnSortMode::Type InSortMode)
+{
+	SetSortModeForColumn(ColumnBeingSorted, InSortMode);
+	TreeView_Refresh();
+}
+
+////////////////////////////////////////////////////////////////////////////////////////////////////
+// SortByColumn action (ContextMenu)
+////////////////////////////////////////////////////////////////////////////////////////////////////
+
+bool STimersView::ContextMenu_SortByColumn_IsChecked(const FName ColumnId)
+{
+	return ColumnId == ColumnBeingSorted;
+}
+
+////////////////////////////////////////////////////////////////////////////////////////////////////
+
+bool STimersView::ContextMenu_SortByColumn_CanExecute(const FName ColumnId) const
+{
+	return ColumnId != ColumnBeingSorted;
+}
+
+////////////////////////////////////////////////////////////////////////////////////////////////////
+
+void STimersView::ContextMenu_SortByColumn_Execute(const FName ColumnId)
+{
+	SetSortModeForColumn(ColumnId, EColumnSortMode::Descending);
+	TreeView_Refresh();
+}
+
+////////////////////////////////////////////////////////////////////////////////////////////////////
+// HideColumn action (HeaderMenu)
+////////////////////////////////////////////////////////////////////////////////////////////////////
+
+bool STimersView::HeaderMenu_HideColumn_CanExecute(const FName ColumnId) const
+{
+	const FTimersViewColumn& Column = TreeViewHeaderColumns.FindChecked(ColumnId);
+	return Column.bCanBeHidden();
+}
+
+////////////////////////////////////////////////////////////////////////////////////////////////////
+
+void STimersView::HeaderMenu_HideColumn_Execute(const FName ColumnId)
+{
+	FTimersViewColumn& Column = TreeViewHeaderColumns.FindChecked(ColumnId);
+	Column.bIsVisible = false;
+	TreeViewHeaderRow->RemoveColumn(ColumnId);
+}
+
+////////////////////////////////////////////////////////////////////////////////////////////////////
+// ToggleColumn action (ContextMenu)
+////////////////////////////////////////////////////////////////////////////////////////////////////
+
+bool STimersView::ContextMenu_ToggleColumn_IsChecked(const FName ColumnId)
+{
+	const FTimersViewColumn& Column = TreeViewHeaderColumns.FindChecked(ColumnId);
+	return Column.bIsVisible;
+}
+
+////////////////////////////////////////////////////////////////////////////////////////////////////
+
+bool STimersView::ContextMenu_ToggleColumn_CanExecute(const FName ColumnId) const
+{
+	const FTimersViewColumn& Column = TreeViewHeaderColumns.FindChecked(ColumnId);
+	return Column.bCanBeHidden();
+}
+
+////////////////////////////////////////////////////////////////////////////////////////////////////
+
+void STimersView::ContextMenu_ToggleColumn_Execute(const FName ColumnId)
+{
+	FTimersViewColumn& Column = TreeViewHeaderColumns.FindChecked(ColumnId);
+	if (Column.bIsVisible)
+	{
+		HeaderMenu_HideColumn_Execute(ColumnId);
+	}
+	else
+	{
+		TreeViewHeaderRow_ShowColumn(ColumnId);
+	}
+}
+
+////////////////////////////////////////////////////////////////////////////////////////////////////
+// "Show All Columns" action (ContextMenu)
+////////////////////////////////////////////////////////////////////////////////////////////////////
+
+bool STimersView::ContextMenu_ShowAllColumns_CanExecute() const
+{
+	return true;
+}
+
+////////////////////////////////////////////////////////////////////////////////////////////////////
+
+void STimersView::ContextMenu_ShowAllColumns_Execute()
+{
+	ColumnSortMode = DefaultColumnSortMode;
+	ColumnBeingSorted = DefaultColumnBeingSorted;
+
+	const int32 NumColumns = FTimersViewColumnFactory::Get().Collection.Num();
+	for (int32 ColumnIndex = 0; ColumnIndex < NumColumns; ColumnIndex++)
+	{
+		const FTimersViewColumn& DefaultColumn = *FTimersViewColumnFactory::Get().Collection[ColumnIndex];
+		const FTimersViewColumn& CurrentColumn = TreeViewHeaderColumns.FindChecked(DefaultColumn.Id);
+
+		if (!CurrentColumn.bIsVisible)
+		{
+			TreeViewHeaderRow_ShowColumn(DefaultColumn.Id);
+		}
+	}
+}
+
+////////////////////////////////////////////////////////////////////////////////////////////////////
+// "Show Min/Max/Median Columns" action (ContextMenu)
+////////////////////////////////////////////////////////////////////////////////////////////////////
+
+bool STimersView::ContextMenu_ShowMinMaxMedColumns_CanExecute() const
+{
+	return true;
+}
+
+////////////////////////////////////////////////////////////////////////////////////////////////////
+
+void STimersView::ContextMenu_ShowMinMaxMedColumns_Execute()
+{
+	TSet<FName> Preset =
+	{
+		FTimersViewColumns::NameColumnID,
+		//FTimersViewColumns::MetaGroupNameColumnID,
+		//FTimersViewColumns::TypeColumnID,
+		FTimersViewColumns::InstanceCountColumnID,
+
+		FTimersViewColumns::TotalInclusiveTimeColumnID,
+		FTimersViewColumns::MaxInclusiveTimeColumnID,
+		FTimersViewColumns::MedianInclusiveTimeColumnID,
+		FTimersViewColumns::MinInclusiveTimeColumnID,
+
+		FTimersViewColumns::TotalExclusiveTimeColumnID,
+		FTimersViewColumns::MaxExclusiveTimeColumnID,
+		FTimersViewColumns::MedianExclusiveTimeColumnID,
+		FTimersViewColumns::MinExclusiveTimeColumnID,
+	};
+
+	ColumnSortMode = EColumnSortMode::Descending;
+	ColumnBeingSorted = FTimersViewColumns::TotalInclusiveTimeColumnID;
+
+	const int32 NumColumns = FTimersViewColumnFactory::Get().Collection.Num();
+	for (int32 ColumnIndex = 0; ColumnIndex < NumColumns; ColumnIndex++)
+	{
+		const FTimersViewColumn& DefaultColumn = *FTimersViewColumnFactory::Get().Collection[ColumnIndex];
+		const FTimersViewColumn& CurrentColumn = TreeViewHeaderColumns.FindChecked(DefaultColumn.Id);
+
+		bool bIsVisible = Preset.Contains(DefaultColumn.Id);
+		if (bIsVisible && !CurrentColumn.bIsVisible)
+		{
+			TreeViewHeaderRow_ShowColumn(DefaultColumn.Id);
+		}
+		else if (!bIsVisible && CurrentColumn.bIsVisible)
+		{
+			HeaderMenu_HideColumn_Execute(DefaultColumn.Id);
+		}
+	}
+}
+
+////////////////////////////////////////////////////////////////////////////////////////////////////
+// ResetColumns action (ContextMenu)
+////////////////////////////////////////////////////////////////////////////////////////////////////
+
+bool STimersView::ContextMenu_ResetColumns_CanExecute() const
+{
+	return true;
+}
+
+////////////////////////////////////////////////////////////////////////////////////////////////////
+
+void STimersView::ContextMenu_ResetColumns_Execute()
+{
+	ColumnSortMode = DefaultColumnSortMode;
+	ColumnBeingSorted = DefaultColumnBeingSorted;
+
+	const int32 NumColumns = FTimersViewColumnFactory::Get().Collection.Num();
+	for (int32 ColumnIndex = 0; ColumnIndex < NumColumns; ColumnIndex++)
+	{
+		const FTimersViewColumn& DefaultColumn = *FTimersViewColumnFactory::Get().Collection[ColumnIndex];
+		const FTimersViewColumn& CurrentColumn = TreeViewHeaderColumns.FindChecked(DefaultColumn.Id);
+
+		if (DefaultColumn.bIsVisible && !CurrentColumn.bIsVisible)
+		{
+			TreeViewHeaderRow_ShowColumn(DefaultColumn.Id);
+		}
+		else if (!DefaultColumn.bIsVisible && CurrentColumn.bIsVisible)
+		{
+			HeaderMenu_HideColumn_Execute(DefaultColumn.Id);
+		}
+	}
+}
+
+////////////////////////////////////////////////////////////////////////////////////////////////////
+
+void STimersView::RebuildTree(bool bResync)
+{
+	bool bListHasChanged = false;
+
+	if (bResync)
+	{
+		TimerNodes.Empty(TimerNodes.Num());
+		//TimerNodesMap.Empty(TimerNodesMap.Num());
+		TimerNodesIdMap.Empty(TimerNodesIdMap.Num());
+		bListHasChanged = true;
+	}
+
+	if (Session.IsValid() && Trace::ReadTimingProfilerProvider(*Session.Get()))
+	{
+		Trace::FAnalysisSessionReadScope SessionReadScope(*Session.Get());
+
+		const Trace::ITimingProfilerProvider& TimingProfilerProvider = *Trace::ReadTimingProfilerProvider(*Session.Get());
+
+		TimingProfilerProvider.ReadTimers([this, &bResync, &bListHasChanged](const Trace::FTimingProfilerTimer* Timers, uint64 TimersCount)
+		{
+			if (!bResync)
+			{
+				bResync = (TimersCount != TimerNodes.Num());
+			}
+
+			if (bResync)
+			{
+				TimerNodes.Empty(TimerNodes.Num());
+				//TimerNodesMap.Empty(TimerNodesMap.Num());
+				TimerNodesIdMap.Empty(TimerNodesIdMap.Num());
+				bListHasChanged = true;
+
+				for (uint64 TimerIndex = 0; TimerIndex < TimersCount; ++TimerIndex)
+				{
+					const Trace::FTimingProfilerTimer& Timer = Timers[TimerIndex];
+					FName Name(Timer.Name);// +TEXT(" [GPU]")));
+					FName Group(Timer.IsGpuTimer ? TEXT("GPU") : TEXT("CPU"));
+					ETimerNodeType Type = Timer.IsGpuTimer ? ETimerNodeType::GpuScope : ETimerNodeType::CpuScope;
+					FTimerNode* TimerPtr = new FTimerNode(Timer.Id, Name, Group, Type);
+					FTimerNodePtr TimerNodePtr = MakeShareable(TimerPtr);
+					TimerNodes.Add(TimerNodePtr);
+					//TimerNodesMap.Add(Name, TimerNodePtr);
+					TimerNodesIdMap.Add(Timer.Id, TimerNodePtr);
+				}
+			}
+		});
+	}
+
+	if (bListHasChanged)
+	{
+		UpdateTree();
+		UpdateStats(StatsStartTime, StatsEndTime);
+	}
+}
+
+////////////////////////////////////////////////////////////////////////////////////////////////////
+
+void STimersView::UpdateStats(double StartTime, double EndTime)
+{
+	for (const FTimerNodePtr& TimerNodePtr : TimerNodes)
+	{
+		TimerNodePtr->ResetAggregatedStats();
+	}
+
+	if (Session.IsValid() && StartTime < EndTime && Trace::ReadTimingProfilerProvider(*Session.Get()))
+	{
+		TSharedPtr<STimingProfilerWindow> Wnd = FTimingProfilerManager::Get()->GetProfilerWindow();
+		TSharedPtr<STimingView> TimingView = Wnd.IsValid() ? Wnd->TimingView : nullptr;
+
+		auto ThreadFilter = [&TimingView](uint32 ThreadId)
+		{
+			return !TimingView.IsValid() || TimingView->IsCpuTrackVisible(ThreadId);
+		};
+
+		const bool bIsGpuTrackVisible = TimingView.IsValid() && TimingView->IsGpuTrackVisible();
+
+		TUniquePtr<Trace::ITable<Trace::FTimingProfilerAggregatedStats>> AggregationResultTable;
+		{
+			Trace::FAnalysisSessionReadScope SessionReadScope(*Session.Get());
+			const Trace::ITimingProfilerProvider& TimingProfilerProvider = *Trace::ReadTimingProfilerProvider(*Session.Get());
+			AggregationResultTable.Reset(TimingProfilerProvider.CreateAggregation(StartTime, EndTime, ThreadFilter, bIsGpuTrackVisible));
+		}
+
+		TUniquePtr<Trace::ITableReader<Trace::FTimingProfilerAggregatedStats>> AggregationResultTableReader(AggregationResultTable->CreateReader());
+		while (AggregationResultTableReader->IsValid())
+		{
+			const Trace::FTimingProfilerAggregatedStats* Row = AggregationResultTableReader->GetCurrentRow();
+			FTimerNodePtr* TimerNodePtrPtr = TimerNodesIdMap.Find(Row->Timer->Id);
+			if (TimerNodePtrPtr != nullptr)
+			{
+				(*TimerNodePtrPtr)->SetAggregatedStats(*Row);
+			}
+			AggregationResultTableReader->NextRow();
+		}
+
+		StatsStartTime = StartTime;
+		StatsEndTime = EndTime;
+
+		UpdateTree();
+		TreeView->RebuildList();
+	}
+}
+
+////////////////////////////////////////////////////////////////////////////////////////////////////
+
+void STimersView::SelectTimerNode(uint64 Id)
+{
+	FTimerNodePtr* TimerNodePtrPtr = TimerNodesIdMap.Find(Id);
+	if (TimerNodePtrPtr != nullptr)
+	{
+		//UE_LOG(TimingProfiler, Log, TEXT("Select and RequestScrollIntoView %s"), *(*TimerNodePtrPtr)->GetName().ToString());
+		TreeView->SetSelection(*TimerNodePtrPtr);
+		TreeView->RequestScrollIntoView(*TimerNodePtrPtr);
+	}
+}
+
+////////////////////////////////////////////////////////////////////////////////////////////////////
+
+#undef LOCTEXT_NAMESPACE
diff --git a/Engine/Source/Developer/TraceInsights/Private/Insights/Widgets/STimersViewTooltip.cpp b/Engine/Source/Developer/TraceInsights/Private/Insights/Widgets/STimersViewTooltip.cpp
new file mode 100644
--- /dev/null
+++ b/Engine/Source/Developer/TraceInsights/Private/Insights/Widgets/STimersViewTooltip.cpp
@@ -0,0 +1,558 @@
+// Copyright 1998-2019 Epic Games, Inc. All Rights Reserved.
+
+#include "STimersViewTooltip.h"
+
+#include "EditorStyleSet.h"
+#include "SlateOptMacros.h"
+#include "Widgets/Images/SImage.h"
+#include "Widgets/Layout/SGridPanel.h"
+#include "Widgets/Layout/SSeparator.h"
+#include "Widgets/SBoxPanel.h"
+#include "Widgets/SToolTip.h"
+#include "Widgets/Text/STextBlock.h"
+
+// Insights
+#include "Insights/Common/TimeUtils.h"
+#include "Insights/ViewModels/TimerNode.h"
+#include "Insights/ViewModels/TimersViewColumn.h"
+#include "Insights/ViewModels/TimerNodeHelper.h"
+
+////////////////////////////////////////////////////////////////////////////////////////////////////
+
+#define LOCTEXT_NAMESPACE "STimersViewTooltip"
+
+////////////////////////////////////////////////////////////////////////////////////////////////////
+
+BEGIN_SLATE_FUNCTION_BUILD_OPTIMIZATION
+
+TSharedPtr<SToolTip> STimersViewTooltip::GetColumnTooltip(const FTimersViewColumn& Column)
+{
+	TSharedPtr<SToolTip> ColumnTooltip =
+		SNew(SToolTip)
+		[
+			SNew(SVerticalBox)
+
+			+ SVerticalBox::Slot()
+			.AutoHeight()
+			.Padding(2.0f)
+			[
+				SNew(STextBlock)
+				.Text(Column.TitleName)
+				.TextStyle(FEditorStyle::Get(), TEXT("Profiler.TooltipBold"))
+			]
+
+			+ SVerticalBox::Slot()
+			.AutoHeight()
+			.Padding(2.0f)
+			[
+				SNew(STextBlock)
+				.Text(Column.Description)
+				.TextStyle(FEditorStyle::Get(), TEXT("Profiler.Tooltip"))
+			]
+		];
+
+	return ColumnTooltip;
+}
+
+////////////////////////////////////////////////////////////////////////////////////////////////////
+
+TSharedPtr<SToolTip> STimersViewTooltip::GetTableCellTooltip(const TSharedPtr<FTimerNode> TimerNodePtr)
+{
+	//const FSlateFontInfo TitleFont = FCoreStyle::GetDefaultFontStyle("Bold", 10);
+	//const FSlateFontInfo DescriptionFont = FCoreStyle::GetDefaultFontStyle("Regular", 8);
+	//const FSlateFontInfo DescriptionFontB = FCoreStyle::GetDefaultFontStyle("Bold", 8);
+
+	const FLinearColor DefaultColor(1.0f,1.0f,1.0f,1.0f);
+	const FLinearColor ThreadColor(5.0f, 0.0f, 0.0f, 1.0f);
+	const float Alpha = 0.0f;//TimerNodePtr->_FramePct * 0.01f;
+	const FLinearColor ColorAndOpacity = FMath::Lerp(DefaultColor, ThreadColor,Alpha);
+
+	const Trace::FAggregatedTimingStats& Stats = TimerNodePtr->GetAggregatedStats();
+
+	const int32 NumDigits = 5;
+
+	const FText TotalInclusiveTimeText = FText::FromString(TimeUtils::FormatTimeAuto(Stats.TotalInclusiveTime));
+	const FText MinInclusiveTimeText = FText::FromString(TimeUtils::FormatTimeMs(Stats.MinInclusiveTime, NumDigits, true));
+	const FText MaxInclusiveTimeText = FText::FromString(TimeUtils::FormatTimeMs(Stats.MaxInclusiveTime, NumDigits, true));
+	const FText AvgInclusiveTimeText = FText::FromString(TimeUtils::FormatTimeMs(Stats.AverageInclusiveTime, NumDigits, true));
+	const FText MedInclusiveTimeText = FText::FromString(TimeUtils::FormatTimeMs(Stats.MedianInclusiveTime, NumDigits, true));
+
+	const FText TotalExclusiveTimeText = FText::FromString(TimeUtils::FormatTimeAuto(Stats.TotalExclusiveTime));
+	const FText MinExclusiveTimeText = FText::FromString(TimeUtils::FormatTimeMs(Stats.MinExclusiveTime, NumDigits, true));
+	const FText MaxExclusiveTimeText = FText::FromString(TimeUtils::FormatTimeMs(Stats.MaxExclusiveTime, NumDigits, true));
+	const FText AvgExclusiveTimeText = FText::FromString(TimeUtils::FormatTimeMs(Stats.AverageExclusiveTime, NumDigits, true));
+	const FText MedExclusiveTimeText = FText::FromString(TimeUtils::FormatTimeMs(Stats.MedianExclusiveTime, NumDigits, true));
+
+	const FText InstanceCountText = FText::AsNumber(Stats.InstanceCount);
+
+	//TSharedPtr<SHorizontalBox> HBoxCaption;
+	TSharedPtr<SGridPanel> GridPanel;
+	TSharedPtr<SHorizontalBox> HBox;
+
+	TSharedPtr<SToolTip> TableCellTooltip =
+		SNew(SToolTip)
+		[
+			SAssignNew(HBox, SHorizontalBox)
+
+			+SHorizontalBox::Slot()
+			.AutoWidth()
+			[
+				SNew(SVerticalBox)
+
+				//+SVerticalBox::Slot()
+				//.AutoHeight()
+				//.Padding(2.0f)
+				//[
+				//	SAssignNew(HBoxCaption, SHorizontalBox)
+				//]
+
+				+SVerticalBox::Slot()
+				.AutoHeight()
+				.Padding(2.0f)
+				[
+					SNew(SSeparator)
+					.Orientation(Orient_Horizontal)
+				]
+
+				+SVerticalBox::Slot()
+				.AutoHeight()
+				.Padding(2.0f)
+				[
+					SNew(SGridPanel)
+
+					// Id: [Id]
+					+SGridPanel::Slot(0, 0)
+					.Padding(2.0f)
+					[
+						SNew(STextBlock)
+						.Text(LOCTEXT("TT_Id", "Id:"))
+						.TextStyle(FEditorStyle::Get(), TEXT("Profiler.TooltipBold"))
+					]
+					+SGridPanel::Slot(1, 0)
+					.Padding(2.0f)
+					[
+						SNew(STextBlock)
+						.Text(FText::AsNumber(TimerNodePtr->GetId()))
+						.TextStyle(FEditorStyle::Get(), TEXT("Profiler.Tooltip"))
+					]
+
+					// Name: [Name]
+					+SGridPanel::Slot(0, 1)
+					.Padding(2.0f)
+					[
+						SNew(STextBlock)
+						.Text(LOCTEXT("TT_Name", "Name:"))
+						.TextStyle(FEditorStyle::Get(), TEXT("Profiler.TooltipBold"))
+					]
+					+SGridPanel::Slot(1, 1)
+					.Padding(2.0f)
+					[
+						SNew(STextBlock)
+						.Text(FText::FromName(TimerNodePtr->GetName()))
+						.TextStyle(FEditorStyle::Get(), TEXT("Profiler.Tooltip"))
+					]
+
+					//// Group: [MetaGroupName]
+					//+SGridPanel::Slot(0, 2)
+					//.Padding(2.0f)
+					//[
+					//	SNew(STextBlock)
+					//	.Text(LOCTEXT("TT_Group", "Group:"))
+					//	.TextStyle(FEditorStyle::Get(), TEXT("Profiler.TooltipBold"))
+					//]
+					//+SGridPanel::Slot(1, 2)
+					//.Padding(2.0f)
+					//[
+					//	SNew(STextBlock)
+					//	.Text(FText::FromName(TimerNodePtr->GetMetaGroupName()))
+					//	.TextStyle(FEditorStyle::Get(), TEXT("Profiler.Tooltip"))
+					//]
+
+					// Type: [Type]
+					+ SGridPanel::Slot(0, 3)
+					.Padding(2.0f)
+					[
+						SNew(STextBlock)
+						.Text(LOCTEXT("TT_Type", "Type:"))
+						.TextStyle(FEditorStyle::Get(), TEXT("Profiler.TooltipBold"))
+					]
+					+ SGridPanel::Slot(1, 3)
+					.Padding(2.0f)
+					[
+						SNew(STextBlock)
+						.Text(TimerNodeTypeHelper::ToName(TimerNodePtr->GetType()))
+						.TextStyle(FEditorStyle::Get(), TEXT("Profiler.Tooltip"))
+					]
+				]
+
+				+SVerticalBox::Slot()
+				.AutoHeight()
+				.Padding(2.0f)
+				[
+					SNew(SSeparator)
+					.Orientation(Orient_Horizontal)
+				]
+
+				+ SVerticalBox::Slot()
+				.AutoHeight()
+				.Padding(2.0f)
+				[
+					SNew(SGridPanel)
+
+					+ SGridPanel::Slot(0, 0)
+					.Padding(2.0f)
+					[
+						SNew(STextBlock)
+						.Text(LOCTEXT("TT_NumInstances", "Num Instances:"))
+						.TextStyle(FEditorStyle::Get(), TEXT("Profiler.TooltipBold"))
+					]
+					+ SGridPanel::Slot(1, 0)
+					.Padding(2.0f)
+					[
+						SNew(STextBlock)
+						.Text(InstanceCountText)
+						.TextStyle(FEditorStyle::Get(), TEXT("Profiler.Tooltip"))
+					]
+				]
+
+				+SVerticalBox::Slot()
+				.AutoHeight()
+				.Padding(2.0f)
+				[
+					SNew(SSeparator)
+					.Orientation(Orient_Horizontal)
+				]
+
+				+ SVerticalBox::Slot()
+				.AutoHeight()
+				.Padding(2.0f)
+				[
+					SAssignNew(GridPanel, SGridPanel)
+
+					+ SGridPanel::Slot(1, 0)
+					.Padding(2.0f)
+					.HAlign(HAlign_Right)
+					[
+						SNew(STextBlock)
+						.Text(LOCTEXT("TT_InclusiveTime", "Inclusive"))
+						.TextStyle(FEditorStyle::Get(), TEXT("Profiler.TooltipBold"))
+					]
+					+ SGridPanel::Slot(2, 0)
+					.Padding(2.0f)
+					.HAlign(HAlign_Right)
+					[
+						SNew(STextBlock)
+						.Text(LOCTEXT("TT_ExclusiveTime", "Exclusive"))
+						.TextStyle(FEditorStyle::Get(), TEXT("Profiler.TooltipBold"))
+					]
+
+					// Stats are added here.
+				]
+
+				+SVerticalBox::Slot()
+				.AutoHeight()
+				.Padding(2.0f)
+				[
+					SNew(SSeparator)
+					.Orientation(Orient_Horizontal)
+				]
+			]
+		];
+
+	int32 Row = 1;
+	AddStatsRow(GridPanel, Row, LOCTEXT("TT_TotalTime",   "Total Time:"),   TotalInclusiveTimeText, TotalExclusiveTimeText);
+	AddStatsRow(GridPanel, Row, LOCTEXT("TT_MaxTime",     "Max Time:"),     MaxInclusiveTimeText,   MaxExclusiveTimeText);
+	AddStatsRow(GridPanel, Row, LOCTEXT("TT_AverageTime", "Average Time:"), AvgInclusiveTimeText,   AvgExclusiveTimeText);
+	AddStatsRow(GridPanel, Row, LOCTEXT("TT_MedianTime",  "Median Time:"),  MedInclusiveTimeText,   MedExclusiveTimeText);
+	AddStatsRow(GridPanel, Row, LOCTEXT("TT_MinTime",     "Min Time:"),     MinInclusiveTimeText,   MinExclusiveTimeText);
+
+	/*
+	//TODO: We need Stats hierarchy (not the grouping hierarchy)!
+	const bool bHasParent = TimerNodePtr->GetStats()->GetParent().IsValid();
+	const bool bHasChildren = TimerNodePtr->GetStats()->GetChildren().Num() > 0;
+
+	if (bHasParent)
+	{
+		const FText ParentName = FText::FromName(TimerNodePtr->GetGroupPtr()->GetName());
+		HBoxCaption->AddSlot()
+			.AutoWidth()
+			[
+				SNew(STextBlock)
+				.Text(ParentName)
+				.TextStyle(FEditorStyle::Get(), TEXT("Profiler.Caption"))
+			];
+
+		HBoxCaption->AddSlot()
+			.AutoWidth()
+			[
+				SNew(SImage)
+				.Image(FEditorStyle::GetBrush("BreadcrumbTrail.Delimiter"))
+			];
+	}
+
+	const FText TimerName = FText::FromName(TimerNodePtr->GetName());
+	HBoxCaption->AddSlot()
+		.AutoWidth()
+		[
+			SNew(STextBlock)
+			.Text(TimerName)
+			.TextStyle(FEditorStyle::Get(), TEXT("Profiler.CaptionBold"))
+		];
+
+	if (bHasChildren)
+	{
+		typedef TKeyValuePair<float, FName> FEventNameAndPct;
+		TArray<FEventNameAndPct> MinimalChildren;
+
+		const TArray<FTimerNodePtr>& Children = TimerNodePtr->GetChildren();
+		for(int32 ChildIndex = 0; ChildIndex < Children.Num(); ChildIndex++)
+		{
+			const FTimerNodePtr Child = Children[ChildIndex];
+			float Percent = static_cast<float>(Child->GetStats().TotalInclusiveTime / TimerNodePtr->GetStats().TotalInclusiveTime * 100.0);
+			MinimalChildren.Add(FEventNameAndPct(Child->GetStats().TotalInclusiveTime, Child->GetName()));
+		}
+
+		struct FCompareByFloatDescending
+		{
+			FORCEINLINE bool operator()(const FEventNameAndPct& A, const FEventNameAndPct& B) const
+			{
+				return A.Key > B.Key;
+			}
+		};
+		MinimalChildren.Sort(FCompareByFloatDescending());
+
+		FString ChildrenNames;
+		const int32 NumChildrenToDisplay = FMath::Min(MinimalChildren.Num(), 3);
+		for(int32 SortedChildIndex = 0; SortedChildIndex < NumChildrenToDisplay; SortedChildIndex++)
+		{
+			const FEventNameAndPct& MinimalChild = MinimalChildren[SortedChildIndex];
+			ChildrenNames += FString::Printf(TEXT("%s (%.1f %%)"), *MinimalChild.Value.ToString(), MinimalChild.Key);
+
+			const bool bAddDelimiter = SortedChildIndex < NumChildrenToDisplay - 1;
+			if (bAddDelimiter)
+			{
+				ChildrenNames += TEXT(", ");
+			}
+		}
+
+		HBoxCaption->AddSlot()
+			.AutoWidth()
+			[
+				SNew(SImage)
+				.Image(FEditorStyle::GetBrush("BreadcrumbTrail.Delimiter"))
+			];
+
+		HBoxCaption->AddSlot()
+			.AutoWidth()
+			[
+				SNew(STextBlock)
+				.Text(FText::FromString(ChildrenNames))
+				.TextStyle(FEditorStyle::Get(), TEXT("Profiler.Caption"))
+			];
+	}
+	*/
+
+	return TableCellTooltip;
+}
+
+////////////////////////////////////////////////////////////////////////////////////////////////////
+
+TSharedRef<SToolTip> STimersViewTooltip::GetTooltip()
+{
+	if (Session.IsValid())
+	{
+		const TSharedRef<SGridPanel> ToolTipGrid = SNew(SGridPanel);
+		int32 CurrentRowPos = 0;
+
+		AddHeader(ToolTipGrid, CurrentRowPos);
+		AddDescription(ToolTipGrid, CurrentRowPos);
+
+		AddNoDataInformation(ToolTipGrid, CurrentRowPos);
+
+		return SNew(SToolTip)
+			[
+				ToolTipGrid
+			];
+	}
+	else
+	{
+		return SNew(SToolTip)
+			.Text(LOCTEXT("NotImplemented", "Tooltip for multiple profiler instances has not been implemented yet"));
+	}
+}
+
+////////////////////////////////////////////////////////////////////////////////////////////////////
+
+void STimersViewTooltip::AddNoDataInformation(const TSharedRef<SGridPanel>& Grid, int32& RowPos)
+{
+	Grid->AddSlot(0, RowPos)
+		.Padding(2.0f)
+		.ColumnSpan(3)
+		[
+			SNew(STextBlock)
+			.TextStyle(FEditorStyle::Get(), TEXT("Profiler.TooltipBold"))
+			.Text(LOCTEXT("NoDataAvailable", "N/A"))
+		];
+	RowPos++;
+}
+
+////////////////////////////////////////////////////////////////////////////////////////////////////
+
+void STimersViewTooltip::AddHeader(const TSharedRef<SGridPanel>& Grid, int32& RowPos)
+{
+	const FString InstanceName = Session->GetName();
+
+	Grid->AddSlot(0, RowPos++)
+		.Padding(2.0f)
+		.ColumnSpan(3)
+		[
+			SNew(STextBlock)
+			.TextStyle(FEditorStyle::Get(), TEXT("Profiler.TooltipBold"))
+			.Text(LOCTEXT("StatInstance", "Stat information for profiler instance"))
+		];
+
+	Grid->AddSlot(0, RowPos++)
+		.Padding(2.0f)
+		.ColumnSpan(3)
+		[
+			SNew(STextBlock)
+			.TextStyle(FEditorStyle::Get(), TEXT("Profiler.Tooltip"))
+			.Text(FText::FromString(InstanceName))
+		];
+
+	AddSeparator(Grid, RowPos);
+}
+
+////////////////////////////////////////////////////////////////////////////////////////////////////
+
+void STimersViewTooltip::AddDescription(const TSharedRef<SGridPanel>& Grid, int32& RowPos)
+{
+	/*
+	const FProfilerStat& ProfilerStat = Session->GetMetaData()->GetTimerById(TimerId);
+	const ETimerNodeType SampleType = Session->GetMetaData()->GetSampleTypeForTimerId(TimerId);
+	const FSlateBrush* const StatIcon = STimersViewHelper::GetIconForStatType(SampleType);
+
+	Grid->AddSlot(0, RowPos)
+	.Padding(2.0f)
+	[
+		SNew(STextBlock)
+		.TextStyle(FEditorStyle::Get(), TEXT("Profiler.TooltipBold"))
+		.Text(LOCTEXT("GroupDesc","Group:"))
+	];
+
+	Grid->AddSlot(1, RowPos)
+	.Padding(2.0f)
+	.ColumnSpan(2)
+	[
+		SNew(STextBlock)
+		.TextStyle(FEditorStyle::Get(), TEXT("Profiler.Tooltip"))
+		.Text(FText::FromName(ProfilerStat.OwningGroup().Name()))
+	];
+	RowPos++;
+
+	Grid->AddSlot(0, RowPos)
+	.Padding(2.0f)
+	[
+		SNew(STextBlock)
+		.TextStyle(FEditorStyle::Get(), TEXT("Profiler.TooltipBold"))
+		.Text(LOCTEXT("NameDesc","Name:"))
+	];
+
+	Grid->AddSlot(1, RowPos)
+	.Padding(2.0f)
+	.ColumnSpan(2)
+	[
+		SNew(STextBlock)
+		.TextStyle(FEditorStyle::Get(), TEXT("Profiler.Tooltip"))
+		.Text(FText::FromName(ProfilerStat.Name()))
+	];
+	RowPos++;
+
+	Grid->AddSlot(0, RowPos)
+	.Padding(2.0f)
+	[
+		SNew(STextBlock)
+		.TextStyle(FEditorStyle::Get(), TEXT("Profiler.TooltipBold"))
+		.Text(LOCTEXT("TypeDesc","Type:"))
+	];
+
+	Grid->AddSlot(1, RowPos)
+	.Padding(2.0f)
+	.ColumnSpan(2)
+	[
+		SNew(SHorizontalBox)
+
+		+SHorizontalBox::Slot()
+		.AutoWidth()
+		[
+			SNew(SImage)
+			.Image(StatIcon)
+		]
+
+		+SHorizontalBox::Slot()
+		.AutoWidth()
+		.HAlign(HAlign_Left)
+		.VAlign(VAlign_Center)
+		[
+			SNew(STextBlock)
+			.Text(FText::FromString(ETimerNodeType::ToDescription(SampleType)))
+			.TextStyle(FEditorStyle::Get(), TEXT("Profiler.Tooltip"))
+		]
+	];
+	RowPos++;
+
+	AddSeparator(Grid, RowPos);
+	*/
+}
+
+////////////////////////////////////////////////////////////////////////////////////////////////////
+
+void STimersViewTooltip::AddSeparator(const TSharedRef<SGridPanel>& Grid, int32& RowPos)
+{
+	Grid->AddSlot(0, RowPos++)
+		.Padding(2.0f)
+		.ColumnSpan(3)
+		[
+			SNew(SSeparator)
+			.Orientation(Orient_Horizontal)
+		];
+}
+
+////////////////////////////////////////////////////////////////////////////////////////////////////
+
+void STimersViewTooltip::AddStatsRow(TSharedPtr<SGridPanel> Grid, int32& Row, const FText& Name, const FText& Value1, const FText& Value2)
+{
+	Grid->AddSlot(0, Row)
+		.Padding(2.0f)
+		[
+			SNew(STextBlock)
+			.Text(Name)
+			.TextStyle(FEditorStyle::Get(), TEXT("Profiler.TooltipBold"))
+		];
+
+	Grid->AddSlot(1, Row)
+		.Padding(2.0f)
+		.HAlign(HAlign_Right)
+		[
+			SNew(STextBlock)
+			.Text(Value1)
+			.TextStyle(FEditorStyle::Get(), TEXT("Profiler.Tooltip"))
+		];
+
+	Grid->AddSlot(2, Row)
+		.Padding(2.0f)
+		.HAlign(HAlign_Right)
+		[
+			SNew(STextBlock)
+			.Text(Value2)
+			.TextStyle(FEditorStyle::Get(), TEXT("Profiler.Tooltip"))
+		];
+
+	Row++;
+}
+
+END_SLATE_FUNCTION_BUILD_OPTIMIZATION
+
+////////////////////////////////////////////////////////////////////////////////////////////////////
+
+#undef LOCTEXT_NAMESPACE
diff --git a/Engine/Source/Developer/TraceInsights/Private/Insights/Widgets/STimersViewTooltip.h b/Engine/Source/Developer/TraceInsights/Private/Insights/Widgets/STimersViewTooltip.h
new file mode 100644
--- /dev/null
+++ b/Engine/Source/Developer/TraceInsights/Private/Insights/Widgets/STimersViewTooltip.h
@@ -0,0 +1,53 @@
+// Copyright 1998-2019 Epic Games, Inc. All Rights Reserved.
+
+#pragma once
+
+#include "CoreMinimal.h"
+#include "EditorStyleSet.h"
+#include "SlateOptMacros.h"
+#include "TraceServices/AnalysisService.h"
+#include "Widgets/Images/SImage.h"
+#include "Widgets/Layout/SGridPanel.h"
+#include "Widgets/Layout/SSeparator.h"
+#include "Widgets/SToolTip.h"
+#include "Widgets/Text/STextBlock.h"
+
+// Insights
+#include "Insights/InsightsManager.h"
+
+// Insights
+class FTimerNode;
+class FTimersViewColumn;
+
+#define LOCTEXT_NAMESPACE "STimersView"
+
+/** Timers View Tooltip */
+class STimersViewTooltip
+{
+	const uint32 TimerId;
+	TSharedPtr<const Trace::IAnalysisSession> Session;
+
+public:
+	STimersViewTooltip(const uint32 InTimerId)
+		: TimerId(InTimerId)
+	{
+		Session = FInsightsManager::Get()->GetSession();
+	}
+
+	TSharedRef<SToolTip> GetTooltip();
+
+protected:
+	void AddNoDataInformation(const TSharedRef<SGridPanel>& Grid, int32& RowPos);
+	void AddHeader(const TSharedRef<SGridPanel>& Grid, int32& RowPos);
+	void AddDescription(const TSharedRef<SGridPanel>& Grid, int32& RowPos);
+	void AddSeparator(const TSharedRef<SGridPanel>& Grid, int32& RowPos);
+
+public:
+	static TSharedPtr<SToolTip> GetColumnTooltip(const FTimersViewColumn& Column);
+	static TSharedPtr<SToolTip> GetTableCellTooltip(const TSharedPtr<FTimerNode> TimerNodePtr);
+
+protected:
+	static void AddStatsRow(TSharedPtr<SGridPanel> Grid, int32& Row, const FText& Name, const FText& Value1, const FText& Value2);
+};
+
+#undef LOCTEXT_NAMESPACE
diff --git a/Engine/Source/Developer/TraceInsights/Private/Insights/Widgets/STimingProfilerToolbar.cpp b/Engine/Source/Developer/TraceInsights/Private/Insights/Widgets/STimingProfilerToolbar.cpp
new file mode 100644
--- /dev/null
+++ b/Engine/Source/Developer/TraceInsights/Private/Insights/Widgets/STimingProfilerToolbar.cpp
@@ -0,0 +1,155 @@
+// Copyright 1998-2019 Epic Games, Inc. All Rights Reserved.
+
+#include "STimingProfilerToolbar.h"
+
+#include "EditorStyleSet.h"
+#include "Framework/Application/SlateApplication.h"
+#include "Framework/Docking/TabManager.h"
+#include "Framework/MultiBox/MultiBoxDefs.h"
+#include "Framework/MultiBox/MultiBoxBuilder.h"
+#include "Widgets/Docking/SDockTab.h"
+#include "Widgets/Layout/SBorder.h"
+#include "Widgets/SBoxPanel.h"
+
+// Insights
+#include "Insights/InsightsCommands.h"
+#include "Insights/InsightsManager.h"
+#include "Insights/TimingProfilerCommands.h"
+#include "Insights/TimingProfilerManager.h"
+
+////////////////////////////////////////////////////////////////////////////////////////////////////
+
+#define LOCTEXT_NAMESPACE "STimingProfilerToolbar"
+
+////////////////////////////////////////////////////////////////////////////////////////////////////
+
+STimingProfilerToolbar::STimingProfilerToolbar()
+{
+}
+
+////////////////////////////////////////////////////////////////////////////////////////////////////
+
+STimingProfilerToolbar::~STimingProfilerToolbar()
+{
+}
+
+////////////////////////////////////////////////////////////////////////////////////////////////////
+
+void STimingProfilerToolbar::Construct(const FArguments& InArgs)
+{
+	CreateCommands();
+
+	struct Local
+	{
+		static void FillToolbar1(FToolBarBuilder& ToolbarBuilder)
+		{
+			//ToolbarBuilder.BeginSection("Session");
+			//{
+			//	ToolbarBuilder.AddToolBarButton(FInsightsCommands::Get().InsightsManager_Live);
+			//	ToolbarBuilder.AddToolBarButton(FInsightsCommands::Get().InsightsManager_Load);
+			//}
+			//ToolbarBuilder.EndSection();
+			ToolbarBuilder.BeginSection("View");
+			{
+				ToolbarBuilder.AddToolBarButton(FTimingProfilerCommands::Get().ToggleFramesTrackVisibility);
+				ToolbarBuilder.AddToolBarButton(FTimingProfilerCommands::Get().ToggleGraphTrackVisibility);
+				ToolbarBuilder.AddToolBarButton(FTimingProfilerCommands::Get().ToggleTimingViewVisibility);
+				ToolbarBuilder.AddToolBarButton(FTimingProfilerCommands::Get().ToggleTimersViewVisibility);
+				ToolbarBuilder.AddToolBarButton(FTimingProfilerCommands::Get().ToggleStatsCountersViewVisibility);
+				ToolbarBuilder.AddToolBarButton(FTimingProfilerCommands::Get().ToggleLogViewVisibility);
+			}
+			ToolbarBuilder.EndSection();
+			//ToolbarBuilder.BeginSection("Options");
+			//{
+			//	ToolbarBuilder.AddToolBarButton(FInsightsCommands::Get().OpenSettings);
+			//}
+			//ToolbarBuilder.EndSection();
+		}
+
+		static void FillToolbar2(FToolBarBuilder& ToolbarBuilder)
+		{
+			ToolbarBuilder.BeginSection("Options");
+			{
+				ToolbarBuilder.AddToolBarButton(FInsightsCommands::Get().ToggleDebugInfo);
+			}
+			ToolbarBuilder.EndSection();
+		}
+	};
+
+	TSharedPtr<FUICommandList> CommandList = FInsightsManager::Get()->GetCommandList();
+
+	FToolBarBuilder ToolbarBuilder1(CommandList.ToSharedRef(), FMultiBoxCustomization::None);
+	Local::FillToolbar1(ToolbarBuilder1);
+
+	FToolBarBuilder ToolbarBuilder2(CommandList.ToSharedRef(), FMultiBoxCustomization::None);
+	Local::FillToolbar2(ToolbarBuilder2);
+
+	// Create the tool bar!
+	ChildSlot
+	[
+		SNew(SHorizontalBox)
+
+		+SHorizontalBox::Slot()
+		.HAlign(HAlign_Left)
+		.VAlign(VAlign_Center)
+		.FillWidth(1.0)
+		.Padding(0.0f)
+		[
+			SNew(SBorder)
+			.Padding(0)
+			.BorderImage(FEditorStyle::GetBrush("NoBorder"))
+			.IsEnabled(FSlateApplication::Get().GetNormalExecutionAttribute())
+			[
+				ToolbarBuilder1.MakeWidget()
+			]
+		]
+
+		+SHorizontalBox::Slot()
+		.HAlign(HAlign_Right)
+		.VAlign(VAlign_Center)
+		.FillWidth(1.0)
+		.Padding(0.0f)
+		[
+			SNew(SBorder)
+			.Padding(0)
+			.BorderImage(FEditorStyle::GetBrush("NoBorder"))
+			.IsEnabled(FSlateApplication::Get().GetNormalExecutionAttribute())
+			[
+				ToolbarBuilder2.MakeWidget()
+			]
+		]
+	];
+}
+
+////////////////////////////////////////////////////////////////////////////////////////////////////
+
+void STimingProfilerToolbar::ShowStats()
+{
+	// do nothing
+}
+
+////////////////////////////////////////////////////////////////////////////////////////////////////
+
+void STimingProfilerToolbar::ShowMemory()
+{
+	// do nothing
+}
+
+////////////////////////////////////////////////////////////////////////////////////////////////////
+
+void STimingProfilerToolbar::CreateCommands()
+{
+	//TSharedPtr<FUICommandList> ProfilerCommandList = FTimingProfilerManager::Get()->GetCommandList();
+	//const FTimingProfilerCommands& Commands = FTimingProfilerCommands::Get();
+	//
+	//// Stats command
+	//ProfilerCommandList->MapAction(Commands.StatsProfiler,
+	//	FExecuteAction::CreateRaw(this, &STimingProfilerToolbar::ShowStats),
+	//	FCanExecuteAction(),
+	//	FIsActionChecked::CreateRaw(this, &STimingProfilerToolbar::IsShowingStats)
+	//);
+}
+
+////////////////////////////////////////////////////////////////////////////////////////////////////
+
+#undef LOCTEXT_NAMESPACE
diff --git a/Engine/Source/Developer/TraceInsights/Private/Insights/Widgets/STimingProfilerToolbar.h b/Engine/Source/Developer/TraceInsights/Private/Insights/Widgets/STimingProfilerToolbar.h
new file mode 100644
--- /dev/null
+++ b/Engine/Source/Developer/TraceInsights/Private/Insights/Widgets/STimingProfilerToolbar.h
@@ -0,0 +1,55 @@
+// Copyright 1998-2019 Epic Games, Inc. All Rights Reserved.
+
+
+#pragma once
+
+#include "CoreMinimal.h"
+#include "Widgets/DeclarativeSyntaxSupport.h"
+#include "Widgets/SCompoundWidget.h"
+
+class SBorder;
+
+/** Ribbon based toolbar used as a main menu in the Profiler window. */
+class STimingProfilerToolbar : public SCompoundWidget
+{
+public:
+	/** Default constructor. */
+	STimingProfilerToolbar();
+
+	/** Virtual destructor. */
+	virtual ~STimingProfilerToolbar();
+
+	SLATE_BEGIN_ARGS( STimingProfilerToolbar )
+		{}
+	SLATE_END_ARGS()
+
+	/**
+	 * Construct this widget
+	 *
+	 * @param	InArgs	The declaration data for this widget
+	 */
+	void Construct( const FArguments& InArgs );
+
+	bool IsImplemented() const { return false; }
+	bool IsShowingStats() const { return true; }
+	bool IsShowingMemory() const { return false; }
+
+protected:
+	/** Called when the list of sessions has changed. @see FProfilerManager.FOnSessionsUpdated */
+	void ProfilerManager_SessionsUpdated();
+
+	/** Create the UI commands for the toolbar */
+	void CreateCommands();
+
+	/** Shows the stats profiler */
+	void ShowStats();
+
+	/** Shows the memory profiler */
+	void ShowMemory();
+
+	/** Shows the FPSChart view */
+	void ShowFPSChart();
+
+protected:
+	TSharedPtr<SBorder> Border;
+};
