Copying //UE4/Dev-Core to Dev-Main (//UE4/Dev-Main)

diff --git a/Engine/Source/Developer/TraceInsights/Private/Insights/Widgets/STimingProfilerWindow.cpp b/Engine/Source/Developer/TraceInsights/Private/Insights/Widgets/STimingProfilerWindow.cpp
new file mode 100644
--- /dev/null
+++ b/Engine/Source/Developer/TraceInsights/Private/Insights/Widgets/STimingProfilerWindow.cpp
@@ -0,0 +1,506 @@
+// Copyright 1998-2019 Epic Games, Inc. All Rights Reserved.
+
+#include "STimingProfilerWindow.h"
+
+#include "EditorStyleSet.h"
+#include "Framework/Docking/WorkspaceItem.h"
+#include "SlateOptMacros.h"
+#include "Widgets/Docking/SDockTab.h"
+#include "Widgets/Images/SImage.h"
+#include "Widgets/Layout/SBorder.h"
+#include "Widgets/Layout/SBox.h"
+#include "Widgets/Layout/SSpacer.h"
+#include "Widgets/Notifications/SNotificationList.h"
+#include "Widgets/SBoxPanel.h"
+#include "Widgets/Text/STextBlock.h"
+//#include "WorkspaceMenuStructure.h"
+//#include "WorkspaceMenuStructureModule.h"
+
+#if WITH_EDITOR
+	#include "EngineAnalytics.h"
+	#include "Runtime/Analytics/Analytics/Public/AnalyticsEventAttribute.h"
+	#include "Runtime/Analytics/Analytics/Public/Interfaces/IAnalyticsProvider.h"
+#endif // WITH_EDITOR
+
+// Insights
+#include "Insights/InsightsManager.h"
+#include "Insights/InsightsStyle.h"
+#include "Insights/TimingProfilerManager.h"
+#include "Insights/Version.h"
+#include "Insights/Widgets/SFrameTrack.h"
+#include "Insights/Widgets/SGraphTrack.h"
+#include "Insights/Widgets/SInsightsSettings.h"
+#include "Insights/Widgets/SLogView.h"
+#include "Insights/Widgets/SStatsView.h"
+#include "Insights/Widgets/STimersView.h"
+#include "Insights/Widgets/STimingProfilerToolbar.h"
+#include "Insights/Widgets/STimingView.h"
+
+////////////////////////////////////////////////////////////////////////////////////////////////////
+
+#define LOCTEXT_NAMESPACE "STimingProfilerWindow"
+
+////////////////////////////////////////////////////////////////////////////////////////////////////
+
+const FName FTimingProfilerTabs::ToolbarID(TEXT("Toolbar"));
+const FName FTimingProfilerTabs::FramesTrackID(TEXT("Frames"));
+const FName FTimingProfilerTabs::GraphTrackID(TEXT("Graph"));
+const FName FTimingProfilerTabs::TimingViewID(TEXT("TimingView"));
+const FName FTimingProfilerTabs::TimersID(TEXT("Timers"));
+const FName FTimingProfilerTabs::StatsCountersID(TEXT("StasCounters"));
+const FName FTimingProfilerTabs::LogViewID(TEXT("LogView"));
+
+////////////////////////////////////////////////////////////////////////////////////////////////////
+
+STimingProfilerWindow::STimingProfilerWindow()
+	: DurationActive(0.0f)
+{
+}
+
+////////////////////////////////////////////////////////////////////////////////////////////////////
+
+STimingProfilerWindow::~STimingProfilerWindow()
+{
+#if WITH_EDITOR
+	if (DurationActive > 0.0f && FEngineAnalytics::IsAvailable())
+	{
+		FEngineAnalytics::GetProvider().RecordEvent(TEXT("Editor.Usage.Profiler"), FAnalyticsEventAttribute(TEXT("Duration"), DurationActive));
+	}
+#endif // WITH_EDITOR
+}
+
+////////////////////////////////////////////////////////////////////////////////////////////////////
+
+BEGIN_SLATE_FUNCTION_BUILD_OPTIMIZATION
+
+TSharedRef<SDockTab> STimingProfilerWindow::SpawnTab_Toolbar(const FSpawnTabArgs& Args)
+{
+	return SNew(SDockTab)
+		.ShouldAutosize(true)
+		.TabRole(ETabRole::PanelTab)
+		//.IsEnabled(this, &STimingProfilerWindow::IsProfilerEnabled)
+		[
+			SNew(STimingProfilerToolbar)
+		];
+}
+
+////////////////////////////////////////////////////////////////////////////////////////////////////
+
+TSharedRef<SDockTab> STimingProfilerWindow::SpawnTab_FramesTrack(const FSpawnTabArgs& Args)
+{
+	FTimingProfilerManager::Get()->SetFramesTrackVisible(true);
+
+	return SNew(SDockTab)
+		.ShouldAutosize(false)
+		.TabRole(ETabRole::PanelTab)
+		//.IsEnabled(this, &STimingProfilerWindow::IsProfilerEnabled)
+		[
+			SAssignNew(FrameTrack, SFrameTrack)
+		];
+}
+
+////////////////////////////////////////////////////////////////////////////////////////////////////
+
+TSharedRef<SDockTab> STimingProfilerWindow::SpawnTab_GraphTrack(const FSpawnTabArgs& Args)
+{
+	FTimingProfilerManager::Get()->SetGraphTrackVisible(true);
+
+	return SNew(SDockTab)
+		.ShouldAutosize(false)
+		.TabRole(ETabRole::PanelTab)
+		//.IsEnabled(this, &STimingProfilerWindow::IsProfilerEnabled)
+		[
+			SAssignNew(GraphTrack, SGraphTrack)
+		];
+}
+
+////////////////////////////////////////////////////////////////////////////////////////////////////
+
+TSharedRef<SDockTab> STimingProfilerWindow::SpawnTab_TimingView(const FSpawnTabArgs& Args)
+{
+	FTimingProfilerManager::Get()->SetTimingViewVisible(true);
+
+	return SNew(SDockTab)
+		.ShouldAutosize(false)
+		.TabRole(ETabRole::PanelTab)
+		//.IsEnabled(this, &STimingProfilerWindow::IsProfilerEnabled)
+		[
+			SAssignNew(TimingView, STimingView)
+		];
+}
+
+////////////////////////////////////////////////////////////////////////////////////////////////////
+
+TSharedRef<SDockTab> STimingProfilerWindow::SpawnTab_Timers(const FSpawnTabArgs& Args)
+{
+	FTimingProfilerManager::Get()->SetTimersViewVisible(true);
+
+	return SNew(SDockTab)
+		.ShouldAutosize(false)
+		.TabRole(ETabRole::PanelTab)
+		//.IsEnabled(this, &STimingProfilerWindow::IsProfilerEnabled)
+		[
+			SAssignNew(TimersView, STimersView)
+		];
+}
+
+////////////////////////////////////////////////////////////////////////////////////////////////////
+
+TSharedRef<SDockTab> STimingProfilerWindow::SpawnTab_StatsCounters(const FSpawnTabArgs& Args)
+{
+	FTimingProfilerManager::Get()->SetStatsCountersViewVisible(true);
+
+	return SNew(SDockTab)
+		.ShouldAutosize(false)
+		.TabRole(ETabRole::PanelTab)
+		//.IsEnabled(this, &STimingProfilerWindow::IsProfilerEnabled)
+		[
+			SAssignNew(StatsView, SStatsView)
+		];
+}
+
+////////////////////////////////////////////////////////////////////////////////////////////////////
+
+TSharedRef<SDockTab> STimingProfilerWindow::SpawnTab_LogView(const FSpawnTabArgs& Args)
+{
+	FTimingProfilerManager::Get()->SetLogViewVisible(true);
+
+	return SNew(SDockTab)
+		.ShouldAutosize(false)
+		.TabRole(ETabRole::PanelTab)
+		//.IsEnabled(this, &STimingProfilerWindow::IsProfilerEnabled)
+		[
+			SAssignNew(LogView, SLogView)
+		];
+}
+
+////////////////////////////////////////////////////////////////////////////////////////////////////
+
+void STimingProfilerWindow::Construct(const FArguments& InArgs, const TSharedRef<SDockTab>& ConstructUnderMajorTab, const TSharedPtr<SWindow>& ConstructUnderWindow)
+{
+	// Create & initialize tab manager.
+	TabManager = FGlobalTabmanager::Get()->NewTabManager(ConstructUnderMajorTab);
+	TSharedRef<FWorkspaceItem> AppMenuGroup = TabManager->AddLocalWorkspaceMenuCategory(LOCTEXT("TimingProfilerMenuGroupName", "Timing Insights"));
+
+	TabManager->RegisterTabSpawner(FTimingProfilerTabs::ToolbarID, FOnSpawnTab::CreateRaw(this, &STimingProfilerWindow::SpawnTab_Toolbar))
+		.SetDisplayName(LOCTEXT("DeviceToolbarTabTitle", "Toolbar"))
+		.SetIcon(FSlateIcon(FInsightsStyle::GetStyleSetName(), "Toolbar.Icon.Small"))
+		.SetGroup(AppMenuGroup);
+
+	TabManager->RegisterTabSpawner(FTimingProfilerTabs::FramesTrackID, FOnSpawnTab::CreateRaw(this, &STimingProfilerWindow::SpawnTab_FramesTrack))
+		.SetDisplayName(LOCTEXT("FramesTrackTabTitle", "Frames"))
+		.SetIcon(FSlateIcon(FInsightsStyle::GetStyleSetName(), "FramesTrack.Icon.Small"))
+		.SetGroup(AppMenuGroup);
+
+	TabManager->RegisterTabSpawner(FTimingProfilerTabs::GraphTrackID, FOnSpawnTab::CreateRaw(this, &STimingProfilerWindow::SpawnTab_GraphTrack))
+		.SetDisplayName(LOCTEXT("GraphTrackTabTitle", "Graph"))
+		.SetIcon(FSlateIcon(FInsightsStyle::GetStyleSetName(), "GraphTrack.Icon.Small"))
+		.SetGroup(AppMenuGroup);
+
+	TabManager->RegisterTabSpawner(FTimingProfilerTabs::TimingViewID, FOnSpawnTab::CreateRaw(this, &STimingProfilerWindow::SpawnTab_TimingView))
+		.SetDisplayName(LOCTEXT("TimingViewTabTitle", "Timing View"))
+		.SetIcon(FSlateIcon(FInsightsStyle::GetStyleSetName(), "TimingView.Icon.Small"))
+		.SetGroup(AppMenuGroup);
+
+	TabManager->RegisterTabSpawner(FTimingProfilerTabs::TimersID, FOnSpawnTab::CreateRaw(this, &STimingProfilerWindow::SpawnTab_Timers))
+		.SetDisplayName(LOCTEXT("TimersTabTitle", "Timers"))
+		.SetIcon(FSlateIcon(FInsightsStyle::GetStyleSetName(), "TimersView.Icon.Small"))
+		.SetGroup(AppMenuGroup);
+
+	TabManager->RegisterTabSpawner(FTimingProfilerTabs::StatsCountersID, FOnSpawnTab::CreateRaw(this, &STimingProfilerWindow::SpawnTab_StatsCounters))
+		.SetDisplayName(LOCTEXT("StatsCountersTabTitle", "Stats Counters"))
+		.SetIcon(FSlateIcon(FInsightsStyle::GetStyleSetName(), "StatsCountersView.Icon.Small"))
+		.SetGroup(AppMenuGroup);
+
+	TabManager->RegisterTabSpawner(FTimingProfilerTabs::LogViewID, FOnSpawnTab::CreateRaw(this, &STimingProfilerWindow::SpawnTab_LogView))
+		.SetDisplayName(LOCTEXT("LogViewTabTitle", "Log View"))
+		.SetIcon(FSlateIcon(FInsightsStyle::GetStyleSetName(), "LogView.Icon.Small"))
+		.SetGroup(AppMenuGroup);
+
+	TSharedPtr<FTimingProfilerManager> TimingProfilerManager = FTimingProfilerManager::Get();
+	ensure(TimingProfilerManager.IsValid());
+
+	// Create tab layout.
+	const TSharedRef<FTabManager::FLayout> Layout = FTabManager::NewLayout("InsightsTimingProfilerLayout_v1.0")
+		->AddArea
+		(
+			FTabManager::NewPrimaryArea()
+			->SetOrientation(Orient_Vertical)
+			->Split
+			(
+				FTabManager::NewStack()
+				->AddTab(FTimingProfilerTabs::ToolbarID, ETabState::OpenedTab)
+				->SetHideTabWell(true)
+			)
+			->Split
+			(
+				FTabManager::NewSplitter()
+				->SetOrientation(Orient_Horizontal)
+				->SetSizeCoefficient(1.0f)
+				->Split
+				(
+					FTabManager::NewSplitter()
+					->SetOrientation(Orient_Vertical)
+					->SetSizeCoefficient(0.65f)
+					->Split
+					(
+						FTabManager::NewStack()
+						->SetSizeCoefficient(0.1f)
+						->SetHideTabWell(true)
+						->AddTab(FTimingProfilerTabs::FramesTrackID, TimingProfilerManager->IsFramesTrackVisible() ? ETabState::OpenedTab : ETabState::ClosedTab)
+					)
+					->Split
+					(
+						FTabManager::NewStack()
+						->SetSizeCoefficient(0.2f)
+						->SetHideTabWell(true)
+						->AddTab(FTimingProfilerTabs::GraphTrackID, TimingProfilerManager->IsGraphTrackVisible() ? ETabState::OpenedTab : ETabState::ClosedTab)
+					)
+					->Split
+					(
+						FTabManager::NewStack()
+						->SetSizeCoefficient(0.5f)
+						->SetHideTabWell(true)
+						->AddTab(FTimingProfilerTabs::TimingViewID, TimingProfilerManager->IsTimingViewVisible() ? ETabState::OpenedTab : ETabState::ClosedTab)
+					)
+					->Split
+					(
+						FTabManager::NewStack()
+						->SetSizeCoefficient(0.2f)
+						->SetHideTabWell(true)
+						->AddTab(FTimingProfilerTabs::LogViewID, TimingProfilerManager->IsLogViewVisible() ? ETabState::OpenedTab : ETabState::ClosedTab)
+					)
+				)
+				->Split
+				(
+					FTabManager::NewStack()
+					->SetSizeCoefficient(0.35f)
+					->AddTab(FTimingProfilerTabs::TimersID, TimingProfilerManager->IsTimersViewVisible() ? ETabState::OpenedTab : ETabState::ClosedTab)
+					->AddTab(FTimingProfilerTabs::StatsCountersID, TimingProfilerManager->IsStatsCountersViewVisible() ? ETabState::OpenedTab : ETabState::ClosedTab)
+					->SetForegroundTab(FTimingProfilerTabs::TimersID)
+				)
+			)
+		);
+
+	// Create & initialize main menu.
+	FMenuBarBuilder MenuBarBuilder = FMenuBarBuilder(TSharedPtr<FUICommandList>());
+	
+	MenuBarBuilder.AddPullDownMenu(
+		LOCTEXT("MenuLabel", "MENU"),
+		FText::GetEmpty(),
+		FNewMenuDelegate::CreateStatic(&STimingProfilerWindow::FillMenu, TabManager),
+		FName(TEXT("MENU"))
+	);
+
+	ChildSlot
+		[
+			SNew(SOverlay)
+
+			// Version
+			+ SOverlay::Slot()
+				.HAlign(HAlign_Right)
+				.VAlign(VAlign_Top)
+				.Padding(0.0f, -16.0f, 0.0f, 0.0f)
+				[
+					SNew(STextBlock)
+						.Clipping(EWidgetClipping::ClipToBoundsWithoutIntersecting)
+						.Text(LOCTEXT("UnrealInsightsVersion", UNREAL_INSIGHTS_VERSION_STRING_EX))
+						.ColorAndOpacity(FLinearColor(0.15f, 0.15f, 0.15f, 1.0f))
+				]
+
+			// Overlay slot for the main profiler window area
+			+ SOverlay::Slot()
+				.HAlign(HAlign_Fill)
+				.VAlign(VAlign_Fill)
+				[
+					SNew(SVerticalBox)
+
+					+ SVerticalBox::Slot()
+						.AutoHeight()
+						[
+							MenuBarBuilder.MakeWidget()
+						]
+
+					+ SVerticalBox::Slot()
+						.FillHeight(1.0f)
+						[
+							TabManager->RestoreFrom(Layout, ConstructUnderWindow).ToSharedRef()
+						]
+				]
+
+			// Session hint overlay
+			+ SOverlay::Slot()
+				.HAlign(HAlign_Center)
+				.VAlign(VAlign_Center)
+				[
+					SNew(SBorder)
+						.Visibility(this, &STimingProfilerWindow::IsSessionOverlayVisible)
+						.BorderImage(FEditorStyle::GetBrush("NotificationList.ItemBackground"))
+						.Padding(8.0f)
+						[
+							SNew(STextBlock)
+								.Text(LOCTEXT("SelectTraceOverlayText", "Please select a trace."))
+						]
+				]
+		];
+
+	// Tell tab-manager about the global menu bar.
+	TabManager->SetMenuMultiBox(MenuBarBuilder.GetMultiBox());
+}
+
+END_SLATE_FUNCTION_BUILD_OPTIMIZATION
+
+////////////////////////////////////////////////////////////////////////////////////////////////////
+
+void STimingProfilerWindow::FillMenu(FMenuBuilder& MenuBuilder, const TSharedPtr<FTabManager> TabManager)
+{
+	if (!TabManager.IsValid())
+	{
+		return;
+	}
+
+#if !WITH_EDITOR
+	//TODO: FGlobalTabmanager::Get()->PopulateTabSpawnerMenu(MenuBuilder, WorkspaceMenu::GetMenuStructure().GetStructureRoot());
+#endif //!WITH_EDITOR
+
+	TabManager->PopulateLocalTabSpawnerMenu(MenuBuilder);
+}
+
+////////////////////////////////////////////////////////////////////////////////////////////////////
+
+void STimingProfilerWindow::ShowTab(const FName& TabID)
+{
+	if (TabManager->CanSpawnTab(TabID))
+	{
+		TabManager->InvokeTab(TabID);
+	}
+}
+
+////////////////////////////////////////////////////////////////////////////////////////////////////
+
+void STimingProfilerWindow::HideTab(const FName& TabID)
+{
+	TSharedPtr<SDockTab> Tab = TabManager->FindExistingLiveTab(TabID);
+	if (Tab.IsValid())
+	{
+		Tab->RequestCloseTab();
+	}
+}
+
+////////////////////////////////////////////////////////////////////////////////////////////////////
+
+EVisibility STimingProfilerWindow::IsSessionOverlayVisible() const
+{
+	if (FInsightsManager::Get()->GetSession().IsValid())
+	{
+		return EVisibility::Hidden;
+	}
+	else
+	{
+		return EVisibility::Visible;
+	}
+}
+
+////////////////////////////////////////////////////////////////////////////////////////////////////
+
+bool STimingProfilerWindow::IsProfilerEnabled() const
+{
+	return FInsightsManager::Get()->GetSession().IsValid();
+}
+
+////////////////////////////////////////////////////////////////////////////////////////////////////
+
+EActiveTimerReturnType STimingProfilerWindow::UpdateActiveDuration(double InCurrentTime, float InDeltaTime)
+{
+	DurationActive += InDeltaTime;
+
+	// The profiler window will explicitly unregister this active timer when the mouse leaves.
+	return EActiveTimerReturnType::Continue;
+}
+
+////////////////////////////////////////////////////////////////////////////////////////////////////
+
+void STimingProfilerWindow::OnMouseEnter(const FGeometry& MyGeometry, const FPointerEvent& MouseEvent)
+{
+	SCompoundWidget::OnMouseEnter(MyGeometry, MouseEvent);
+
+	if (!ActiveTimerHandle.IsValid())
+	{
+		ActiveTimerHandle = RegisterActiveTimer(0.f, FWidgetActiveTimerDelegate::CreateSP(this, &STimingProfilerWindow::UpdateActiveDuration));
+	}
+}
+
+////////////////////////////////////////////////////////////////////////////////////////////////////
+
+void STimingProfilerWindow::OnMouseLeave(const FPointerEvent& MouseEvent)
+{
+	SCompoundWidget::OnMouseLeave(MouseEvent);
+
+	auto PinnedActiveTimerHandle = ActiveTimerHandle.Pin();
+	if (PinnedActiveTimerHandle.IsValid())
+	{
+		UnRegisterActiveTimer(PinnedActiveTimerHandle.ToSharedRef());
+	}
+}
+
+////////////////////////////////////////////////////////////////////////////////////////////////////
+
+FReply STimingProfilerWindow::OnKeyDown(const FGeometry& MyGeometry, const FKeyEvent& InKeyEvent)
+{
+	return FTimingProfilerManager::Get()->GetCommandList()->ProcessCommandBindings(InKeyEvent) ? FReply::Handled() : FReply::Unhandled();
+}
+
+////////////////////////////////////////////////////////////////////////////////////////////////////
+
+FReply STimingProfilerWindow::OnDragOver(const FGeometry& MyGeometry, const FDragDropEvent& DragDropEvent)
+{
+	TSharedPtr<FExternalDragOperation> DragDropOp = DragDropEvent.GetOperationAs<FExternalDragOperation>();
+	if (DragDropOp.IsValid())
+	{
+		if (DragDropOp->HasFiles())
+		{
+			const TArray<FString>& Files = DragDropOp->GetFiles();
+			if (Files.Num() == 1)
+			{
+				const FString DraggedFileExtension = FPaths::GetExtension(Files[0], true);
+				if (DraggedFileExtension == TEXT(".utrace"))
+				{
+					return FReply::Handled();
+				}
+			}
+		}
+	}
+
+	return SCompoundWidget::OnDragOver(MyGeometry,DragDropEvent);
+}
+
+////////////////////////////////////////////////////////////////////////////////////////////////////
+
+FReply STimingProfilerWindow::OnDrop(const FGeometry& MyGeometry, const FDragDropEvent& DragDropEvent)
+{
+	TSharedPtr<FExternalDragOperation> DragDropOp = DragDropEvent.GetOperationAs<FExternalDragOperation>();
+	if (DragDropOp.IsValid())
+	{
+		if (DragDropOp->HasFiles())
+		{
+			// For now, only allow a single file.
+			const TArray<FString>& Files = DragDropOp->GetFiles();
+			if (Files.Num() == 1)
+			{
+				const FString DraggedFileExtension = FPaths::GetExtension(Files[0], true);
+				if (DraggedFileExtension == TEXT(".utrace"))
+				{
+					// Enqueue load operation.
+					FInsightsManager::Get()->LoadTraceFile(Files[0]);
+					return FReply::Handled();
+				}
+			}
+		}
+	}
+
+	return SCompoundWidget::OnDrop(MyGeometry,DragDropEvent);
+}
+
+////////////////////////////////////////////////////////////////////////////////////////////////////
+
+#undef LOCTEXT_NAMESPACE
diff --git a/Engine/Source/Developer/TraceInsights/Private/Insights/Widgets/STimingProfilerWindow.h b/Engine/Source/Developer/TraceInsights/Private/Insights/Widgets/STimingProfilerWindow.h
new file mode 100644
--- /dev/null
+++ b/Engine/Source/Developer/TraceInsights/Private/Insights/Widgets/STimingProfilerWindow.h
@@ -0,0 +1,175 @@
+// Copyright 1998-2019 Epic Games, Inc. All Rights Reserved.
+
+#pragma once
+
+#include "CoreMinimal.h"
+#include "Framework/Docking/TabManager.h"
+#include "Input/Reply.h"
+#include "Layout/Visibility.h"
+#include "Misc/Guid.h"
+#include "SlateFwd.h"
+#include "Widgets/DeclarativeSyntaxSupport.h"
+#include "Widgets/Docking/SDockTab.h"
+#include "Widgets/SCompoundWidget.h"
+
+// Insights
+#include "Insights/InsightsManager.h"
+
+////////////////////////////////////////////////////////////////////////////////////////////////////
+
+class FActiveTimerHandle;
+class SVerticalBox;
+class SFrameTrack;
+class SGraphTrack;
+class SStatsView;
+class STimersView;
+class STimingView;
+class SLogView;
+
+////////////////////////////////////////////////////////////////////////////////////////////////////
+
+struct FTimingProfilerTabs
+{
+	// Tab identifiers
+	static const FName ToolbarID;
+	static const FName FramesTrackID;
+	static const FName GraphTrackID;
+	static const FName TimingViewID;
+	static const FName TimersID;
+	static const FName StatsCountersID;
+	static const FName LogViewID;
+};
+
+////////////////////////////////////////////////////////////////////////////////////////////////////
+
+/** Implements the timing profiler window. */
+class STimingProfilerWindow : public SCompoundWidget
+{
+public:
+	/** Default constructor. */
+	STimingProfilerWindow();
+
+	/** Virtual destructor. */
+	virtual ~STimingProfilerWindow();
+
+	SLATE_BEGIN_ARGS(STimingProfilerWindow){}
+	SLATE_END_ARGS()
+
+	/** Constructs this widget. */
+	void Construct(const FArguments& InArgs, const TSharedRef<SDockTab>& ConstructUnderMajorTab, const TSharedPtr<SWindow>& ConstructUnderWindow);
+
+	void ShowTab(const FName& TabID);
+	void HideTab(const FName& TabID);
+	void ShowHideTab(const FName& TabID, bool bShow) { bShow ? ShowTab(TabID) : HideTab(TabID); }
+
+	TSharedPtr<FTabManager> GetTabManager() const { return TabManager; }
+
+private:
+	TSharedRef<SDockTab> SpawnTab_Toolbar(const FSpawnTabArgs& Args);
+	TSharedRef<SDockTab> SpawnTab_FramesTrack(const FSpawnTabArgs& Args);
+	TSharedRef<SDockTab> SpawnTab_GraphTrack(const FSpawnTabArgs& Args);
+	TSharedRef<SDockTab> SpawnTab_TimingView(const FSpawnTabArgs& Args);
+	TSharedRef<SDockTab> SpawnTab_Timers(const FSpawnTabArgs& Args);
+	TSharedRef<SDockTab> SpawnTab_StatsCounters(const FSpawnTabArgs& Args);
+	TSharedRef<SDockTab> SpawnTab_LogView(const FSpawnTabArgs& Args);
+
+	/**
+	 * Fill the main menu with menu items.
+	 *
+	 * @param MenuBuilder The multi-box builder that should be filled with content for this pull-down menu.
+	 * @param TabManager A Tab Manager from which to populate tab spawner menu items.
+	 */
+	static void FillMenu(FMenuBuilder& MenuBuilder, const TSharedPtr<FTabManager> TabManager);
+
+	/** Callback for determining the visibility of the 'Select a session' overlay. */
+	EVisibility IsSessionOverlayVisible() const;
+
+	/** Callback for getting the enabled state of the profiler window. */
+	bool IsProfilerEnabled() const;
+
+	/** Updates the amount of time the profiler has been active. */
+	EActiveTimerReturnType UpdateActiveDuration(double InCurrentTime, float InDeltaTime);
+
+	/**
+	 * Ticks this widget. Override in derived classes, but always call the parent implementation.
+	 *
+	 * @param  AllottedGeometry The space allotted for this widget
+	 * @param  InCurrentTime  Current absolute real time
+	 * @param  InDeltaTime  Real time passed since last tick
+	 */
+	//virtual void Tick(const FGeometry& AllottedGeometry, const double InCurrentTime, const float InDeltaTime) override;
+
+	/**
+	 * The system will use this event to notify a widget that the cursor has entered it. This event is NOT bubbled.
+	 *
+	 * @param MyGeometry The Geometry of the widget receiving the event
+	 * @param MouseEvent Information about the input event
+	 */
+	virtual void OnMouseEnter(const FGeometry& MyGeometry, const FPointerEvent& MouseEvent) override;
+
+	/**
+	 * The system will use this event to notify a widget that the cursor has left it. This event is NOT bubbled.
+	 *
+	 * @param MouseEvent Information about the input event
+	 */
+	virtual void OnMouseLeave(const FPointerEvent& MouseEvent) override;
+
+	/**
+	 * Called after a key is pressed when this widget has focus
+	 *
+	 * @param MyGeometry The Geometry of the widget receiving the event
+	 * @param  InKeyEvent  Key event
+	 *
+	 * @return  Returns whether the event was handled, along with other possible actions
+	 */
+	virtual FReply OnKeyDown(const FGeometry& MyGeometry, const FKeyEvent& InKeyEvent) override;
+
+	/**
+	 * Called when the user is dropping something onto a widget; terminates drag and drop.
+	 *
+	 * @param MyGeometry      The geometry of the widget receiving the event.
+	 * @param DragDropEvent   The drag and drop event.
+	 *
+	 * @return A reply that indicated whether this event was handled.
+	 */
+	virtual FReply OnDrop(const FGeometry& MyGeometry, const FDragDropEvent& DragDropEvent) override;
+
+	/**
+	 * Called during drag and drop when the the mouse is being dragged over a widget.
+	 *
+	 * @param MyGeometry      The geometry of the widget receiving the event.
+	 * @param DragDropEvent   The drag and drop event.
+	 *
+	 * @return A reply that indicated whether this event was handled.
+	 */
+	virtual FReply OnDragOver(const FGeometry& MyGeometry, const FDragDropEvent& DragDropEvent)  override;
+
+public:
+	/** Widget for the frame track */
+	TSharedPtr<SFrameTrack> FrameTrack;
+
+	/** Widget for the graph track */
+	TSharedPtr<SGraphTrack> GraphTrack;
+
+	/** Widget for the timing track */
+	TSharedPtr<STimingView> TimingView;
+
+	/** Holds the Timers view widget. */
+	TSharedPtr<STimersView> TimersView;
+
+	/** Holds the Stats (Counters) view widget. */
+	TSharedPtr<SStatsView> StatsView;
+
+	/** Widget for the log view */
+	TSharedPtr<SLogView> LogView;
+
+	/** The number of seconds the profiler has been active */
+	float DurationActive;
+
+private:
+	/** Holds the tab manager that manages the front-end's tabs. */
+	TSharedPtr<FTabManager> TabManager;
+
+	/** The handle to the active update duration tick */
+	TWeakPtr<FActiveTimerHandle> ActiveTimerHandle;
+};
