Copying //UE4/Dev-Core to Dev-Main (//UE4/Dev-Main)

diff --git a/Engine/Source/Developer/TraceInsights/Private/Insights/Widgets/STimingView.cpp b/Engine/Source/Developer/TraceInsights/Private/Insights/Widgets/STimingView.cpp
new file mode 100644
--- /dev/null
+++ b/Engine/Source/Developer/TraceInsights/Private/Insights/Widgets/STimingView.cpp
@@ -0,0 +1,3806 @@
+// Copyright 1998-2019 Epic Games, Inc. All Rights Reserved.
+
+#include "STimingView.h"
+
+#include "Brushes/SlateBorderBrush.h"
+#include "Brushes/SlateBoxBrush.h"
+#include "Brushes/SlateColorBrush.h"
+#include "Containers/ArrayBuilder.h"
+#include "Containers/MapBuilder.h"
+#include "EditorStyleSet.h"
+#include "Fonts/FontMeasure.h"
+#include "Fonts/SlateFontInfo.h"
+#include "Framework/Application/MenuStack.h"
+#include "Framework/Application/SlateApplication.h"
+#include "Framework/MultiBox/MultiBoxBuilder.h"
+#include "HAL/PlatformTime.h"
+#include "Layout/WidgetPath.h"
+#include "Misc/Paths.h"
+#include "Rendering/DrawElements.h"
+#include "SlateOptMacros.h"
+#include "Styling/CoreStyle.h"
+#include "TraceServices/AnalysisService.h"
+#include "TraceServices/SessionService.h"
+#include "Widgets/Input/SButton.h"
+#include "Widgets/Input/SComboButton.h"
+#include "Widgets/Layout/SScrollBar.h"
+#include "Widgets/Layout/SSpacer.h"
+#include "Widgets/SOverlay.h"
+#include "Widgets/Text/STextBlock.h"
+
+// Insights
+#include "Insights/Common/PaintUtils.h"
+#include "Insights/Common/Stopwatch.h"
+#include "Insights/Common/TimeUtils.h"
+#include "Insights/InsightsManager.h"
+#include "Insights/TimingProfilerCommon.h"
+#include "Insights/TimingProfilerManager.h"
+#include "Insights/ViewModels/BaseTimingTrack.h"
+#include "Insights/ViewModels/TimingViewDrawHelper.h"
+#include "Insights/Widgets/SStatsView.h"
+#include "Insights/Widgets/STimersView.h"
+#include "Insights/Widgets/STimingProfilerWindow.h"
+
+#include <limits>
+
+////////////////////////////////////////////////////////////////////////////////////////////////////
+
+#define DEBUG_TIMING_TRACK 0
+#define LOCTEXT_NAMESPACE "STimingView"
+
+// start auto generated ids from a big number (MSB set to 1) to avoid collisions with ids for gpu/cpu tracks based on 32bit timeline index
+uint64 FBaseTimingTrack::IdGenerator = (1ULL << 63);
+
+const TCHAR* GetName(ELoadTimeProfilerPackageEventType Type);
+const TCHAR* GetName(ELoadTimeProfilerObjectEventType Type);
+
+////////////////////////////////////////////////////////////////////////////////////////////////////
+
+STimingView::STimingView()
+	: bAssetLoadingMode(false)
+	, TimeRulerTrack(FBaseTimingTrack::GenerateId())
+	, MarkersTrack(FBaseTimingTrack::GenerateId())
+	, GraphTrack(FBaseTimingTrack::GenerateId())
+	, TooltipWidth(MinTooltipWidth)
+	, TooltipAlpha(0.0f)
+	, WhiteBrush(FCoreStyle::Get().GetBrush("WhiteBrush"))
+	, MainFont(FCoreStyle::GetDefaultFontStyle("Regular", 8))
+{
+	Reset();
+}
+
+////////////////////////////////////////////////////////////////////////////////////////////////////
+
+STimingView::~STimingView()
+{
+	for (const auto& KV : CachedTimelines)
+	{
+		delete KV.Value;
+	}
+	CachedTimelines.Reset();
+}
+
+////////////////////////////////////////////////////////////////////////////////////////////////////
+
+void STimingView::Construct(const FArguments& InArgs)
+{
+	OnSelectionChanged = InArgs._OnSelectionChanged;
+	OnHoveredEventChanged = InArgs._OnHoveredEventChanged;
+	OnSelectedEventChanged = InArgs._OnSelectedEventChanged;
+
+	ChildSlot
+	[
+		SNew(SOverlay)
+		.Visibility(EVisibility::SelfHitTestInvisible)
+
+		+ SOverlay::Slot()
+		.VAlign(VAlign_Bottom)
+		.Padding(FMargin(0, 0, 12, 0))
+		[
+			SAssignNew(HorizontalScrollBar, SScrollBar)
+			.Orientation(Orient_Horizontal)
+			.AlwaysShowScrollbar(false)
+			.Visibility(EVisibility::Visible)
+			.Thickness(FVector2D(5.0f, 5.0f))
+			//.RenderOpacity(0.75)
+			.OnUserScrolled(this, &STimingView::HorizontalScrollBar_OnUserScrolled)
+		]
+
+		+ SOverlay::Slot()
+		.HAlign(HAlign_Right)
+		.Padding(FMargin(0, 22, 0, 12))
+		[
+			SAssignNew(VerticalScrollBar, SScrollBar)
+			.Orientation(Orient_Vertical)
+			.AlwaysShowScrollbar(false)
+			.Visibility(EVisibility::Visible)
+			.Thickness(FVector2D(5.0f, 5.0f))
+			//.RenderOpacity(0.75)
+			.OnUserScrolled(this, &STimingView::VerticalScrollBar_OnUserScrolled)
+		]
+
+		+ SOverlay::Slot()
+		.HAlign(HAlign_Left)
+		.VAlign(VAlign_Top)
+		.Padding(FMargin(0.0f, 0.0f, 0.0f, 0.0f))
+		[
+			// Tracks Filter
+			SNew(SComboButton)
+			.ComboButtonStyle(FEditorStyle::Get(), "GenericFilters.ComboButtonStyle")
+			.ForegroundColor(FLinearColor::White)
+			.ToolTipText(LOCTEXT("TracksFilterToolTip", "Filter timing tracks."))
+			.OnGetMenuContent(this, &STimingView::MakeTracksFilterMenu)
+			.HasDownArrow(true)
+			.ContentPadding(FMargin(1.0f, 1.0f, 1.0f, 1.0f))
+			.ButtonContent()
+			[
+				SNew(SHorizontalBox)
+
+				+SHorizontalBox::Slot()
+				.Padding(0.0f)
+				.AutoWidth()
+				[
+					SNew(STextBlock)
+					.TextStyle(FEditorStyle::Get(), "GenericFilters.TextStyle")
+					.Font(FEditorStyle::Get().GetFontStyle("FontAwesome.9"))
+					.Text(FText::FromString(FString(TEXT("\xf0b0"))) /*fa-filter*/)
+				]
+
+				+SHorizontalBox::Slot()
+				.Padding(2.0f, 0.0f, 0.0f, 0.0f)
+				.AutoWidth()
+				[
+					SNew(STextBlock)
+					.TextStyle(FEditorStyle::Get(), "GenericFilters.TextStyle")
+					.Text(LOCTEXT("TracksFilter", "Tracks"))
+				]
+			]
+		]
+	];
+
+	UpdateHorizontalScrollBar();
+	UpdateVerticalScrollBar();
+
+	BindCommands();
+}
+
+////////////////////////////////////////////////////////////////////////////////////////////////////
+
+void STimingView::Reset()
+{
+	bShowHideAllGpuTracks = true;
+	bShowHideAllCpuTracks = true;
+	bShowHideAllLoadingTracks = false;
+	bShowHideAllIoTracks = false;
+
+	//////////////////////////////////////////////////
+
+	Viewport.Reset();
+
+	bIsViewportDirty = true;
+	bIsVerticalViewportDirty = true;
+
+	//////////////////////////////////////////////////
+
+	for (const auto& KV : CachedTimelines)
+	{
+		delete KV.Value;
+	}
+	CachedTimelines.Reset();
+
+	//TODO: TopTracks.Reset();
+	//TODO: ScrollableTracks.Reset();
+	//TODO: BottomTracks.Reset();
+	//TODO: ForegroundTracks.Reset();
+
+	//////////////////////////////////////////////////
+
+	TimingEventsTracks.Reset();
+
+	bAreTimingEventsTracksDirty = true;
+
+	bUseDownSampling = true;
+
+	//////////////////////////////////////////////////
+
+	GpuTrack = nullptr;
+	CpuTracks.Reset();
+
+	ThreadGroups.Reset();
+
+	//////////////////////////////////////////////////
+
+	LoadingMainThreadTrack = nullptr;
+	LoadingAsyncThreadTrack = nullptr;
+
+	LoadingMainThreadId = 0;
+	LoadingAsyncThreadId = 0;
+
+	LoadingGetEventNameFn = FLoadingTrackGetEventNameDelegate::CreateRaw(this, &STimingView::GetLoadTimeProfilerEventName);
+
+	EventAggregationStr.Reset();
+	ObjectTypeAggregationStr.Reset();
+
+	IoOverviewTrack = nullptr;
+	IoActivityTrack = nullptr;
+
+	bForceIoEventsUpdate = true;
+	bMergeIoLanes = true;
+	AllIoEvents.Reset();
+	//AllIoEvents.Reserve(1000000);
+
+	//////////////////////////////////////////////////
+
+	TimeRulerTrack.Reset();
+	MarkersTrack.Reset();
+
+	GraphTrack.Reset();
+	GraphTrack.SetVisibilityFlag(false);
+
+	//////////////////////////////////////////////////
+
+	MousePosition = FVector2D::ZeroVector;
+
+	MousePositionOnButtonDown = FVector2D::ZeroVector;
+	ViewportStartTimeOnButtonDown = 0.0;
+	ViewportScrollPosYOnButtonDown = 0.0f;
+
+	MousePositionOnButtonUp = FVector2D::ZeroVector;
+
+	bIsLMB_Pressed = false;
+	bIsRMB_Pressed = false;
+
+	bIsSpaceBarKeyPressed = false;
+	bIsDragging = false;
+
+	bIsPanning = false;
+	PanningMode = EPanningMode::None;
+
+	bIsSelecting = false;
+
+	SelectionStartTime = 0.0;
+	SelectionEndTime = 0.0;
+
+	HoveredTimingEvent.Reset();
+	SelectedTimingEvent.Reset();
+
+	LastSelectionType = ESelectionType::None;
+
+	TimeMarker = std::numeric_limits<double>::infinity();
+
+	//ThisGeometry
+
+	Layout.ForceNormalMode();
+
+	//////////////////////////////////////////////////
+
+	NumUpdatedEvents = 0;
+	TimelineCacheUpdateDurationHistory.Reset();
+	TimelineCacheUpdateDurationHistory.AddValue(0);
+	TimeMarkerCacheUpdateDurationHistory.Reset();
+	TimeMarkerCacheUpdateDurationHistory.AddValue(0);
+	DrawDurationHistory.Reset();
+	OnPaintDurationHistory.Reset();
+	LastOnPaintTime = FPlatformTime::Cycles64();
+
+	if (bAssetLoadingMode)
+	{
+		EnableAssetLoadingMode();
+	}
+}
+
+////////////////////////////////////////////////////////////////////////////////////////////////////
+
+void STimingView::EnableAssetLoadingMode()
+{
+	bAssetLoadingMode = true;
+
+	bShowHideAllGpuTracks = false;
+	bShowHideAllCpuTracks = false;
+	bShowHideAllLoadingTracks = true;
+}
+
+////////////////////////////////////////////////////////////////////////////////////////////////////
+
+void STimingView::Tick(const FGeometry& AllottedGeometry, const double InCurrentTime, const float InDeltaTime)
+{
+	ThisGeometry = AllottedGeometry;
+
+	const float ViewWidth = AllottedGeometry.GetLocalSize().X;
+	const float ViewHeight = AllottedGeometry.GetLocalSize().Y;
+	if (Viewport.UpdateSize(ViewWidth, ViewHeight))
+	{
+		bIsViewportDirty = true;
+	}
+
+	//////////////////////////////////////////////////
+	// Update Y postion of non-scrollable tracks (docked on top).
+	// Also compute total height (TopOffset) of top docked tracks.
+
+	float TopOffset = 0.0f;
+	if (TimeRulerTrack.IsVisible())
+	{
+		TimeRulerTrack.SetPosY(TopOffset);
+		TopOffset += TimeRulerTrack.GetHeight();
+	}
+	if (MarkersTrack.IsVisible())
+	{
+		MarkersTrack.SetPosY(TopOffset);
+		TopOffset += MarkersTrack.GetHeight();
+	}
+	if (GraphTrack.IsVisible())
+	{
+		GraphTrack.SetPosY(TopOffset);
+		GraphTrack.SetHeight(200.0f);
+		TopOffset += GraphTrack.GetHeight();
+	}
+	Viewport.TopOffset = TopOffset;
+
+	//////////////////////////////////////////////////
+
+	if (!bIsPanning)
+	{
+		//////////////////////////////////////////////////
+		// Elastic snap to vertical scroll limits.
+
+		const float MinY = 0.0f;// -0.5f * Viewport.Height;
+		const float MaxY = Viewport.ScrollHeight - Viewport.Height + TopOffset + 7.0f;
+		const float U = 0.5f;
+
+		float ScrollPosY = Viewport.ScrollPosY;
+		if (ScrollPosY < MinY)
+		{
+			ScrollPosY = ScrollPosY * U + (1.0f - U) * MinY;
+			if (FMath::IsNearlyEqual(ScrollPosY, MinY, 0.5f))
+			{
+				ScrollPosY = MinY;
+			}
+		}
+		else if (ScrollPosY > MaxY)
+		{
+			ScrollPosY = ScrollPosY * U + (1.0f - U) * MaxY;
+			if (FMath::IsNearlyEqual(ScrollPosY, MaxY, 0.5f))
+			{
+				ScrollPosY = MaxY;
+			}
+			if (ScrollPosY < MinY)
+			{
+				ScrollPosY = MinY;
+			}
+		}
+
+		if (Viewport.ScrollPosY != ScrollPosY)
+		{
+			Viewport.ScrollPosY = ScrollPosY;
+			UpdateVerticalScrollBar();
+			bIsVerticalViewportDirty = true;
+		}
+
+		//////////////////////////////////////////////////
+		// Elastic snap to horizontal time limits.
+
+		if (Viewport.EnforceHorizontalScrollLimits(0.5)) // 0.5 is the interpolation factor
+		{
+			UpdateHorizontalScrollBar();
+			bIsViewportDirty = true;
+		}
+	}
+
+	// We need to check if TimersView or StatsView needs to update their lists of timers / counters.
+	// But, ensure we do not check too often.
+	static uint64 NextTimestamp = 0;
+	uint64 Time = FPlatformTime::Cycles64();
+	if (Time > NextTimestamp)
+	{
+		const uint64 WaitTime = static_cast<uint64>(0.2 / FPlatformTime::GetSecondsPerCycle64()); // 200ms
+		NextTimestamp = Time + WaitTime;
+
+		TSharedPtr<STimingProfilerWindow> Wnd = FTimingProfilerManager::Get()->GetProfilerWindow();
+		if (Wnd)
+		{
+			if (Wnd->TimersView)
+			{
+				Wnd->TimersView->RebuildTree(false);
+			}
+			if (Wnd->StatsView)
+			{
+				Wnd->StatsView->RebuildTree(false);
+			}
+		}
+	}
+
+	UpdateIo();
+
+	TSharedPtr<const Trace::IAnalysisSession> Session = FInsightsManager::Get()->GetSession();
+	if (Session)
+	{
+		Trace::FAnalysisSessionReadScope SessionReadScope(*Session.Get());
+
+		// Check if horizontal scroll area has changed.
+		double SessionTime = Session->GetDurationSeconds();
+		if (SessionTime > Viewport.MaxValidTime &&
+			SessionTime != DBL_MAX &&
+			SessionTime != std::numeric_limits<double>::infinity())
+		{
+			//UE_LOG(TimingProfiler, Log, TEXT("Session Duration: %g"), DT);
+			Viewport.MaxValidTime = SessionTime;
+			UpdateHorizontalScrollBar();
+			//bIsMaxValidTimeDirty = true;
+
+			if (SessionTime >= Viewport.StartTime && SessionTime <= Viewport.EndTime)
+			{
+				bAreTimingEventsTracksDirty = true;
+				MarkersTrack.SetDirtyFlag();
+			}
+		}
+
+		bool bIsTimingEventsTrackDirty = false;
+
+		if (Trace::ReadLoadTimeProfilerProvider(*Session.Get()))
+		{
+			const Trace::ILoadTimeProfilerProvider& LoadTimeProfilerProvider = *Trace::ReadLoadTimeProfilerProvider(*Session.Get());
+
+			LoadingMainThreadId = LoadTimeProfilerProvider.GetMainThreadId();
+			LoadingAsyncThreadId = LoadTimeProfilerProvider.GetAsyncLoadingThreadId();
+
+			if (LoadingMainThreadTrack == nullptr)
+			{
+				uint64 TrackId = FBaseTimingTrack::GenerateId();
+				LoadingMainThreadTrack = AddTimingEventsTrack(TrackId, ETimingEventsTrackType::Loading, TEXT("Loading - Main Thread"), nullptr, -3);
+				LoadingMainThreadTrack->SetVisibilityFlag(bShowHideAllLoadingTracks);
+				bIsTimingEventsTrackDirty = true;
+			}
+
+			if (LoadingAsyncThreadTrack == nullptr)
+			{
+				uint64 TrackId = FBaseTimingTrack::GenerateId();
+				LoadingAsyncThreadTrack = AddTimingEventsTrack(TrackId, ETimingEventsTrackType::Loading, TEXT("Loading - Async Thread"), nullptr, -2);
+				LoadingAsyncThreadTrack->SetVisibilityFlag(bShowHideAllLoadingTracks);
+				bIsTimingEventsTrackDirty = true;
+			}
+		}
+
+		if (Trace::ReadFileActivityProvider(*Session.Get()))
+		{
+			if (IoOverviewTrack == nullptr)
+			{
+				// Note: The I/O timelines are just prototypes for now (will be removed once the functionality is moved in analyzer).
+				const uint64 TrackId = FBaseTimingTrack::GenerateId();
+				IoOverviewTrack = AddTimingEventsTrack(TrackId, ETimingEventsTrackType::Io, TEXT("I/O Overview"), nullptr, -1);
+				IoOverviewTrack->SetVisibilityFlag(bShowHideAllIoTracks);
+				bIsTimingEventsTrackDirty = true;
+			}
+		}
+
+		if (Trace::ReadTimingProfilerProvider(*Session.Get()))
+		{
+			const Trace::ITimingProfilerProvider& TimingProfilerProvider = *Trace::ReadTimingProfilerProvider(*Session.Get());
+
+			// Check if we have a GPU track.
+			uint32 GpuTimelineIndex;
+			if (TimingProfilerProvider.GetGpuTimelineIndex(GpuTimelineIndex))
+			{
+				const uint64 TrackId = static_cast<uint64>(GpuTimelineIndex);
+				if (!CachedTimelines.Contains(TrackId))
+				{
+					GpuTrack = AddTimingEventsTrack(TrackId, ETimingEventsTrackType::Gpu, TEXT("GPU"), nullptr, 0);
+					GpuTrack->SetVisibilityFlag(bShowHideAllGpuTracks);
+					bIsTimingEventsTrackDirty = true;
+				}
+			}
+
+			// Check available CPU tracks.
+			int32 Order = 1;
+			const Trace::IThreadProvider& ThreadProvider = Trace::ReadThreadProvider(*Session.Get());
+			ThreadProvider.EnumerateThreads([this, &Order, &bIsTimingEventsTrackDirty, &TimingProfilerProvider](const Trace::FThreadInfo& ThreadInfo)
+			{
+				const TCHAR* GroupName = ThreadInfo.GroupName;
+				if (GroupName == nullptr)
+				{
+					GroupName = ThreadInfo.Name;
+				}
+
+				bool bIsGroupVisible = bShowHideAllCpuTracks;
+				if (GroupName != nullptr)
+				{
+					if (!ThreadGroups.Contains(GroupName))
+					{
+						ThreadGroups.Add(GroupName, { GroupName, bIsGroupVisible, 0 });
+					}
+					else
+					{
+						bIsGroupVisible = ThreadGroups[GroupName].bIsVisible;
+					}
+				}
+
+				uint32 CpuTimelineIndex;
+				if (TimingProfilerProvider.GetCpuThreadTimelineIndex(ThreadInfo.Id, CpuTimelineIndex))
+				{
+					FTimingEventsTrack* Track = nullptr;
+					const uint64 TrackId = static_cast<uint64>(CpuTimelineIndex);
+
+					if (!CachedTimelines.Contains(TrackId))
+					{
+						FString TrackName(ThreadInfo.Name && *ThreadInfo.Name ? ThreadInfo.Name : FString::Printf(TEXT("Thread %u")));
+
+						// Create new Timing Events track for the CPU thread.
+						Track = AddTimingEventsTrack(TrackId, ETimingEventsTrackType::Cpu, TrackName, GroupName, Order);
+
+						Track->SetThreadId(ThreadInfo.Id);
+						CpuTracks.Add(ThreadInfo.Id, Track);
+
+						FThreadGroup& ThreadGroup = ThreadGroups[GroupName];
+						ThreadGroup.NumTimelines++;
+
+						if (bAssetLoadingMode && (ThreadInfo.Id == LoadingMainThreadId || ThreadInfo.Id == LoadingAsyncThreadId))
+						{
+							Track->SetVisibilityFlag(true);
+							ThreadGroup.bIsVisible = true;
+						}
+						else
+						{
+							Track->SetVisibilityFlag(bIsGroupVisible);
+						}
+
+						bIsTimingEventsTrackDirty = true;
+					}
+					else
+					{
+						Track = CachedTimelines[TrackId];
+
+						if (Track->GetOrder() != Order)
+						{
+							Track->SetOrder(Order);
+							bIsTimingEventsTrackDirty = true;
+						}
+					}
+
+					Order++;
+				}
+			});
+		}
+
+		if (Trace::ReadFileActivityProvider(*Session.Get()))
+		{
+			if (IoActivityTrack == nullptr)
+			{
+				// Note: The I/O timelines are just prototypes for now (will be removed once the functionality is moved in analyzer).
+				uint64 TrackId = FBaseTimingTrack::GenerateId();
+				IoActivityTrack = AddTimingEventsTrack(TrackId, ETimingEventsTrackType::Io, TEXT("I/O Activity"), nullptr, 999999);
+				IoActivityTrack->SetVisibilityFlag(bShowHideAllIoTracks);
+				bIsTimingEventsTrackDirty = true;
+			}
+		}
+
+		if (bIsTimingEventsTrackDirty)
+		{
+			// The list has changed. Sort the list again.
+			TimingEventsTracks.Sort([this](const FTimingEventsTrack& A, const FTimingEventsTrack& B) -> bool
+			{
+				return A.GetOrder() < B.GetOrder();
+			});
+		}
+	}
+
+	// Compute total height of visible tracks.
+	float TotalScrollHeight = 0.0f;
+	for (FTimingEventsTrack* TrackPtr : TimingEventsTracks)
+	{
+		ensure(TrackPtr != nullptr);
+		FTimingEventsTrack& Track = *TrackPtr;
+
+		if (Track.IsVisible())
+		{
+			Track.SetPosY(TotalScrollHeight);
+			TotalScrollHeight += Track.GetHeight();
+		}
+
+		// Reset track's hovered/selected state.
+		Track.SetHoveredState(false);
+		Track.SetSelectedFlag(false);
+	}
+	TotalScrollHeight += 1.0f; // allow 1 pixel at the bottom (for last horizontal line)
+
+	// Set hovered/selected state for actual hovered/selected tracks, if any.
+	if (HoveredTimingEvent.Track != nullptr)
+	{
+		const_cast<FTimingEventsTrack*>(HoveredTimingEvent.Track)->SetHoveredState(true);
+	}
+	if (SelectedTimingEvent.Track != nullptr)
+	{
+		const_cast<FTimingEventsTrack*>(SelectedTimingEvent.Track)->SetSelectedFlag(true);
+	}
+
+	//TimeRulerTrack.Tick(InCurrentTime, InDeltaTime);
+	MarkersTrack.Tick(InCurrentTime, InDeltaTime);
+
+	// Check if vertical scroll area has changed.
+	if (TotalScrollHeight != Viewport.ScrollHeight)
+	{
+		Viewport.ScrollHeight = TotalScrollHeight;
+		UpdateVerticalScrollBar();
+		bIsVerticalViewportDirty = true;
+	}
+
+	// Update layout change animation (i.e. compact mode <-> normal mode).
+	Layout.Update();
+
+	//////////////////////////////////////////////////
+	// Update dirty tracks.
+
+	if (bIsViewportDirty)
+	{
+		bIsViewportDirty = false;
+
+		//TimeRuler.SetDirtyFlag();
+		MarkersTrack.SetDirtyFlag();
+		bAreTimingEventsTracksDirty = true;
+	}
+
+	if (bIsVerticalViewportDirty)
+	{
+		bIsVerticalViewportDirty = false;
+
+		bAreTimingEventsTracksDirty = true;
+	}
+
+	//if (TimeRuler.IsVisible() && TimeRuler.IsDirty())
+	//{
+	//	TimeRuler.ClearDirtyFlag();
+	//
+	//	TimeRuler.OnViewportChanged(Viewport);
+	//}
+
+	if (MarkersTrack.IsVisible() && MarkersTrack.IsDirty())
+	{
+		MarkersTrack.ClearDirtyFlag();
+
+		FStopwatch Stopwatch;
+		Stopwatch.Start();
+
+		MarkersTrack.Update(Viewport);
+
+		Stopwatch.Stop();
+		TimeMarkerCacheUpdateDurationHistory.AddValue(Stopwatch.AccumulatedTime);
+	}
+
+	if (bAreTimingEventsTracksDirty)
+	{
+		bAreTimingEventsTracksDirty = false;
+
+		FStopwatch Stopwatch;
+		Stopwatch.Start();
+
+		if (GraphTrack.IsVisible())
+		{
+			GraphTrack.Update(Viewport);
+		}
+
+		for (int32 TrackIndex = 0; TrackIndex < TimingEventsTracks.Num(); ++TrackIndex)
+		{
+			TimingEventsTracks[TrackIndex]->Update(Viewport);
+		}
+
+		Stopwatch.Stop();
+		TimelineCacheUpdateDurationHistory.AddValue(Stopwatch.AccumulatedTime);
+	}
+}
+
+////////////////////////////////////////////////////////////////////////////////////////////////////
+
+int32 STimingView::OnPaint(const FPaintArgs& Args, const FGeometry& AllottedGeometry, const FSlateRect& MyCullingRect, FSlateWindowElementList& OutDrawElements, int32 LayerId, const FWidgetStyle& InWidgetStyle, bool bParentEnabled) const
+{
+	const bool bEnabled = ShouldBeEnabled(bParentEnabled);
+	const ESlateDrawEffect DrawEffects = bEnabled ? ESlateDrawEffect::None : ESlateDrawEffect::DisabledEffect;
+	FDrawContext DrawContext(AllottedGeometry, MyCullingRect, InWidgetStyle, DrawEffects, OutDrawElements, LayerId);
+
+#if 0 // Enabling this may further increase UI performance (TODO: profile if this is really needed again).
+	// Avoids multiple resizes of Slate's draw elements buffers.
+	OutDrawElements.GetRootDrawLayer().DrawElements.Reserve(50000);
+#endif
+
+	const TSharedRef<FSlateFontMeasure> FontMeasureService = FSlateApplication::Get().GetRenderer()->GetFontMeasureService();
+
+	const float ViewWidth = AllottedGeometry.GetLocalSize().X;
+	const float ViewHeight = AllottedGeometry.GetLocalSize().Y;
+
+#if 0 // Enabling this may further increase UI performance (TODO: profile if this is really needed again).
+	// Warm up Slate vertex/index buffers to avoid initial freezes due to multiple resizes of those buffers.
+	static bool bWarmingUp = false;
+	if (!bWarmingUp)
+	{
+		bWarmingUp = true;
+
+		FRandomStream RandomStream(0);
+		const int32 Count = 1'000'000;
+		for (int32 Index = 0; Index < Count; ++Index)
+		{
+			float X = ViewWidth * RandomStream.GetFraction();
+			float Y = ViewHeight * RandomStream.GetFraction();
+			FLinearColor Color(RandomStream.GetFraction(), RandomStream.GetFraction(), RandomStream.GetFraction(), 1.0f);
+			DrawContext.DrawBox(X, Y, 1.0f, 1.0f, WhiteBrush, Color);
+		}
+		LayerId++;
+		LayerId++;
+	}
+#endif
+
+	//////////////////////////////////////////////////
+
+	FStopwatch Stopwatch;
+	Stopwatch.Start();
+
+	FTimingViewDrawHelper Helper(DrawContext, Viewport, Layout);
+
+	// Draw background.
+	Helper.DrawBackground();
+
+	// Draw scrollable tracks.
+	{
+		Helper.BeginTimelines();
+
+		for (int32 TrackIndex = 0; TrackIndex < TimingEventsTracks.Num(); ++TrackIndex)
+		{
+			FTimingEventsTrack& Track = *TimingEventsTracks[TrackIndex];
+			if (Track.IsVisible())
+			{
+				//TODO: Track.Draw(Helper);
+				if (Track.GetType() == ETimingEventsTrackType::Cpu)
+				{
+					DrawTimingProfilerTrack(Helper, Track);
+				}
+				else if (Track.GetType() == ETimingEventsTrackType::Gpu)
+				{
+					DrawTimingProfilerTrack(Helper, Track);
+				}
+				else if (Track.GetType() == ETimingEventsTrackType::Loading)
+				{
+					DrawLoadTimeProfilerTrack(Helper, Track);
+				}
+				else if (Track.GetType() == ETimingEventsTrackType::Io)
+				{
+					if (&Track == IoOverviewTrack)
+					{
+						DrawIoOverviewTrack(Helper, Track);
+					}
+					else if (&Track == IoActivityTrack)
+					{
+						DrawIoActivityTrack(Helper, Track);
+					}
+				}
+			}
+		}
+
+		Helper.EndTimelines();
+	}
+
+	if (GraphTrack.IsVisible())
+	{
+		GraphTrack.Draw(DrawContext, Viewport);
+	}
+
+	if (!FTimingEvent::AreEquals(SelectedTimingEvent, HoveredTimingEvent))
+	{
+		// Highlight the selected timing event (if any).
+		if (SelectedTimingEvent.IsValid())
+		{
+			const float Y = Viewport.GetViewportY(SelectedTimingEvent.Track->GetPosY()) + Layout.GetLaneY(SelectedTimingEvent.Depth);
+			Helper.DrawTimingEventHighlight(SelectedTimingEvent.StartTime, SelectedTimingEvent.EndTime, Y, FTimingViewDrawHelper::EHighlightMode::Selected);
+		}
+
+		// Highlight the hovered timing event (if any).
+		if (HoveredTimingEvent.IsValid())
+		{
+			const float Y = Viewport.GetViewportY(HoveredTimingEvent.Track->GetPosY()) + Layout.GetLaneY(HoveredTimingEvent.Depth);
+			Helper.DrawTimingEventHighlight(HoveredTimingEvent.StartTime, HoveredTimingEvent.EndTime, Y, FTimingViewDrawHelper::EHighlightMode::Hovered);
+		}
+	}
+	else
+	{
+		// Highlight the selected and hovered timing event (if any).
+		if (SelectedTimingEvent.IsValid())
+		{
+			const float Y = Viewport.GetViewportY(SelectedTimingEvent.Track->GetPosY()) + Layout.GetLaneY(SelectedTimingEvent.Depth);
+			Helper.DrawTimingEventHighlight(SelectedTimingEvent.StartTime, SelectedTimingEvent.EndTime, Y, FTimingViewDrawHelper::EHighlightMode::SelectedAndHovered);
+		}
+	}
+
+	if (MarkersTrack.IsVisible())
+	{
+		// Draw the time markers (vertical lines + category & message texts).
+		MarkersTrack.Draw(DrawContext, Viewport);
+	}
+
+	//////////////////////////////////////////////////
+
+	// Draw the time ruler.
+	if (TimeRulerTrack.IsVisible())
+	{
+		TimeRulerTrack.Draw(DrawContext, Viewport, MousePosition, bIsSelecting, SelectionStartTime, SelectionEndTime);
+	}
+
+	// Fill background for the Tracks filter combobox.
+	DrawContext.DrawBox(0, 0, 66, 22, WhiteBrush, FLinearColor(0.05f, 0.05f, 0.05f, 1.0f));
+
+	//////////////////////////////////////////////////
+
+	// Draw the time range selection.
+	DrawTimeRangeSelection(DrawContext);
+
+	//////////////////////////////////////////////////
+
+	// Draw the time marker (orange vertical line).
+	//DrawTimeMarker(OnPaintState);
+	float TimeMarkerX = Viewport.TimeToSlateUnitsRounded(TimeMarker);
+	if (TimeMarkerX >= 0.0f && TimeMarkerX < Viewport.Width)
+	{
+		DrawContext.DrawBox(TimeMarkerX, 0.0f, 1.0f, Viewport.Height, WhiteBrush, FLinearColor(0.85f, 0.5f, 0.03f, 0.5f));
+		DrawContext.LayerId++;
+	}
+
+	//////////////////////////////////////////////////
+
+	Stopwatch.Stop();
+	DrawDurationHistory.AddValue(Stopwatch.AccumulatedTime);
+
+	//////////////////////////////////////////////////
+
+	const bool bShouldDisplayDebugInfo = FInsightsManager::Get()->IsDebugInfoEnabled();
+	if (bShouldDisplayDebugInfo)
+	{
+		const FSlateFontInfo& SummaryFont = MainFont;
+
+		const float MaxFontCharHeight = FontMeasureService->Measure(TEXT("!"), SummaryFont).Y;
+		const float DbgDY = MaxFontCharHeight;
+
+		const float DbgW = 320.0f;
+		const float DbgH = DbgDY * 9 + 3.0f;
+		const float DbgX = ViewWidth - DbgW - 20.0f;
+		float DbgY = Viewport.TopOffset + 10.0f;
+
+		DrawContext.LayerId++;
+
+		DrawContext.DrawBox(DbgX - 2.0f, DbgY - 2.0f, DbgW, DbgH, WhiteBrush, FLinearColor(1.0f, 1.0f, 1.0f, 0.9f));
+		DrawContext.LayerId++;
+
+		FLinearColor DbgTextColor(0.0f, 0.0f, 0.0f, 0.9f);
+
+		//////////////////////////////////////////////////
+		// Display the "Draw" performance info.
+
+		// Time interval since last OnPaint call.
+		const uint64 CurrentTime = FPlatformTime::Cycles64();
+		const uint64 OnPaintDuration = CurrentTime - LastOnPaintTime;
+		LastOnPaintTime = CurrentTime;
+		OnPaintDurationHistory.AddValue(OnPaintDuration); // saved for last 32 OnPaint calls
+		const uint64 AvgOnPaintDuration = OnPaintDurationHistory.ComputeAverage();
+		const uint64 AvgOnPaintDurationMs = FStopwatch::Cycles64ToMilliseconds(AvgOnPaintDuration);
+		const double AvgOnPaintFps = AvgOnPaintDurationMs != 0 ? 1.0 / FStopwatch::Cycles64ToSeconds(AvgOnPaintDuration) : 0.0;
+
+		const uint64 AvgDrawDurationMs = FStopwatch::Cycles64ToMilliseconds(DrawDurationHistory.ComputeAverage());
+
+		DrawContext.DrawText
+		(
+			DbgX, DbgY,
+			FString::Printf(TEXT("D: %llu ms + %llu ms = %llu ms (%d fps)"),
+				AvgDrawDurationMs, // drawing time
+				AvgOnPaintDurationMs - AvgDrawDurationMs, // other overhead to OnPaint calls
+				AvgOnPaintDurationMs, // average time between two OnPaint calls
+				FMath::RoundToInt(AvgOnPaintFps)), // framerate of OnPaint calls
+			SummaryFont, DbgTextColor
+		);
+		DbgY += DbgDY;
+
+		//////////////////////////////////////////////////
+		// Display the "timelines cache update" info.
+
+		const uint64 AvgTimelineUpdateDurationMs = FStopwatch::Cycles64ToMilliseconds(TimelineCacheUpdateDurationHistory.ComputeAverage());
+		const uint64 LastTimelineUpdateDurationMs = FStopwatch::Cycles64ToMilliseconds(TimelineCacheUpdateDurationHistory.GetValue(0));
+
+		DrawContext.DrawText
+		(
+			DbgX, DbgY,
+			FString::Printf(TEXT("U1 avg: %llu ms, last: %llu ms"),
+				AvgTimelineUpdateDurationMs, // time to update the Timeline cache
+				LastTimelineUpdateDurationMs), // last time to update the Timeline cache
+			SummaryFont, DbgTextColor
+		);
+		DbgY += DbgDY;
+
+		//////////////////////////////////////////////////
+		// Display the "time markers cache update" info.
+
+		const uint64 AvgTimeMarkerUpdateDurationMs = FStopwatch::Cycles64ToMilliseconds(TimeMarkerCacheUpdateDurationHistory.ComputeAverage());
+		const uint64 LastTimeMarkerUpdateDurationMs = FStopwatch::Cycles64ToMilliseconds(TimeMarkerCacheUpdateDurationHistory.GetValue(0));
+
+		DrawContext.DrawText
+		(
+			DbgX, DbgY,
+			FString::Printf(TEXT("U2 avg: %llu ms, last: %llu ms"),
+				AvgTimeMarkerUpdateDurationMs, // time to update the TimeMarker cache
+				LastTimeMarkerUpdateDurationMs), // last time to update the Timeline cache
+			SummaryFont, DbgTextColor
+		);
+		DbgY += DbgDY;
+
+		//////////////////////////////////////////////////
+		// Display timing events stats.
+
+		DrawContext.DrawText
+		(
+			DbgX, DbgY,
+			FString::Format(TEXT("{0} events : {1} ({2}) boxes, {3} borders, {4} texts"),
+			{
+				FText::AsNumber(Helper.GetNumEvents()).ToString(),
+				FText::AsNumber(Helper.GetNumDrawBoxes()).ToString(),
+				FText::AsPercent((double)Helper.GetNumDrawBoxes() / (Helper.GetNumDrawBoxes() + Helper.GetNumMergedBoxes())).ToString(),
+				FText::AsNumber(Helper.GetNumDrawBorders()).ToString(),
+				FText::AsNumber(Helper.GetNumDrawTexts()).ToString(),
+				//OutDrawElements.GetRootDrawLayer().GetElementCount(),
+			}),
+			SummaryFont, DbgTextColor
+		);
+		DbgY += DbgDY;
+
+		//////////////////////////////////////////////////
+		// Display time markers stats.
+
+		DrawContext.DrawText
+		(
+			DbgX, DbgY,
+			FString::Format(TEXT("{0}{1} logs : {2} boxes, {3} texts"),
+			{
+				MarkersTrack.IsVisible() ? TEXT("") : TEXT("*"),
+				FText::AsNumber(MarkersTrack.GetNumLogMessages()).ToString(),
+				FText::AsNumber(MarkersTrack.GetNumBoxes()).ToString(),
+				FText::AsNumber(MarkersTrack.GetNumTexts()).ToString(),
+			}),
+			SummaryFont, DbgTextColor
+		);
+		DbgY += DbgDY;
+
+		//////////////////////////////////////////////////
+		// Display Graph track stats.
+
+		DrawContext.DrawText
+		(
+			DbgX, DbgY,
+			FString::Format(TEXT("{0} events : {1} points, {2} lines, {3} boxes"),
+				{
+					FText::AsNumber(GraphTrack.GetNumAddedEvents()).ToString(),
+					FText::AsNumber(GraphTrack.GetNumDrawPoints()).ToString(),
+					FText::AsNumber(GraphTrack.GetNumDrawLines()).ToString(),
+					FText::AsNumber(GraphTrack.GetNumDrawBoxes()).ToString(),
+				}),
+				SummaryFont, DbgTextColor
+				);
+		DbgY += DbgDY;
+
+		//////////////////////////////////////////////////
+		// Display viewport's horizontal info.
+
+		DrawContext.DrawText
+		(
+			DbgX, DbgY,
+			FString::Printf(TEXT("SX: %g, ST: %g, ET: %s"), Viewport.ScaleX, Viewport.StartTime, *TimeUtils::FormatTimeAuto(Viewport.MaxValidTime)),
+			SummaryFont, DbgTextColor
+		);
+		DbgY += DbgDY;
+
+		//////////////////////////////////////////////////
+		// Display viewport's vertical info.
+
+		DrawContext.DrawText
+		(
+			DbgX, DbgY,
+			FString::Printf(TEXT("Y: %.2f, H: %g, VH: %g"), Viewport.ScrollPosY, Viewport.ScrollHeight, Viewport.Height),
+			SummaryFont, DbgTextColor
+		);
+		DbgY += DbgDY;
+
+		//////////////////////////////////////////////////
+		// Display input related debug info.
+
+		FString InputStr = FString::Printf(TEXT("(%.0f, %.0f)"), MousePosition.X, MousePosition.Y);
+		if (bIsSpaceBarKeyPressed) InputStr += " Space";
+		if (bIsLMB_Pressed) InputStr += " LMB";
+		if (bIsRMB_Pressed) InputStr += " RMB";
+		if (bIsPanning) InputStr += " Panning";
+		if (bIsSelecting) InputStr += " Selecting";
+		if (bIsDragging) InputStr += " Dragging";
+		DrawContext.DrawText(DbgX, DbgY, InputStr, SummaryFont, DbgTextColor);
+		DbgY += DbgDY;
+	}
+
+	//////////////////////////////////////////////////
+
+	{
+		// Draw info about selected event (bottom-right corner).
+		if (SelectedTimingEvent.IsValid())
+		{
+			if (SelectedTimingEvent.Track->GetType() == ETimingEventsTrackType::Cpu ||
+				SelectedTimingEvent.Track->GetType() == ETimingEventsTrackType::Gpu)
+			{
+				const FTimerNodePtr TimerNodePtr = GetTimerNode(SelectedTimingEvent.TypeId);
+
+				FString Str = FString::Printf(TEXT("%s (Incl.: %s, Excl.: %s)"),
+					TimerNodePtr ? *(TimerNodePtr->GetName().ToString()) : TEXT("N/A"),
+					*TimeUtils::FormatTimeAuto(SelectedTimingEvent.Duration()),
+					*TimeUtils::FormatTimeAuto(SelectedTimingEvent.ExclusiveTime));
+
+				const FVector2D Size = FontMeasureService->Measure(Str, MainFont);
+				const float X = Viewport.Width - Size.X - 23.0f;
+				const float Y = Viewport.Height - Size.Y - 18.0f;
+
+				const FLinearColor BackgroundColor(0.05f, 0.05f, 0.05f, 1.0f);
+				const FLinearColor TextColor(0.7f, 0.7f, 0.7f, 1.0f);
+
+				DrawContext.DrawBox(X - 8.0f, Y - 2.0f, Size.X + 16.0f, Size.Y + 4.0f, WhiteBrush, BackgroundColor);
+				DrawContext.LayerId++;
+
+				DrawContext.DrawText(X, Y, Str, MainFont, TextColor);
+				DrawContext.LayerId++;
+			}
+			else if (SelectedTimingEvent.Track->GetType() == ETimingEventsTrackType::Loading)
+			{
+				//TODO: ...
+			}
+		}
+
+		// Draw info about hovered event (like a tooltip at mouse position).
+		if (HoveredTimingEvent.IsValid() && !MousePosition.IsZero())
+		{
+			if (HoveredTimingEvent.Track->GetType() == ETimingEventsTrackType::Cpu ||
+				HoveredTimingEvent.Track->GetType() == ETimingEventsTrackType::Gpu)
+			{
+				const FTimerNodePtr TimerNodePtr = GetTimerNode(HoveredTimingEvent.TypeId);
+				FString Name = TimerNodePtr ? TimerNodePtr->GetName().ToString() : TEXT("N/A");
+
+				float W = FontMeasureService->Measure(Name, MainFont).X;
+
+				if (W < MinTooltipWidth)
+				{
+					W = MinTooltipWidth;
+				}
+				if (TooltipWidth != W)
+				{
+					TooltipWidth = TooltipWidth * 0.75f + W * 0.25f;
+				}
+
+				const float MaxX = FMath::Max(0.0f, Viewport.Width - TooltipWidth - 21.0f);
+				const float X = FMath::Clamp<float>(MousePosition.X + 18.0f, 0.0f, MaxX);
+
+				constexpr float LineH = 14.0f;
+				constexpr float H = 4 * LineH;
+				const float MaxY = FMath::Max(0.0f, Viewport.Height - H - 18.0f);
+				float Y = FMath::Clamp<float>(MousePosition.Y + 18.0f, 0.0f, MaxY);
+
+				const float Alpha = 1.0f - FMath::Abs(TooltipWidth - W) / W;
+				if (TooltipAlpha < Alpha)
+				{
+					TooltipAlpha = TooltipAlpha * 0.9f + Alpha * 0.1f;
+				}
+				else
+				{
+					TooltipAlpha = Alpha;
+				}
+
+				const FLinearColor BackgroundColor(0.05f, 0.05f, 0.05f, TooltipAlpha);
+				const FLinearColor NameColor(0.9f, 0.9f, 0.5f, TooltipAlpha);
+				const FLinearColor TextColor(0.6f, 0.6f, 0.6f, TooltipAlpha);
+				const FLinearColor ValueColor(1.0f, 1.0f, 1.0f, TooltipAlpha);
+
+				DrawContext.DrawBox(X - 6.0f, Y - 3.0f, TooltipWidth + 12.0f, H + 6.0f, WhiteBrush, BackgroundColor);
+				DrawContext.LayerId++;
+
+				DrawContext.DrawText(X, Y, Name, MainFont, NameColor);
+				Y += LineH;
+
+				const float ValueX = X + 58.0f;
+
+				DrawContext.DrawText(X + 3.0f, Y, TEXT("Incl. Time:"), MainFont, TextColor);
+				FString InclStr = TimeUtils::FormatTimeAuto(HoveredTimingEvent.Duration());
+				DrawContext.DrawText(ValueX, Y, InclStr, MainFont, ValueColor);
+				Y += LineH;
+
+				DrawContext.DrawText(X, Y, TEXT("Excl. Time:"), MainFont, TextColor);
+				FString ExclStr = TimeUtils::FormatTimeAuto(HoveredTimingEvent.ExclusiveTime);
+				DrawContext.DrawText(ValueX, Y, ExclStr, MainFont, ValueColor);
+				Y += LineH;
+
+				DrawContext.DrawText(X + 24.0f, Y, TEXT("Depth:"), MainFont, TextColor);
+				FString DepthStr = FString::Printf(TEXT("%d"), HoveredTimingEvent.Depth);
+				DrawContext.DrawText(ValueX, Y, DepthStr, MainFont, ValueColor);
+				Y += LineH;
+
+				DrawContext.LayerId++;
+			}
+			else if (HoveredTimingEvent.Track->GetType() == ETimingEventsTrackType::Loading)
+			{
+				FString Name(LoadingGetEventNameFn.Execute(HoveredTimingEvent.Depth, HoveredTimingEvent.LoadingInfo));
+
+				const Trace::FPackageInfo* Package = HoveredTimingEvent.LoadingInfo.Package;
+				const Trace::FPackageExportInfo* Export = HoveredTimingEvent.LoadingInfo.Export;
+
+				FString PackageName = Package ? Package->Name : TEXT("N/A");
+
+				constexpr float ValueOffsetX = 100.0f;
+				constexpr float MinValueTextWidth = 220.0f;
+
+				float W = FontMeasureService->Measure(Name, MainFont).X;
+
+				float PackageNameW = ValueOffsetX + FontMeasureService->Measure(PackageName, MainFont).X;
+				if (W < PackageNameW)
+				{
+					W = PackageNameW;
+				}
+				if (W < ValueOffsetX + MinValueTextWidth)
+				{
+					W = ValueOffsetX + MinValueTextWidth;
+				}
+				if (W < MinTooltipWidth)
+				{
+					W = MinTooltipWidth;
+				}
+				if (TooltipWidth != W)
+				{
+					TooltipWidth = TooltipWidth * 0.75f + W * 0.25f;
+				}
+
+				const float MaxX = FMath::Max(0.0f, Viewport.Width - TooltipWidth - 21.0f);
+				const float X = FMath::Clamp<float>(MousePosition.X + 18.0f, 0.0f, MaxX);
+				const float ValueX = X + ValueOffsetX;
+
+				constexpr float LineH = 14.0f;
+				constexpr float H = 10 * LineH;
+				const float MaxY = FMath::Max(0.0f, Viewport.Height - H - 18.0f);
+				float Y = FMath::Clamp<float>(MousePosition.Y + 18.0f, 0.0f, MaxY);
+
+				const float Alpha = 1.0f - FMath::Abs(TooltipWidth - W) / W;
+				if (TooltipAlpha < Alpha)
+				{
+					TooltipAlpha = TooltipAlpha * 0.9f + Alpha * 0.1f;
+				}
+				else
+				{
+					TooltipAlpha = Alpha;
+				}
+
+				const FLinearColor BackgroundColor(0.05f, 0.05f, 0.05f, TooltipAlpha);
+				const FLinearColor NameColor(0.9f, 0.9f, 0.5f, TooltipAlpha);
+				const FLinearColor TextColor(0.6f, 0.6f, 0.6f, TooltipAlpha);
+				const FLinearColor ValueColor(1.0f, 1.0f, 1.0f, TooltipAlpha);
+
+				DrawContext.DrawBox(X - 6.0f, Y - 3.0f, TooltipWidth + 12.0f, H + 6.0f, WhiteBrush, BackgroundColor);
+				DrawContext.LayerId++;
+
+				DrawContext.DrawText(X, Y, Name, MainFont, NameColor);
+				Y += LineH;
+
+				DrawContext.DrawText(X, Y, TEXT("Duration:"), MainFont, TextColor);
+				FString InclStr = TimeUtils::FormatTimeAuto(HoveredTimingEvent.Duration());
+				DrawContext.DrawText(ValueX, Y, InclStr, MainFont, ValueColor);
+				Y += LineH;
+
+				DrawContext.DrawText(X, Y, TEXT("Depth:"), MainFont, TextColor);
+				FString DepthStr = FString::Printf(TEXT("%d"), HoveredTimingEvent.Depth);
+				DrawContext.DrawText(ValueX, Y, DepthStr, MainFont, ValueColor);
+				Y += LineH;
+
+				DrawContext.DrawText(X, Y, TEXT("Package Type:"), MainFont, TextColor);
+				DrawContext.DrawText(ValueX, Y, GetName(HoveredTimingEvent.LoadingInfo.PackageEventType), MainFont, ValueColor);
+				Y += LineH;
+
+				if (Package)
+				{
+					DrawContext.DrawText(X, Y, TEXT("Package Name:"), MainFont, TextColor);
+					DrawContext.DrawText(ValueX, Y, PackageName, MainFont, ValueColor);
+					Y += LineH;
+
+					DrawContext.DrawText(X, Y, TEXT("Header Size:"), MainFont, TextColor);
+					FString HeaderSizeStr = FString::Printf(TEXT("%d bytes"), Package->Summary.TotalHeaderSize);
+					DrawContext.DrawText(ValueX, Y, HeaderSizeStr, MainFont, ValueColor);
+					Y += LineH;
+
+					DrawContext.DrawText(X, Y, TEXT("Package Summary:"), MainFont, TextColor);
+					FString SummaryStr = FString::Printf(TEXT("%d names, %d imports, %d exports"), Package->Summary.NameCount, Package->Summary.ImportCount, Package->Summary.ExportCount);
+					DrawContext.DrawText(ValueX, Y, SummaryStr, MainFont, ValueColor);
+					Y += LineH;
+				}
+
+				{
+					DrawContext.DrawText(X, Y, TEXT("Export Type:"), MainFont, TextColor);
+					FString ExportTypeStr = FString::Printf(TEXT("%s%s"), GetName(HoveredTimingEvent.LoadingInfo.ExportEventType), Export && Export->IsAsset ? TEXT(" [asset]") : TEXT(""));
+					DrawContext.DrawText(ValueX, Y, ExportTypeStr, MainFont, ValueColor);
+					Y += LineH;
+				}
+
+				if (Export)
+				{
+					DrawContext.DrawText(X, Y, TEXT("Export Class:"), MainFont, TextColor);
+					FString ClassStr = FString::Printf(TEXT("%s"), Export->Class ? Export->Class->Name : TEXT("N/A"));
+					DrawContext.DrawText(ValueX, Y, ClassStr, MainFont, ValueColor);
+					Y += LineH;
+
+					DrawContext.DrawText(X, Y, TEXT("Serial:"), MainFont, TextColor);
+					FString SerialStr = FString::Printf(TEXT("Offset: %llu, Size: %llu%s"), Export->SerialOffset, Export->SerialSize);
+					DrawContext.DrawText(ValueX, Y, SerialStr, MainFont, ValueColor);
+					Y += LineH;
+				}
+
+				DrawContext.LayerId++;
+			}
+		}
+		else
+		{
+			TooltipWidth = MinTooltipWidth;
+			TooltipAlpha = 0.0f;
+		}
+	}
+
+	if (bAssetLoadingMode)
+	{
+		const FLinearColor BackgroundColor(0.01f, 0.01f, 0.01f, 0.9f);
+		const FLinearColor TextColor(1.0f, 1.0f, 1.0f, 0.9f);
+
+		constexpr float MarginX = 8.0f;
+		constexpr float MarginY = 8.0f;
+
+		constexpr float BorderX = 4.0f;
+		constexpr float BorderY = 4.0f;
+		float Y = ViewHeight;
+
+		if (ObjectTypeAggregationStr.Len() > 0)
+		{
+			FVector2D TextSize = FontMeasureService->Measure(ObjectTypeAggregationStr, MainFont);
+			const float X = ViewWidth - MarginX - TextSize.X - 2 * BorderX;
+			Y -= MarginY + TextSize.Y + 2 * BorderY;
+			DrawContext.DrawBox(X, Y, TextSize.X + 2 * BorderX, TextSize.Y + 2 * BorderY, WhiteBrush, BackgroundColor);
+			DrawContext.LayerId++;
+			DrawContext.DrawText(X + BorderX, Y + BorderY, ObjectTypeAggregationStr, MainFont, TextColor);
+			DrawContext.LayerId++;
+		}
+
+		if (EventAggregationStr.Len() > 0)
+		{
+			FVector2D TextSize = FontMeasureService->Measure(EventAggregationStr, MainFont);
+			const float X = ViewWidth - MarginX - TextSize.X - 2 * BorderX;
+			Y -= MarginY + TextSize.Y + 2 * BorderY;
+			DrawContext.DrawBox(X, Y, TextSize.X + 2 * BorderX, TextSize.Y + 2 * BorderY, WhiteBrush, BackgroundColor);
+			DrawContext.LayerId++;
+			DrawContext.DrawText(X + BorderX, Y + BorderY, EventAggregationStr, MainFont, TextColor);
+			DrawContext.LayerId++;
+		}
+	}
+
+	return SCompoundWidget::OnPaint(Args, AllottedGeometry, MyCullingRect, OutDrawElements, LayerId, InWidgetStyle, bParentEnabled && IsEnabled());
+}
+
+////////////////////////////////////////////////////////////////////////////////////////////////////
+
+const TCHAR* GetName(ETimingEventsTrackType Type)
+{
+	switch (Type)
+	{
+		case ETimingEventsTrackType::Gpu:		return TEXT("GPU");
+		case ETimingEventsTrackType::Cpu:		return TEXT("CPU");
+		case ETimingEventsTrackType::Loading:	return TEXT("Loading");
+		case ETimingEventsTrackType::Io:		return TEXT("I/O");
+		default:								return TEXT("unknown");
+	}
+}
+
+////////////////////////////////////////////////////////////////////////////////////////////////////
+
+FTimingEventsTrack* STimingView::AddTimingEventsTrack(uint64 TrackId, ETimingEventsTrackType TrackType, const FString& TrackName, const TCHAR* GroupName, int32 Order)
+{
+	FTimingEventsTrack* Track = new FTimingEventsTrack(TrackId, TrackType, TrackName, GroupName);
+
+	Track->SetOrder(Order);
+
+	UE_LOG(TimingProfiler, Log, TEXT("New Timing Events Track (%d) : %s (\"%s\")"), TimingEventsTracks.Num() + 1, GetName(TrackType), *TrackName);
+
+	CachedTimelines.Add(TrackId, Track);
+	TimingEventsTracks.Add(Track);
+
+	return Track;
+}
+
+////////////////////////////////////////////////////////////////////////////////////////////////////
+
+bool STimingView::IsGpuTrackVisible() const
+{
+	return GpuTrack != nullptr && GpuTrack->IsVisible();
+}
+
+////////////////////////////////////////////////////////////////////////////////////////////////////
+
+bool STimingView::IsCpuTrackVisible(uint32 ThreadId) const
+{
+	return CpuTracks.Contains(ThreadId) && CpuTracks[ThreadId]->IsVisible();
+}
+
+////////////////////////////////////////////////////////////////////////////////////////////////////
+
+void STimingView::DrawTimingProfilerTrack(FTimingViewDrawHelper& Helper, FTimingEventsTrack& Track) const
+{
+	if (Helper.BeginTimeline(Track))
+	{
+		TSharedPtr<const Trace::IAnalysisSession> Session = FInsightsManager::Get()->GetSession();
+		if (Session.IsValid() && Trace::ReadTimingProfilerProvider(*Session.Get()))
+		{
+			Trace::FAnalysisSessionReadScope SessionReadScope(*Session.Get());
+
+			const Trace::ITimingProfilerProvider& TimingProfilerProvider = *Trace::ReadTimingProfilerProvider(*Session.Get());
+
+			TimingProfilerProvider.ReadTimers([this, &Helper, &Track, &TimingProfilerProvider](const Trace::FTimingProfilerTimer* Timers, uint64 TimersCount)
+			{
+				TimingProfilerProvider.ReadTimeline(Track.GetId(), [this, &Helper, &Track, Timers](const Trace::ITimingProfilerProvider::Timeline& Timeline)
+				{
+					if (bUseDownSampling)
+					{
+						const double SecondsPerPixel = 1.0 / Helper.GetViewport().ScaleX;
+						Timeline.EnumerateEventsDownSampled(Helper.GetViewport().StartTime, Helper.GetViewport().EndTime, SecondsPerPixel, [this, &Helper, Timers](double StartTime, double EndTime, uint32 Depth, const Trace::FTimingProfilerEvent& Event)
+						{
+							Helper.AddEvent(StartTime, EndTime, Depth, Timers[Event.TimerIndex].Name);
+						});
+					}
+					else
+					{
+						Timeline.EnumerateEvents(Helper.GetViewport().StartTime, Helper.GetViewport().EndTime, [this, &Helper, Timers](double StartTime, double EndTime, uint32 Depth, const Trace::FTimingProfilerEvent& Event)
+						{
+							Helper.AddEvent(StartTime, EndTime, Depth, Timers[Event.TimerIndex].Name);
+						});
+					}
+				});
+			});
+
+			Helper.EndTimeline(Track);
+		}
+	}
+}
+
+////////////////////////////////////////////////////////////////////////////////////////////////////
+
+const TCHAR* GetName(ELoadTimeProfilerPackageEventType Type)
+{
+	switch (Type)
+	{
+		case LoadTimeProfilerPackageEventType_CreateLinker:				return TEXT("CreateLinker");
+		case LoadTimeProfilerPackageEventType_FinishLinker:				return TEXT("FinishLinker");
+		case LoadTimeProfilerPackageEventType_StartImportPackages:		return TEXT("StartImportPackages");
+		case LoadTimeProfilerPackageEventType_SetupImports:				return TEXT("SetupImports");
+		case LoadTimeProfilerPackageEventType_SetupExports:				return TEXT("SetupExports");
+		case LoadTimeProfilerPackageEventType_ProcessImportsAndExports:	return TEXT("ProcessImportsAndExports");
+		case LoadTimeProfilerPackageEventType_ExportsDone:				return TEXT("ExportsDone");
+		case LoadTimeProfilerPackageEventType_PostLoadWait:				return TEXT("PostLoadWait");
+		case LoadTimeProfilerPackageEventType_StartPostLoad:			return TEXT("StartPostLoad");
+		case LoadTimeProfilerPackageEventType_Tick:						return TEXT("Tick");
+		case LoadTimeProfilerPackageEventType_Finish:					return TEXT("Finish");
+		case LoadTimeProfilerPackageEventType_DeferredPostLoad:			return TEXT("DeferredPostLoad");
+		case LoadTimeProfilerPackageEventType_None:						return TEXT("None");
+		default:														return TEXT("");
+	}
+}
+
+const TCHAR* GetName(ELoadTimeProfilerObjectEventType Type)
+{
+	switch (Type)
+	{
+		case LoadTimeProfilerObjectEventType_Create:	return TEXT("Create");
+		case LoadTimeProfilerObjectEventType_Serialize:	return TEXT("Serialize");
+		case LoadTimeProfilerObjectEventType_PostLoad:	return TEXT("PostLoad");
+		case LoadTimeProfilerObjectEventType_None:		return TEXT("None");
+		default:										return TEXT("");
+	};
+}
+
+const TCHAR* STimingView::GetLoadTimeProfilerEventNameByPackageEventType(uint32 Depth, const Trace::FLoadTimeProfilerCpuEvent& Event) const
+{
+	return GetName(Event.PackageEventType);
+}
+
+const TCHAR* STimingView::GetLoadTimeProfilerEventNameByExportEventType(uint32 Depth, const Trace::FLoadTimeProfilerCpuEvent& Event) const
+{
+	return GetName(Event.ExportEventType);
+}
+
+const TCHAR* STimingView::GetLoadTimeProfilerEventNameByPackageName(uint32 Depth, const Trace::FLoadTimeProfilerCpuEvent& Event) const
+{
+	return Event.Package ?  Event.Package->Name : TEXT("");
+}
+
+const TCHAR* STimingView::GetLoadTimeProfilerEventNameByExportClassName(uint32 Depth, const Trace::FLoadTimeProfilerCpuEvent& Event) const
+{
+	return Event.Export && Event.Export->Class ? Event.Export->Class->Name : TEXT("");
+}
+
+const TCHAR* STimingView::GetLoadTimeProfilerEventName(uint32 Depth, const Trace::FLoadTimeProfilerCpuEvent& Event) const
+{
+	if (Depth == 0)
+	{
+		if (Event.Package)
+		{
+			return Event.Package->Name;
+		}
+	}
+
+	if (Event.Export && Event.Export->Class)
+	{
+		return Event.Export->Class->Name;
+	}
+
+	return TEXT("");
+}
+
+////////////////////////////////////////////////////////////////////////////////////////////////////
+
+void STimingView::DrawLoadTimeProfilerTrack(FTimingViewDrawHelper& Helper, FTimingEventsTrack& Track) const
+{
+	if (Helper.BeginTimeline(Track))
+	{
+		TSharedPtr<const Trace::IAnalysisSession> Session = FInsightsManager::Get()->GetSession();
+		if (Session.IsValid() && Trace::ReadLoadTimeProfilerProvider(*Session.Get()))
+		{
+			Trace::FAnalysisSessionReadScope SessionReadScope(*Session.Get());
+
+			const Trace::ILoadTimeProfilerProvider& LoadTimeProfilerProvider = *Trace::ReadLoadTimeProfilerProvider(*Session.Get());
+
+			if (&Track == LoadingMainThreadTrack)
+			{
+				LoadTimeProfilerProvider.ReadMainThreadCpuTimeline([this, &Helper, &Track](const Trace::ILoadTimeProfilerProvider::CpuTimeline& Timeline)
+				{
+					if (bUseDownSampling)
+					{
+						const double SecondsPerPixel = 1.0 / Helper.GetViewport().ScaleX;
+						Timeline.EnumerateEventsDownSampled(Helper.GetViewport().StartTime, Helper.GetViewport().EndTime, SecondsPerPixel, [this, &Helper](double StartTime, double EndTime, uint32 Depth, const Trace::FLoadTimeProfilerCpuEvent& Event)
+						{
+							const TCHAR* Name = LoadingGetEventNameFn.Execute(Depth, Event);
+							Helper.AddEvent(StartTime, EndTime, Depth, Name);
+						});
+					}
+					else
+					{
+						Timeline.EnumerateEvents(Helper.GetViewport().StartTime, Helper.GetViewport().EndTime, [this, &Helper](double StartTime, double EndTime, uint32 Depth, const Trace::FLoadTimeProfilerCpuEvent& Event)
+						{
+							const TCHAR* Name = LoadingGetEventNameFn.Execute(Depth, Event);
+							Helper.AddEvent(StartTime, EndTime, Depth, Name);
+						});
+					}
+				});
+			}
+			else if (&Track == LoadingAsyncThreadTrack)
+			{
+				LoadTimeProfilerProvider.ReadAsyncLoadingThreadCpuTimeline([this, &Helper, &Track](const Trace::ILoadTimeProfilerProvider::CpuTimeline& Timeline)
+				{
+					if (bUseDownSampling)
+					{
+						const double SecondsPerPixel = 1.0 / Helper.GetViewport().ScaleX;
+						Timeline.EnumerateEventsDownSampled(Helper.GetViewport().StartTime, Helper.GetViewport().EndTime, SecondsPerPixel, [this, &Helper](double StartTime, double EndTime, uint32 Depth, const Trace::FLoadTimeProfilerCpuEvent& Event)
+						{
+							const TCHAR* Name = LoadingGetEventNameFn.Execute(Depth, Event);
+							Helper.AddEvent(StartTime, EndTime, Depth, Name);
+						});
+					}
+					else
+					{
+						Timeline.EnumerateEvents(Helper.GetViewport().StartTime, Helper.GetViewport().EndTime, [this, &Helper](double StartTime, double EndTime, uint32 Depth, const Trace::FLoadTimeProfilerCpuEvent& Event)
+						{
+							const TCHAR* Name = LoadingGetEventNameFn.Execute(Depth, Event);
+							Helper.AddEvent(StartTime, EndTime, Depth, Name);
+						});
+					}
+				});
+			}
+		}
+
+		Helper.EndTimeline(Track);
+	}
+}
+
+////////////////////////////////////////////////////////////////////////////////////////////////////
+// "I/O - File Activity" prototype
+////////////////////////////////////////////////////////////////////////////////////////////////////
+
+// The I/O timelines are just prototypes for now.
+// Below code will be removed once the functionality is moved in analyzer.
+
+////////////////////////////////////////////////////////////////////////////////////////////////////
+
+const TCHAR* GetFileActivityTypeName(Trace::EFileActivityType Type)
+{
+	static_assert(Trace::FileActivityType_Open == 0, "Trace::EFileActivityType enum has changed!?");
+	static_assert(Trace::FileActivityType_Close == 1, "Trace::EFileActivityType enum has changed!?");
+	static_assert(Trace::FileActivityType_Read == 2, "Trace::EFileActivityType enum has changed!?");
+	static_assert(Trace::FileActivityType_Write == 3, "Trace::EFileActivityType enum has changed!?");
+	static_assert(Trace::FileActivityType_Count == 4, "Trace::EFileActivityType enum has changed!?");
+	static const TCHAR* GFileActivityTypeNames[] =
+	{
+		TEXT("Open"),
+		TEXT("Close"),
+		TEXT("Read"),
+		TEXT("Write"),
+		TEXT("Idle"), // virtual events added for cases where Close event is more than 1s away from last Open/Read/Write event.
+		TEXT("NotClosed") // virtual events added when an Open activity never closes
+	};
+	return GFileActivityTypeNames[Type];
+}
+
+////////////////////////////////////////////////////////////////////////////////////////////////////
+
+uint32 GetFileActivityTypeColor(Trace::EFileActivityType Type)
+{
+	static const uint32 GFileActivityTypeColors[] =
+	{
+		0xFFCCAA33, // open
+		0xFF33AACC, // close
+		0xFF33AA33, // read
+		0xFFDD33CC, // write
+		0x55333333, // idle
+		0x55553333, // close
+	};
+	return GFileActivityTypeColors[Type];
+}
+
+////////////////////////////////////////////////////////////////////////////////////////////////////
+
+void STimingView::DrawIoOverviewTrack(FTimingViewDrawHelper& Helper, FTimingEventsTrack& Track) const
+{
+	if (Helper.BeginTimeline(Track))
+	{
+		// Draw the IO Overiew track using cached events (as those are sorted by Start Time).
+		for (const FIoTimingEvent& Event : AllIoEvents)
+		{
+			const Trace::EFileActivityType ActivityType = static_cast<Trace::EFileActivityType>(Event.Type & 0x0F);
+
+			if (ActivityType >= Trace::FileActivityType_Count)
+			{
+				// Ignore "Idle" and "NotClosed" events.
+				continue;
+			}
+			if (Event.EndTime <= Helper.GetViewport().StartTime)
+			{
+				continue;
+			}
+			if (Event.StartTime >= Helper.GetViewport().EndTime)
+			{
+				break;
+			}
+
+			uint32 Color = GetFileActivityTypeColor(ActivityType);
+
+			const bool bHasFailed = ((Event.Type & 0xF0) != 0);
+
+			if (bHasFailed)
+			{
+				FString Name = TEXT("Failed ");
+				Name += GetFileActivityTypeName(ActivityType);
+				Name += " [";
+				Name += Event.Path;
+				Name += "]";
+				Color = 0xFFAA0000;
+				Helper.AddEvent(Event.StartTime, Event.EndTime, 0, *Name, Color);
+			}
+			else if (ActivityType == Trace::FileActivityType_Open)
+			{
+				FString Name = GetFileActivityTypeName(ActivityType);
+				Name += " [";
+				Name += Event.Path;
+				Name += "]";
+				Helper.AddEvent(Event.StartTime, Event.EndTime, 0, *Name, Color);
+			}
+			else
+			{
+				Helper.AddEvent(Event.StartTime, Event.EndTime, 0, GetFileActivityTypeName(ActivityType), Color);
+			}
+		}
+
+		Helper.EndTimeline(Track);
+	}
+}
+
+////////////////////////////////////////////////////////////////////////////////////////////////////
+
+void STimingView::UpdateIo()
+{
+	if (bForceIoEventsUpdate)
+	{
+		bForceIoEventsUpdate = false;
+
+		AllIoEvents.Reset();
+
+		TSharedPtr<const Trace::IAnalysisSession> Session = FInsightsManager::Get()->GetSession();
+		if (Session.IsValid() && Trace::ReadFileActivityProvider(*Session.Get()))
+		{
+			Trace::FAnalysisSessionReadScope SessionReadScope(*Session.Get());
+
+			const Trace::IFileActivityProvider& FileActivityProvider = *Trace::ReadFileActivityProvider(*Session.Get());
+
+			// Enumerate all IO events and cache them.
+			FileActivityProvider.EnumerateFileActivity([this](const Trace::FFileInfo& FileInfo, const Trace::IFileActivityProvider::Timeline& Timeline)
+			{
+				Timeline.EnumerateEvents(-std::numeric_limits<double>::infinity(), +std::numeric_limits<double>::infinity(),
+					[this, &FileInfo, &Timeline](double EventStartTime, double EventEndTime, uint32 EventDepth, const Trace::FFileActivity& FileActivity)
+				{
+					if (bMergeIoLanes)
+					{
+						EventDepth = static_cast<uint32>(FileInfo.Id); // used by MergeLanes algorithm
+					}
+					else
+					{
+						EventDepth = FileInfo.Id % 32; // simple layout
+					}
+					uint32 Type = ((uint32)FileActivity.ActivityType & 0x0F) | (FileActivity.Failed ? 0x80 : 0);
+					AllIoEvents.Add(FIoTimingEvent{ EventStartTime, EventEndTime, EventDepth, Type, FileInfo.Path });
+				});
+
+				return true;
+			});
+		}
+
+		// Sort cached IO events by Start Time.
+		AllIoEvents.Sort([](const FIoTimingEvent& A, const FIoTimingEvent& B) { return A.StartTime < B.StartTime; });
+
+		if (bMergeIoLanes)
+		{
+			struct FIoLane
+			{
+				uint32 FileId;
+				const TCHAR* Path;
+				double LastEndTime;
+				double EndTime;
+			};
+			TArray<FIoLane> Lanes;
+
+			TArray<FIoTimingEvent> IoEventsToAdd;
+
+			for (FIoTimingEvent& Event : AllIoEvents)
+			{
+				uint64 TimelineId = Event.Depth;
+
+				int32 Depth = -1;
+
+				bool bIsCloseEvent = false;
+
+				const Trace::EFileActivityType ActivityType = static_cast<Trace::EFileActivityType>(Event.Type & 0x0F);
+
+				if (ActivityType == Trace::FileActivityType_Open)
+				{
+					// Find lane (avoiding overlaps with other opened files).
+					for (int32 LaneIndex = 0; LaneIndex < Lanes.Num(); ++LaneIndex)
+					{
+						FIoLane& Lane = Lanes[LaneIndex];
+						if (Event.StartTime >= Lane.EndTime)
+						{
+							Depth = LaneIndex;
+							break;
+						}
+					}
+
+					if (Depth < 0)
+					{
+						// Add new lane.
+						Depth = Lanes.AddDefaulted();
+					}
+
+					const bool bHasFailed = ((Event.Type & 0xF0) != 0);
+
+					FIoLane& Lane = Lanes[Depth];
+					Lane.FileId = TimelineId;
+					Lane.Path = Event.Path;
+					Lane.LastEndTime = Event.EndTime;
+					Lane.EndTime = bHasFailed ? Event.EndTime : std::numeric_limits<double>::infinity();
+				}
+				else if (ActivityType == Trace::FileActivityType_Close)
+				{
+					bIsCloseEvent = true;
+
+					// Find lane with same id.
+					for (int32 LaneIndex = 0; LaneIndex < Lanes.Num(); ++LaneIndex)
+					{
+						FIoLane& Lane = Lanes[LaneIndex];
+						if (Lane.FileId == TimelineId)
+						{
+							// Adds an Idle event, but only if has passed at least 1s since last open/read/write activity.
+							if (Event.StartTime - Lane.LastEndTime > 1.0)
+							{
+								IoEventsToAdd.Add(FIoTimingEvent{ Lane.LastEndTime, Event.StartTime, static_cast<uint32>(LaneIndex), Trace::FileActivityType_Count, Event.Path });
+							}
+							Lane.LastEndTime = Event.EndTime;
+							Lane.EndTime = Event.EndTime;
+							Depth = LaneIndex;
+							break;
+						}
+					}
+					ensure(Depth >= 0);
+				}
+				else
+				{
+					// All other events should be inside the virtual Open-Close event.
+					// Find lane with same id.
+					for (int32 LaneIndex = 0; LaneIndex < Lanes.Num(); ++LaneIndex)
+					{
+						FIoLane& Lane = Lanes[LaneIndex];
+						if (Lane.FileId == TimelineId)
+						{
+							Lane.LastEndTime = Event.EndTime;
+							Depth = LaneIndex;
+							break;
+						}
+					}
+					ensure(Depth >= 0);
+				}
+
+				if (Depth < 0) // just in case
+				{
+					// Add new lane.
+					Depth = Lanes.AddDefaulted();
+					FIoLane& Lane = Lanes[Depth];
+					Lane.FileId = TimelineId;
+					Lane.Path = Event.Path;
+					Lane.LastEndTime = Event.EndTime;
+					Lane.EndTime = bIsCloseEvent ? Event.EndTime : std::numeric_limits<double>::infinity();
+				}
+
+				Event.Depth = Depth;
+			}
+
+			for (int32 LaneIndex = 0; LaneIndex < Lanes.Num(); ++LaneIndex)
+			{
+				FIoLane& Lane = Lanes[LaneIndex];
+				if (Lane.EndTime == std::numeric_limits<double>::infinity())
+				{
+					IoEventsToAdd.Add(FIoTimingEvent{ Lane.LastEndTime, Lane.EndTime, static_cast<uint32>(LaneIndex), Trace::FileActivityType_Count + 1, Lane.Path });
+				}
+			}
+
+			for (const FIoTimingEvent& Event : IoEventsToAdd)
+			{
+				AllIoEvents.Add(Event);
+			}
+
+			// Sort cached IO events one more time, also by Start Time.
+			AllIoEvents.Sort([](const FIoTimingEvent& A, const FIoTimingEvent& B) { return A.StartTime < B.StartTime; });
+
+			//// Sort cached IO events again, by Depth and then by Start Time.
+			//AllIoEvents.Sort([](const FIoTimingEvent& A, const FIoTimingEvent& B)
+			//{
+			//	return A.Depth == B.Depth ? A.StartTime < B.StartTime : A.Depth < B.Depth;
+			//});
+		}
+	}
+}
+
+////////////////////////////////////////////////////////////////////////////////////////////////////
+
+void STimingView::DrawIoActivityTrack(FTimingViewDrawHelper& Helper, FTimingEventsTrack& Track) const
+{
+	if (Helper.BeginTimeline(Track))
+	{
+		// Draw IO track using cached events.
+		for (const FIoTimingEvent& Event : AllIoEvents)
+		{
+			if (Event.EndTime <= Helper.GetViewport().StartTime)
+			{
+				continue;
+			}
+			if (Event.StartTime >= Helper.GetViewport().EndTime)
+			{
+				break;
+			}
+
+			const Trace::EFileActivityType ActivityType = static_cast<Trace::EFileActivityType>(Event.Type & 0x0F);
+
+			uint32 Color = GetFileActivityTypeColor(ActivityType);
+
+			if (ActivityType < Trace::FileActivityType_Count)
+			{
+				const bool bHasFailed = ((Event.Type & 0xF0) != 0);
+
+				if (bHasFailed)
+				{
+					FString Name = TEXT("Failed ");
+					Name += GetFileActivityTypeName(ActivityType);
+					Name += " [";
+					Name += Event.Path;
+					Name += "]";
+					Color = 0xFFAA0000;
+					Helper.AddEvent(Event.StartTime, Event.EndTime, Event.Depth, *Name, Color);
+				}
+				else if (ActivityType == Trace::FileActivityType_Open)
+				{
+					FString Name = GetFileActivityTypeName(ActivityType);
+					Name += " [";
+					Name += Event.Path;
+					Name += "]";
+					Helper.AddEvent(Event.StartTime, Event.EndTime, Event.Depth, *Name, Color);
+				}
+				else
+				{
+					Helper.AddEvent(Event.StartTime, Event.EndTime, Event.Depth, GetFileActivityTypeName(ActivityType), Color);
+				}
+			}
+			else
+			{
+				FString Name = GetFileActivityTypeName(ActivityType);
+				Name += " [";
+				Name += Event.Path;
+				Name += "]";
+				Helper.AddEvent(Event.StartTime, Event.EndTime, Event.Depth, *Name, Color);
+			}
+		}
+
+		Helper.EndTimeline(Track);
+	}
+}
+
+////////////////////////////////////////////////////////////////////////////////////////////////////
+// end of "I/O - File Activity" prototype
+////////////////////////////////////////////////////////////////////////////////////////////////////
+
+void STimingView::DrawTimeRangeSelection(FDrawContext& DrawContext) const
+{
+	if (SelectionEndTime > SelectionStartTime)
+	{
+		float SelectionX1 = Viewport.TimeToSlateUnitsRounded(SelectionStartTime);
+		float SelectionX2 = Viewport.TimeToSlateUnitsRounded(SelectionEndTime);
+
+		if (SelectionX1 <= Viewport.Width &&
+			SelectionX2 >= 0)
+		{
+			float ClipLeft = 0.0f;
+			float ClipRight = 0.0f;
+			if (SelectionX1 < 0.0f)
+			{
+				ClipLeft = -SelectionX1;
+				SelectionX1 = 0.0f;
+			}
+			if (SelectionX2 > Viewport.Width)
+			{
+				ClipRight = SelectionX2 - Viewport.Width;
+				SelectionX2 = Viewport.Width;
+			}
+
+			// Draw selection area.
+			DrawContext.DrawBox(SelectionX1, 0.0f, SelectionX2 - SelectionX1, Viewport.Height, WhiteBrush, FLinearColor(0.25f, 0.5f, 1.0f, 0.25f));
+			DrawContext.LayerId++;
+
+			FColor ArrowFillColor(32, 64, 128, 255);
+			FLinearColor ArrowColor(ArrowFillColor);
+
+			if (SelectionX1 > 0.0f)
+			{
+				// Draw left side (vertical line).
+				DrawContext.DrawBox(SelectionX1 - 1.0f, 0.0f, 1.0f, Viewport.Height, WhiteBrush, ArrowColor);
+			}
+
+			if (SelectionX2 < Viewport.Width)
+			{
+				// Draw right side (vertical line).
+				DrawContext.DrawBox(SelectionX2, 0.0f, 1.0f, Viewport.Height, WhiteBrush, ArrowColor);
+			}
+
+			DrawContext.LayerId++;
+
+			const float ArrowSize = 6.0f;
+			const float ArrowY = 6.0f;
+
+			if (SelectionX2 - SelectionX1 > 2 * ArrowSize)
+			{
+				// Draw horizontal line.
+				float HorizLineX1 = SelectionX1;
+				if (ClipLeft == 0.0f)
+				{
+					HorizLineX1 += 1.0f;
+				}
+				float HorizLineX2 = SelectionX2;
+				if (ClipRight == 0.0f)
+				{
+					HorizLineX2 -= 1.0f;
+				}
+				DrawContext.DrawBox(HorizLineX1, ArrowY - 1.0f, HorizLineX2 - HorizLineX1, 3.0f, WhiteBrush, ArrowColor);
+
+				if (ClipLeft < ArrowSize)
+				{
+					// Draw left arrow.
+					for (float AH = 0.0f; AH < ArrowSize; AH += 1.0f)
+					{
+						DrawContext.DrawBox(SelectionX1 - ClipLeft + AH, ArrowY - AH, 1.0f, 2.0f * AH + 1.0f, WhiteBrush, ArrowColor);
+					}
+				}
+
+				if (ClipRight < ArrowSize)
+				{
+					// Draw right arrow.
+					for (float AH = 0.0f; AH < ArrowSize; AH += 1.0f)
+					{
+						DrawContext.DrawBox(SelectionX2 + ClipRight - AH - 1.0f, ArrowY - AH, 1.0f, 2.0f * AH + 1.0f, WhiteBrush, ArrowColor);
+					}
+				}
+
+				DrawContext.LayerId++;
+
+#if 0
+				//im: This should be a more efficeint way top draw the arrows, but it renders them with artifacts (missing vertical lines; shader bug?)!
+
+				const FSlateBrush* MyBrush = WhiteBrush;
+				FSlateShaderResourceProxy* ResourceProxy = FSlateDataPayload::ResourceManager->GetShaderResource(*MyBrush);
+				FSlateResourceHandle ResourceHandle = FSlateApplication::Get().GetRenderer()->GetResourceHandle(*MyBrush);
+
+				FVector2D AtlasOffset = ResourceProxy ? ResourceProxy->StartUV : FVector2D(0.0f, 0.0f);
+				FVector2D AtlasUVSize = ResourceProxy ? ResourceProxy->SizeUV : FVector2D(1.0f, 1.0f);
+
+				const FVector2D Pos = AllottedGeometry.GetAbsolutePosition() + FVector2D(0.0f, 40.0f);
+				const float Scale = AllottedGeometry.Scale;
+
+				FSlateRenderTransform RenderTransform;
+
+				TArray<FSlateVertex> Verts;
+				Verts.Reserve(6);
+				Verts.Add(FSlateVertex::Make<ESlateVertexRounding::Disabled>(RenderTransform, Pos + FVector2D(0.5f + SelectionX1 + ArrowSize, 0.5f + ArrowY + ArrowSize) * Scale, AtlasOffset + FVector2D(0.0f, 1.0f) * AtlasUVSize, ArrowFillColor));
+				Verts.Add(FSlateVertex::Make<ESlateVertexRounding::Disabled>(RenderTransform, Pos + FVector2D(0.5f + SelectionX1,             0.5f + ArrowY            ) * Scale, AtlasOffset + FVector2D(1.0f, 0.5f) * AtlasUVSize, ArrowFillColor));
+				Verts.Add(FSlateVertex::Make<ESlateVertexRounding::Disabled>(RenderTransform, Pos + FVector2D(0.5f + SelectionX1 + ArrowSize, 0.5f + ArrowY - ArrowSize) * Scale, AtlasOffset + FVector2D(0.0f, 0.0f) * AtlasUVSize, ArrowFillColor));
+				Verts.Add(FSlateVertex::Make<ESlateVertexRounding::Disabled>(RenderTransform, Pos + FVector2D(0.5f + SelectionX2 - ArrowSize, 0.5f + ArrowY - ArrowSize) * Scale, AtlasOffset + FVector2D(0.0f, 0.0f) * AtlasUVSize, ArrowFillColor));
+				Verts.Add(FSlateVertex::Make<ESlateVertexRounding::Disabled>(RenderTransform, Pos + FVector2D(0.5f + SelectionX2,             0.5f + ArrowY            ) * Scale, AtlasOffset + FVector2D(1.0f, 0.5f) * AtlasUVSize, ArrowFillColor));
+				Verts.Add(FSlateVertex::Make<ESlateVertexRounding::Disabled>(RenderTransform, Pos + FVector2D(0.5f + SelectionX2 - ArrowSize, 0.5f + ArrowY + ArrowSize) * Scale, AtlasOffset + FVector2D(0.0f, 1.0f) * AtlasUVSize, ArrowFillColor));
+
+				TArray<SlateIndex> Indices;
+				Indices.Reserve(6);
+				if (ClipLeft < ArrowSize)
+				{
+					Indices.Add(0);
+					Indices.Add(1);
+					Indices.Add(2);
+				}
+				if (ClipRight < ArrowSize)
+				{
+					Indices.Add(3);
+					Indices.Add(4);
+					Indices.Add(5);
+				}
+
+				FSlateDrawElement::MakeCustomVerts(
+					OutDrawElements,
+					LayerId,
+					ResourceHandle,
+					Verts,
+					Indices,
+					nullptr,
+					0,
+					0,
+					ESlateDrawEffect::PreMultipliedAlpha);
+
+				DrawContext.LayerId++;
+#endif
+			}
+
+			//////////////////////////////////////////////////
+			// Draw duration for selected time interval.
+
+			double Duration = SelectionEndTime - SelectionStartTime;
+			FString Text = TimeUtils::FormatTimeAuto(Duration);
+
+			const TSharedRef<FSlateFontMeasure> FontMeasureService = FSlateApplication::Get().GetRenderer()->GetFontMeasureService();
+			const float TextWidth = FontMeasureService->Measure(Text, MainFont).X;
+
+			const float CenterX = (SelectionX1 + SelectionX2) / 2.0f;
+
+			DrawContext.DrawBox(CenterX - TextWidth / 2 - 2.0, ArrowY - 6.0f, TextWidth + 4.0f, 13.0f, WhiteBrush, ArrowColor);
+			DrawContext.LayerId++;
+
+			DrawContext.DrawText(CenterX - TextWidth / 2, ArrowY - 6.0f, Text, MainFont, FLinearColor::White);
+			DrawContext.LayerId++;
+		}
+	}
+}
+
+////////////////////////////////////////////////////////////////////////////////////////////////////
+
+FReply STimingView::OnMouseButtonDown(const FGeometry& MyGeometry, const FPointerEvent& MouseEvent)
+{
+	FReply Reply = FReply::Unhandled();
+	MousePositionOnButtonDown = MyGeometry.AbsoluteToLocal(MouseEvent.GetScreenSpacePosition());
+
+	bool bStartPanning = false;
+
+	if (MouseEvent.GetEffectingButton() == EKeys::LeftMouseButton)
+	{
+		if (MarkersTrack.IsVisible() && MarkersTrack.IsHeaderHovered())
+		{
+			MarkersTrack.ToggleCollapsed();
+		}
+
+		if (!bIsRMB_Pressed)
+		{
+			bIsLMB_Pressed = true;
+
+			if (bIsSpaceBarKeyPressed || MousePositionOnButtonDown.Y > TimeRulerTrack.GetHeight())
+			{
+				bStartPanning = true;
+			}
+			else
+			{
+				bIsSelecting = true;
+				bIsDragging = false;
+
+				SelectionStartTime = Viewport.SlateUnitsToTime(MousePositionOnButtonDown.X);
+				SelectionEndTime = SelectionStartTime;
+				LastSelectionType = ESelectionType::None;
+				//TODO: SelectionChangingEvent.Broadcast(SelectionStartTime, SelectionEndTime);
+				EventAggregationStr.Reset();
+				ObjectTypeAggregationStr.Reset();
+			}
+
+			// Capture mouse, so we can drag outside this widget.
+			Reply = FReply::Handled().CaptureMouse(SharedThis(this));
+		}
+	}
+	else if (MouseEvent.GetEffectingButton() == EKeys::RightMouseButton)
+	{
+		if (!bIsLMB_Pressed)
+		{
+			bIsRMB_Pressed = true;
+
+			bStartPanning = true;
+
+			// Capture mouse, so we can drag outside this widget.
+			Reply = FReply::Handled().CaptureMouse(SharedThis(this));
+		}
+	}
+
+	if (bStartPanning)
+	{
+		bIsPanning = true;
+		bIsDragging = false;
+
+		ViewportStartTimeOnButtonDown = Viewport.StartTime;
+		ViewportScrollPosYOnButtonDown = Viewport.ScrollPosY;
+
+		if (MouseEvent.GetModifierKeys().IsControlDown())
+		{
+			// Allow panning only horizontally.
+			PanningMode = EPanningMode::Horizontal;
+		}
+		else if (MouseEvent.GetModifierKeys().IsShiftDown())
+		{
+			// Allow panning only vertically.
+			PanningMode = EPanningMode::Vertical;
+		}
+		else
+		{
+			// Allow panning both horizontally and vertically.
+			PanningMode = EPanningMode::HorizontalAndVertical;
+		}
+	}
+
+	return Reply;
+}
+
+////////////////////////////////////////////////////////////////////////////////////////////////////
+
+FReply STimingView::OnMouseButtonUp(const FGeometry& MyGeometry, const FPointerEvent& MouseEvent)
+{
+	FReply Reply = FReply::Unhandled();
+	MousePositionOnButtonUp = MyGeometry.AbsoluteToLocal(MouseEvent.GetScreenSpacePosition());
+
+	const bool bIsValidForMouseClick = MousePositionOnButtonUp.Equals(MousePositionOnButtonDown, 2.0f);
+
+	if (MouseEvent.GetEffectingButton() == EKeys::LeftMouseButton)
+	{
+		if (bIsLMB_Pressed)
+		{
+			if (bIsPanning)
+			{
+				PanningMode = EPanningMode::None;
+
+				bIsPanning = false;
+			}
+			else if (bIsSelecting)
+			{
+				//TODO: SelectionChangedEvent.Broadcast(SelectionStartTime, SelectionEndTime);
+				UpdateAggregatedStats();
+
+				bIsSelecting = false;
+			}
+
+			if (bIsValidForMouseClick)
+			{
+				// Select the hovered timing event (if any).
+				UpdateHoveredTimingEvent(MousePositionOnButtonUp.X, MousePositionOnButtonUp.Y);
+				SelectHoveredTimingEvent();
+
+				// When clicking on an empty space...
+				if (!SelectedTimingEvent.IsValid())
+				{
+					// ...reset selection.
+					SelectionEndTime = SelectionStartTime = 0.0;
+					LastSelectionType = ESelectionType::None;
+					//TODO: SelectionChangedEvent.Broadcast(SelectionStartTime, SelectionEndTime);
+					EventAggregationStr.Reset();
+					ObjectTypeAggregationStr.Reset();
+				}
+			}
+
+			bIsDragging = false;
+
+			// Release mouse as we no longer drag.
+			Reply = FReply::Handled().ReleaseMouseCapture();
+
+			bIsLMB_Pressed = false;
+		}
+	}
+	else if (MouseEvent.GetEffectingButton() == EKeys::RightMouseButton)
+	{
+		if (bIsRMB_Pressed)
+		{
+			if (bIsPanning)
+			{
+				PanningMode = EPanningMode::None;
+
+				bIsPanning = false;
+			}
+
+			if (!bIsDragging && !bIsSpaceBarKeyPressed && bIsValidForMouseClick)
+			{
+				ShowContextMenu(MouseEvent.GetScreenSpacePosition(), MouseEvent);
+				Reply = FReply::Handled();
+			}
+
+			bIsDragging = false;
+
+			// Release mouse as we no longer drag.
+			Reply = FReply::Handled().ReleaseMouseCapture();
+
+			bIsRMB_Pressed = false;
+		}
+	}
+
+	return Reply;
+}
+
+////////////////////////////////////////////////////////////////////////////////////////////////////
+
+FReply STimingView::OnMouseButtonDoubleClick(const FGeometry& MyGeometry, const FPointerEvent& MouseEvent)
+{
+	FReply Reply = FReply::Unhandled();
+
+	if (MouseEvent.GetEffectingButton() == EKeys::LeftMouseButton)
+	{
+		if (MarkersTrack.IsVisible() && MarkersTrack.IsHovered())
+		{
+			MarkersTrack.ToggleCollapsed();
+		}
+		else if (HoveredTimingEvent.IsValid())
+		{
+			double EndTime = Viewport.RestrictEndTime(HoveredTimingEvent.EndTime);
+			SelectTimeInterval(HoveredTimingEvent.StartTime, EndTime - HoveredTimingEvent.StartTime);
+		}
+
+		Reply = FReply::Handled();
+	}
+
+	return Reply;
+}
+
+////////////////////////////////////////////////////////////////////////////////////////////////////
+
+FReply STimingView::OnMouseMove(const FGeometry& MyGeometry, const FPointerEvent& MouseEvent)
+{
+	FReply Reply = FReply::Unhandled();
+
+	MousePosition = MyGeometry.AbsoluteToLocal(MouseEvent.GetScreenSpacePosition());
+
+	if (bIsPanning)
+	{
+		if (HasMouseCapture() && !MouseEvent.GetCursorDelta().IsZero())
+		{
+			bIsDragging = true;
+
+			if ((int32)PanningMode & (int32)EPanningMode::Horizontal)
+			{
+				ScrollAtTime(ViewportStartTimeOnButtonDown + static_cast<double>(MousePositionOnButtonDown.X - MousePosition.X) / Viewport.ScaleX);
+			}
+
+			if ((int32)PanningMode & (int32)EPanningMode::Vertical)
+			{
+				Viewport.ScrollPosY = ViewportScrollPosYOnButtonDown + (MousePositionOnButtonDown.Y - MousePosition.Y);
+				UpdateVerticalScrollBar();
+				bIsVerticalViewportDirty = true;
+			}
+		}
+
+		Reply = FReply::Handled();
+	}
+	else if (bIsSelecting)
+	{
+		if (HasMouseCapture() && !MouseEvent.GetCursorDelta().IsZero())
+		{
+			bIsDragging = true;
+
+			SelectionStartTime = Viewport.SlateUnitsToTime(MousePositionOnButtonDown.X);
+			SelectionEndTime = Viewport.SlateUnitsToTime(MousePosition.X);
+			if (SelectionStartTime > SelectionEndTime)
+			{
+				double Temp = SelectionStartTime;
+				SelectionStartTime = SelectionEndTime;
+				SelectionEndTime = Temp;
+			}
+			LastSelectionType = ESelectionType::TimeRange;
+			//TODO: SelectionChangingEvent.Broadcast(SelectionStartTime, SelectionEndTime);
+			EventAggregationStr.Reset();
+			ObjectTypeAggregationStr.Reset();
+		}
+
+		Reply = FReply::Handled();
+	}
+	else
+	{
+		if (MarkersTrack.IsVisible())
+		{
+			MarkersTrack.UpdateHoveredState(MousePosition.X, MousePosition.Y, Viewport);
+		}
+
+		UpdateHoveredTimingEvent(MousePosition.X, MousePosition.Y);
+	}
+
+	return Reply; // SAssetViewItem::CreateToolTipWidget
+}
+
+////////////////////////////////////////////////////////////////////////////////////////////////////
+
+void STimingView::OnMouseEnter(const FGeometry& MyGeometry, const FPointerEvent& MouseEvent)
+{
+
+}
+
+////////////////////////////////////////////////////////////////////////////////////////////////////
+
+void STimingView::OnMouseLeave(const FPointerEvent& MouseEvent)
+{
+	if (!HasMouseCapture())
+	{
+		// No longer dragging (unless we have mouse capture).
+		bIsDragging = false;
+		bIsPanning = false;
+		bIsSelecting = false;
+
+		bIsLMB_Pressed = false;
+		bIsRMB_Pressed = false;
+
+		MousePosition = FVector2D::ZeroVector;
+	}
+}
+
+////////////////////////////////////////////////////////////////////////////////////////////////////
+
+FReply STimingView::OnMouseWheel(const FGeometry& MyGeometry, const FPointerEvent& MouseEvent)
+{
+	if (MouseEvent.GetModifierKeys().IsShiftDown())
+	{
+		// Scroll vertically.
+		constexpr float ScrollSpeedY = 16.0f * 3;
+		Viewport.ScrollPosY -= ScrollSpeedY * MouseEvent.GetWheelDelta();
+		UpdateVerticalScrollBar();
+		bIsVerticalViewportDirty = true;
+	}
+	else if (MouseEvent.GetModifierKeys().IsControlDown())
+	{
+		// Scroll horizontally.
+		const double ScrollSpeedX = Viewport.GetDurationForViewportDX(16.0 * 3);
+		ScrollAtTime(Viewport.StartTime - ScrollSpeedX * MouseEvent.GetWheelDelta());
+	}
+	else
+	{
+		// Zoom in/out horizontally.
+		const double Delta = MouseEvent.GetWheelDelta();
+		constexpr double ZoomStep = 0.25; // as percent
+		double ScaleX;
+
+		if (Delta > 0)
+		{
+			ScaleX = Viewport.ScaleX * FMath::Pow(1.0 + ZoomStep, Delta);
+		}
+		else
+		{
+			ScaleX = Viewport.ScaleX * FMath::Pow(1.0 / (1.0 + ZoomStep), -Delta);
+		}
+
+		//UE_LOG(TimingProfiler, Log, TEXT("%.2f, %.2f, %.2f"), Delta, Viewport.ScaleX, ScaleX);
+		MousePosition = MyGeometry.AbsoluteToLocal(MouseEvent.GetScreenSpacePosition());
+
+		if (Viewport.ZoomWithFixedX(ScaleX, MousePosition.X))
+		{
+			//Viewport.EnforceHorizontalScrollLimits(1.0);
+			UpdateHorizontalScrollBar();
+			bIsViewportDirty = true;
+		}
+	}
+
+	return FReply::Handled();
+}
+
+////////////////////////////////////////////////////////////////////////////////////////////////////
+
+void STimingView::OnDragEnter(const FGeometry& MyGeometry, const FDragDropEvent& DragDropEvent)
+{
+	SCompoundWidget::OnDragEnter(MyGeometry, DragDropEvent);
+
+	//TSharedPtr<FStatIDDragDropOp> Operation = DragDropEvent.GetOperationAs<FStatIDDragDropOp>();
+	//if (Operation.IsValid())
+	//{
+	//	Operation->ShowOK();
+	//}
+}
+
+////////////////////////////////////////////////////////////////////////////////////////////////////
+
+void STimingView::OnDragLeave(const FDragDropEvent& DragDropEvent)
+{
+	SCompoundWidget::OnDragLeave(DragDropEvent);
+
+	//TSharedPtr<FStatIDDragDropOp> Operation = DragDropEvent.GetOperationAs<FStatIDDragDropOp>();
+	//if (Operation.IsValid())
+	//{
+	//	Operation->ShowError();
+	//}
+}
+
+////////////////////////////////////////////////////////////////////////////////////////////////////
+
+FReply STimingView::OnDragOver(const FGeometry& MyGeometry, const FDragDropEvent& DragDropEvent)
+{
+	return SCompoundWidget::OnDragOver(MyGeometry,DragDropEvent);
+}
+
+////////////////////////////////////////////////////////////////////////////////////////////////////
+
+FReply STimingView::OnDrop(const FGeometry& MyGeometry, const FDragDropEvent& DragDropEvent)
+{
+	//TSharedPtr<FStatIDDragDropOp> Operation = DragDropEvent.GetOperationAs<FStatIDDragDropOp>();
+	//if (Operation.IsValid())
+	//{
+	//	return FReply::Handled();
+	//}
+
+	return SCompoundWidget::OnDrop(MyGeometry,DragDropEvent);
+}
+
+////////////////////////////////////////////////////////////////////////////////////////////////////
+
+FCursorReply STimingView::OnCursorQuery(const FGeometry& MyGeometry, const FPointerEvent& CursorEvent) const
+{
+	if (bIsPanning)
+	{
+		if (bIsDragging)
+		{
+			//return FCursorReply::Cursor(EMouseCursor::GrabHandClosed);
+			return FCursorReply::Cursor(EMouseCursor::GrabHand);
+		}
+	}
+	else if (bIsSelecting)
+	{
+		if (bIsDragging)
+		{
+			return FCursorReply::Cursor(EMouseCursor::ResizeLeftRight);
+		}
+	}
+	else if (bIsSpaceBarKeyPressed)
+	{
+		return FCursorReply::Cursor(EMouseCursor::GrabHand);
+	}
+
+	return FCursorReply::Unhandled();
+}
+
+////////////////////////////////////////////////////////////////////////////////////////////////////
+
+FReply STimingView::OnKeyDown(const FGeometry& MyGeometry, const FKeyEvent& InKeyEvent)
+{
+	if (InKeyEvent.GetKey() == EKeys::B)
+	{
+		// Toggle Bookmarks.
+		if (MarkersTrack.IsVisible())
+		{
+			if (!MarkersTrack.IsBookmarksTrack())
+			{
+				SetDrawOnlyBookmarks(true);
+			}
+			else
+			{
+				SetTimeMarkersVisible(false);
+			}
+		}
+		else
+		{
+			SetTimeMarkersVisible(true);
+			SetDrawOnlyBookmarks(true);
+		}
+		return FReply::Handled();
+	}
+	else if (InKeyEvent.GetKey() == EKeys::M)
+	{
+		// Toggle Time Markers.
+		if (MarkersTrack.IsVisible())
+		{
+			if (MarkersTrack.IsBookmarksTrack())
+			{
+				SetDrawOnlyBookmarks(false);
+			}
+			else
+			{
+				SetTimeMarkersVisible(false);
+			}
+		}
+		else
+		{
+			SetTimeMarkersVisible(true);
+			SetDrawOnlyBookmarks(false);
+		}
+
+		return FReply::Handled();
+	}
+	else if (InKeyEvent.GetKey() == EKeys::F)
+	{
+		FrameSelection();
+		return FReply::Handled();
+	}
+	else if (InKeyEvent.GetKey() == EKeys::C)
+	{
+		Layout.bIsCompactMode = !Layout.bIsCompactMode;
+		bAreTimingEventsTracksDirty = true;
+		return FReply::Handled();
+	}
+	else if (InKeyEvent.GetKey() == EKeys::V)
+	{
+		ToggleAutoHideEmptyTracks();
+		return FReply::Handled();
+	}
+	else if (InKeyEvent.GetKey() == EKeys::Equals ||
+			 InKeyEvent.GetKey() == EKeys::Add)
+	{
+		// Zoom In
+		double ScaleX = Viewport.ScaleX * 1.25;
+		if (Viewport.ZoomWithFixedX(ScaleX, Viewport.Width/2))
+		{
+			//Viewport.EnforceHorizontalScrollLimits(1.0);
+			UpdateHorizontalScrollBar();
+			bIsViewportDirty = true;
+		}
+		return FReply::Handled();
+	}
+	else if (InKeyEvent.GetKey() == EKeys::Hyphen ||
+			 InKeyEvent.GetKey() == EKeys::Subtract)
+	{
+		// Zoom Out
+		double ScaleX = Viewport.ScaleX / 1.25;
+		if (Viewport.ZoomWithFixedX(ScaleX, Viewport.Width / 2))
+		{
+			//Viewport.EnforceHorizontalScrollLimits(1.0);
+			UpdateHorizontalScrollBar();
+			bIsViewportDirty = true;
+		}
+		return FReply::Handled();
+	}
+	else if (InKeyEvent.GetKey() == EKeys::Left)
+	{
+		if (InKeyEvent.GetModifierKeys().IsControlDown())
+		{
+			// Scroll Left
+			double DT = Viewport.EndTime - Viewport.StartTime;
+			//ScrollAtTime(Viewport.StartTime - DT * 0.05);
+			if (Viewport.ScrollAtTime(Viewport.StartTime - DT * 0.05))
+			{
+				//Viewport.EnforceHorizontalScrollLimits(1.0);
+				UpdateHorizontalScrollBar();
+				bIsViewportDirty = true;
+			}
+		}
+		else
+		{
+			SelectLeftTimingEvent();
+		}
+		return FReply::Handled();
+	}
+	else if (InKeyEvent.GetKey() == EKeys::Right)
+	{
+		if (InKeyEvent.GetModifierKeys().IsControlDown())
+		{
+			// Scroll Right
+			double DT = Viewport.EndTime - Viewport.StartTime;
+			//ScrollAtTime(Viewport.StartTime + DT * 0.05);
+			if (Viewport.ScrollAtTime(Viewport.StartTime + DT * 0.05))
+			{
+				//Viewport.EnforceHorizontalScrollLimits(1.0);
+				UpdateHorizontalScrollBar();
+				bIsViewportDirty = true;
+			}
+		}
+		else
+		{
+			SelectRightTimingEvent();
+		}
+		return FReply::Handled();
+	}
+	else if (InKeyEvent.GetKey() == EKeys::Up)
+	{
+		if (InKeyEvent.GetModifierKeys().IsControlDown())
+		{
+			// Scroll Up
+			Viewport.ScrollPosY -= 16.0 * 3;
+			//Viewport.EnforceVerticalScrollLimits(1.0);
+			UpdateVerticalScrollBar();
+			bIsVerticalViewportDirty = true;
+		}
+		else
+		{
+			SelectUpTimingEvent();
+		}
+		return FReply::Handled();
+	}
+	else if (InKeyEvent.GetKey() == EKeys::Down)
+	{
+		if (InKeyEvent.GetModifierKeys().IsControlDown())
+		{
+			// Scroll Down
+			Viewport.ScrollPosY += 16.0 * 3;
+			//Viewport.EnforceVerticalScrollLimits(1.0);
+			UpdateVerticalScrollBar();
+			bIsVerticalViewportDirty = true;
+		}
+		else
+		{
+			SelectDownTimingEvent();
+		}
+		return FReply::Handled();
+	}
+	else if (InKeyEvent.GetKey() == EKeys::Enter)
+	{
+		if (SelectedTimingEvent.IsValid())
+		{
+			const double Duration = Viewport.RestrictDuration(SelectedTimingEvent.StartTime, SelectedTimingEvent.EndTime);
+			SelectTimeInterval(SelectedTimingEvent.StartTime, Duration);
+		}
+		return FReply::Handled();
+	}
+	else if (InKeyEvent.GetKey() == EKeys::SpaceBar)
+	{
+		bIsSpaceBarKeyPressed = true;
+		FSlateApplication::Get().QueryCursor();
+		return FReply::Handled();
+	}
+	else if (InKeyEvent.GetKey() == EKeys::D) // debug: toggles down-sampling on/off
+	{
+		bUseDownSampling = !bUseDownSampling;
+		bAreTimingEventsTracksDirty = true;
+		return FReply::Handled();
+	}
+	else if (InKeyEvent.GetKey() == EKeys::G) // debug: toggles Graph track on/off
+	{
+		GraphTrack.ToggleVisibility();
+		bAreTimingEventsTracksDirty = true;
+		return FReply::Handled();
+	}
+	else if (InKeyEvent.GetKey() == EKeys::Y) // debug: toggles GPU track on/off
+	{
+		ShowHideAllGpuTracks_Execute();
+		return FReply::Handled();
+	}
+	else if (InKeyEvent.GetKey() == EKeys::U) // debug: toggles CPU tracks on/off
+	{
+		ShowHideAllCpuTracks_Execute();
+		return FReply::Handled();
+	}
+	else if (InKeyEvent.GetKey() == EKeys::L)  // debug: toggles Loading tracks on/off
+	{
+		ShowHideAllLoadingTracks_Execute();
+		return FReply::Handled();
+	}
+	else if (InKeyEvent.GetKey() == EKeys::I)  // debug: toggles IO tracks on/off
+	{
+		ShowHideAllIoTracks_Execute();
+		return FReply::Handled();
+	}
+	else if (InKeyEvent.GetKey() == EKeys::O)  // debug: toggles IO merge lanes algorithm
+	{
+		bMergeIoLanes = !bMergeIoLanes;
+		bForceIoEventsUpdate = true;
+		return FReply::Handled();
+	}
+	else if (InKeyEvent.GetKey() == EKeys::One)
+	{
+		LoadingGetEventNameFn = FLoadingTrackGetEventNameDelegate::CreateRaw(this, &STimingView::GetLoadTimeProfilerEventNameByPackageEventType);
+		return FReply::Handled();
+	}
+	else if (InKeyEvent.GetKey() == EKeys::Two)
+	{
+		LoadingGetEventNameFn = FLoadingTrackGetEventNameDelegate::CreateRaw(this, &STimingView::GetLoadTimeProfilerEventNameByExportEventType);
+		return FReply::Handled();
+	}
+	else if (InKeyEvent.GetKey() == EKeys::Three)
+	{
+		LoadingGetEventNameFn = FLoadingTrackGetEventNameDelegate::CreateRaw(this, &STimingView::GetLoadTimeProfilerEventNameByPackageName);
+		return FReply::Handled();
+	}
+	else if (InKeyEvent.GetKey() == EKeys::Four)
+	{
+		LoadingGetEventNameFn = FLoadingTrackGetEventNameDelegate::CreateRaw(this, &STimingView::GetLoadTimeProfilerEventNameByExportClassName);
+		return FReply::Handled();
+	}
+	else if (InKeyEvent.GetKey() == EKeys::Five)
+	{
+		LoadingGetEventNameFn = FLoadingTrackGetEventNameDelegate::CreateRaw(this, &STimingView::GetLoadTimeProfilerEventName);
+		return FReply::Handled();
+	}
+
+	return SCompoundWidget::OnKeyDown(MyGeometry, InKeyEvent);
+}
+
+////////////////////////////////////////////////////////////////////////////////////////////////////
+
+FReply STimingView::OnKeyUp(const FGeometry& MyGeometry, const FKeyEvent& InKeyEvent)
+{
+	if (InKeyEvent.GetKey() == EKeys::SpaceBar)
+	{
+		bIsSpaceBarKeyPressed = false;
+		FSlateApplication::Get().QueryCursor();
+		return FReply::Handled();
+	}
+
+	return SCompoundWidget::OnKeyUp(MyGeometry, InKeyEvent);
+}
+
+////////////////////////////////////////////////////////////////////////////////////////////////////
+
+void STimingView::ShowContextMenu(const FVector2D& ScreenSpacePosition, const FPointerEvent& MouseEvent)
+{
+	TSharedPtr<FUICommandList> ProfilerCommandList = FTimingProfilerManager::Get()->GetCommandList();
+	const FTimingProfilerCommands& ProfilerCommands = FTimingProfilerManager::GetCommands();
+	const FTimingProfilerActionManager& ProfilerActionManager = FTimingProfilerManager::GetActionManager();
+
+	const bool bShouldCloseWindowAfterMenuSelection = true;
+	FMenuBuilder MenuBuilder(bShouldCloseWindowAfterMenuSelection, ProfilerCommandList);
+
+	MenuBuilder.BeginSection(TEXT("Misc"), LOCTEXT("Miscellaneous", "Miscellaneous"));
+	{
+		//MenuBuilder.AddMenuEntry(FTimingProfilerManager::GetCommands().EventGraph_SelectAllFrames);
+		//MenuBuilder.AddMenuEntry(FTimingProfilerManager::GetCommands().ProfilerManager_ToggleLivePreview);
+	}
+	MenuBuilder.EndSection();
+
+	TSharedRef<SWidget> MenuWidget = MenuBuilder.MakeWidget();
+
+	FWidgetPath EventPath = MouseEvent.GetEventPath() != nullptr ? *MouseEvent.GetEventPath() : FWidgetPath();
+	FSlateApplication::Get().PushMenu(SharedThis(this), EventPath, MenuWidget, ScreenSpacePosition, FPopupTransitionEffect::ContextMenu);
+}
+
+////////////////////////////////////////////////////////////////////////////////////////////////////
+
+void STimingView::BindCommands()
+{
+	//TSharedPtr<FUICommandList> CommandList = FTimingProfilerManager::Get()->GetCommandList();
+	//const FTimingProfilerCommands& Commands = FTimingProfilerManager::GetCommands();
+}
+
+////////////////////////////////////////////////////////////////////////////////////////////////////
+
+void STimingView::HorizontalScrollBar_OnUserScrolled(float ScrollOffset)
+{
+	const double SX = 1.0 / (Viewport.MaxValidTime - Viewport.MinValidTime);
+	const float ThumbSizeFraction = FMath::Clamp<float>((Viewport.EndTime - Viewport.StartTime) * SX, 0.0f, 1.0f);
+	const float OffsetFraction = FMath::Clamp<float>(ScrollOffset, 0.0f, 1.0f - ThumbSizeFraction);
+
+	const double Time = Viewport.MinValidTime + static_cast<double>(OffsetFraction) * (Viewport.MaxValidTime - Viewport.MinValidTime);
+	if (Viewport.StartTime != Time)
+	{
+		Viewport.ScrollAtTime(Time);
+		bIsViewportDirty = true;
+	}
+
+	HorizontalScrollBar->SetState(OffsetFraction, ThumbSizeFraction);
+}
+
+////////////////////////////////////////////////////////////////////////////////////////////////////
+
+void STimingView::UpdateHorizontalScrollBar()
+{
+	const double SX = 1.0 / (Viewport.MaxValidTime - Viewport.MinValidTime);
+	const float ThumbSizeFraction = FMath::Clamp<float>((Viewport.EndTime - Viewport.StartTime) * SX, 0.0f, 1.0f);
+	const float ScrollOffset = static_cast<float>((Viewport.StartTime - Viewport.MinValidTime) * SX);
+	const float OffsetFraction = FMath::Clamp<float>(ScrollOffset, 0.0f, 1.0f - ThumbSizeFraction);
+
+	HorizontalScrollBar->SetState(OffsetFraction, ThumbSizeFraction);
+}
+
+////////////////////////////////////////////////////////////////////////////////////////////////////
+
+void STimingView::VerticalScrollBar_OnUserScrolled(float ScrollOffset)
+{
+	const float SY = 1.0 / Viewport.ScrollHeight;
+	const float H = Viewport.Height - Viewport.TopOffset;
+	const float ThumbSizeFraction = FMath::Clamp<float>(H * SY, 0.0f, 1.0f);
+	const float OffsetFraction = FMath::Clamp<float>(ScrollOffset, 0.0f, 1.0f - ThumbSizeFraction);
+
+	const float ScrollPosY = OffsetFraction * Viewport.ScrollHeight;
+	if (Viewport.ScrollPosY != ScrollPosY)
+	{
+		Viewport.ScrollPosY = ScrollPosY;
+		bIsVerticalViewportDirty = true;
+	}
+
+	VerticalScrollBar->SetState(OffsetFraction, ThumbSizeFraction);
+}
+
+////////////////////////////////////////////////////////////////////////////////////////////////////
+
+void STimingView::UpdateVerticalScrollBar()
+{
+	const float SY = 1.0 / Viewport.ScrollHeight;
+	const float H = Viewport.Height - Viewport.TopOffset;
+	const float ThumbSizeFraction = FMath::Clamp<float>(H * SY, 0.0f, 1.0f);
+	const float ScrollOffset = Viewport.ScrollPosY * SY;
+	float OffsetFraction = FMath::Clamp<float>(ScrollOffset, 0.0f, 1.0f - ThumbSizeFraction);
+
+	VerticalScrollBar->SetState(OffsetFraction, ThumbSizeFraction);
+}
+
+////////////////////////////////////////////////////////////////////////////////////////////////////
+
+void STimingView::ScrollAtTime(double StartTime)
+{
+	if (Viewport.ScrollAtTime(StartTime))
+	{
+		UpdateHorizontalScrollBar();
+		bIsViewportDirty = true;
+	}
+}
+
+////////////////////////////////////////////////////////////////////////////////////////////////////
+
+void STimingView::CenterOnTimeInterval(double IntervalStartTime, double IntervalDuration)
+{
+	if (Viewport.CenterOnTimeInterval(IntervalStartTime, IntervalDuration))
+	{
+		Viewport.EnforceHorizontalScrollLimits(1.0);
+
+		UpdateHorizontalScrollBar();
+		bIsViewportDirty = true;
+	}
+}
+
+////////////////////////////////////////////////////////////////////////////////////////////////////
+
+void STimingView::BringIntoView(double StartTime, double EndTime)
+{
+	EndTime = Viewport.RestrictEndTime(EndTime);
+
+	// Increase interval with 8% on each side.
+	const double DT = (Viewport.EndTime - Viewport.StartTime) * 0.08;
+	StartTime -= DT;
+	EndTime += DT;
+
+	double NewStartTime = Viewport.StartTime;
+
+	if (EndTime > Viewport.EndTime)
+	{
+		NewStartTime += EndTime - Viewport.EndTime;
+	}
+
+	if (StartTime < NewStartTime)
+	{
+		NewStartTime = StartTime;
+	}
+
+	ScrollAtTime(NewStartTime);
+}
+
+////////////////////////////////////////////////////////////////////////////////////////////////////
+
+void STimingView::SelectTimeInterval(double IntervalStartTime, double IntervalDuration)
+{
+	SelectionStartTime = IntervalStartTime;
+	SelectionEndTime = IntervalStartTime + IntervalDuration;
+	LastSelectionType = ESelectionType::TimeRange;
+
+	//TODO: SelectionChangedEvent.Broadcast(SelectionStartTime, SelectionEndTime);
+	UpdateAggregatedStats();
+}
+
+////////////////////////////////////////////////////////////////////////////////////////////////////
+
+void STimingView::UpdateAggregatedStats()
+{
+	TSharedPtr<STimingProfilerWindow> Wnd = FTimingProfilerManager::Get()->GetProfilerWindow();
+	if (Wnd)
+	{
+		if (Wnd->TimersView)
+		{
+			Wnd->TimersView->UpdateStats(SelectionStartTime, SelectionEndTime);
+		}
+		if (Wnd->StatsView)
+		{
+			Wnd->StatsView->UpdateStats(SelectionStartTime, SelectionEndTime);
+		}
+	}
+
+	if (bAssetLoadingMode)
+	{
+		TSharedPtr<const Trace::IAnalysisSession> Session = FInsightsManager::Get()->GetSession();
+		if (Session.IsValid() && Trace::ReadLoadTimeProfilerProvider(*Session.Get()))
+		{
+			Trace::FAnalysisSessionReadScope SessionReadScope(*Session.Get());
+
+			const Trace::ILoadTimeProfilerProvider& LoadTimeProfilerProvider = *Trace::ReadLoadTimeProfilerProvider(*Session.Get());
+
+			Trace::ITable<Trace::FLoadTimeProfilerAggregatedStats>* EventAggregationTable = LoadTimeProfilerProvider.CreateEventAggregation(SelectionStartTime, SelectionEndTime);
+			Trace::ITable<Trace::FLoadTimeProfilerAggregatedStats>* ObjectTypeAggregationTable = LoadTimeProfilerProvider.CreateObjectTypeAggregation(SelectionStartTime, SelectionEndTime);
+
+			auto TableToString = [](const TCHAR* TableName, Trace::ITable<Trace::FLoadTimeProfilerAggregatedStats>* Table) -> FString
+			{
+				FString Str;
+
+				const uint32 TotalRowCount = static_cast<uint32>(Table->GetRowCount());
+
+				Str.Append(TableName).Append(": ").Append(FText::AsNumber(TotalRowCount).ToString()).Append(" records, sorted by Total time");
+
+				constexpr uint32 MaxRowCount = 10; // only get first 10 records
+				const uint32 RowCount = FMath::Min<uint32>(TotalRowCount, MaxRowCount);
+
+				const Trace::ITableLayout& TableLayout = Table->GetLayout();
+				const int32 ColumnCount = TableLayout.GetColumnCount();
+
+				auto Reader = Table->CreateReader();
+
+				//////////////////////////////////////////////////
+
+				struct FSortedIndexEntry
+				{
+					uint32 RowIndex;
+					double Value;
+				};
+
+				TArray<FSortedIndexEntry> SortedIndex;
+				SortedIndex.Reserve(TotalRowCount);
+
+				constexpr uint32 TotalTimeColumnIndex = 2;
+				ensure(TableLayout.GetColumnType(TotalTimeColumnIndex) == Trace::ETableColumnType::TableColumnType_Double);
+
+				for (uint32 RowIndex = 0; RowIndex < TotalRowCount; ++RowIndex)
+				{
+					Reader->SetRowIndex(RowIndex);
+					double Value = Reader->GetValueDouble(TotalTimeColumnIndex);
+					SortedIndex.Add({ RowIndex, Value });
+				}
+
+				SortedIndex.Sort([](const FSortedIndexEntry& A, const FSortedIndexEntry& B) { return A.Value > B.Value; });
+
+				//////////////////////////////////////////////////
+
+				for (uint32 Index = 0; Index < RowCount; ++Index)
+				{
+					uint32 RowIndex = SortedIndex[Index].RowIndex;
+					Reader->SetRowIndex(RowIndex);
+
+					Str.Append("\n").Append(FText::AsNumber(Index + 1).ToString()).Append(". ");
+
+					for (int32 ColumnIndex = 0; ColumnIndex < ColumnCount; ++ColumnIndex)
+					{
+						Str.Append(TableLayout.GetColumnName(ColumnIndex)).Append("=");
+
+						Trace::ETableColumnType ColumnType = TableLayout.GetColumnType(ColumnIndex);
+						switch (ColumnType)
+						{
+						case Trace::ETableColumnType::TableColumnType_Bool:
+							Str.Append(Reader->GetValueBool(ColumnIndex) ? TEXT("True") : TEXT("False")).Append(", ");
+							break;
+						case Trace::ETableColumnType::TableColumnType_Int:
+							Str.Append(FText::AsNumber(Reader->GetValueInt(ColumnIndex)).ToString()).Append(", ");
+							break;
+						case Trace::ETableColumnType::TableColumnType_Float:
+							Str.Append(FText::AsNumber(Reader->GetValueFloat(ColumnIndex)).ToString()).Append(", ");
+							break;
+						case Trace::ETableColumnType::TableColumnType_Double:
+							Str.Append(FText::AsNumber(Reader->GetValueDouble(ColumnIndex)).ToString()).Append(", ");
+							break;
+						case Trace::ETableColumnType::TableColumnType_CString:
+							Str.Append(Reader->GetValueCString(ColumnIndex)).Append(", ");
+							break;
+						}
+					}
+				}
+
+				if (RowCount != TotalRowCount)
+				{
+					Str.Append("\n[...]");
+				}
+
+				return Str;
+			};
+
+			EventAggregationStr = TableToString(TEXT("Event Aggregation"), EventAggregationTable);
+			ObjectTypeAggregationStr = TableToString(TEXT("Object Type Aggregation"), ObjectTypeAggregationTable);
+		}
+	}
+}
+
+////////////////////////////////////////////////////////////////////////////////////////////////////
+
+void STimingView::SetAndCenterOnTimeMarker(double Time)
+{
+	TimeMarker = Time;
+
+	double MinT, MaxT;
+	Viewport.GetHorizontalScrollLimits(MinT, MaxT);
+	const double ViewportDuration = static_cast<double>(Viewport.Width) / Viewport.ScaleX;
+	MinT += ViewportDuration / 2;
+	MaxT += ViewportDuration / 2;
+	Time = FMath::Clamp<double>(Time, MinT, MaxT);
+
+	Time = Viewport.AlignTimeToPixel(Time);
+	CenterOnTimeInterval(Time, 0.0);
+}
+
+////////////////////////////////////////////////////////////////////////////////////////////////////
+
+void STimingView::SelectToTimeMarker(double Time)
+{
+	if (TimeMarker < Time)
+	{
+		SelectTimeInterval(TimeMarker, Time - TimeMarker);
+	}
+	else
+	{
+		SelectTimeInterval(Time, TimeMarker - Time);
+	}
+
+	TimeMarker = Time;
+}
+
+////////////////////////////////////////////////////////////////////////////////////////////////////
+
+void STimingView::SetTimeMarkersVisible(bool bIsMarkersTrackVisible)
+{
+	if (MarkersTrack.IsVisible() != bIsMarkersTrackVisible)
+	{
+		MarkersTrack.SetVisibilityFlag(bIsMarkersTrackVisible);
+
+		if (MarkersTrack.IsVisible())
+		{
+			if (Viewport.ScrollPosY != 0.0f)
+			{
+				Viewport.ScrollPosY += MarkersTrack.GetHeight();
+			}
+
+			MarkersTrack.SetDirtyFlag();
+		}
+		else
+		{
+			Viewport.ScrollPosY -= MarkersTrack.GetHeight();
+		}
+	}
+}
+
+////////////////////////////////////////////////////////////////////////////////////////////////////
+
+void STimingView::SetDrawOnlyBookmarks(bool bIsBookmarksTrack)
+{
+	if (MarkersTrack.IsBookmarksTrack() != bIsBookmarksTrack)
+	{
+		const float PrevHeight = MarkersTrack.GetHeight();
+		MarkersTrack.SetBookmarksTrackFlag(bIsBookmarksTrack);
+
+		if (MarkersTrack.IsVisible())
+		{
+			if (Viewport.ScrollPosY != 0.0f)
+			{
+				Viewport.ScrollPosY += MarkersTrack.GetHeight() - PrevHeight;
+			}
+
+			MarkersTrack.SetDirtyFlag();
+		}
+	}
+}
+
+////////////////////////////////////////////////////////////////////////////////////////////////////
+
+void STimingView::UpdateHoveredTimingEvent(float MX, float MY)
+{
+	HoveredTimingEvent.Track = nullptr;
+	HoveredTimingEvent.TypeId = FTimerNode::InvalidId;
+
+	if (MY >= Viewport.TopOffset && MY < Viewport.Height)
+	{
+		for (const auto& KV : CachedTimelines)
+		{
+			const FTimingEventsTrack& Track = *KV.Value;
+			if (Track.IsVisible())
+			{
+				const float Y = Viewport.TopOffset + Track.GetPosY() - Viewport.ScrollPosY;
+				if (MY >= Y && MY < Y + Track.GetHeight())
+				{
+					HoveredTimingEvent.Track = &Track;
+					break;
+				}
+			}
+		}
+
+		if (HoveredTimingEvent.Track)
+		{
+			const float Y0 = Viewport.TopOffset + HoveredTimingEvent.Track->GetPosY() - Viewport.ScrollPosY + 1.0f + Layout.TimelineDY;
+
+			// If mouse is not above first sub-track or below last sub-track...
+			if (MY >= Y0 && MY < Y0 + HoveredTimingEvent.Track->GetHeight() + Layout.TimelineDY)
+			{
+				int32 Depth = (MY - Y0) / (Layout.EventDY + Layout.EventH);
+				float EventMY = (MY - Y0) - Depth * (Layout.EventDY + Layout.EventH);
+
+				const double StartTime = Viewport.SlateUnitsToTime(MX);
+				const double EndTime = StartTime + 2.0 / Viewport.ScaleX; // +2px
+				constexpr bool bStopAtFirstMatch = true; // get first one matching
+				constexpr bool bSearchForLargestEvent = false;
+				SearchTimingEvent(StartTime, EndTime,
+					[Depth](double, double, uint32 EventDepth, uint32)
+					{
+						return EventDepth == Depth;
+					},
+					HoveredTimingEvent, bStopAtFirstMatch, bSearchForLargestEvent);
+
+				//TODO: ComputeSingleTimingEventStats(HoveredTimingEvent) --> compute ExclusiveTime
+			}
+		}
+	}
+}
+
+////////////////////////////////////////////////////////////////////////////////////////////////////
+
+bool STimingView::SearchTimingEvent(const double InStartTime,
+									const double InEndTime,
+									TFunctionRef<bool(double, double, uint32, uint32)> InPredicate,
+									FTimingEvent& InOutTimingEvent,
+									bool bInStopAtFirstMatch,
+									bool bInSearchForLargestEvent) const
+{
+	struct FSearchTimingEventContext
+	{
+		const double StartTime;
+		const double EndTime;
+		TFunctionRef<bool(double, double, uint32, uint32)> Predicate;
+		FTimingEvent& TimingEvent;
+		const bool bStopAtFirstMatch;
+		const bool bSearchForLargestEvent;
+		mutable bool bFound;
+		mutable bool bContinueSearching;
+		mutable double LargestDuration;
+
+		FSearchTimingEventContext(const double InStartTime, const double InEndTime, TFunctionRef<bool(double, double, uint32, uint32)> InPredicate, FTimingEvent& InOutTimingEvent, bool bInStopAtFirstMatch, bool bInSearchForLargestEvent)
+			: StartTime(InStartTime)
+			, EndTime(InEndTime)
+			, Predicate(InPredicate)
+			, TimingEvent(InOutTimingEvent)
+			, bStopAtFirstMatch(bInStopAtFirstMatch)
+			, bSearchForLargestEvent(bInSearchForLargestEvent)
+			, bFound(false)
+			, bContinueSearching(true)
+			, LargestDuration(-1.0)
+		{
+		}
+
+		void CheckEvent(double EventStartTime, double EventEndTime, uint32 EventDepth, const Trace::FTimingProfilerEvent& Event)
+		{
+			if (bContinueSearching && Predicate(EventStartTime, EventEndTime, EventDepth, Event.TimerIndex))
+			{
+				if (!bSearchForLargestEvent || EventEndTime - EventStartTime > LargestDuration)
+				{
+					LargestDuration = EventEndTime - EventStartTime;
+
+					TimingEvent.TypeId = Event.TimerIndex;
+					TimingEvent.Depth = EventDepth;
+					TimingEvent.StartTime = EventStartTime;
+					TimingEvent.EndTime = EventEndTime;
+
+					bFound = true;
+					bContinueSearching = !bStopAtFirstMatch || bSearchForLargestEvent;
+				}
+			}
+		}
+
+		void CheckEvent(double EventStartTime, double EventEndTime, uint32 EventDepth, const Trace::FLoadTimeProfilerCpuEvent& Event)
+		{
+			if (bContinueSearching && Predicate(EventStartTime, EventEndTime, EventDepth, 0))
+			{
+				if (!bSearchForLargestEvent || EventEndTime - EventStartTime > LargestDuration)
+				{
+					LargestDuration = EventEndTime - EventStartTime;
+
+					TimingEvent.TypeId = 0;
+					TimingEvent.Depth = EventDepth;
+					TimingEvent.StartTime = EventStartTime;
+					TimingEvent.EndTime = EventEndTime;
+
+					TimingEvent.LoadingInfo = Event;
+
+					bFound = true;
+					bContinueSearching = !bStopAtFirstMatch || bSearchForLargestEvent;
+				}
+			}
+		}
+	};
+
+	FSearchTimingEventContext Ctx(InStartTime, InEndTime, InPredicate, InOutTimingEvent, bInStopAtFirstMatch, bInSearchForLargestEvent);
+
+	TSharedPtr<const Trace::IAnalysisSession> Session = FInsightsManager::Get()->GetSession();
+	if (Session.IsValid())
+	{
+		Trace::FAnalysisSessionReadScope SessionReadScope(*Session.Get());
+
+		const FTimingEventsTrack* Track = Ctx.TimingEvent.Track;
+
+		if (Track->GetType() == ETimingEventsTrackType::Cpu ||
+			Track->GetType() == ETimingEventsTrackType::Gpu)
+		{
+			if (Trace::ReadTimingProfilerProvider(*Session.Get()))
+			{
+				const Trace::ITimingProfilerProvider& TimingProfilerProvider = *Trace::ReadTimingProfilerProvider(*Session.Get());
+
+				TimingProfilerProvider.ReadTimeline(Track->GetId(), [&Ctx](const Trace::ITimingProfilerProvider::Timeline& Timeline)
+				{
+					Timeline.EnumerateEvents(Ctx.StartTime, Ctx.EndTime, [&Ctx](double EventStartTime, double EventEndTime, uint32 EventDepth, const Trace::FTimingProfilerEvent& Event)
+					{
+						Ctx.CheckEvent(EventStartTime, EventEndTime, EventDepth, Event);
+					});
+				});
+
+				if (Ctx.bFound)
+				{
+					// Compute Exclusive Time.
+					TimingProfilerProvider.ReadTimeline(Track->GetId(), [&Ctx](const Trace::ITimingProfilerProvider::Timeline& Timeline)
+					{
+						struct FEnumerationState
+						{
+							double EventStartTime;
+							double EventEndTime;
+							uint64 EventDepth;
+
+							uint64 CurrentDepth;
+							double LastTime;
+							double ExclusiveTime;
+							bool IsInEventScope;
+						};
+						FEnumerationState State;
+
+						State.EventStartTime = Ctx.TimingEvent.StartTime;
+						State.EventEndTime = Ctx.TimingEvent.EndTime;
+						State.EventDepth = Ctx.TimingEvent.Depth;
+
+						State.CurrentDepth = 0;
+						State.LastTime = 0.0;
+						State.ExclusiveTime = 0.0;
+						State.IsInEventScope = false;
+
+						Timeline.EnumerateEvents(Ctx.TimingEvent.StartTime, Ctx.TimingEvent.EndTime, [&State](bool IsEnter, double Time, const Trace::FTimingProfilerEvent& Event)
+						{
+							if (IsEnter)
+							{
+								if (State.IsInEventScope && State.CurrentDepth == State.EventDepth + 1)
+								{
+									State.ExclusiveTime += Time - State.LastTime;
+								}
+								if (State.CurrentDepth == State.EventDepth && Time == State.EventStartTime)
+								{
+									State.IsInEventScope = true;
+								}
+								++State.CurrentDepth;
+							}
+							else
+							{
+								--State.CurrentDepth;
+								if (State.CurrentDepth == State.EventDepth && Time == State.EventEndTime)
+								{
+									State.IsInEventScope = false;
+									State.ExclusiveTime += Time - State.LastTime;
+								}
+							}
+							State.LastTime = Time;
+						});
+
+						Ctx.TimingEvent.ExclusiveTime = State.ExclusiveTime;
+					});
+				}
+			}
+		}
+		else if (Track->GetType() == ETimingEventsTrackType::Loading)
+		{
+			if (Trace::ReadLoadTimeProfilerProvider(*Session.Get()))
+			{
+				const Trace::ILoadTimeProfilerProvider& LoadTimeProfilerProvider = *Trace::ReadLoadTimeProfilerProvider(*Session.Get());
+
+				if (Track == LoadingMainThreadTrack)
+				{
+					LoadTimeProfilerProvider.ReadMainThreadCpuTimeline([&Ctx](const Trace::ILoadTimeProfilerProvider::CpuTimeline& Timeline)
+					{
+						Timeline.EnumerateEvents(Ctx.StartTime, Ctx.EndTime, [&Ctx](double EventStartTime, double EventEndTime, uint32 EventDepth, const Trace::FLoadTimeProfilerCpuEvent& Event)
+						{
+							Ctx.CheckEvent(EventStartTime, EventEndTime, EventDepth, Event);
+						});
+					});
+				}
+				else if (Track == LoadingAsyncThreadTrack)
+				{
+					LoadTimeProfilerProvider.ReadAsyncLoadingThreadCpuTimeline([&Ctx](const Trace::ILoadTimeProfilerProvider::CpuTimeline& Timeline)
+					{
+						Timeline.EnumerateEvents(Ctx.StartTime, Ctx.EndTime, [&Ctx](double EventStartTime, double EventEndTime, uint32 EventDepth, const Trace::FLoadTimeProfilerCpuEvent& Event)
+						{
+							Ctx.CheckEvent(EventStartTime, EventEndTime, EventDepth, Event);
+						});
+					});
+				}
+			}
+		}
+	}
+
+	return Ctx.bFound;
+}
+
+////////////////////////////////////////////////////////////////////////////////////////////////////
+
+void STimingView::OnSelectedTimingEventChanged()
+{
+	// Select the timer node coresponding to timing event type of selected timing event.
+	if (SelectedTimingEvent.IsValid() &&
+		(SelectedTimingEvent.Track->GetType() == ETimingEventsTrackType::Cpu ||
+		 SelectedTimingEvent.Track->GetType() == ETimingEventsTrackType::Gpu))
+	{
+		TSharedPtr<STimingProfilerWindow> Wnd = FTimingProfilerManager::Get()->GetProfilerWindow();
+		if (Wnd)
+		{
+			if (Wnd->TimersView)
+			{
+				Wnd->TimersView->SelectTimerNode(SelectedTimingEvent.TypeId);
+			}
+		}
+	}
+}
+
+////////////////////////////////////////////////////////////////////////////////////////////////////
+
+void STimingView::SelectHoveredTimingEvent()
+{
+	SelectedTimingEvent = HoveredTimingEvent;
+	if (SelectedTimingEvent.IsValid())
+	{
+		LastSelectionType = ESelectionType::TimingEvent;
+		BringIntoView(SelectedTimingEvent.StartTime, SelectedTimingEvent.EndTime);
+	}
+	OnSelectedTimingEventChanged();
+}
+
+////////////////////////////////////////////////////////////////////////////////////////////////////
+
+void STimingView::SelectLeftTimingEvent()
+{
+	if (SelectedTimingEvent.IsValid())
+	{
+		const uint32 Depth = SelectedTimingEvent.Depth;
+		const double StartTime = SelectedTimingEvent.StartTime;
+		const double EndTime = SelectedTimingEvent.EndTime;
+		const bool bStopAtFirstMatch = false; // get last one matching
+		const bool bSearchForLargestEvent = false;
+		if (SearchTimingEvent(0.0, StartTime,
+			[Depth, StartTime, EndTime](double EventStartTime, double EventEndTime, uint32 EventDepth, uint32 EventTypeId)
+			{
+				return EventDepth == Depth &&
+					(EventStartTime < StartTime || EventEndTime < EndTime);
+			},
+			SelectedTimingEvent, bStopAtFirstMatch, bSearchForLargestEvent))
+		{
+			BringIntoView(SelectedTimingEvent.StartTime, SelectedTimingEvent.EndTime);
+			OnSelectedTimingEventChanged();
+		}
+	}
+}
+
+////////////////////////////////////////////////////////////////////////////////////////////////////
+
+void STimingView::SelectRightTimingEvent()
+{
+	if (SelectedTimingEvent.IsValid())
+	{
+		const uint32 Depth = SelectedTimingEvent.Depth;
+		const double StartTime = SelectedTimingEvent.StartTime;
+		const double EndTime = SelectedTimingEvent.EndTime;
+		const bool bStopAtFirstMatch = true; // get first one matching
+		const bool bSearchForLargestEvent = false;
+		if (SearchTimingEvent(EndTime, Viewport.MaxValidTime,
+			[Depth, StartTime, EndTime](double EventStartTime, double EventEndTime, uint32 EventDepth, uint32 EventTypeId)
+			{
+				return EventDepth == Depth &&
+					(EventStartTime > StartTime || EventEndTime > EndTime);
+			},
+			SelectedTimingEvent, bStopAtFirstMatch, bSearchForLargestEvent))
+		{
+			BringIntoView(SelectedTimingEvent.StartTime, SelectedTimingEvent.EndTime);
+			OnSelectedTimingEventChanged();
+		}
+	}
+}
+
+////////////////////////////////////////////////////////////////////////////////////////////////////
+
+void STimingView::SelectUpTimingEvent()
+{
+	if (SelectedTimingEvent.IsValid() &&
+		SelectedTimingEvent.Depth > 0)
+	{
+		const uint32 Depth = SelectedTimingEvent.Depth - 1;
+		double StartTime = SelectedTimingEvent.StartTime;
+		double EndTime = SelectedTimingEvent.EndTime;
+		const bool bStopAtFirstMatch = true; // get first one matching
+		const bool bSearchForLargestEvent = false;
+		if (SearchTimingEvent(StartTime, EndTime,
+			[Depth, StartTime, EndTime](double EventStartTime, double EventEndTime, uint32 EventDepth, uint32 EventTypeId)
+			{
+				return EventDepth == Depth
+					&& EventStartTime <= EndTime
+					&& EventEndTime >= StartTime;
+			},
+			SelectedTimingEvent, bStopAtFirstMatch, bSearchForLargestEvent))
+		{
+			BringIntoView(SelectedTimingEvent.StartTime, SelectedTimingEvent.EndTime);
+			OnSelectedTimingEventChanged();
+		}
+	}
+}
+
+////////////////////////////////////////////////////////////////////////////////////////////////////
+
+void STimingView::SelectDownTimingEvent()
+{
+	if (SelectedTimingEvent.IsValid())
+	{
+		const uint32 Depth = SelectedTimingEvent.Depth + 1;
+		double StartTime = SelectedTimingEvent.StartTime;
+		double EndTime = SelectedTimingEvent.EndTime;
+		const bool bStopAtFirstMatch = false; // check all timing events
+		const bool bSearchForLargestEvent = true; // get largest timing event
+		if (SearchTimingEvent(StartTime, EndTime,
+			[Depth, StartTime, EndTime](double EventStartTime, double EventEndTime, uint32 EventDepth, uint32 EventTypeId)
+			{
+				return EventDepth == Depth
+					&& EventStartTime <= EndTime
+					&& EventEndTime >= StartTime;
+			},
+			SelectedTimingEvent, bStopAtFirstMatch, bSearchForLargestEvent))
+		{
+			BringIntoView(SelectedTimingEvent.StartTime, SelectedTimingEvent.EndTime);
+			OnSelectedTimingEventChanged();
+		}
+	}
+}
+
+////////////////////////////////////////////////////////////////////////////////////////////////////
+
+void STimingView::FrameSelection()
+{
+	double StartTime, EndTime;
+
+	ESelectionType Type = ESelectionType::None;
+
+	if (LastSelectionType == ESelectionType::TimingEvent)
+	{
+		// Try framing the selected timing event.
+		if (SelectedTimingEvent.IsValid())
+		{
+			Type = ESelectionType::TimingEvent;
+		}
+
+		// Next time, try framing the selected time range.
+		LastSelectionType = ESelectionType::TimeRange;
+	}
+	else if (LastSelectionType == ESelectionType::TimeRange)
+	{
+		// Try framing the selected time range.
+		if (SelectionEndTime > SelectionStartTime)
+		{
+			Type = ESelectionType::TimeRange;
+		}
+
+		// Next time, try framing the selected timing event.
+		LastSelectionType = ESelectionType::TimingEvent;
+	}
+
+	// If no last selection or last selection is empty...
+	if (LastSelectionType == ESelectionType::None || Type == ESelectionType::None)
+	{
+		// First, try framing the selected timing event...
+		if (SelectedTimingEvent.IsValid())
+		{
+			Type = ESelectionType::TimingEvent;
+		}
+		else // ...otherwise, try framing the selected time range
+		{
+			Type = ESelectionType::TimeRange;
+		}
+	}
+
+	if (Type == ESelectionType::TimingEvent)
+	{
+		// Frame the selected event.
+		StartTime = SelectedTimingEvent.StartTime;
+		EndTime = Viewport.RestrictEndTime(SelectedTimingEvent.EndTime);
+		if (EndTime == StartTime)
+		{
+			EndTime += 1.0 / Viewport.ScaleX; // +1px
+		}
+	}
+	else
+	{
+		// Frame the selected time range.
+		StartTime = SelectionStartTime;
+		EndTime = Viewport.RestrictEndTime(SelectionEndTime);
+	}
+
+	if (EndTime > StartTime)
+	{
+		const double Duration = EndTime - StartTime;
+		if (Viewport.ZoomOnTimeInterval(StartTime - Duration * 0.1, Duration * 1.2))
+		{
+			UpdateHorizontalScrollBar();
+			bIsViewportDirty = true;
+		}
+	}
+}
+
+////////////////////////////////////////////////////////////////////////////////////////////////////
+
+const FTimerNodePtr STimingView::GetTimerNode(uint64 TypeId) const
+{
+	const FTimerNodePtr* TimerNodePtrPtr = nullptr;
+	TSharedPtr<STimingProfilerWindow> Wnd = FTimingProfilerManager::Get()->GetProfilerWindow();
+	if (Wnd && Wnd->TimersView)
+	{
+		TimerNodePtrPtr = Wnd->TimersView->GetTimerNode(TypeId);
+		if (TimerNodePtrPtr == nullptr)
+		{
+			// List of timers in TimersView not up to date?
+			// Refresh and try again.
+			Wnd->TimersView->RebuildTree(false);
+			TimerNodePtrPtr = Wnd->TimersView->GetTimerNode(TypeId);
+		}
+	}
+	return TimerNodePtrPtr ? *TimerNodePtrPtr : nullptr;
+}
+
+////////////////////////////////////////////////////////////////////////////////////////////////////
+
+TSharedRef<SWidget> STimingView::MakeTracksFilterMenu()
+{
+	FMenuBuilder MenuBuilder(/*bInShouldCloseWindowAfterMenuSelection=*/true, nullptr);
+
+	MenuBuilder.BeginSection("QuickFilter", LOCTEXT("TracksFilterHeading", "Quick Filter"));
+	{
+		MenuBuilder.AddMenuEntry(
+			LOCTEXT("ShowAllGpuTracks", "GPU Track (Y)"),
+			LOCTEXT("ShowAllGpuTracks_Tooltip", "Show/hide the GPU track"),
+			FSlateIcon(),
+			FUIAction(FExecuteAction::CreateSP(this, &STimingView::ShowHideAllGpuTracks_Execute),
+				FCanExecuteAction::CreateLambda([] { return true; }),
+				FIsActionChecked::CreateSP(this, &STimingView::ShowHideAllGpuTracks_IsChecked)),
+			NAME_None,
+			EUserInterfaceActionType::ToggleButton
+		);
+
+		MenuBuilder.AddMenuEntry(
+			LOCTEXT("ShowAllCpuTracks", "CPU Thread Tracks (U)"),
+			LOCTEXT("ShowAllCpuTracks_Tooltip", "Show/hide all CPU tracks (and all CPU thread groups)"),
+			FSlateIcon(),
+			FUIAction(FExecuteAction::CreateSP(this, &STimingView::ShowHideAllCpuTracks_Execute),
+				FCanExecuteAction::CreateLambda([] { return true; }),
+				FIsActionChecked::CreateSP(this, &STimingView::ShowHideAllCpuTracks_IsChecked)),
+			NAME_None,
+			EUserInterfaceActionType::ToggleButton
+		);
+
+		MenuBuilder.AddMenuEntry(
+			LOCTEXT("ShowAllLoadingTracks", "Asset Loading Tracks (L)"),
+			LOCTEXT("ShowAllLoadingTracks_Tooltip", "Show/hide the Asset Loading tracks"),
+			FSlateIcon(),
+			FUIAction(FExecuteAction::CreateSP(this, &STimingView::ShowHideAllLoadingTracks_Execute),
+				FCanExecuteAction::CreateLambda([] { return true; }),
+				FIsActionChecked::CreateSP(this, &STimingView::ShowHideAllLoadingTracks_IsChecked)),
+			NAME_None,
+			EUserInterfaceActionType::ToggleButton
+		);
+
+		MenuBuilder.AddMenuEntry(
+			LOCTEXT("ShowAllIoTracks", "I/O Tracks (I)"),
+			LOCTEXT("ShowAllIoTracks_Tooltip", "Show/hide the I/O (File Activity) tracks"),
+			FSlateIcon(),
+			FUIAction(FExecuteAction::CreateSP(this, &STimingView::ShowHideAllIoTracks_Execute),
+				FCanExecuteAction::CreateLambda([] { return true; }),
+				FIsActionChecked::CreateSP(this, &STimingView::ShowHideAllIoTracks_IsChecked)),
+			NAME_None,
+			EUserInterfaceActionType::ToggleButton
+		);
+
+		MenuBuilder.AddMenuEntry(
+			LOCTEXT("AutoHideEmptyTracks", "Auto Hide Empty Tracks (V)"),
+			LOCTEXT("AutoHideEmptyTracks_Tooltip", "Auto hide empty tracks (ones without timing events in current viewport)"),
+			FSlateIcon(),
+			FUIAction(FExecuteAction::CreateSP(this, &STimingView::ToggleAutoHideEmptyTracks),
+				FCanExecuteAction::CreateLambda([] { return true; }),
+				FIsActionChecked::CreateSP(this, &STimingView::IsAutoHideEmptyTracksEnabled)),
+			NAME_None,
+			EUserInterfaceActionType::ToggleButton
+		);
+	}
+	MenuBuilder.EndSection();
+
+	MenuBuilder.BeginSection("ThreadGroups", LOCTEXT("ThreadGroupsHeading", "CPU Thread Groups"));
+	CreateThreadGroupsMenu(MenuBuilder);
+	MenuBuilder.EndSection();
+
+	//MenuBuilder.BeginSection("Tracks", LOCTEXT("TracksHeading", "Tracks"));
+	//CreateTracksMenu(MenuBuilder);
+	//MenuBuilder.EndSection();
+
+	return MenuBuilder.MakeWidget();
+}
+
+////////////////////////////////////////////////////////////////////////////////////////////////////
+
+void STimingView::CreateThreadGroupsMenu(FMenuBuilder& MenuBuilder)
+{
+	for (const auto& KV : ThreadGroups)
+	{
+		const FThreadGroup& ThreadGroup = KV.Value;
+		if (ThreadGroup.NumTimelines > 0)
+		{
+			MenuBuilder.AddMenuEntry(
+				FText::FromString(ThreadGroup.Name),
+				TAttribute<FText>(), // no tooltip
+				FSlateIcon(),
+				FUIAction(FExecuteAction::CreateSP(this, &STimingView::ToggleTrackVisibilityByGroup_Execute, ThreadGroup.Name),
+					FCanExecuteAction::CreateLambda([] { return true; }),
+					FIsActionChecked::CreateSP(this, &STimingView::ToggleTrackVisibilityByGroup_IsChecked, ThreadGroup.Name)),
+				NAME_None,
+				EUserInterfaceActionType::ToggleButton
+			);
+		}
+	}
+}
+
+////////////////////////////////////////////////////////////////////////////////////////////////////
+
+void STimingView::CreateTracksMenu(FMenuBuilder& MenuBuilder)
+{
+	for (int32 TrackIndex = 0; TrackIndex < TimingEventsTracks.Num(); ++TrackIndex)
+	{
+		const FTimingEventsTrack& Track = *TimingEventsTracks[TrackIndex];
+		if (Track.GetHeight() > 0.0f || Layout.TargetMinTimelineH > 0.0f)
+		{
+			MenuBuilder.AddMenuEntry(
+				FText::FromString(Track.GetName()),
+				TAttribute<FText>(), // no tooltip
+				FSlateIcon(),
+				FUIAction(FExecuteAction::CreateSP(this, &STimingView::ToggleTrackVisibility_Execute, Track.GetId()),
+					FCanExecuteAction::CreateLambda([] { return true; }),
+					FIsActionChecked::CreateSP(this, &STimingView::ToggleTrackVisibility_IsChecked, Track.GetId())),
+				NAME_None,
+				EUserInterfaceActionType::ToggleButton
+			);
+		}
+	}
+}
+
+////////////////////////////////////////////////////////////////////////////////////////////////////
+
+bool STimingView::ShowHideAllCpuTracks_IsChecked() const
+{
+	return bShowHideAllCpuTracks;
+}
+
+////////////////////////////////////////////////////////////////////////////////////////////////////
+
+void STimingView::ShowHideAllCpuTracks_Execute()
+{
+	bShowHideAllCpuTracks = !bShowHideAllCpuTracks;
+
+	for (int32 TrackIndex = 0; TrackIndex < TimingEventsTracks.Num(); ++TrackIndex)
+	{
+		FTimingEventsTrack& Track = *TimingEventsTracks[TrackIndex];
+		if (Track.GetType() == ETimingEventsTrackType::Cpu)
+		{
+			Track.SetVisibilityFlag(bShowHideAllCpuTracks);
+		}
+	}
+
+	for (auto& KV: ThreadGroups)
+	{
+		KV.Value.bIsVisible = bShowHideAllCpuTracks;
+	}
+
+	HoveredTimingEvent.Reset();
+	SelectedTimingEvent.Reset();
+
+	bAreTimingEventsTracksDirty = true;
+}
+
+////////////////////////////////////////////////////////////////////////////////////////////////////
+
+bool STimingView::ShowHideAllGpuTracks_IsChecked() const
+{
+	return bShowHideAllGpuTracks;
+}
+
+////////////////////////////////////////////////////////////////////////////////////////////////////
+
+void STimingView::ShowHideAllGpuTracks_Execute()
+{
+	bShowHideAllGpuTracks = !bShowHideAllGpuTracks;
+
+	for (int32 TrackIndex = 0; TrackIndex < TimingEventsTracks.Num(); ++TrackIndex)
+	{
+		FTimingEventsTrack& Track = *TimingEventsTracks[TrackIndex];
+		if (Track.GetType() == ETimingEventsTrackType::Gpu)
+		{
+			Track.SetVisibilityFlag(bShowHideAllGpuTracks);
+		}
+	}
+
+	HoveredTimingEvent.Reset();
+	SelectedTimingEvent.Reset();
+
+	bAreTimingEventsTracksDirty = true;
+}
+
+////////////////////////////////////////////////////////////////////////////////////////////////////
+
+bool STimingView::ShowHideAllLoadingTracks_IsChecked() const
+{
+	return bShowHideAllLoadingTracks;
+}
+
+////////////////////////////////////////////////////////////////////////////////////////////////////
+
+void STimingView::ShowHideAllLoadingTracks_Execute()
+{
+	bShowHideAllLoadingTracks = !bShowHideAllLoadingTracks;
+
+	for (int32 TrackIndex = 0; TrackIndex < TimingEventsTracks.Num(); ++TrackIndex)
+	{
+		FTimingEventsTrack& Track = *TimingEventsTracks[TrackIndex];
+		if (Track.GetType() == ETimingEventsTrackType::Loading)
+		{
+			Track.SetVisibilityFlag(bShowHideAllLoadingTracks);
+		}
+	}
+
+	bAreTimingEventsTracksDirty = true;
+}
+
+////////////////////////////////////////////////////////////////////////////////////////////////////
+
+bool STimingView::ShowHideAllIoTracks_IsChecked() const
+{
+	return bShowHideAllIoTracks;
+}
+
+////////////////////////////////////////////////////////////////////////////////////////////////////
+
+void STimingView::ShowHideAllIoTracks_Execute()
+{
+	bShowHideAllIoTracks = !bShowHideAllIoTracks;
+
+	for (int32 TrackIndex = 0; TrackIndex < TimingEventsTracks.Num(); ++TrackIndex)
+	{
+		FTimingEventsTrack& Track = *TimingEventsTracks[TrackIndex];
+		if (Track.GetType() == ETimingEventsTrackType::Io)
+		{
+			Track.SetVisibilityFlag(bShowHideAllIoTracks);
+		}
+	}
+
+	bAreTimingEventsTracksDirty = true;
+
+	if (bShowHideAllIoTracks)
+	{
+		bForceIoEventsUpdate = true;
+	}
+}
+
+////////////////////////////////////////////////////////////////////////////////////////////////////
+
+bool STimingView::IsAutoHideEmptyTracksEnabled() const
+{
+	return (Layout.TargetMinTimelineH == 0.0f);
+}
+
+////////////////////////////////////////////////////////////////////////////////////////////////////
+
+void STimingView::ToggleAutoHideEmptyTracks()
+{
+	if (Layout.TargetMinTimelineH == 0.0f)
+	{
+		Layout.TargetMinTimelineH = RealMinTimelineH;
+	}
+	else
+	{
+		Layout.TargetMinTimelineH = 0.0f;
+	}
+
+	Layout.MinTimelineH = Layout.TargetMinTimelineH; // no layout animation
+
+	for (auto& CachedTimelineKV : CachedTimelines)
+	{
+		CachedTimelineKV.Value->SetHeight(0.0f);
+	}
+
+	Viewport.ScrollPosY = 0.0f;
+	UpdateVerticalScrollBar();
+	bIsVerticalViewportDirty = true;
+}
+
+////////////////////////////////////////////////////////////////////////////////////////////////////
+
+bool STimingView::ToggleTrackVisibility_IsChecked(uint64 InTrackId) const
+{
+	if (CachedTimelines.Contains(InTrackId))
+	{
+		const FTimingEventsTrack* Track = CachedTimelines[InTrackId];
+		return Track->IsVisible();
+	}
+	return false;
+}
+
+////////////////////////////////////////////////////////////////////////////////////////////////////
+
+void STimingView::ToggleTrackVisibility_Execute(uint64 InTrackId)
+{
+	if (CachedTimelines.Contains(InTrackId))
+	{
+		FTimingEventsTrack* Track = CachedTimelines[InTrackId];
+		Track->ToggleVisibility();
+
+		HoveredTimingEvent.Reset();
+		SelectedTimingEvent.Reset();
+	}
+}
+
+////////////////////////////////////////////////////////////////////////////////////////////////////
+
+bool STimingView::ToggleTrackVisibilityByGroup_IsChecked(const TCHAR* InGroupName) const
+{
+	if (ThreadGroups.Contains(InGroupName))
+	{
+		const FThreadGroup& ThreadGroup = ThreadGroups[InGroupName];
+		return ThreadGroup.bIsVisible;
+	}
+	return false;
+}
+
+////////////////////////////////////////////////////////////////////////////////////////////////////
+
+void STimingView::ToggleTrackVisibilityByGroup_Execute(const TCHAR* InGroupName)
+{
+	if (ThreadGroups.Contains(InGroupName))
+	{
+		FThreadGroup& ThreadGroup = ThreadGroups[InGroupName];
+		ThreadGroup.bIsVisible = !ThreadGroup.bIsVisible;
+
+		for (auto& KV : CachedTimelines)
+		{
+			FTimingEventsTrack& Track = *KV.Value;
+			if (Track.GetType() == ETimingEventsTrackType::Cpu &&
+				Track.GetGroupName() == InGroupName)
+			{
+				Track.SetVisibilityFlag(ThreadGroup.bIsVisible);
+			}
+		}
+
+		HoveredTimingEvent.Reset();
+		SelectedTimingEvent.Reset();
+	}
+}
+
+////////////////////////////////////////////////////////////////////////////////////////////////////
+
+#undef LOCTEXT_NAMESPACE
