Copying //UE4/Dev-Core to Dev-Main (//UE4/Dev-Main)

diff --git a/Engine/Source/Developer/TraceInsights/Private/Insights/Widgets/STimingView.h b/Engine/Source/Developer/TraceInsights/Private/Insights/Widgets/STimingView.h
new file mode 100644
--- /dev/null
+++ b/Engine/Source/Developer/TraceInsights/Private/Insights/Widgets/STimingView.h
@@ -0,0 +1,545 @@
+// Copyright 1998-2019 Epic Games, Inc. All Rights Reserved.
+
+#pragma once
+
+#include "CoreMinimal.h"
+#include "Framework/MultiBox/MultiBoxBuilder.h"
+#include "Input/CursorReply.h"
+#include "Input/Reply.h"
+#include "Layout/Geometry.h"
+#include "TraceServices/AnalysisService.h"
+#include "Widgets/DeclarativeSyntaxSupport.h"
+#include "Widgets/SCompoundWidget.h"
+#include "Widgets/SWidget.h"
+
+// Insights
+#include "Insights/Common/FixedCircularBuffer.h"
+#include "Insights/ViewModels/GraphTrack.h"
+#include "Insights/ViewModels/MarkersTimingTrack.h"
+#include "Insights/ViewModels/TimerNode.h"
+#include "Insights/ViewModels/TimeRulerTrack.h"
+#include "Insights/ViewModels/TimingEvent.h"
+#include "Insights/ViewModels/TimingEventsTrack.h"
+#include "Insights/ViewModels/TimingTrackViewport.h"
+#include "Insights/ViewModels/TimingViewDrawHelper.h"
+
+class FTimingViewDrawHelper;
+class SScrollBar;
+
+/** The delegate to be invoked when the selection have been changed. */
+DECLARE_DELEGATE_TwoParams(FSelectionChangedDelegate, double /*StartTime*/, double /*EndTime*/);
+
+/** The delegate to be invoked when the timing event being hovered by the mouse has changed. Returns id of timing event or -1 (if no event is hovered). */
+DECLARE_DELEGATE_RetVal(int32, FHoveredEventChangedDelegate);
+
+/** The delegate to be invoked when the selected timing event has changed. Returns id of timing event or -1 (if no event is hovered). */
+DECLARE_DELEGATE_RetVal(int32, FSelectedEventChangedDelegate);
+
+/** Defines FLoadingTrackGetEventNameDelegate delegate interface. Returns the name for a timing event in a Loading track. */
+DECLARE_DELEGATE_RetVal_TwoParams(const TCHAR*, FLoadingTrackGetEventNameDelegate, uint32 /*Depth*/, const Trace::FLoadTimeProfilerCpuEvent& /*Event*/);
+
+/** A custom widget used to display timing events. */
+class STimingView : public SCompoundWidget
+{
+private:
+	static constexpr float MinTooltipWidth = 110.0f;
+
+public:
+	/** Default constructor. */
+	STimingView();
+
+	/** Virtual destructor. */
+	virtual ~STimingView();
+
+	SLATE_BEGIN_ARGS(STimingView)
+		: _OnSelectionChanged()
+		, _OnHoveredEventChanged()
+		, _OnSelectedEventChanged()
+		{
+			_Clipping = EWidgetClipping::ClipToBounds;
+		}
+
+		SLATE_EVENT(FSelectionChangedDelegate, OnSelectionChanged)
+		SLATE_EVENT(FHoveredEventChangedDelegate, OnHoveredEventChanged)
+		SLATE_EVENT(FSelectedEventChangedDelegate, OnSelectedEventChanged)
+	SLATE_END_ARGS()
+
+	/**
+	 * Construct this widget
+	 *
+	 * @param	InArgs	The declaration data for this widget
+	 */
+	void Construct(const FArguments& InArgs);
+
+	TSharedRef<SWidget> MakeTracksFilterMenu();
+	void CreateThreadGroupsMenu(FMenuBuilder& MenuBuilder);
+	void CreateTracksMenu(FMenuBuilder& MenuBuilder);
+
+	bool ShowHideAllCpuTracks_IsChecked() const;
+	void ShowHideAllCpuTracks_Execute();
+
+	bool ShowHideAllGpuTracks_IsChecked() const;
+	void ShowHideAllGpuTracks_Execute();
+
+	bool ShowHideAllLoadingTracks_IsChecked() const;
+	void ShowHideAllLoadingTracks_Execute();
+
+	bool ShowHideAllIoTracks_IsChecked() const;
+	void ShowHideAllIoTracks_Execute();
+
+	bool IsAutoHideEmptyTracksEnabled() const;
+	void ToggleAutoHideEmptyTracks();
+
+	bool ToggleTrackVisibility_IsChecked(uint64 InTrackId) const;
+	void ToggleTrackVisibility_Execute(uint64 InTrackId);
+
+	bool ToggleTrackVisibilityByGroup_IsChecked(const TCHAR* InGroupName) const;
+	void ToggleTrackVisibilityByGroup_Execute(const TCHAR* InGroupName);
+
+	void EnableAssetLoadingMode();
+
+	/** Resets internal widget's data to the default one. */
+	void Reset();
+
+	/**
+	 * Ticks this widget.  Override in derived classes, but always call the parent implementation.
+	 *
+	 * @param  AllottedGeometry The space allotted for this widget
+	 * @param  InCurrentTime  Current absolute real time
+	 * @param  InDeltaTime  Real time passed since last tick
+	 */
+	virtual void Tick(const FGeometry& AllottedGeometry, const double InCurrentTime, const float InDeltaTime) override;
+
+	virtual int32 OnPaint(const FPaintArgs& Args, const FGeometry& AllottedGeometry, const FSlateRect& MyCullingRect, FSlateWindowElementList& OutDrawElements, int32 LayerId, const FWidgetStyle& InWidgetStyle, bool bParentEnabled) const override;
+
+	/**
+	 * The system calls this method to notify the widget that a mouse button was pressed within it. This event is bubbled.
+	 *
+	 * @param MyGeometry The Geometry of the widget receiving the event
+	 * @param MouseEvent Information about the input event
+	 *
+	 * @return Whether the event was handled along with possible requests for the system to take action.
+	 */
+	virtual FReply OnMouseButtonDown(const FGeometry& MyGeometry, const FPointerEvent& MouseEvent) override;
+
+	/**
+	 * The system calls this method to notify the widget that a mouse button was release within it. This event is bubbled.
+	 *
+	 * @param MyGeometry The Geometry of the widget receiving the event
+	 * @param MouseEvent Information about the input event
+	 *
+	 * @return Whether the event was handled along with possible requests for the system to take action.
+	 */
+	virtual FReply OnMouseButtonUp(const FGeometry& MyGeometry, const FPointerEvent& MouseEvent) override;
+
+	/**
+	 * Called when a mouse button is double clicked.  Override this in derived classes.
+	 *
+	 * @param  InMyGeometry  Widget geometry
+	 * @param  InMouseEvent  Mouse button event
+	 *
+	 * @return  Returns whether the event was handled, along with other possible actions
+	 */
+	virtual FReply OnMouseButtonDoubleClick(const FGeometry& MyGeometry, const FPointerEvent& MouseEvent) override;
+
+	/**
+	 * The system calls this method to notify the widget that a mouse moved within it. This event is bubbled.
+	 *
+	 * @param MyGeometry The Geometry of the widget receiving the event
+	 * @param MouseEvent Information about the input event
+	 *
+	 * @return Whether the event was handled along with possible requests for the system to take action.
+	 */
+	virtual FReply OnMouseMove(const FGeometry& MyGeometry, const FPointerEvent& MouseEvent) override;
+
+	/**
+	 * The system will use this event to notify a widget that the cursor has entered it. This event is NOT bubbled.
+	 *
+	 * @param MyGeometry The Geometry of the widget receiving the event
+	 * @param MouseEvent Information about the input event
+	 */
+	virtual void OnMouseEnter(const FGeometry& MyGeometry, const FPointerEvent& MouseEvent) override;
+
+	/**
+	 * The system will use this event to notify a widget that the cursor has left it. This event is NOT bubbled.
+	 *
+	 * @param MouseEvent Information about the input event
+	 */
+	virtual void OnMouseLeave(const FPointerEvent& MouseEvent) override;
+
+	/**
+	 * Called when the mouse wheel is spun. This event is bubbled.
+	 *
+	 * @param  MouseEvent  Mouse event
+	 *
+	 * @return  Returns whether the event was handled, along with other possible actions
+	 */
+	virtual FReply OnMouseWheel(const FGeometry& MyGeometry, const FPointerEvent& MouseEvent) override;
+
+	/**
+	 * Called when the user is dropping something onto a widget; terminates drag and drop.
+	 *
+	 * @param MyGeometry      The geometry of the widget receiving the event.
+	 * @param DragDropEvent   The drag and drop event.
+	 *
+	 * @return A reply that indicated whether this event was handled.
+	 */
+	virtual FReply OnDrop(const FGeometry& MyGeometry, const FDragDropEvent& DragDropEvent) override;
+
+	/**
+	 * Called during drag and drop when the drag enters a widget.
+	 *
+	 * Enter/Leave events in slate are meant as lightweight notifications.
+	 * So we do not want to capture mouse or set focus in response to these.
+	 * However, OnDragEnter must also support external APIs (e.g. OLE Drag/Drop)
+	 * Those require that we let them know whether we can handle the content
+	 * being dragged OnDragEnter.
+	 *
+	 * The concession is to return a can_handled/cannot_handle
+	 * boolean rather than a full FReply.
+	 *
+	 * @param MyGeometry      The geometry of the widget receiving the event.
+	 * @param DragDropEvent   The drag and drop event.
+	 *
+	 * @return A reply that indicated whether the contents of the DragDropEvent can potentially be processed by this widget.
+	 */
+	virtual void OnDragEnter(const FGeometry& MyGeometry, const FDragDropEvent& DragDropEvent) override;
+
+	/**
+	 * Called during drag and drop when the drag leaves a widget.
+	 *
+	 * @param DragDropEvent   The drag and drop event.
+	 */
+	virtual void OnDragLeave(const FDragDropEvent& DragDropEvent) override;
+
+	/**
+	 * Called during drag and drop when the the mouse is being dragged over a widget.
+	 *
+	 * @param MyGeometry      The geometry of the widget receiving the event.
+	 * @param DragDropEvent   The drag and drop event.
+	 *
+	 * @return A reply that indicated whether this event was handled.
+	 */
+	virtual FReply OnDragOver(const FGeometry& MyGeometry, const FDragDropEvent& DragDropEvent) override;
+
+	/**
+	 * Called when the system wants to know which cursor to display for this Widget.  This event is bubbled.
+	 *
+	 * @return  The cursor requested (can be None.)
+	 */
+	virtual FCursorReply OnCursorQuery(const FGeometry& MyGeometry, const FPointerEvent& CursorEvent) const override;
+
+	virtual bool SupportsKeyboardFocus() const override { return true; }
+	virtual FReply OnKeyDown(const FGeometry& MyGeometry, const FKeyEvent& InKeyEvent) override;
+	virtual FReply OnKeyUp(const FGeometry& MyGeometry, const FKeyEvent& InKeyEvent) override;
+
+	////////////////////////////////////////////////////////////////////////////////////////////////////
+
+	const FTimingTrackViewport& GetViewport() { return Viewport; }
+
+	bool IsTimeSelected(double Time) const { return Time >= SelectionStartTime && Time < SelectionEndTime; }
+	bool IsTimeSelectedInclusive(double Time) const { return Time >= SelectionStartTime && Time <= SelectionEndTime; }
+
+	void ScrollAtTime(double StartTime);
+	void CenterOnTimeInterval(double IntervalStartTime, double IntervalDuration);
+	void BringIntoView(double StartTime, double EndTime);
+	void SelectTimeInterval(double IntervalStartTime, double IntervalDuration);
+	void SetAndCenterOnTimeMarker(double Time);
+	void SelectToTimeMarker(double Time);
+
+	bool AreTimeMarkersVisible() { return MarkersTrack.IsVisible(); }
+	void SetTimeMarkersVisible(bool bOnOff);
+	bool IsDrawOnlyBookmarksEnabled() { return MarkersTrack.IsBookmarksTrack(); }
+	void SetDrawOnlyBookmarks(bool bOnOff);
+
+	bool IsGpuTrackVisible() const;
+	bool IsCpuTrackVisible(uint32 ThreadId) const;
+
+protected:
+	virtual FVector2D ComputeDesiredSize(float) const override
+	{
+		return FVector2D(16.0f, 16.0f);
+	}
+
+	FTimingEventsTrack* AddTimingEventsTrack(uint64 TrackId, ETimingEventsTrackType TrackType, const FString& TrackName, const TCHAR* GroupName, int32 Order);
+
+	void DrawTimingProfilerTrack(FTimingViewDrawHelper& Helper, FTimingEventsTrack& Track) const;
+
+	const TCHAR* GetLoadTimeProfilerEventNameByPackageEventType(uint32 Depth, const Trace::FLoadTimeProfilerCpuEvent& Event) const;
+	const TCHAR* GetLoadTimeProfilerEventNameByExportEventType(uint32 Depth, const Trace::FLoadTimeProfilerCpuEvent& Event) const;
+	const TCHAR* GetLoadTimeProfilerEventNameByPackageName(uint32 Depth, const Trace::FLoadTimeProfilerCpuEvent& Event) const;
+	const TCHAR* GetLoadTimeProfilerEventNameByExportClassName(uint32 Depth, const Trace::FLoadTimeProfilerCpuEvent& Event) const;
+	const TCHAR* GetLoadTimeProfilerEventName(uint32 Depth, const Trace::FLoadTimeProfilerCpuEvent& Event) const;
+	void DrawLoadTimeProfilerTrack(FTimingViewDrawHelper& Helper, FTimingEventsTrack& Track) const;
+
+	void UpdateIo();
+	void DrawIoOverviewTrack(FTimingViewDrawHelper& Helper, FTimingEventsTrack& Track) const;
+	void DrawIoActivityTrack(FTimingViewDrawHelper& Helper, FTimingEventsTrack& Track) const;
+
+	void DrawTimeRangeSelection(FDrawContext& DH) const;
+
+	void ShowContextMenu(const FVector2D& ScreenSpacePosition, const FPointerEvent& MouseEvent);
+
+	/** Binds our UI commands to delegates. */
+	void BindCommands();
+
+	////////////////////////////////////////////////////////////////////////////////////////////////////
+
+	/**
+	 * Called when the user scrolls the horizontal scrollbar.
+	 *
+	 * @param ScrollOffset Scroll offset as a fraction between 0 and 1.
+	 */
+	void HorizontalScrollBar_OnUserScrolled(float ScrollOffset);
+	void UpdateHorizontalScrollBar();
+
+	/**
+	 * Called when the user scrolls the vertical scrollbar.
+	 *
+	 * @param ScrollOffset Scroll offset as a fraction between 0 and 1.
+	 */
+	void VerticalScrollBar_OnUserScrolled(float ScrollOffset);
+	void UpdateVerticalScrollBar();
+
+	////////////////////////////////////////////////////////////////////////////////////////////////////
+
+	void UpdateAggregatedStats();
+
+	void UpdateHoveredTimingEvent(float MX, float MY);
+
+	bool SearchTimingEvent(const double InStartTime,
+						   const double InEndTime,
+						   TFunctionRef<bool(double, double, uint32, uint32)> InPredicate,
+						   FTimingEvent& InOutTimingEvent,
+						   bool bInStopAtFirstMatch,
+						   bool bInSearchForLargestEvent) const;
+
+	void OnSelectedTimingEventChanged();
+	void SelectHoveredTimingEvent();
+	void SelectLeftTimingEvent();
+	void SelectRightTimingEvent();
+	void SelectUpTimingEvent();
+	void SelectDownTimingEvent();
+
+	void FrameSelection();
+
+	const FTimerNodePtr GetTimerNode(uint64 TypeId) const;
+
+	////////////////////////////////////////////////////////////////////////////////////////////////////
+	// FrameSelectionChanged Event
+
+public:
+	/**
+	 * The event to execute when the selected frames have been changed.
+	 *
+	 * @param FrameStartIndex	- The index of the first frame selected.
+	 * @param FrameEndIndex		- The index of the last frame selected.
+	 *
+	 */
+	DECLARE_EVENT_TwoParams(STimingView, FFrameSelectionChangedEvent, int32 /*FrameStartIndex*/, int32 /*FrameEndIndex*/);
+	FFrameSelectionChangedEvent& OnFrameSelectionChanged()
+	{
+		return FrameSelectionChangedEvent;
+	}
+
+protected:
+	/** The event to execute when the selected frames have been changed. */
+	FFrameSelectionChangedEvent FrameSelectionChangedEvent;
+
+	////////////////////////////////////////////////////////////////////////////////////////////////////
+
+protected:
+	bool bShowHideAllGpuTracks;
+	bool bShowHideAllCpuTracks;
+	bool bShowHideAllLoadingTracks;
+	bool bShowHideAllIoTracks;
+
+	////////////////////////////////////////////////////////////
+
+	/** The track's viewport. Encapsulates info about position and scale. */
+	FTimingTrackViewport Viewport;
+
+	bool bIsViewportDirty;
+	bool bIsVerticalViewportDirty;
+
+	////////////////////////////////////////////////////////////
+
+	/** All created tracks.
+	  * Maps track id to track pointer.
+	  */
+	TMap<uint64, FTimingEventsTrack*> CachedTimelines;
+
+	//TODO: TArray<FBaseTimingTrack*> TopTracks; /**< tracks docked on top, in order to be displayed (top to bottom) */
+	//TODO: TArray<FBaseTimingTrack*> ScrollableTracks; /**< tracks in scrollable area, in order to be displayed (top to bottom) */
+	//TODO: TArray<FBaseTimingTrack*> BottomTracks; /**< tracks docked on bottom, in order to be displayed (top to bottom) */
+	//TODO: TArray<FBaseTimingTrack*> ForegorundTracks; /**< tracks to draw over top/scrollable/bottom tracks (can use entire area), in order to be displayed (back to front) */
+
+	////////////////////////////////////////////////////////////
+
+	TArray<FTimingEventsTrack*> TimingEventsTracks; /**< all timing events tracks in order to be displayed */
+
+	bool bAreTimingEventsTracksDirty;
+
+	bool bUseDownSampling;
+
+	////////////////////////////////////////////////////////////
+	// Cpu/Gpu
+
+	FTimingEventsTrack* GpuTrack;
+	TMap<uint32, FTimingEventsTrack*> CpuTracks; /**< maps thread id to track pointer */
+
+	struct FThreadGroup
+	{
+		const TCHAR* Name; /**< the thread group name; pointer to string owned by ThreadProvider */
+		bool bIsVisible;
+		uint32 NumTimelines;
+	};
+
+	TMap<const TCHAR*, FThreadGroup> ThreadGroups; /**< maps thread group name to thread group info */
+
+	////////////////////////////////////////////////////////////
+	// Asset Loading
+
+	FTimingEventsTrack* LoadingMainThreadTrack;
+	FTimingEventsTrack* LoadingAsyncThreadTrack;
+
+	uint32 LoadingMainThreadId;
+	uint32 LoadingAsyncThreadId;
+
+	FLoadingTrackGetEventNameDelegate LoadingGetEventNameFn;
+
+	bool bAssetLoadingMode;
+
+	FString EventAggregationStr;
+	FString ObjectTypeAggregationStr;
+
+	////////////////////////////////////////////////////////////
+	// File activity (I/O)
+
+	FTimingEventsTrack* IoOverviewTrack;
+	FTimingEventsTrack* IoActivityTrack;
+
+	bool bForceIoEventsUpdate;
+
+	bool bMergeIoLanes;
+
+	struct FIoTimingEvent
+	{
+		double StartTime;
+		double EndTime;
+		uint32 Depth;
+		uint32 Type; // Trace::EFileActivityType + Failed
+		const TCHAR* Path;
+	};
+
+	/** All IO events, cached. */
+	TArray<FIoTimingEvent> AllIoEvents;
+
+	////////////////////////////////////////////////////////////
+
+	/** The time ruler track. */
+	FTimeRulerTrack TimeRulerTrack;
+
+	/** The time markers track. */
+	FMarkersTimingTrack MarkersTrack;
+
+	/** A graph track for frame times. */
+	FFramesGraphTrack GraphTrack;
+
+	////////////////////////////////////////////////////////////
+
+	/** Horizontal scroll bar, used for scrolling timing events' viewport. */
+	TSharedPtr<SScrollBar> HorizontalScrollBar;
+
+	/** Vertical scroll bar, used for scrolling timing events' viewport. */
+	TSharedPtr<SScrollBar> VerticalScrollBar;
+
+	////////////////////////////////////////////////////////////
+
+	/** The current mouse position. */
+	FVector2D MousePosition;
+
+	/** Mouse position during the call on mouse button down. */
+	FVector2D MousePositionOnButtonDown;
+	double ViewportStartTimeOnButtonDown;
+	float ViewportScrollPosYOnButtonDown;
+
+	/** Mouse position during the call on mouse button up. */
+	FVector2D MousePositionOnButtonUp;
+
+	bool bIsLMB_Pressed;
+	bool bIsRMB_Pressed;
+
+	bool bIsSpaceBarKeyPressed;
+	bool bIsDragging;
+
+	////////////////////////////////////////////////////////////
+	// Panning
+
+	/**
+	 * True, if the user is currently interactively panning the view (horizontally and/or vertically),
+	 * either by holding the right mouse button and dragging
+	 * or by holding spacebar pressed and dragging with left mouse button.
+	 */
+	bool bIsPanning;
+
+	/** How to pan. */
+	enum class EPanningMode : uint8
+	{
+		None = 0,
+		Horizontal = 0x01,
+		Vertical = 0x02,
+		HorizontalAndVertical = Horizontal | Vertical,
+	};
+	EPanningMode PanningMode;
+
+	////////////////////////////////////////////////////////////
+	// Selection
+
+	/** True, if the user is currently changing the selection (by holding the left mouse button and dragging). */
+	bool bIsSelecting;
+
+	double SelectionStartTime;
+	double SelectionEndTime;
+
+	mutable float TooltipWidth;
+	mutable float TooltipAlpha;
+
+	FTimingEvent HoveredTimingEvent;
+	FTimingEvent SelectedTimingEvent;
+
+	enum class ESelectionType
+	{
+		None,
+		TimeRange,
+		TimingEvent
+	};
+	ESelectionType LastSelectionType;
+
+	double TimeMarker;
+
+	////////////////////////////////////////////////////////////
+	// Misc
+
+	FGeometry ThisGeometry;
+
+	const FSlateBrush* WhiteBrush;
+	const FSlateFontInfo MainFont;
+
+	FTimingEventsTrackLayout Layout;
+
+	// Debug stats
+	int32 NumUpdatedEvents;
+	TFixedCircularBuffer<uint64, 32> TimelineCacheUpdateDurationHistory;
+	TFixedCircularBuffer<uint64, 32> TimeMarkerCacheUpdateDurationHistory;
+	mutable TFixedCircularBuffer<uint64, 32> DrawDurationHistory;
+	mutable TFixedCircularBuffer<uint64, 32> OnPaintDurationHistory;
+	mutable uint64 LastOnPaintTime;
+
+	////////////////////////////////////////////////////////////
+	// Delegates
+
+	FSelectionChangedDelegate OnSelectionChanged;
+	FHoveredEventChangedDelegate OnHoveredEventChanged;
+	FSelectedEventChangedDelegate OnSelectedEventChanged;
+};
diff --git a/Engine/Source/Developer/TraceInsights/Public/Insights/IUnrealInsightsModule.h b/Engine/Source/Developer/TraceInsights/Public/Insights/IUnrealInsightsModule.h
new file mode 100644
--- /dev/null
+++ b/Engine/Source/Developer/TraceInsights/Public/Insights/IUnrealInsightsModule.h
@@ -0,0 +1,26 @@
+// Copyright 1998-2019 Epic Games, Inc. All Rights Reserved.
+
+#pragma once
+
+#include "CoreMinimal.h"
+#include "Modules/ModuleInterface.h"
+#include "Framework/Docking/TabManager.h"
+
+/** Interface for an Unreal Insights module. */
+class IUnrealInsightsModule : public IModuleInterface
+{
+public:
+	/**
+	 * Called when application starts and a new layout is created. It allows the module to create its own areas.
+	 *
+	 * @param NewLayout The newly created layout.
+	 */
+	virtual void OnNewLayout(TSharedRef<FTabManager::FLayout> NewLayout) = 0;
+
+	/**
+	 * Called after application layout was restored.
+	 *
+	 * @param TabManager The tab manager used by Unreal Insights tools.
+	 */
+	virtual void OnLayoutRestored(TSharedPtr<FTabManager> TabManager) = 0;
+};
diff --git a/Engine/Source/Developer/TraceInsights/Public/Insights/TimingProfilerCommon.h b/Engine/Source/Developer/TraceInsights/Public/Insights/TimingProfilerCommon.h
new file mode 100644
--- /dev/null
+++ b/Engine/Source/Developer/TraceInsights/Public/Insights/TimingProfilerCommon.h
@@ -0,0 +1,32 @@
+// Copyright 1998-2019 Epic Games, Inc. All Rights Reserved.
+
+#pragma once
+
+#include "CoreMinimal.h"
+#include "Logging/LogMacros.h"
+#include "Stats/Stats.h"
+#include "Stats/StatsMisc.h"
+
+DECLARE_LOG_CATEGORY_EXTERN(TimingProfiler, Log, All);
+
+#define DEBUG_TIMING_PROFILER_PERFORMANCE 0
+
+#if DEBUG_TIMING_PROFILER_PERFORMANCE==1
+#define TP_SCOPE_LOG_TIME(arg0,arg1) SCOPE_LOG_TIME(arg0,arg1)
+#else
+#define TP_SCOPE_LOG_TIME(arg0,arg1)
+#endif
+
+DECLARE_STATS_GROUP(TEXT("TimingProfiler"), STATGROUP_TimingProfiler, STATCAT_Advanced);
+
+/** Time spent on Frame Track drawing. */
+DECLARE_CYCLE_STAT_EXTERN(TEXT("FrameTrackOnPaint"), STAT_FT_OnPaint, STATGROUP_TimingProfiler, );
+
+/** Time spent on Graph Track drawing. */
+DECLARE_CYCLE_STAT_EXTERN(TEXT("GraphTrackOnPaint"), STAT_GT_OnPaint, STATGROUP_TimingProfiler, );
+
+/** Time spent on Timing Track drawing. */
+DECLARE_CYCLE_STAT_EXTERN(TEXT("TimingTrackOnPaint"), STAT_TT_OnPaint, STATGROUP_TimingProfiler, );
+
+/** Time spent on ticking profiler manager. */
+DECLARE_CYCLE_STAT_EXTERN(TEXT("TimingProfilerTick"), STAT_TPM_Tick, STATGROUP_TimingProfiler, );
diff --git a/Engine/Source/Developer/TraceInsights/Public/Insights/Version.h b/Engine/Source/Developer/TraceInsights/Public/Insights/Version.h
new file mode 100644
--- /dev/null
+++ b/Engine/Source/Developer/TraceInsights/Public/Insights/Version.h
@@ -0,0 +1,20 @@
+// Copyright 1998-2019 Epic Games, Inc. All Rights Reserved.
+
+#pragma once
+
+// See https://semver.org/
+
+#define UNREAL_INSIGHTS_VERSION_MAJOR 0
+#define UNREAL_INSIGHTS_VERSION_MINOR 6
+#define UNREAL_INSIGHTS_VERSION_PATCH 0
+
+#define UNREAL_INSIGHTS_VERSION_STRING "0.06"
+
+/** Extra identifier string (like: "alpha", "beta.1", "beta.2", "test", "featureX", etc.). */
+//#define UNREAL_INSIGHTS_VERSION_ID_STRING "alpha"
+
+#ifdef UNREAL_INSIGHTS_VERSION_ID_STRING
+#define UNREAL_INSIGHTS_VERSION_STRING_EX "v" UNREAL_INSIGHTS_VERSION_STRING "-" UNREAL_INSIGHTS_VERSION_ID_STRING
+#else
+#define UNREAL_INSIGHTS_VERSION_STRING_EX "v" UNREAL_INSIGHTS_VERSION_STRING
+#endif
diff --git a/Engine/Source/Developer/TraceServices/Private/AnalysisService.cpp b/Engine/Source/Developer/TraceServices/Private/AnalysisService.cpp
new file mode 100644
--- /dev/null
+++ b/Engine/Source/Developer/TraceServices/Private/AnalysisService.cpp
@@ -0,0 +1,187 @@
+// Copyright 1998-2019 Epic Games, Inc. All Rights Reserved.
+
+#include "TraceServices/AnalysisService.h"
+#include "AnalysisServicePrivate.h"
+#include "Trace/Analyzer.h"
+#include "Trace/Analysis.h"
+#include "Trace/DataStream.h"
+#include "HAL/PlatformFile.h"
+#include "Analyzers/MiscTraceAnalysis.h"
+#include "Analyzers/LogTraceAnalysis.h"
+#include "Math/RandomStream.h"
+#include "ModuleServicePrivate.h"
+#include "Model/LogPrivate.h"
+#include "Model/FramesPrivate.h"
+#include "Model/BookmarksPrivate.h"
+#include "Model/ThreadsPrivate.h"
+
+namespace Trace
+{
+
+void FAnalysisSessionLock::ReadAccessCheck() const
+{
+	checkf(IsReadOnly || FPlatformTLS::GetCurrentThreadId() == OwnerThread, TEXT("Trying to read from session while someone else is writing"));
+}
+
+void FAnalysisSessionLock::WriteAccessCheck() const
+{
+	checkf(!IsReadOnly, TEXT("Trying to edit session while in read only mode"));
+	checkf(FPlatformTLS::GetCurrentThreadId() == OwnerThread, TEXT("Trying to edit session from thread without write access"));
+}
+
+void FAnalysisSessionLock::BeginRead()
+{
+	CriticalSection.Lock();
+	check(OwnerThread == 0);
+	OwnerThread = FPlatformTLS::GetCurrentThreadId();
+	IsReadOnly = true;
+}
+
+void FAnalysisSessionLock::EndRead()
+{
+	check(FPlatformTLS::GetCurrentThreadId() == OwnerThread);
+	OwnerThread = 0;
+	IsReadOnly = false;
+	CriticalSection.Unlock();
+}
+
+void FAnalysisSessionLock::BeginEdit()
+{
+	CriticalSection.Lock();
+	check(OwnerThread == 0);
+	check(!IsReadOnly);
+	OwnerThread = FPlatformTLS::GetCurrentThreadId();
+}
+
+void FAnalysisSessionLock::EndEdit()
+{
+	check(OwnerThread == FPlatformTLS::GetCurrentThreadId());
+	OwnerThread = 0;
+	CriticalSection.Unlock();
+}
+
+FAnalysisSession::FAnalysisSession(const TCHAR* SessionName)
+	: Name(SessionName)
+	, DurationSeconds(0.0)
+	, Allocator(32 << 20)
+	, StringStore(Allocator)
+{
+
+}
+
+FAnalysisSession::~FAnalysisSession()
+{
+	for (auto& KV : Providers)
+	{
+		IProvider* Provider = KV.Value;
+		delete Provider;
+	}
+}
+
+void FAnalysisSession::AddProvider(const FName& InName, IProvider* Provider)
+{
+	Providers.Add(InName, Provider);
+}
+
+const IProvider* FAnalysisSession::ReadProviderPrivate(const FName& InName) const
+{
+	IProvider* const* FindIt = Providers.Find(InName);
+	if (FindIt)
+	{
+		return *FindIt;
+	}
+	else
+	{
+		return nullptr;
+	}
+}
+
+FAnalysisService::FAnalysisService(FModuleService& InModuleService)
+	: ModuleService(InModuleService)
+{
+
+}
+
+FAnalysisService::FAnalysisWorker::FAnalysisWorker(FAnalysisService& InOuter, TUniquePtr<Trace::IInDataStream>&& InDataStream, TSharedRef<FAnalysisSession> InAnalysisSession)
+	: Outer(InOuter)
+	, DataStream(MoveTemp(InDataStream))
+	, AnalysisSession(InAnalysisSession)
+{
+	
+}
+
+void FAnalysisService::AnalyzeInternal(TSharedRef<FAnalysisSession> AnalysisSession, Trace::IInDataStream* DataStream)
+{
+	OnAnalysisStarted().Broadcast(AnalysisSession);
+
+	TArray<IAnalyzer*> Analyzers;
+	{
+		IAnalysisSession& Session = AnalysisSession.Get();
+		Trace::FAnalysisSessionEditScope _(Session);
+		FBookmarkProvider* BookmarkProvider = new FBookmarkProvider(Session);
+		Session.AddProvider(FBookmarkProvider::ProviderName, BookmarkProvider);
+		FLogProvider* LogProvider = new FLogProvider(Session);
+		Session.AddProvider(FLogProvider::ProviderName, LogProvider);
+		FThreadProvider* ThreadProvider = new FThreadProvider(Session);
+		Session.AddProvider(FThreadProvider::ProviderName, ThreadProvider);
+		FFrameProvider* FrameProvider = new FFrameProvider(Session);
+		Session.AddProvider(FFrameProvider::ProviderName, FrameProvider);
+		Analyzers.Add(new FMiscTraceAnalyzer(Session, *ThreadProvider, *BookmarkProvider, *LogProvider, *FrameProvider));
+		Analyzers.Add(new FLogTraceAnalyzer(Session, *LogProvider));
+		ModuleService.OnAnalysisBegin(Session, Analyzers);
+	}
+	
+	FAnalysisContext Context;
+	for (Trace::IAnalyzer* Analyzer : Analyzers)
+	{
+		Context.AddAnalyzer(*Analyzer);
+	}
+	Trace::FAnalysisProcessor Processor = Context.Process();
+
+	Processor.Start(*DataStream);
+
+	for (Trace::IAnalyzer* Analyzer : Analyzers)
+	{
+		delete Analyzer;
+	}
+	Analyzers.Empty();
+
+	delete DataStream;
+
+	AnalysisSession->SetComplete();
+
+	OnAnalysisFinished().Broadcast(AnalysisSession);
+}
+
+void FAnalysisService::FAnalysisWorker::DoWork()
+{
+	Outer.AnalyzeInternal(AnalysisSession.ToSharedRef(), DataStream.Release());
+	AnalysisSession = nullptr;
+}
+
+FAnalysisService::~FAnalysisService()
+{
+	for (TSharedPtr<FAsyncTask<FAnalysisWorker>> Task : Tasks)
+	{
+		Task->EnsureCompletion();
+	}
+}
+
+TSharedPtr<const IAnalysisSession> FAnalysisService::Analyze(const TCHAR* SessionName, TUniquePtr<Trace::IInDataStream>&& InDataStream)
+{
+	TSharedRef<FAnalysisSession> AnalysisSession = MakeShared<FAnalysisSession>(SessionName);
+	TUniquePtr<Trace::IInDataStream> DataStream = MoveTemp(InDataStream);
+	AnalyzeInternal(AnalysisSession, DataStream.Release());
+	return AnalysisSession;
+}
+
+TSharedPtr<const IAnalysisSession> FAnalysisService::StartAnalysis(const TCHAR* SessionName, TUniquePtr<Trace::IInDataStream>&& DataStream)
+{
+	TSharedRef<FAnalysisSession> AnalysisSession = MakeShared<FAnalysisSession>(SessionName);
+	TSharedPtr<FAsyncTask<FAnalysisWorker>> Task = MakeShared<FAsyncTask<FAnalysisWorker>>(*this, MoveTemp(DataStream), AnalysisSession);
+	Tasks.Add(Task);
+	Task->StartBackgroundTask();
+	return AnalysisSession;
+}
+
+}
diff --git a/Engine/Source/Developer/TraceServices/Private/AnalysisServicePrivate.h b/Engine/Source/Developer/TraceServices/Private/AnalysisServicePrivate.h
new file mode 100644
--- /dev/null
+++ b/Engine/Source/Developer/TraceServices/Private/AnalysisServicePrivate.h
@@ -0,0 +1,110 @@
+// Copyright 1998-2019 Epic Games, Inc. All Rights Reserved.
+
+#pragma once
+
+#include "TraceServices/AnalysisService.h"
+#include "Async/AsyncWork.h"
+#include "Containers/HashTable.h"
+#include "Misc/ScopeLock.h"
+#include "Common/StringStore.h"
+
+namespace Trace
+{
+
+class FModuleService;
+
+class FAnalysisSessionLock
+{
+public:
+	void ReadAccessCheck() const;
+	void WriteAccessCheck() const;
+
+	void BeginRead();
+	void EndRead();
+
+	void BeginEdit();
+	void EndEdit();
+
+private:
+	FCriticalSection CriticalSection;
+	uint32 OwnerThread = 0;
+	bool IsReadOnly = false;
+};
+
+class FAnalysisSession
+	: public IAnalysisSession
+{
+public:
+	FAnalysisSession(const TCHAR* SessionName);
+	virtual ~FAnalysisSession();
+
+	virtual const TCHAR* GetName() const { return *Name; }
+	virtual bool IsAnalysisComplete() const override { return IsComplete; }
+	virtual double GetDurationSeconds() const { Lock.ReadAccessCheck(); return DurationSeconds; }
+	virtual void UpdateDurationSeconds(double Duration) override { Lock.WriteAccessCheck(); DurationSeconds = FMath::Max(Duration, DurationSeconds); }
+	void SetComplete() { IsComplete = true; }
+
+	virtual ILinearAllocator& GetLinearAllocator() override { return Allocator; }
+	virtual const TCHAR* StoreString(const TCHAR* String) override { return StringStore.Store(String); }
+
+	virtual void BeginRead() const override { Lock.BeginRead(); }
+	virtual void EndRead() const override { Lock.EndRead(); }
+
+	virtual void BeginEdit() override { Lock.BeginEdit(); }
+	virtual void EndEdit() override { Lock.EndEdit(); }
+
+	virtual void ReadAccessCheck() const override { return Lock.ReadAccessCheck(); }
+	virtual void WriteAccessCheck() override { return Lock.WriteAccessCheck(); }
+
+	virtual void AddProvider(const FName& Name, IProvider* Provider) override;
+
+private:
+	virtual const IProvider* ReadProviderPrivate(const FName& Name) const override;
+
+	mutable FAnalysisSessionLock Lock;
+
+	FString Name;
+	bool IsComplete = false;
+	double DurationSeconds = 0.0;
+	FSlabAllocator Allocator;
+	FStringStore StringStore;
+	TMap<FName, IProvider*> Providers;
+};
+
+class FAnalysisService
+	: public IAnalysisService
+{
+public:
+	FAnalysisService(FModuleService& ModuleService);
+	virtual ~FAnalysisService();
+	virtual TSharedPtr<const IAnalysisSession> Analyze(const TCHAR* SessionName, TUniquePtr<Trace::IInDataStream>&& DataStream) override;
+	virtual TSharedPtr<const IAnalysisSession> StartAnalysis(const TCHAR* SessionName, TUniquePtr<Trace::IInDataStream>&& DataStream) override;
+	virtual FAnalysisStartedEvent& OnAnalysisStarted() override { return AnalysisStartedEvent; }
+	virtual FAnalysisFinishedEvent& OnAnalysisFinished() override { return AnalysisFinishedEvent; }
+
+private:
+	class FAnalysisWorker : public FNonAbandonableTask
+	{
+	public:
+		FAnalysisWorker(FAnalysisService& Outer, TUniquePtr<Trace::IInDataStream>&& InDataStream, TSharedRef<FAnalysisSession> InAnalysisSession);
+		void DoWork();
+		FORCEINLINE TStatId GetStatId() const
+		{
+			RETURN_QUICK_DECLARE_CYCLE_STAT(FAnalysisWorker, STATGROUP_ThreadPoolAsyncTasks);
+		}
+
+	private:
+		FAnalysisService& Outer;
+		TUniquePtr<Trace::IInDataStream> DataStream;
+		TSharedPtr<FAnalysisSession> AnalysisSession;
+	};
+
+	void AnalyzeInternal(TSharedRef<FAnalysisSession> AnalysisSession, Trace::IInDataStream* DataStream);
+
+	FModuleService& ModuleService;
+	FAnalysisStartedEvent AnalysisStartedEvent;
+	FAnalysisFinishedEvent AnalysisFinishedEvent;
+	TArray<TSharedPtr<FAsyncTask<FAnalysisWorker>>> Tasks;
+};
+
+}
diff --git a/Engine/Source/Developer/TraceServices/Private/Analyzers/CpuProfilerTraceAnalysis.cpp b/Engine/Source/Developer/TraceServices/Private/Analyzers/CpuProfilerTraceAnalysis.cpp
new file mode 100644
--- /dev/null
+++ b/Engine/Source/Developer/TraceServices/Private/Analyzers/CpuProfilerTraceAnalysis.cpp
@@ -0,0 +1,123 @@
+// Copyright 1998-2019 Epic Games, Inc. All Rights Reserved.
+#include "CpuProfilerTraceAnalysis.h"
+#include "AnalysisServicePrivate.h"
+#include "Common/Utils.h"
+#include "TraceServices/Model/Threads.h"
+
+FCpuProfilerAnalyzer::FCpuProfilerAnalyzer(Trace::IAnalysisSession& InSession, Trace::FTimingProfilerProvider& InTimingProfilerProvider)
+	: Session(InSession)
+	, TimingProfilerProvider(InTimingProfilerProvider)
+{
+
+}
+
+void FCpuProfilerAnalyzer::OnAnalysisBegin(const FOnAnalysisContext& Context)
+{
+	auto& Builder = Context.InterfaceBuilder;
+
+	Builder.RouteEvent(RouteId_EventSpec, "CpuProfiler", "EventSpec");
+	Builder.RouteEvent(RouteId_EventBatch, "CpuProfiler", "EventBatch");
+	Builder.RouteEvent(RouteId_EndCapture, "CpuProfiler", "EndCapture");
+}
+
+void FCpuProfilerAnalyzer::OnAnalysisEnd()
+{
+}
+
+void FCpuProfilerAnalyzer::OnEvent(uint16 RouteId, const FOnEventContext& Context)
+{
+	Trace::FAnalysisSessionEditScope _(Session);
+
+	const auto& EventData = Context.EventData;
+	switch (RouteId)
+	{
+	case RouteId_EventSpec:
+	{
+		uint16 Id = EventData.GetValue("Id").As<uint16>();
+		if (ScopeIdToEventIdMap.Contains(Id))
+		{
+			TimingProfilerProvider.SetTimerName(ScopeIdToEventIdMap[Id], reinterpret_cast<const TCHAR*>(EventData.GetAttachment()));
+		}
+		else
+		{
+			ScopeIdToEventIdMap.Add(Id, TimingProfilerProvider.AddCpuTimer(reinterpret_cast<const TCHAR*>(EventData.GetAttachment())));
+		}
+		break;
+	}
+	case RouteId_EventBatch:
+	case RouteId_EndCapture:
+	{
+		TotalEventSize += EventData.GetTotalSize();
+		uint32 ThreadId = EventData.GetValue("ThreadId").As<uint32>();
+		TSharedRef<FThreadState> ThreadState = GetThreadState(ThreadId);
+		uint64 LastCycle = ThreadState->LastCycle;
+		uint64 BufferSize = EventData.GetAttachmentSize();
+		const uint8* BufferPtr = EventData.GetAttachment();
+		const uint8* BufferEnd = BufferPtr + BufferSize;
+		while (BufferPtr < BufferEnd)
+		{
+			uint64 DecodedCycle = FTraceAnalyzerUtils::Decode7bit(BufferPtr);
+			uint64 ActualCycle = (DecodedCycle >> 1) + LastCycle;
+			LastCycle = ActualCycle;
+			if (DecodedCycle & 1ull)
+			{
+				EventScopeState& ScopeState = ThreadState->ScopeStack.AddDefaulted_GetRef();
+				ScopeState.StartCycle = ActualCycle;
+				uint16 SpecId = FTraceAnalyzerUtils::Decode7bit(BufferPtr);
+				uint32* FindIt = ScopeIdToEventIdMap.Find(SpecId);
+				if (!FindIt)
+				{
+					ScopeState.EventTypeId = ScopeIdToEventIdMap.Add(SpecId, TimingProfilerProvider.AddCpuTimer(TEXT("<unknown>")));
+				}
+				else
+				{
+					ScopeState.EventTypeId = *FindIt;
+				}
+				Trace::FTimingProfilerEvent Event;
+				Event.TimerIndex = ScopeState.EventTypeId;
+				ThreadState->Timeline->AppendBeginEvent(Context.SessionContext.TimestampFromCycle(ScopeState.StartCycle), Event);
+				++TotalScopeCount;
+			}
+			else if (ThreadState->ScopeStack.Num())
+			{
+				ThreadState->ScopeStack.Pop();
+				ThreadState->Timeline->AppendEndEvent(Context.SessionContext.TimestampFromCycle(ActualCycle));
+			}
+		}
+		check(BufferPtr == BufferEnd);
+		if (LastCycle)
+		{
+			double LastTimestamp = Context.SessionContext.TimestampFromCycle(LastCycle);
+			Session.UpdateDurationSeconds(LastTimestamp);
+			if (RouteId == RouteId_EndCapture)
+			{
+
+				while (ThreadState->ScopeStack.Num())
+				{
+					ThreadState->ScopeStack.Pop();
+					ThreadState->Timeline->AppendEndEvent(LastTimestamp);
+				}
+			}
+		}
+		ThreadState->LastCycle = LastCycle;
+		BytesPerScope = double(TotalEventSize) / double(TotalScopeCount);
+		break;
+	}
+	}
+}
+
+TSharedRef<FCpuProfilerAnalyzer::FThreadState> FCpuProfilerAnalyzer::GetThreadState(uint32 ThreadId)
+{
+	if (!ThreadStatesMap.Contains(ThreadId))
+	{
+		TSharedRef<FThreadState> ThreadState = MakeShared<FThreadState>();
+		ThreadState->Timeline = &TimingProfilerProvider.EditCpuThreadTimeline(ThreadId);
+		ThreadStatesMap.Add(ThreadId, ThreadState);
+		return ThreadState;
+	}
+	else
+	{
+		return ThreadStatesMap[ThreadId];
+	}
+}
+
diff --git a/Engine/Source/Developer/TraceServices/Private/Analyzers/CpuProfilerTraceAnalysis.h b/Engine/Source/Developer/TraceServices/Private/Analyzers/CpuProfilerTraceAnalysis.h
new file mode 100644
--- /dev/null
+++ b/Engine/Source/Developer/TraceServices/Private/Analyzers/CpuProfilerTraceAnalysis.h
@@ -0,0 +1,54 @@
+// Copyright 1998-2019 Epic Games, Inc. All Rights Reserved.
+
+#pragma once
+
+#include "Trace/Trace.h"
+#include "Trace/Analyzer.h"
+#include "Containers/UnrealString.h"
+#include "Model/TimingProfilerPrivate.h"
+
+namespace Trace
+{
+	class IAnalysisSession;
+}
+
+class FCpuProfilerAnalyzer
+	: public Trace::IAnalyzer
+{
+public:
+	FCpuProfilerAnalyzer(Trace::IAnalysisSession& Session, Trace::FTimingProfilerProvider& TimingProfilerProvider);
+	virtual void OnAnalysisBegin(const FOnAnalysisContext& Context) override;
+	virtual void OnEvent(uint16 RouteId, const FOnEventContext& Context) override;
+	virtual void OnAnalysisEnd() override;
+
+private:
+	struct EventScopeState
+	{
+		uint64 StartCycle;
+		uint32 EventTypeId;
+	};
+
+	struct FThreadState
+	{
+		TArray<EventScopeState> ScopeStack;
+		Trace::FTimingProfilerProvider::TimelineInternal* Timeline;
+		double LastCycle = 0.0;
+	};
+
+	TSharedRef<FThreadState> GetThreadState(uint32 ThreadId);
+
+	enum : uint16
+	{
+		RouteId_EventSpec,
+		RouteId_EventBatch,
+		RouteId_EndCapture,
+	};
+
+	Trace::IAnalysisSession& Session;
+	Trace::FTimingProfilerProvider& TimingProfilerProvider;
+	TMap<uint32, TSharedRef<FThreadState>> ThreadStatesMap;
+	TMap<uint16, uint32> ScopeIdToEventIdMap;
+	uint64 TotalEventSize = 0;
+	uint64 TotalScopeCount = 0;
+	double BytesPerScope = 0.0;
+};
diff --git a/Engine/Source/Developer/TraceServices/Private/Analyzers/GpuProfilerTraceAnalysis.cpp b/Engine/Source/Developer/TraceServices/Private/Analyzers/GpuProfilerTraceAnalysis.cpp
new file mode 100644
--- /dev/null
+++ b/Engine/Source/Developer/TraceServices/Private/Analyzers/GpuProfilerTraceAnalysis.cpp
@@ -0,0 +1,89 @@
+// Copyright 1998-2019 Epic Games, Inc. All Rights Reserved.
+#include "GpuProfilerTraceAnalysis.h"
+#include "AnalysisServicePrivate.h"
+#include "Common/Utils.h"
+
+FGpuProfilerAnalyzer::FGpuProfilerAnalyzer(Trace::FAnalysisSession& InSession, Trace::FTimingProfilerProvider& InTimingProfilerProvider)
+	: Session(InSession)
+	, TimingProfilerProvider(InTimingProfilerProvider)
+	, Timeline(TimingProfilerProvider.EditGpuTimeline())
+	, Calibrated(false)
+{
+
+}
+
+void FGpuProfilerAnalyzer::OnAnalysisBegin(const FOnAnalysisContext& Context)
+{
+	auto& Builder = Context.InterfaceBuilder;
+
+	Builder.RouteEvent(RouteId_EventSpec, "GpuProfiler", "EventSpec");
+	Builder.RouteEvent(RouteId_Frame, "GpuProfiler", "Frame");
+}
+
+void FGpuProfilerAnalyzer::OnEvent(uint16 RouteId, const FOnEventContext& Context)
+{
+	Trace::FAnalysisSessionEditScope _(Session);
+
+	const auto& EventData = Context.EventData;
+
+	switch (RouteId)
+	{
+	case RouteId_EventSpec:
+	{
+		uint64 EventType = EventData.GetValue("EventType").As<uint64>();
+		FString EventName(reinterpret_cast<const TCHAR*>(EventData.GetAttachment()), EventData.GetValue("NameLength").As<uint16>());
+		EventTypeMap.Add(EventType, TimingProfilerProvider.AddGpuTimer(*EventName));
+		break;
+	}
+	case RouteId_Frame:
+	{
+		uint64 BufferSize = EventData.GetAttachmentSize();
+		const uint8* BufferPtr = EventData.GetAttachment();
+		const uint8* BufferEnd = BufferPtr + BufferSize;
+
+		uint32 CurrentDepth = 0;
+
+		uint64 LastTimestamp = EventData.GetValue("TimestampBase").As<uint64>();
+		double LastTime = 0.0;
+		while (BufferPtr < BufferEnd)
+		{
+			uint64 DecodedTimestamp = FTraceAnalyzerUtils::Decode7bit(BufferPtr);
+			uint64 ActualTimestamp = (DecodedTimestamp >> 1) + LastTimestamp;
+			LastTimestamp = ActualTimestamp;
+			LastTime = GpuTimestampToSessionTime(ActualTimestamp);
+			if (DecodedTimestamp & 1ull)
+			{
+				uint64 EventType = *reinterpret_cast<const uint64*>(BufferPtr);
+				BufferPtr += sizeof(uint64);
+				check(EventTypeMap.Contains(EventType));
+				Trace::FTimingProfilerEvent Event;
+				Event.TimerIndex = EventTypeMap[EventType];
+				Timeline.AppendBeginEvent(LastTime, Event);
+				++CurrentDepth;
+			}
+			else
+			{
+				check(CurrentDepth > 0);
+				--CurrentDepth;
+				Timeline.AppendEndEvent(LastTime);
+			}
+		}
+		Session.UpdateDurationSeconds(LastTime);
+		check(BufferPtr == BufferEnd);
+		check(CurrentDepth == 0);
+		break;
+	}
+		
+	}
+}
+
+double FGpuProfilerAnalyzer::GpuTimestampToSessionTime(uint64 GpuMicroseconds)
+{
+	if (!Calibrated)
+	{
+		uint64 SessionTimeMicroseconds = uint64(Session.GetDurationSeconds() * 1000000.0);
+		GpuTimeOffset = GpuMicroseconds - SessionTimeMicroseconds;
+		Calibrated = true;
+	}
+	return (GpuMicroseconds - GpuTimeOffset) / 1000000.0;
+}
diff --git a/Engine/Source/Developer/TraceServices/Private/Analyzers/GpuProfilerTraceAnalysis.h b/Engine/Source/Developer/TraceServices/Private/Analyzers/GpuProfilerTraceAnalysis.h
new file mode 100644
--- /dev/null
+++ b/Engine/Source/Developer/TraceServices/Private/Analyzers/GpuProfilerTraceAnalysis.h
@@ -0,0 +1,38 @@
+// Copyright 1998-2019 Epic Games, Inc. All Rights Reserved.
+
+#pragma once
+
+#include "Trace/Trace.h"
+#include "Trace/Analyzer.h"
+#include "Model/TimingProfilerPrivate.h"
+
+namespace Trace
+{
+	class FAnalysisSession;
+}
+
+class FGpuProfilerAnalyzer
+	: public Trace::IAnalyzer
+{
+public:
+	FGpuProfilerAnalyzer(Trace::FAnalysisSession& Session, Trace::FTimingProfilerProvider& TimingProfilerProvider);
+	virtual void OnAnalysisBegin(const FOnAnalysisContext& Context) override;
+	virtual void OnEvent(uint16 RouteId, const FOnEventContext& Context) override;
+	virtual void OnAnalysisEnd() override {};
+
+private:
+	enum : uint16
+	{
+		RouteId_EventSpec,
+		RouteId_Frame,
+	};
+
+	double GpuTimestampToSessionTime(uint64 GpuMicroseconds);
+
+	Trace::FAnalysisSession& Session;
+	Trace::FTimingProfilerProvider& TimingProfilerProvider;
+	Trace::FTimingProfilerProvider::TimelineInternal& Timeline;
+	TMap<uint64, uint32> EventTypeMap;
+	bool Calibrated;
+	uint64 GpuTimeOffset;
+};
diff --git a/Engine/Source/Developer/TraceServices/Private/Analyzers/LoadTimeTraceAnalysis.cpp b/Engine/Source/Developer/TraceServices/Private/Analyzers/LoadTimeTraceAnalysis.cpp
new file mode 100644
--- /dev/null
+++ b/Engine/Source/Developer/TraceServices/Private/Analyzers/LoadTimeTraceAnalysis.cpp
@@ -0,0 +1,457 @@
+// Copyright 1998-2019 Epic Games, Inc. All Rights Reserved.
+#include "LoadTimeTraceAnalysis.h"
+
+#include "HAL/FileManager.h"
+#include "Serialization/LoadTimeTrace.h"
+#include "Trace/Trace.h"
+#include "Model/LoadTimeProfilerPrivate.h"
+#include "Analyzers/MiscTraceAnalysis.h"
+
+FAsyncLoadingTraceAnalyzer::FAsyncLoadingTraceAnalyzer(Trace::IAnalysisSession& InSession, Trace::FLoadTimeProfilerProvider& InLoadTimeProfilerProvider)
+	: Session(InSession)
+	, LoadTimeProfilerProvider(InLoadTimeProfilerProvider)
+{
+}
+
+FAsyncLoadingTraceAnalyzer::~FAsyncLoadingTraceAnalyzer()
+{
+
+}
+
+TSharedRef<FAsyncLoadingTraceAnalyzer::FThreadState> FAsyncLoadingTraceAnalyzer::GetThreadState(uint32 ThreadId)
+{
+	if (!ThreadStatesMap.Contains(ThreadId))
+	{
+		TSharedRef<FThreadState> ThreadState = MakeShared<FThreadState>();
+		ThreadStatesMap.Add(ThreadId, ThreadState);
+		if (MainThreadId == -1)
+		{
+			MainThreadId = ThreadId;
+			LoadTimeProfilerProvider.SetMainThreadId(MainThreadId);
+			ThreadState->CpuTimeline = LoadTimeProfilerProvider.EditMainThreadCpuTimeline();
+		}
+		else if (ThreadId == AsyncLoadingThreadId)
+		{
+			ThreadState->CpuTimeline = LoadTimeProfilerProvider.EditAsyncLoadingThreadCpuTimeline();
+		}
+		return ThreadState;
+	}
+	else
+	{
+		return ThreadStatesMap[ThreadId];
+	}
+}
+
+const Trace::FClassInfo* FAsyncLoadingTraceAnalyzer::GetClassInfo(uint64 ClassPtr) const
+{
+	const Trace::FClassInfo* const* ClassInfo = ClassInfosMap.Find(ClassPtr);
+	if (ClassInfo)
+	{
+		return *ClassInfo;
+	}
+	else
+	{
+		return nullptr;
+	}
+}
+
+void FAsyncLoadingTraceAnalyzer::OnAnalysisBegin(const FOnAnalysisContext& Context)
+{
+	auto& Builder = Context.InterfaceBuilder;
+
+	Builder.RouteEvent(RouteId_StartAsyncLoading, "LoadTime", "StartAsyncLoading");
+	Builder.RouteEvent(RouteId_SuspendAsyncLoading, "LoadTime", "SuspendAsyncLoading");
+	Builder.RouteEvent(RouteId_ResumeAsyncLoading, "LoadTime", "ResumeAsyncLoading");
+	Builder.RouteEvent(RouteId_NewLinker, "LoadTime", "NewLinker");
+	Builder.RouteEvent(RouteId_DestroyLinker, "LoadTime", "DestroyLinker");
+	Builder.RouteEvent(RouteId_PackageSummary, "LoadTime", "PackageSummary");
+	Builder.RouteEvent(RouteId_BeginCreateExport, "LoadTime", "BeginCreateExport");
+	Builder.RouteEvent(RouteId_EndCreateExport, "LoadTime", "EndCreateExport");
+	Builder.RouteEvent(RouteId_BeginObjectScope, "LoadTime", "BeginObjectScope");
+	Builder.RouteEvent(RouteId_EndObjectScope, "LoadTime", "EndObjectScope");
+	Builder.RouteEvent(RouteId_NewAsyncPackage, "LoadTime", "NewAsyncPackage");
+	Builder.RouteEvent(RouteId_DestroyAsyncPackage, "LoadTime", "DestroyAsyncPackage");
+	Builder.RouteEvent(RouteId_BeginRequest, "LoadTime", "BeginRequest");
+	Builder.RouteEvent(RouteId_EndRequest, "LoadTime", "EndRequest");
+	Builder.RouteEvent(RouteId_BeginLoadMap, "LoadTime", "BeginLoadMap");
+	Builder.RouteEvent(RouteId_EndLoadMap, "LoadTime", "EndLoadMap");
+	Builder.RouteEvent(RouteId_NewStreamableHandle, "LoadTime", "NewStreamableHandle");
+	Builder.RouteEvent(RouteId_DestroyStreamableHandle, "LoadTime", "DestroyStreamableHandle");
+	Builder.RouteEvent(RouteId_BeginLoadStreamableHandle, "LoadTime", "BeginLoadStreamableHandle");
+	Builder.RouteEvent(RouteId_EndLoadStreamableHandle, "LoadTime", "EndLoadStreamableHandle");
+	Builder.RouteEvent(RouteId_BeginWaitForStreamableHandle, "LoadTime", "BeginWaitForStreamableHandle");
+	Builder.RouteEvent(RouteId_EndWaitForStreamableHandle, "LoadTime", "EndWaitForStreamableHandle");
+	Builder.RouteEvent(RouteId_StreamableHandleRequestAssociation, "LoadTime", "StreamableHandleRequestAssociation");
+	Builder.RouteEvent(RouteId_AsyncPackageRequestAssociation, "LoadTime", "AsyncPackageRequestAssociation");
+	Builder.RouteEvent(RouteId_AsyncPackageLinkerAssociation, "LoadTime", "AsyncPackageLinkerAssociation");
+	Builder.RouteEvent(RouteId_LinkerArchiveAssociation, "LoadTime", "LinkerArchiveAssociation");
+	Builder.RouteEvent(RouteId_BeginAsyncPackageScope, "LoadTime", "BeginAsyncPackageScope");
+	Builder.RouteEvent(RouteId_EndAsyncPackageScope, "LoadTime", "EndAsyncPackageScope");
+	Builder.RouteEvent(RouteId_BeginFlushAsyncLoading, "LoadTime", "BeginFlushAsyncLoading");
+	Builder.RouteEvent(RouteId_EndFlushAsyncLoading, "LoadTime", "EndFlushAsyncLoading");
+	Builder.RouteEvent(RouteId_ClassInfo, "LoadTime", "ClassInfo");
+}
+
+void FAsyncLoadingTraceAnalyzer::OnAnalysisEnd()
+{
+	
+}
+
+void FAsyncLoadingTraceAnalyzer::OnEvent(uint16 RouteId, const FOnEventContext& Context)
+{
+	const auto& EventData = Context.EventData;
+	switch (RouteId)
+	{
+	case RouteId_StartAsyncLoading:
+	{
+		Trace::FAnalysisSessionEditScope _(Session);
+		AsyncLoadingThreadId = EventData.GetValue("ThreadId").As<uint32>();
+		LoadTimeProfilerProvider.SetAsyncLoadingThreadId(AsyncLoadingThreadId);
+		break;
+	}
+	case RouteId_NewLinker:
+	{
+		uint64 LinkerPtr = EventData.GetValue("Linker").As<uint64>();
+		//check(!ActiveLinkersMap.Contains(LinkerPtr));
+		TSharedRef<FLinkerState> LinkerState = MakeShared<FLinkerState>();
+		ActiveLinkersMap.Add(LinkerPtr, LinkerState);
+		break;
+	}
+	case RouteId_DestroyLinker:
+	{
+		uint64 Ptr = EventData.GetValue("Linker").As<uint64>();
+		//check(ActiveLinkersMap.Contains(Ptr));
+		ActiveLinkersMap.Remove(Ptr);
+		break;
+	}
+	case RouteId_PackageSummary:
+	{
+		uint64 LinkerPtr = EventData.GetValue("Linker").As<uint64>();
+		TSharedRef<FLinkerState>* LinkerState = ActiveLinkersMap.Find(LinkerPtr);
+		if (LinkerState && (*LinkerState)->PackageInfo)
+		{
+			Trace::FAnalysisSessionEditScope _(Session);
+			Trace::FPackageSummaryInfo& Summary = (*LinkerState)->PackageInfo->Summary;
+			Summary.TotalHeaderSize = EventData.GetValue("TotalHeaderSize").As<uint32>();
+			Summary.NameCount = EventData.GetValue("NameCount").As<uint32>();
+			Summary.ImportCount = EventData.GetValue("ImportCount").As<uint32>();
+			Summary.ExportCount = EventData.GetValue("ExportCount").As<uint32>();
+		}
+		break;
+	}
+	case RouteId_BeginCreateExport:
+	{
+		Trace::FAnalysisSessionEditScope _(Session);
+
+		Trace::FPackageExportInfo& Export = LoadTimeProfilerProvider.CreateExport();
+		Export.SerialOffset = EventData.GetValue("SerialOffset").As<uint64>();
+		Export.SerialSize = EventData.GetValue("SerialSize").As<uint64>();
+		Export.IsAsset = EventData.GetValue("IsAsset").As<bool>();
+		uint32 ThreadId = EventData.GetValue("ThreadId").As<uint32>();
+		TSharedRef<FThreadState> ThreadState = GetThreadState(ThreadId);
+		ThreadState->EnterExportScope(Context.SessionContext.TimestampFromCycle(EventData.GetValue("Cycle").As<uint64>()), &Export, LoadTimeProfilerObjectEventType_Create);
+		uint64 LinkerPtr = EventData.GetValue("Linker").As<uint64>();
+		TSharedRef<FLinkerState>* LinkerState = ActiveLinkersMap.Find(LinkerPtr);
+		if (LinkerState && (*LinkerState)->PackageInfo)
+		{
+			(*LinkerState)->PackageInfo->Exports.Add(&Export);
+		}
+		break;
+	}
+	case RouteId_EndCreateExport:
+	{
+		Trace::FAnalysisSessionEditScope _(Session);
+
+		uint32 ThreadId = EventData.GetValue("ThreadId").As<uint32>();
+		TSharedRef<FThreadState> ThreadState = GetThreadState(ThreadId);
+		uint64 ObjectPtr = EventData.GetValue("Object").As<uint64>();
+		Trace::FPackageExportInfo* Export = ThreadState->GetCurrentExportScope();
+		if (Export)
+		{
+			ExportsMap.Add(ObjectPtr, Export);
+			const Trace::FClassInfo* ObjectClass = GetClassInfo(EventData.GetValue("Class").As<uint64>());
+			Export->Class = ObjectClass;
+		}
+		ThreadState->LeaveScope(Context.SessionContext.TimestampFromCycle(EventData.GetValue("Cycle").As<uint64>()));
+		break;
+	}
+	case RouteId_BeginObjectScope:
+	{
+		uint64 ObjectPtr = EventData.GetValue("Object").As<uint64>();
+		Trace::FPackageExportInfo** Export = ExportsMap.Find(ObjectPtr);
+		uint32 ThreadId = EventData.GetValue("ThreadId").As<uint32>();
+		TSharedRef<FThreadState> ThreadState = GetThreadState(ThreadId);
+		ELoadTimeProfilerObjectEventType EventType = static_cast<ELoadTimeProfilerObjectEventType>(EventData.GetValue("EventType").As<uint8>());
+		{
+			Trace::FAnalysisSessionEditScope _(Session);
+			ThreadState->EnterExportScope(Context.SessionContext.TimestampFromCycle(EventData.GetValue("Cycle").As<uint64>()), Export ? *Export : nullptr, EventType);
+		}
+		break;
+	}
+	case RouteId_EndObjectScope:
+	{
+		uint32 ThreadId = EventData.GetValue("ThreadId").As<uint32>();
+		TSharedRef<FThreadState> ThreadState = GetThreadState(ThreadId);
+		{
+			Trace::FAnalysisSessionEditScope _(Session);
+			ThreadState->LeaveScope(Context.SessionContext.TimestampFromCycle(EventData.GetValue("Cycle").As<uint64>()));
+		}
+		break;
+	}
+	case RouteId_BeginRequest:
+	{
+		uint64 RequestId = EventData.GetValue("RequestId").As<uint64>();
+		//check(!ActiveRequestsMap.Contains(RequestId));
+		TSharedRef<FRequestState> RequestState = MakeShared<FRequestState>();
+		RequestState->Id = RequestId;
+		RequestState->WallTimeStartCycle = EventData.GetValue("Cycle").As<uint64>();
+		RequestState->WallTimeEndCycle = 0;
+		Requests.Add(RequestState);
+		ActiveRequestsMap.Add(RequestId, RequestState);
+		uint32 ThreadId = EventData.GetValue("ThreadId").As<uint32>();
+		TSharedRef<FThreadState> ThreadState = GetThreadState(ThreadId);
+		if (ThreadState->ActiveLoadMap)
+		{
+			ThreadState->ActiveLoadMap->Requests.Add(RequestState);
+		}
+		
+		{
+			/*Trace::FAnalysisSessionEditScope _(Session.Get());
+			RequestTimelineEventTypeMap.Add(RequestState->Id, TimelineProvider->AddEventType(*RequestState->PackageName, 0));*/
+		}
+		
+		break;
+	}
+	case RouteId_EndRequest:
+	{
+		uint64 RequestId = EventData.GetValue("RequestId").As<uint64>();
+		TSharedRef<FRequestState>* RequestState = ActiveRequestsMap.Find(RequestId);
+		if (RequestState)
+		{
+			(*RequestState)->WallTimeEndCycle = EventData.GetValue("Cycle").As<uint64>();
+		}
+		break;
+	}
+	case RouteId_BeginLoadMap:
+	{
+		TSharedRef<FLoadMapState> LoadMapState = MakeShared<FLoadMapState>();
+		LoadMapState->Id = NextMapId++;
+		LoadMapState->WallTimeStartCycle = EventData.GetValue("Cycle").As<uint64>();
+		LoadMapState->WallTimeEndCycle = 0;
+		LoadMapState->Name = reinterpret_cast<const TCHAR*>(EventData.GetAttachment());
+		Maps.Add(LoadMapState);
+		uint32 ThreadId = EventData.GetValue("ThreadId").As<uint32>();
+		TSharedRef<FThreadState> ThreadState = GetThreadState(ThreadId);
+		ThreadState->ActiveLoadMap = LoadMapState;
+		break;
+	}
+	case RouteId_EndLoadMap:
+	{
+		uint32 ThreadId = EventData.GetValue("ThreadId").As<uint32>();
+		TSharedRef<FThreadState> ThreadState = GetThreadState(ThreadId);
+		if (ThreadState->ActiveLoadMap)
+		{
+			ThreadState->ActiveLoadMap->WallTimeEndCycle = EventData.GetValue("Cycle").As<uint64>();
+			ThreadState->ActiveLoadMap = nullptr;
+		}
+		break;
+	}
+	case RouteId_NewStreamableHandle:
+	{
+		uint64 Ptr = EventData.GetValue("StreamableHandle").As<uint64>();
+		//check(!ActiveStreamableHandlesMap.Contains(Ptr));
+		TSharedRef<FStreamableHandleState> StreamableHandleState = MakeShared<FStreamableHandleState>();
+		StreamableHandleState->Id = NextStreamableHandleId++;
+		StreamableHandleState->DebugName = FString((const char*)EventData.GetAttachment(), EventData.GetValue("DebugNameSize").As<uint16>());
+		StreamableHandleState->WallTimeStartCycle = 0;
+		StreamableHandleState->WallTimeEndCycle = 0;
+		StreamableHandles.Add(StreamableHandleState);
+		ActiveStreamableHandlesMap.Add(Ptr, StreamableHandleState);
+		
+		{
+			/*Trace::FAnalysisSessionEditScope _(Session.Get());
+			StreamableHandlesTimelineEventTypeMap.Add(StreamableHandleState->Id, TimelineProvider->AddEventType(*StreamableHandleState->DebugName, 0));*/
+		}
+		
+		break;
+	}
+	case RouteId_DestroyStreamableHandle:
+	{
+		uint64 Ptr = EventData.GetValue("StreamableHandle").As<uint64>();
+		ActiveStreamableHandlesMap.Remove(Ptr);
+		break;
+	}
+	case RouteId_BeginLoadStreamableHandle:
+	{
+		uint64 Ptr = EventData.GetValue("StreamableHandle").As<uint64>();
+		TSharedRef<FStreamableHandleState>* StreamableHandleState = ActiveStreamableHandlesMap.Find(Ptr);
+		if (StreamableHandleState)
+		{
+			(*StreamableHandleState)->WallTimeStartCycle = EventData.GetValue("Cycle").As<uint64>();
+		}
+		break;
+	}
+	case RouteId_EndLoadStreamableHandle:
+	{
+		uint64 Ptr = EventData.GetValue("StreamableHandle").As<uint64>();
+		TSharedRef<FStreamableHandleState>* StreamableHandleState = ActiveStreamableHandlesMap.Find(Ptr);
+		if (StreamableHandleState)
+		{
+			(*StreamableHandleState)->WallTimeEndCycle = EventData.GetValue("Cycle").As<uint64>();
+		}
+		break;
+	}
+	case RouteId_BeginWaitForStreamableHandle:
+	{
+		uint64 Ptr = EventData.GetValue("StreamableHandle").As<uint64>();
+		TSharedRef<FStreamableHandleState>* StreamableHandleState = ActiveStreamableHandlesMap.Find(Ptr);
+		if (StreamableHandleState)
+		{
+			uint32 ThreadId = EventData.GetValue("ThreadId").As<uint32>();
+			TSharedRef<FThreadState> ThreadState = GetThreadState(ThreadId);
+			ThreadState->WaitForStreamableHandleHandle = *StreamableHandleState;
+			ThreadState->WaitForStreamableHandleStartCycle = EventData.GetValue("Cycle").As<uint64>();
+		}
+		/*uint64 EventTypeId = StreamableHandlesTimelineEventTypeMap[ThreadState->WaitForStreamableHandleHandle->Id];
+
+		{
+			Trace::FAnalysisSessionEditScope _(Session.Get());
+			BlockingRequestsTimeline->AppendBeginEvent(Context.SessionContext.TimestampFromCycle(ThreadState->WaitForStreamableHandleStartCycle), EventTypeId);
+		}*/
+
+		break;
+	}
+	case RouteId_EndWaitForStreamableHandle:
+	{
+		uint32 ThreadId = EventData.GetValue("ThreadId").As<uint32>();
+		TSharedRef<FThreadState> ThreadState = GetThreadState(ThreadId);
+		ThreadState->WaitForStreamableHandleHandle = nullptr;
+
+		/*{
+			Trace::FAnalysisSessionEditScope _(Session.Get());
+			BlockingRequestsTimeline->AppendEndEvent(Context.SessionContext.TimestampFromCycle(EventData.GetValue("Cycle").As<uint64>()));
+		}*/
+
+		break;
+	}
+	case RouteId_NewAsyncPackage:
+	{
+		Trace::FAnalysisSessionEditScope _(Session);
+
+		uint64 AsyncPackagePtr = EventData.GetValue("AsyncPackage").As<uint64>();
+		//check(!ActivePackagesMap.Contains(AsyncPackagePtr));
+		TSharedRef<FAsyncPackageState> AsyncPackageState = MakeShared<FAsyncPackageState>();
+		AsyncPackageState->PackageInfo = &LoadTimeProfilerProvider.CreatePackage(reinterpret_cast<const TCHAR*>(EventData.GetAttachment()));
+		ActiveAsyncPackagesMap.Add(AsyncPackagePtr, AsyncPackageState);
+		break;
+	}
+	case RouteId_DestroyAsyncPackage:
+	{
+		uint64 AsyncPackagePtr = EventData.GetValue("AsyncPackage").As<uint64>();
+		//check(ActiveAsyncPackagesMap.Contains(AsyncPackagePtr));
+		ActiveAsyncPackagesMap.Remove(AsyncPackagePtr);
+		break;
+	}
+
+	case RouteId_StreamableHandleRequestAssociation:
+	{
+		uint64 StreamableHandlePtr = EventData.GetValue("StreamableHandle").As<uint64>();
+		TSharedRef<FStreamableHandleState>* StreamableHandleState = ActiveStreamableHandlesMap.Find(StreamableHandlePtr);
+		uint64 RequestId = EventData.GetValue("RequestId").As<uint64>();
+		TSharedRef<FRequestState>* RequestState = ActiveRequestsMap.Find(RequestId);
+		if (StreamableHandleState && RequestState)
+		{
+			(*StreamableHandleState)->Requests.Add(*RequestState);
+		}
+		break;
+	}
+	case RouteId_AsyncPackageRequestAssociation:
+	{
+		uint64 AsyncPackagePtr = EventData.GetValue("AsyncPackage").As<uint64>();
+		TSharedRef<FAsyncPackageState>* AsyncPackageState = ActiveAsyncPackagesMap.Find(AsyncPackagePtr);
+		uint64 RequestId = EventData.GetValue("RequestId").As<uint64>();
+		TSharedRef<FRequestState>* RequestState = ActiveRequestsMap.Find(RequestId);
+		if (AsyncPackageState && RequestState)
+		{
+			(*RequestState)->AsyncPackages.Add(*AsyncPackageState);
+			(*AsyncPackageState)->Requests.Add(*RequestState);
+		}
+		break;
+	}
+	case RouteId_AsyncPackageLinkerAssociation:
+	{
+		uint64 LinkerPtr = EventData.GetValue("Linker").As<uint64>();
+		TSharedRef<FLinkerState>* LinkerState = ActiveLinkersMap.Find(LinkerPtr);
+		uint64 AsyncPackagePtr = EventData.GetValue("AsyncPackage").As<uint64>();
+		TSharedRef<FAsyncPackageState>* AsyncPackageState = ActiveAsyncPackagesMap.Find(AsyncPackagePtr);
+		if (LinkerState && AsyncPackageState)
+		{
+			(*AsyncPackageState)->LinkerState = *LinkerState;
+			(*LinkerState)->AsyncPackageState = *AsyncPackageState;
+			(*LinkerState)->PackageInfo = (*AsyncPackageState)->PackageInfo;
+		}
+		break;
+	}
+	case RouteId_LinkerArchiveAssociation:
+	{
+		uint64 ArchivePtr = EventData.GetValue("Archive").As<uint64>();
+		uint64 LinkerPtr = EventData.GetValue("Linker").As<uint64>();
+		break;
+	}
+	case RouteId_BeginAsyncPackageScope:
+	{
+		Trace::FAnalysisSessionEditScope _(Session);
+
+		uint64 AsyncPackagePtr = EventData.GetValue("AsyncPackage").As<uint64>();
+		TSharedRef<FAsyncPackageState>* AsyncPackageState = ActiveAsyncPackagesMap.Find(AsyncPackagePtr);
+		uint32 ThreadId = EventData.GetValue("ThreadId").As<uint32>();
+		TSharedRef<FThreadState> ThreadState = GetThreadState(ThreadId);
+		ELoadTimeProfilerPackageEventType EventType = static_cast<ELoadTimeProfilerPackageEventType>(EventData.GetValue("EventType").As<uint8>());
+		ThreadState->EnterPackageScope(Context.SessionContext.TimestampFromCycle(EventData.GetValue("Cycle").As<uint64>()), AsyncPackageState ? (*AsyncPackageState)->PackageInfo : nullptr, EventType);
+		break;
+	}
+	case RouteId_EndAsyncPackageScope:
+	{
+		Trace::FAnalysisSessionEditScope _(Session);
+
+		uint32 ThreadId = EventData.GetValue("ThreadId").As<uint32>();
+		TSharedRef<FThreadState> ThreadState = GetThreadState(ThreadId);
+		uint64 Cycle = EventData.GetValue("Cycle").As<uint64>();
+		ThreadState->LeaveScope(Context.SessionContext.TimestampFromCycle(EventData.GetValue("Cycle").As<uint64>()));
+		break;
+	}
+	case RouteId_BeginFlushAsyncLoading:
+	{
+		Trace::FAnalysisSessionEditScope _(Session);
+
+		FlushAsyncLoadingRequestId = EventData.GetValue("RequestId").As<uint64>();
+		FlushAsyncLoadingStartCycle = EventData.GetValue("Cycle").As<uint64>();
+		if (FlushAsyncLoadingRequestId != INDEX_NONE)
+		{
+			//check(ActiveRequestsMap.Contains(FlushAsyncLoadingRequestId));
+			/*uint64 EventTypeId = RequestTimelineEventTypeMap[FlushAsyncLoadingRequestId];
+			BlockingRequestsTimeline->AppendBeginEvent(Context.SessionContext.TimestampFromCycle(FlushAsyncLoadingStartCycle), EventTypeId);*/
+		}
+		else
+		{
+			//BlockingRequestsTimeline->AppendBeginEvent(Context.SessionContext.TimestampFromCycle(FlushAsyncLoadingStartCycle), FlushAsyncLoadingEventId);
+		}
+		break;
+	}
+	case RouteId_EndFlushAsyncLoading:
+	{
+		uint64 Cycle = EventData.GetValue("Cycle").As<uint64>();
+		//BlockingRequestsTimeline->AppendEndEvent(Context.SessionContext.TimestampFromCycle(Cycle));
+		break;
+	}
+	case RouteId_ClassInfo:
+	{
+		Trace::FAnalysisSessionEditScope _(Session);
+
+		uint64 ClassPtr = EventData.GetValue("Class").As<uint64>();
+		const Trace::FClassInfo& ClassInfo = LoadTimeProfilerProvider.AddClassInfo(reinterpret_cast<const TCHAR*>(EventData.GetAttachment()));
+		ClassInfosMap.Add(ClassPtr, &ClassInfo);
+		break;
+	}
+	}
+}
diff --git a/Engine/Source/Developer/TraceServices/Private/Analyzers/LoadTimeTraceAnalysis.h b/Engine/Source/Developer/TraceServices/Private/Analyzers/LoadTimeTraceAnalysis.h
new file mode 100644
--- /dev/null
+++ b/Engine/Source/Developer/TraceServices/Private/Analyzers/LoadTimeTraceAnalysis.h
@@ -0,0 +1,224 @@
+// Copyright 1998-2019 Epic Games, Inc. All Rights Reserved.
+
+#pragma once
+
+#include "Templates/SharedPointer.h"
+#include "Trace/Analyzer.h"
+#include "AnalysisServicePrivate.h"
+#include "Model/LoadTimeProfilerPrivate.h"
+
+namespace Trace
+{
+	struct FClassInfo;
+}
+
+inline bool operator!=(const Trace::FLoadTimeProfilerCpuEvent& Lhs, const Trace::FLoadTimeProfilerCpuEvent& Rhs)
+{
+	return Lhs.Package != Rhs.Package ||
+		Lhs.Export != Rhs.Export ||
+		Lhs.PackageEventType != Rhs.PackageEventType ||
+		Lhs.ExportEventType != Rhs.ExportEventType;
+}
+
+class FAsyncLoadingTraceAnalyzer : public Trace::IAnalyzer
+{
+public:
+	FAsyncLoadingTraceAnalyzer(Trace::IAnalysisSession& Session, Trace::FLoadTimeProfilerProvider& LoadTimeProfilerProvider);
+	virtual ~FAsyncLoadingTraceAnalyzer();
+
+	virtual void OnAnalysisBegin(const FOnAnalysisContext& Context) override;
+	virtual void OnEvent(uint16 RouteId, const FOnEventContext& Context) override;
+	virtual void OnAnalysisEnd() override;
+	
+private:
+	struct FRequestState;
+	struct FAsyncPackageState;
+
+	struct FLoadMapState
+	{
+		uint64 Id;
+		FString Name;
+		uint64 WallTimeStartCycle;
+		uint64 WallTimeEndCycle;
+		TArray<TSharedRef<FRequestState>> Requests;
+	};
+
+	struct FStreamableHandleState
+	{
+		uint64 Id;
+		FString DebugName;
+		uint64 WallTimeStartCycle;
+		uint64 WallTimeEndCycle;
+		TArray<TSharedRef<FRequestState>> Requests;
+	};
+
+	struct FRequestState
+	{
+		uint64 Id;
+		uint64 WallTimeStartCycle;
+		uint64 WallTimeEndCycle;
+		TArray<TSharedRef<FAsyncPackageState>> AsyncPackages;
+	};
+
+	struct FLinkerState
+	{
+		TSharedPtr<FAsyncPackageState> AsyncPackageState;
+		Trace::FPackageInfo* PackageInfo = nullptr;
+	};
+
+	struct FAsyncPackageState
+	{
+		Trace::FPackageInfo* PackageInfo = nullptr;
+		TSharedPtr<FLinkerState> LinkerState;
+		TArray<TSharedRef<FRequestState>> Requests;
+	};
+
+	struct FScopeStackEntry
+	{
+		Trace::FLoadTimeProfilerCpuEvent Event;
+		bool EnteredEvent;
+	};
+
+	struct FThreadState
+	{
+		FScopeStackEntry CpuScopeStack[256];
+		uint64 CpuScopeStackDepth = 0;
+		Trace::FLoadTimeProfilerCpuEvent CurrentEvent;
+		uint64 WaitForStreamableHandleStartCycle = 0;
+		TSharedPtr<FStreamableHandleState> WaitForStreamableHandleHandle;
+		TSharedPtr<FLoadMapState> ActiveLoadMap;
+		
+		TSharedPtr<Trace::FLoadTimeProfilerProvider::CpuTimelineInternal> CpuTimeline;
+
+		void EnterPackageScope(double Time, const Trace::FPackageInfo* PackageInfo, ELoadTimeProfilerPackageEventType EventType)
+		{
+			FScopeStackEntry& StackEntry = CpuScopeStack[CpuScopeStackDepth++];
+			StackEntry.Event.Export = CurrentEvent.Export;
+			StackEntry.Event.ExportEventType = CurrentEvent.ExportEventType;
+			StackEntry.Event.Package = PackageInfo;
+			StackEntry.Event.PackageEventType = EventType;
+			if (PackageInfo && CurrentEvent != StackEntry.Event)
+			{
+				CurrentEvent = StackEntry.Event;
+				CpuTimeline->AppendBeginEvent(Time, StackEntry.Event);
+				StackEntry.EnteredEvent = true;
+			}
+			else
+			{
+				StackEntry.EnteredEvent = false;
+			}
+		}
+
+		void EnterExportScope(double Time, const Trace::FPackageExportInfo* ExportInfo, ELoadTimeProfilerObjectEventType EventType)
+		{
+			FScopeStackEntry& StackEntry = CpuScopeStack[CpuScopeStackDepth++];
+			StackEntry.Event.Export = ExportInfo;
+			StackEntry.Event.ExportEventType = EventType;
+			StackEntry.Event.Package = CurrentEvent.Package;
+			StackEntry.Event.PackageEventType = CurrentEvent.PackageEventType;
+
+			if (ExportInfo && CurrentEvent != StackEntry.Event)
+			{
+				CurrentEvent = StackEntry.Event;
+				CpuTimeline->AppendBeginEvent(Time, StackEntry.Event);
+				StackEntry.EnteredEvent = true;
+			}
+			else
+			{
+				StackEntry.EnteredEvent = false;
+			}
+		}
+
+		void LeaveScope(double Time)
+		{
+			FScopeStackEntry& StackEntry = CpuScopeStack[--CpuScopeStackDepth];
+			if (StackEntry.EnteredEvent)
+			{
+				CpuTimeline->AppendEndEvent(Time);
+			}
+			if (CpuScopeStackDepth > 0)
+			{
+				CurrentEvent = CpuScopeStack[CpuScopeStackDepth - 1].Event;
+			}
+			else
+			{
+				CurrentEvent = Trace::FLoadTimeProfilerCpuEvent();
+			}
+		}
+
+		Trace::FPackageExportInfo* GetCurrentExportScope()
+		{
+			if (CpuScopeStackDepth > 0)
+			{
+				FScopeStackEntry& StackEntry = CpuScopeStack[CpuScopeStackDepth - 1];
+				return const_cast<Trace::FPackageExportInfo*>(StackEntry.Event.Export);
+			}
+			else
+			{
+				return nullptr;
+			}
+		}
+	};
+
+	TSharedRef<FThreadState> GetThreadState(uint32 ThreadId);
+	const Trace::FClassInfo* GetClassInfo(uint64 ClassPtr) const;
+
+	enum : uint16
+	{
+		RouteId_StartAsyncLoading,
+		RouteId_SuspendAsyncLoading,
+		RouteId_ResumeAsyncLoading,
+		RouteId_NewLinker,
+		RouteId_DestroyLinker,
+		RouteId_NewAsyncPackage,
+		RouteId_DestroyAsyncPackage,
+		RouteId_BeginRequest,
+		RouteId_EndRequest,
+		RouteId_BeginLoadMap,
+		RouteId_EndLoadMap,
+		RouteId_NewStreamableHandle,
+		RouteId_DestroyStreamableHandle,
+		RouteId_BeginLoadStreamableHandle,
+		RouteId_EndLoadStreamableHandle,
+		RouteId_BeginWaitForStreamableHandle,
+		RouteId_EndWaitForStreamableHandle,
+		RouteId_PackageSummary,
+		RouteId_StreamableHandleRequestAssociation,
+		RouteId_AsyncPackageRequestAssociation,
+		RouteId_AsyncPackageLinkerAssociation,
+		RouteId_LinkerArchiveAssociation,
+		RouteId_BeginAsyncPackageScope,
+		RouteId_EndAsyncPackageScope,
+		RouteId_BeginCreateExport,
+		RouteId_EndCreateExport,
+		RouteId_BeginObjectScope,
+		RouteId_EndObjectScope,
+		RouteId_BeginFlushAsyncLoading,
+		RouteId_EndFlushAsyncLoading,
+		RouteId_ClassInfo,
+	};
+
+	Trace::IAnalysisSession& Session;
+	Trace::FLoadTimeProfilerProvider& LoadTimeProfilerProvider;
+	//TSharedPtr<Trace::FMonotonicTimeline<uint64>> BlockingRequestsTimeline;
+
+	TArray<TSharedRef<FLoadMapState>> Maps;
+	TArray<TSharedRef<FStreamableHandleState>> StreamableHandles;
+	TArray<TSharedRef<FRequestState>> Requests;
+
+	TMap<uint64, TSharedRef<FLinkerState>> ActiveLinkersMap;
+	TMap<uint64, TSharedRef<FAsyncPackageState>> ActiveAsyncPackagesMap;
+	TMap<uint64, Trace::FPackageExportInfo*> ExportsMap;
+	TMap<uint64, TSharedRef<FStreamableHandleState>> ActiveStreamableHandlesMap;
+	TMap<uint64, TSharedRef<FRequestState>> ActiveRequestsMap;
+	TMap<uint32, TSharedRef<FThreadState>> ThreadStatesMap;
+	TMap<uint64, const Trace::FClassInfo*> ClassInfosMap;
+
+	uint64 FlushAsyncLoadingRequestId;
+	uint64 FlushAsyncLoadingStartCycle;
+
+	uint64 NextMapId = 0;
+	uint64 NextStreamableHandleId = 0;
+	int64 MainThreadId = -1;
+	int64 AsyncLoadingThreadId = -1;
+};
diff --git a/Engine/Source/Developer/TraceServices/Private/Analyzers/LogTraceAnalysis.cpp b/Engine/Source/Developer/TraceServices/Private/Analyzers/LogTraceAnalysis.cpp
new file mode 100644
--- /dev/null
+++ b/Engine/Source/Developer/TraceServices/Private/Analyzers/LogTraceAnalysis.cpp
@@ -0,0 +1,60 @@
+// Copyright 1998-2019 Epic Games, Inc. All Rights Reserved.
+#include "LogTraceAnalysis.h"
+#include "AnalysisServicePrivate.h"
+#include "Logging/LogTrace.h"
+#include "Model/LogPrivate.h"
+
+FLogTraceAnalyzer::FLogTraceAnalyzer(Trace::IAnalysisSession& InSession, Trace::FLogProvider& InLogProvider)
+	: Session(InSession)
+	, LogProvider(InLogProvider)
+{
+
+}
+
+void FLogTraceAnalyzer::OnAnalysisBegin(const FOnAnalysisContext& Context)
+{
+	auto& Builder = Context.InterfaceBuilder;
+
+	Builder.RouteEvent(RouteId_LogCategory, "Logging", "LogCategory");
+	Builder.RouteEvent(RouteId_LogMessageSpec, "Logging", "LogMessageSpec");
+	Builder.RouteEvent(RouteId_LogMessage, "Logging", "LogMessage");
+}
+
+void FLogTraceAnalyzer::OnEvent(uint16 RouteId, const FOnEventContext& Context)
+{
+	Trace::FAnalysisSessionEditScope _(Session);
+
+	const auto& EventData = Context.EventData;
+	switch (RouteId)
+	{
+	case RouteId_LogCategory:
+	{
+		uint64 CategoryPointer = EventData.GetValue("CategoryPointer").As<uint64>();
+		Trace::FLogCategory& Category = LogProvider.GetCategory(CategoryPointer);
+		Category.Name = Session.StoreString(reinterpret_cast<const TCHAR*>(EventData.GetAttachment()));
+		Category.DefaultVerbosity = static_cast<ELogVerbosity::Type>(EventData.GetValue("DefaultVerbosity").As<uint8>());
+		break;
+	}
+	case RouteId_LogMessageSpec:
+	{
+		uint64 LogPoint = EventData.GetValue("LogPoint").As<uint64>();
+		Trace::FLogMessageSpec& Spec = LogProvider.GetMessageSpec(LogPoint);
+		uint64 CategoryPointer = EventData.GetValue("CategoryPointer").As<uint64>();
+		Trace::FLogCategory& Category = LogProvider.GetCategory(CategoryPointer);
+		Spec.Category = &Category;
+		Spec.Line = EventData.GetValue("Line").As<int32>();
+		Spec.Verbosity = static_cast<ELogVerbosity::Type>(EventData.GetValue("Verbosity").As<uint8>());
+		const ANSICHAR* File = reinterpret_cast<const ANSICHAR*>(EventData.GetAttachment());
+		Spec.File = Session.StoreString(ANSI_TO_TCHAR(File));
+		Spec.FormatString = Session.StoreString(reinterpret_cast<const TCHAR*>(EventData.GetAttachment() + strlen(File) + 1));
+		break;
+	}
+	case RouteId_LogMessage:
+	{
+		uint64 LogPoint = EventData.GetValue("LogPoint").As<uint64>();
+		uint64 Cycle = EventData.GetValue("Cycle").As<uint64>();
+		LogProvider.AppendMessage(LogPoint, Context.SessionContext.TimestampFromCycle(Cycle), EventData.GetAttachment());
+		break;
+	}
+	}
+}
diff --git a/Engine/Source/Developer/TraceServices/Private/Analyzers/LogTraceAnalysis.h b/Engine/Source/Developer/TraceServices/Private/Analyzers/LogTraceAnalysis.h
new file mode 100644
--- /dev/null
+++ b/Engine/Source/Developer/TraceServices/Private/Analyzers/LogTraceAnalysis.h
@@ -0,0 +1,33 @@
+// Copyright 1998-2019 Epic Games, Inc. All Rights Reserved.
+
+#pragma once
+
+#include "Trace/Trace.h"
+#include "Trace/Analyzer.h"
+
+namespace Trace
+{
+	class IAnalysisSession;
+	class FLogProvider;
+}
+
+class FLogTraceAnalyzer
+	: public Trace::IAnalyzer
+{
+public:
+	FLogTraceAnalyzer(Trace::IAnalysisSession& Session, Trace::FLogProvider& LogProvider);
+	virtual void OnAnalysisBegin(const FOnAnalysisContext& Context) override;
+	virtual void OnEvent(uint16 RouteId, const FOnEventContext& Context) override;
+	virtual void OnAnalysisEnd() override {};
+
+private:
+	enum : uint16
+	{
+		RouteId_LogCategory,
+		RouteId_LogMessageSpec,
+		RouteId_LogMessage,
+	};
+
+	Trace::IAnalysisSession& Session;
+	Trace::FLogProvider& LogProvider;
+};
diff --git a/Engine/Source/Developer/TraceServices/Private/Analyzers/MiscTraceAnalysis.cpp b/Engine/Source/Developer/TraceServices/Private/Analyzers/MiscTraceAnalysis.cpp
new file mode 100644
--- /dev/null
+++ b/Engine/Source/Developer/TraceServices/Private/Analyzers/MiscTraceAnalysis.cpp
@@ -0,0 +1,174 @@
+// Copyright 1998-2019 Epic Games, Inc. All Rights Reserved.
+#include "MiscTraceAnalysis.h"
+#include "TraceServices/Model/AnalysisSession.h"
+#include "Model/LogPrivate.h"
+#include "Model/ThreadsPrivate.h"
+#include "Model/BookmarksPrivate.h"
+#include "Model/FramesPrivate.h"
+#include "Common/Utils.h"
+
+FMiscTraceAnalyzer::FMiscTraceAnalyzer(Trace::IAnalysisSession& InSession,
+									   Trace::FThreadProvider& InThreadProvider,
+									   Trace::FBookmarkProvider& InBookmarkProvider,
+									   Trace::FLogProvider& InLogProvider,
+									   Trace::FFrameProvider& InFrameProvider)
+	: Session(InSession)
+	, ThreadProvider(InThreadProvider)
+	, BookmarkProvider(InBookmarkProvider)
+	, LogProvider(InLogProvider)
+	, FrameProvider(InFrameProvider)
+{
+	Trace::FLogCategory& BookmarkLogCategory = LogProvider.GetCategory(Trace::FLogProvider::ReservedLogCategory_Bookmark);
+	BookmarkLogCategory.Name = TEXT("LogBookmark");
+	BookmarkLogCategory.DefaultVerbosity = ELogVerbosity::All;
+}
+
+void FMiscTraceAnalyzer::OnAnalysisBegin(const FOnAnalysisContext& Context)
+{
+	auto& Builder = Context.InterfaceBuilder;
+
+	Builder.RouteEvent(RouteId_RegisterGameThread, "Misc", "RegisterGameThread");
+	Builder.RouteEvent(RouteId_CreateThread, "Misc", "CreateThread");
+	Builder.RouteEvent(RouteId_SetThreadGroup, "Misc", "SetThreadGroup");
+	Builder.RouteEvent(RouteId_BeginThreadGroupScope, "Misc", "BeginThreadGroupScope");
+	Builder.RouteEvent(RouteId_EndThreadGroupScope, "Misc", "EndThreadGroupScope");
+	Builder.RouteEvent(RouteId_BookmarkSpec, "Misc", "BookmarkSpec");
+	Builder.RouteEvent(RouteId_Bookmark, "Misc", "Bookmark");
+	Builder.RouteEvent(RouteId_BeginFrame, "Misc", "BeginFrame");
+	Builder.RouteEvent(RouteId_EndFrame, "Misc", "EndFrame");
+	Builder.RouteEvent(RouteId_BeginGameFrame, "Misc", "BeginGameFrame");
+	Builder.RouteEvent(RouteId_EndGameFrame, "Misc", "EndGameFrame");
+	Builder.RouteEvent(RouteId_BeginRenderFrame, "Misc", "BeginRenderFrame");
+	Builder.RouteEvent(RouteId_EndRenderFrame, "Misc", "EndRenderFrame");
+}
+
+void FMiscTraceAnalyzer::OnEvent(uint16 RouteId, const FOnEventContext& Context)
+{
+	Trace::FAnalysisSessionEditScope _(Session);
+
+	const auto& EventData = Context.EventData;
+	switch (RouteId)
+	{
+	case RouteId_RegisterGameThread:
+	{
+		uint32 ThreadId = EventData.GetValue("ThreadId").As<uint32>();
+		ThreadProvider.AddGameThread(ThreadId);
+		break;
+	}
+	case RouteId_CreateThread:
+	{
+		uint32 CreatedThreadId = EventData.GetValue("CreatedThreadId").As<uint32>();
+		EThreadPriority Priority = static_cast<EThreadPriority>(EventData.GetValue("Priority").As<uint32>());
+		ThreadProvider.AddThread(CreatedThreadId, reinterpret_cast<const TCHAR*>(EventData.GetAttachment()), Priority);
+		uint32 CurrentThreadId = EventData.GetValue("CurrentThreadId").As<uint32>();
+		FThreadState* ThreadState = GetThreadState(CurrentThreadId);
+		if (ThreadState->ThreadGroupStack.Num())
+		{
+			ThreadProvider.SetThreadGroup(CreatedThreadId, ThreadState->ThreadGroupStack.Top());
+		}
+		break;
+	}
+	case RouteId_SetThreadGroup:
+	{
+		const TCHAR* GroupName = Session.StoreString(ANSI_TO_TCHAR(reinterpret_cast<const char*>(EventData.GetAttachment())));
+		uint32 ThreadId = EventData.GetValue("ThreadId").As<uint32>();
+		ThreadProvider.SetThreadGroup(ThreadId, GroupName);
+		break;
+	}
+	case RouteId_BeginThreadGroupScope:
+	{
+		const TCHAR* GroupName = Session.StoreString(ANSI_TO_TCHAR(reinterpret_cast<const char*>(EventData.GetAttachment())));
+		uint32 CurrentThreadId = EventData.GetValue("CurrentThreadId").As<uint32>();
+		FThreadState* ThreadState = GetThreadState(CurrentThreadId);
+		ThreadState->ThreadGroupStack.Push(GroupName);
+		break;
+	}
+	case RouteId_EndThreadGroupScope:
+	{
+		uint32 CurrentThreadId = EventData.GetValue("CurrentThreadId").As<uint32>();
+		FThreadState* ThreadState = GetThreadState(CurrentThreadId);
+		ThreadState->ThreadGroupStack.Pop();
+		break;
+	}
+	case RouteId_BookmarkSpec:
+	{
+		uint64 BookmarkPoint = EventData.GetValue("BookmarkPoint").As<uint64>();
+		Trace::FBookmarkSpec& Spec = BookmarkProvider.GetSpec(BookmarkPoint);
+		Spec.Line = EventData.GetValue("Line").As<int32>();
+		const ANSICHAR* File = reinterpret_cast<const ANSICHAR*>(EventData.GetAttachment());
+		Spec.File = Session.StoreString(ANSI_TO_TCHAR(File));
+		Spec.FormatString = Session.StoreString(reinterpret_cast<const TCHAR*>(EventData.GetAttachment() + strlen(File) + 1));
+
+		Trace::FLogMessageSpec& LogMessageSpec = LogProvider.GetMessageSpec(BookmarkPoint);
+		LogMessageSpec.Category = &LogProvider.GetCategory(Trace::FLogProvider::ReservedLogCategory_Bookmark);
+		LogMessageSpec.Line = Spec.Line;
+		LogMessageSpec.File = Spec.File;
+		LogMessageSpec.FormatString = Spec.FormatString;
+		LogMessageSpec.Verbosity = ELogVerbosity::Log;
+		break;
+	}
+	case RouteId_Bookmark:
+	{
+		uint64 BookmarkPoint = EventData.GetValue("BookmarkPoint").As<uint64>();
+		uint64 Cycle = EventData.GetValue("Cycle").As<uint64>();
+		double Timestamp = Context.SessionContext.TimestampFromCycle(Cycle);
+		BookmarkProvider.AppendBookmark(Timestamp, BookmarkPoint, EventData.GetAttachment());
+		LogProvider.AppendMessage(BookmarkPoint, Timestamp, EventData.GetAttachment());
+		break;
+	}
+	case RouteId_BeginFrame:
+	{
+		uint64 Cycle = EventData.GetValue("Cycle").As<uint64>();
+		uint8 FrameType = EventData.GetValue("FrameType").As<uint8>();
+		check(FrameType < TraceFrameType_Count);
+		FrameProvider.BeginFrame(ETraceFrameType(FrameType), Context.SessionContext.TimestampFromCycle(Cycle));
+		break;
+	}
+	case RouteId_EndFrame:
+	{
+		uint64 Cycle = EventData.GetValue("Cycle").As<uint64>();
+		uint8 FrameType = EventData.GetValue("FrameType").As<uint8>();
+		check(FrameType < TraceFrameType_Count);
+		FrameProvider.EndFrame(ETraceFrameType(FrameType), Context.SessionContext.TimestampFromCycle(Cycle));
+		break;
+	}
+	case RouteId_BeginGameFrame:
+	case RouteId_EndGameFrame:
+	case RouteId_BeginRenderFrame:
+	case RouteId_EndRenderFrame:
+	{
+		ETraceFrameType FrameType;
+		if (RouteId == RouteId_BeginGameFrame || RouteId == RouteId_EndGameFrame)
+		{
+			FrameType = TraceFrameType_Game;
+		}
+		else
+		{
+			FrameType = TraceFrameType_Rendering;
+		}
+		const uint8* BufferPtr = EventData.GetAttachment();
+		uint64 CycleDiff = FTraceAnalyzerUtils::Decode7bit(BufferPtr);
+		uint64 Cycle = LastFrameCycle[FrameType] + CycleDiff;
+		LastFrameCycle[FrameType] = Cycle;
+		if (RouteId == RouteId_BeginGameFrame || RouteId == RouteId_BeginRenderFrame)
+		{
+			FrameProvider.BeginFrame(FrameType, Context.SessionContext.TimestampFromCycle(Cycle));
+		}
+		else
+		{
+			FrameProvider.EndFrame(FrameType, Context.SessionContext.TimestampFromCycle(Cycle));
+		}
+		break;
+	}
+	}
+}
+
+FMiscTraceAnalyzer::FThreadState* FMiscTraceAnalyzer::GetThreadState(uint32 ThreadId)
+{
+	if (!ThreadStateMap.Contains(ThreadId))
+	{
+		TSharedRef<FThreadState> ThreadState = MakeShared<FThreadState>();
+		ThreadStateMap.Add(ThreadId, ThreadState);
+	}
+	return &ThreadStateMap[ThreadId].Get();
+}
diff --git a/Engine/Source/Developer/TraceServices/Private/Analyzers/MiscTraceAnalysis.h b/Engine/Source/Developer/TraceServices/Private/Analyzers/MiscTraceAnalysis.h
new file mode 100644
--- /dev/null
+++ b/Engine/Source/Developer/TraceServices/Private/Analyzers/MiscTraceAnalysis.h
@@ -0,0 +1,66 @@
+// Copyright 1998-2019 Epic Games, Inc. All Rights Reserved.
+
+#pragma once
+
+#include "Trace/Trace.h"
+#include "Trace/Analyzer.h"
+#include "Templates/SharedPointer.h"
+#include "ProfilingDebugging/MiscTrace.h"
+#include "Common/PagedArray.h"
+
+namespace Trace
+{
+	class IAnalysisSession;
+	class FThreadProvider;
+	class FBookmarkProvider;
+	class FLogProvider;
+	class FFrameProvider;
+}
+
+class FMiscTraceAnalyzer
+	: public Trace::IAnalyzer
+{
+public:
+	FMiscTraceAnalyzer(Trace::IAnalysisSession& Session,
+					   Trace::FThreadProvider& ThreadProvider,
+					   Trace::FBookmarkProvider& BookmarkProvider,
+					   Trace::FLogProvider& LogProvider,
+					   Trace::FFrameProvider& FrameProvider);
+	virtual void OnAnalysisBegin(const FOnAnalysisContext& Context) override;
+	virtual void OnEvent(uint16 RouteId, const FOnEventContext& Context) override;
+	virtual void OnAnalysisEnd() override {};
+
+private:
+	enum : uint16
+	{
+		RouteId_RegisterGameThread,
+		RouteId_CreateThread,
+		RouteId_SetThreadGroup,
+		RouteId_BeginThreadGroupScope,
+		RouteId_EndThreadGroupScope,
+		RouteId_BookmarkSpec,
+		RouteId_Bookmark,
+		RouteId_BeginFrame,
+		RouteId_EndFrame,
+		RouteId_BeginGameFrame,
+		RouteId_EndGameFrame,
+		RouteId_BeginRenderFrame,
+		RouteId_EndRenderFrame,
+	};
+
+	struct FThreadState
+	{
+		TArray<const TCHAR*> ThreadGroupStack;
+	};
+
+	FThreadState* GetThreadState(uint32 ThreadId);
+
+	Trace::IAnalysisSession& Session;
+	Trace::FThreadProvider& ThreadProvider;
+	Trace::FBookmarkProvider& BookmarkProvider;
+	Trace::FLogProvider& LogProvider;
+	Trace::FFrameProvider& FrameProvider;
+	TMap<uint32, TSharedRef<FThreadState>> ThreadStateMap;
+	uint64 LastFrameCycle[TraceFrameType_Count] = { 0, 0 };
+};
+
diff --git a/Engine/Source/Developer/TraceServices/Private/Analyzers/PlatformFileTraceAnalysis.cpp b/Engine/Source/Developer/TraceServices/Private/Analyzers/PlatformFileTraceAnalysis.cpp
new file mode 100644
--- /dev/null
+++ b/Engine/Source/Developer/TraceServices/Private/Analyzers/PlatformFileTraceAnalysis.cpp
@@ -0,0 +1,152 @@
+// Copyright 1998-2019 Epic Games, Inc. All Rights Reserved.
+#include "PlatformFileTraceAnalysis.h"
+#include "AnalysisServicePrivate.h"
+#include "Model/FileActivity.h"
+
+FPlatformFileTraceAnalyzer::FPlatformFileTraceAnalyzer(Trace::IAnalysisSession& InSession, Trace::FFileActivityProvider& InFileActivityProvider)
+	: Session(InSession)
+	, FileActivityProvider(InFileActivityProvider)
+{
+
+}
+
+void FPlatformFileTraceAnalyzer::OnAnalysisBegin(const FOnAnalysisContext& Context)
+{
+	auto& Builder = Context.InterfaceBuilder;
+
+	Builder.RouteEvent(RouteId_BeginOpen, "PlatformFile", "BeginOpen");
+	Builder.RouteEvent(RouteId_EndOpen, "PlatformFile", "EndOpen");
+	Builder.RouteEvent(RouteId_BeginClose, "PlatformFile", "BeginClose");
+	Builder.RouteEvent(RouteId_EndClose, "PlatformFile", "EndClose");
+	Builder.RouteEvent(RouteId_BeginRead, "PlatformFile", "BeginRead");
+	Builder.RouteEvent(RouteId_EndRead, "PlatformFile", "EndRead");
+	Builder.RouteEvent(RouteId_BeginWrite, "PlatformFile", "BeginWrite");
+	Builder.RouteEvent(RouteId_EndWrite, "PlatformFile", "EndWrite");
+}
+
+void FPlatformFileTraceAnalyzer::OnEvent(uint16 RouteId, const FOnEventContext& Context)
+{
+	Trace::FAnalysisSessionEditScope _(Session);
+
+	const auto& EventData = Context.EventData;
+	switch (RouteId)
+	{
+	case RouteId_BeginOpen:
+	{
+		double Time = Context.SessionContext.TimestampFromCycle(EventData.GetValue("Cycle").As<uint64>());
+		uint64 TempHandle = EventData.GetValue("TempHandle").As<uint64>();
+		uint32 ThreadId = EventData.GetValue("ThreadId").As<uint32>();
+		check(!PendingOpenMap.Contains(TempHandle));
+		uint32 FileIndex = FileActivityProvider.GetFileIndex(reinterpret_cast<const TCHAR*>(EventData.GetAttachment()));
+		FPendingActivity& Open = PendingOpenMap.Add(TempHandle);
+		Open.ActivityIndex = FileActivityProvider.BeginActivity(FileIndex, Trace::FileActivityType_Open, 0, 0, Time);
+		Open.FileIndex = FileIndex;
+		break;
+	}
+	case RouteId_EndOpen:
+	{
+		double Time = Context.SessionContext.TimestampFromCycle(EventData.GetValue("Cycle").As<uint64>());
+		uint64 TempHandle = EventData.GetValue("TempHandle").As<uint64>();
+		uint64 FileHandle = EventData.GetValue("FileHandle").As<uint64>();
+		check(PendingOpenMap.Contains(TempHandle));
+		FPendingActivity& Open = PendingOpenMap[TempHandle];
+		if (FileHandle == uint64(-1)) // TODO: Portable invalid file handle?
+		{
+			FileActivityProvider.EndActivity(Open.FileIndex, Open.ActivityIndex, Time, true);
+		}
+		else
+		{
+			check(!OpenFilesMap.Contains(FileHandle));
+			OpenFilesMap.Add(FileHandle, Open.FileIndex);
+			FileActivityProvider.EndActivity(Open.FileIndex, Open.ActivityIndex, Time, false);
+		}
+		PendingOpenMap.Remove(TempHandle);
+		break;
+	}
+	case RouteId_BeginClose:
+	{
+		double Time = Context.SessionContext.TimestampFromCycle(EventData.GetValue("Cycle").As<uint64>());
+		uint64 TempHandle = EventData.GetValue("TempHandle").As<uint64>();
+		uint64 FileHandle = EventData.GetValue("FileHandle").As<uint64>();
+		uint32 ThreadId = EventData.GetValue("ThreadId").As<uint32>();
+		check(OpenFilesMap.Contains(FileHandle));
+		check(!PendingCloseMap.Contains(TempHandle));
+		uint64 FileIndex = OpenFilesMap[FileHandle];
+		OpenFilesMap.Remove(FileHandle);
+		FPendingActivity& Close = PendingCloseMap.Add(TempHandle);
+		Close.ActivityIndex = FileActivityProvider.BeginActivity(FileIndex, Trace::FileActivityType_Close, 0, 0, Time);
+		Close.FileIndex = FileIndex;
+		break;
+	}
+	case RouteId_EndClose:
+	{
+		double Time = Context.SessionContext.TimestampFromCycle(EventData.GetValue("Cycle").As<uint64>());
+		uint64 TempHandle = EventData.GetValue("TempHandle").As<uint64>();
+		check(PendingCloseMap.Contains(TempHandle));
+		FPendingActivity& Close = PendingCloseMap[TempHandle];
+		FileActivityProvider.EndActivity(Close.FileIndex, Close.ActivityIndex, Time, false);
+		PendingCloseMap.Remove(TempHandle);
+		break;
+	}
+	case RouteId_BeginRead:
+	{
+		double Time = Context.SessionContext.TimestampFromCycle(EventData.GetValue("Cycle").As<uint64>());
+		uint64 ReadHandle = EventData.GetValue("ReadHandle").As<uint64>();
+		uint64 FileHandle = EventData.GetValue("FileHandle").As<uint64>();
+		uint64 Offset = EventData.GetValue("Offset").As<uint64>();
+		uint64 Size = EventData.GetValue("Size").As<uint64>();
+		uint32 ThreadId = EventData.GetValue("ThreadId").As<uint32>();
+		check(OpenFilesMap.Contains(FileHandle));
+		uint32 FileIndex = OpenFilesMap[FileHandle];
+		uint64 ReadIndex = FileActivityProvider.BeginActivity(FileIndex, Trace::FileActivityType_Read, Offset, Size, Time);
+		FPendingActivity& Read = ActiveReadsMap.Add(ReadHandle);
+		Read.FileIndex = FileIndex;
+		Read.ActivityIndex = ReadIndex;
+		break;
+	}
+	case RouteId_EndRead:
+	{
+		double Time = Context.SessionContext.TimestampFromCycle(EventData.GetValue("Cycle").As<uint64>());
+		uint64 ReadHandle = EventData.GetValue("ReadHandle").As<uint64>();
+		uint64 SizeRead = EventData.GetValue("SizeRead").As<uint64>();
+		uint32 ThreadId = EventData.GetValue("ThreadId").As<uint32>();
+		check(ActiveReadsMap.Contains(ReadHandle));
+		const FPendingActivity& Read = ActiveReadsMap[ReadHandle];
+		FileActivityProvider.EndActivity(Read.FileIndex, Read.ActivityIndex, Time, false);
+		ActiveReadsMap.Remove(ReadHandle);
+		break;
+	}
+	case RouteId_BeginWrite:
+	{
+		double Time = Context.SessionContext.TimestampFromCycle(EventData.GetValue("Cycle").As<uint64>());
+		uint64 WriteHandle = EventData.GetValue("WriteHandle").As<uint64>();
+		uint64 FileHandle = EventData.GetValue("FileHandle").As<uint64>();
+		uint64 Offset = EventData.GetValue("Offset").As<uint64>();
+		uint64 Size = EventData.GetValue("Size").As<uint64>();
+		uint32 ThreadId = EventData.GetValue("ThreadId").As<uint32>();
+		check(OpenFilesMap.Contains(FileHandle));
+		uint32 FileIndex = OpenFilesMap[FileHandle];
+		uint64 WriteIndex = FileActivityProvider.BeginActivity(FileIndex, Trace::FileActivityType_Write, Offset, Size, Time);
+		FPendingActivity& Write = ActiveWritesMap.Add(WriteHandle);
+		Write.FileIndex = FileIndex;
+		Write.ActivityIndex = WriteIndex;
+		break;
+	}
+	case RouteId_EndWrite:
+	{
+		double Time = Context.SessionContext.TimestampFromCycle(EventData.GetValue("Cycle").As<uint64>());
+		uint64 WriteHandle = EventData.GetValue("WriteHandle").As<uint64>();
+		uint64 SizeWritten = EventData.GetValue("SizeWritten").As<uint64>();
+		uint32 ThreadId = EventData.GetValue("ThreadId").As<uint32>();
+		check(ActiveWritesMap.Contains(WriteHandle));
+		const FPendingActivity& Write = ActiveWritesMap[WriteHandle];
+		FileActivityProvider.EndActivity(Write.FileIndex, Write.ActivityIndex, Time, false);
+		ActiveWritesMap.Remove(WriteHandle);
+		break;
+	}
+	}
+}
+
+void FPlatformFileTraceAnalyzer::OnAnalysisEnd()
+{
+}
\ No newline at end of file
diff --git a/Engine/Source/Developer/TraceServices/Private/Analyzers/PlatformFileTraceAnalysis.h b/Engine/Source/Developer/TraceServices/Private/Analyzers/PlatformFileTraceAnalysis.h
new file mode 100644
--- /dev/null
+++ b/Engine/Source/Developer/TraceServices/Private/Analyzers/PlatformFileTraceAnalysis.h
@@ -0,0 +1,53 @@
+// Copyright 1998-2019 Epic Games, Inc. All Rights Reserved.
+
+#pragma once
+
+#include "Trace/Trace.h"
+#include "Trace/Analyzer.h"
+#include "Containers/Map.h"
+#include "TraceServices/Model/AnalysisSession.h"
+
+namespace Trace
+{
+	class FTimeline;
+	class FFileActivityProvider;
+}
+
+class FPlatformFileTraceAnalyzer
+	: public Trace::IAnalyzer
+{
+public:
+	FPlatformFileTraceAnalyzer(Trace::IAnalysisSession& Session, Trace::FFileActivityProvider& FileActivityProvider);
+	virtual void OnAnalysisBegin(const FOnAnalysisContext& Context) override;
+	virtual void OnEvent(uint16 RouteId, const FOnEventContext& Context) override;
+	virtual void OnAnalysisEnd() override;
+
+private:
+	enum : uint16
+	{
+		RouteId_BeginOpen,
+		RouteId_EndOpen,
+		RouteId_BeginClose,
+		RouteId_EndClose,
+		RouteId_BeginRead,
+		RouteId_EndRead,
+		RouteId_BeginWrite,
+		RouteId_EndWrite,
+
+		RouteId_Count
+	};
+
+	struct FPendingActivity
+	{
+		uint64 ActivityIndex;
+		uint32 FileIndex;
+	};
+
+	Trace::IAnalysisSession& Session;
+	Trace::FFileActivityProvider& FileActivityProvider;
+	TMap<uint64, uint32> OpenFilesMap;
+	TMap<uint64, FPendingActivity> PendingOpenMap;
+	TMap<uint64, FPendingActivity> PendingCloseMap;
+	TMap<uint64, FPendingActivity> ActiveReadsMap;
+	TMap<uint64, FPendingActivity> ActiveWritesMap;
+};
diff --git a/Engine/Source/Developer/TraceServices/Private/Analyzers/StatsTraceAnalysis.cpp b/Engine/Source/Developer/TraceServices/Private/Analyzers/StatsTraceAnalysis.cpp
new file mode 100644
--- /dev/null
+++ b/Engine/Source/Developer/TraceServices/Private/Analyzers/StatsTraceAnalysis.cpp
@@ -0,0 +1,140 @@
+// Copyright 1998-2019 Epic Games, Inc. All Rights Reserved.
+#include "StatsTraceAnalysis.h"
+#include "AnalysisServicePrivate.h"
+#include "Common/Utils.h"
+#include "Model/Counters.h"
+
+FStatsAnalyzer::FStatsAnalyzer(Trace::FAnalysisSession& InSession, Trace::FCounterProvider& InCounterProvider)
+	: Session(InSession)
+	, CounterProvider(InCounterProvider)
+{
+
+}
+
+void FStatsAnalyzer::OnAnalysisBegin(const FOnAnalysisContext& Context)
+{
+	auto& Builder = Context.InterfaceBuilder;
+
+	Builder.RouteEvent(RouteId_Spec, "Stats", "Spec");
+	Builder.RouteEvent(RouteId_EventBatch, "Stats", "EventBatch");
+}
+
+void FStatsAnalyzer::OnAnalysisEnd()
+{
+}
+
+void FStatsAnalyzer::OnEvent(uint16 RouteId, const FOnEventContext& Context)
+{
+	const auto& EventData = Context.EventData;
+	switch (RouteId)
+	{
+	case RouteId_Spec:
+	{
+		Trace::FAnalysisSessionEditScope _(Session);
+		uint32 StatId = EventData.GetValue("Id").As<uint32>();
+		check(!CountersMap.Contains(StatId));
+		const ANSICHAR* Name = reinterpret_cast<const ANSICHAR*>(EventData.GetAttachment());
+		const TCHAR* Description = reinterpret_cast<const TCHAR*>(EventData.GetAttachment() + strlen(Name) + 1);
+		Trace::ECounterDisplayHint DisplayHint = Trace::CounterDisplayHint_None;
+		if (EventData.GetValue("IsFloatingPoint").As<bool>())
+		{
+			DisplayHint = Trace::CounterDisplayHint_FloatingPoint;
+		}
+		else if (EventData.GetValue("IsMemory").As<bool>())
+		{
+			DisplayHint = Trace::CounterDisplayHint_Memory;
+		}
+		Trace::FCounterInternal* Counter = CounterProvider.CreateCounter(ANSI_TO_TCHAR(Name), Description, DisplayHint);
+		CountersMap.Add(StatId, Counter);
+		break;
+	}
+	case RouteId_EventBatch:
+	{
+		Trace::FAnalysisSessionEditScope _(Session);
+		uint32 ThreadId = EventData.GetValue("ThreadId").As<uint32>();
+		TSharedRef<FThreadState> ThreadState = GetThreadState(ThreadId);
+		uint64 BufferSize = EventData.GetAttachmentSize();
+		const uint8* BufferPtr = EventData.GetAttachment();
+		const uint8* BufferEnd = BufferPtr + BufferSize;
+		while (BufferPtr < BufferEnd)
+		{
+			enum EOpType
+			{
+				Increment = 0,
+				Decrement = 1,
+				AddInteger = 2,
+				SetInteger = 3,
+				AddFloat = 4,
+				SetFloat = 5,
+			};
+
+			uint64 DecodedIdAndOp = FTraceAnalyzerUtils::Decode7bit(BufferPtr);
+			uint32 StatId = DecodedIdAndOp >> 3;
+			//check(CountersMap.Contains(StatId));
+			Trace::FCounterInternal& Counter = *CountersMap[StatId];
+			uint8 Op = DecodedIdAndOp & 0x7;
+			uint64 CycleDiff = FTraceAnalyzerUtils::Decode7bit(BufferPtr);
+			uint64 Cycle = ThreadState->LastCycle + CycleDiff;
+			double Time = Context.SessionContext.TimestampFromCycle(Cycle);
+			ThreadState->LastCycle = Cycle;
+			switch (Op)
+			{
+			case Increment:
+			{
+				CounterProvider.Add(Counter, Time, int64(1));
+				break;
+			}
+			case Decrement:
+			{
+				CounterProvider.Add(Counter, Time, int64(-1));
+				break;
+			}
+			case AddInteger:
+			{
+				int64 Amount = FTraceAnalyzerUtils::DecodeZigZag(BufferPtr);
+				CounterProvider.Add(Counter, Time, Amount);
+				break;
+			}
+			case SetInteger:
+			{
+				int64 Value = FTraceAnalyzerUtils::DecodeZigZag(BufferPtr);
+				CounterProvider.Set(Counter, Time, Value);
+				break;
+			}
+			case AddFloat:
+			{
+				double Amount;
+				memcpy(&Amount, BufferPtr, sizeof(double));
+				BufferPtr += sizeof(double);
+				CounterProvider.Add(Counter, Time, Amount);
+				break;
+			}
+			case SetFloat:
+			{
+				double Value;
+				memcpy(&Value, BufferPtr, sizeof(double));
+				BufferPtr += sizeof(double);
+				break;
+			}
+			}
+		}
+		check(BufferPtr == BufferEnd);
+		break;
+	}
+	}
+}
+
+TSharedRef<FStatsAnalyzer::FThreadState> FStatsAnalyzer::GetThreadState(uint32 ThreadId)
+{
+	if (!ThreadStatesMap.Contains(ThreadId))
+	{
+		TSharedRef<FThreadState> ThreadState = MakeShared<FThreadState>();
+		ThreadState->LastCycle = 0;
+		ThreadStatesMap.Add(ThreadId, ThreadState);
+		return ThreadState;
+	}
+	else
+	{
+		return ThreadStatesMap[ThreadId];
+	}
+}
diff --git a/Engine/Source/Developer/TraceServices/Private/Analyzers/StatsTraceAnalysis.h b/Engine/Source/Developer/TraceServices/Private/Analyzers/StatsTraceAnalysis.h
new file mode 100644
--- /dev/null
+++ b/Engine/Source/Developer/TraceServices/Private/Analyzers/StatsTraceAnalysis.h
@@ -0,0 +1,43 @@
+// Copyright 1998-2019 Epic Games, Inc. All Rights Reserved.
+
+#pragma once
+
+#include "Trace/Trace.h"
+#include "Trace/Analyzer.h"
+#include "Templates/SharedPointer.h"
+
+namespace Trace
+{
+	class FAnalysisSession;
+	class FCounterInternal;
+	class FCounterProvider;
+}
+
+class FStatsAnalyzer
+	: public Trace::IAnalyzer
+{
+public:
+	FStatsAnalyzer(Trace::FAnalysisSession& Session, Trace::FCounterProvider& CounterProvider);
+	virtual void OnAnalysisBegin(const FOnAnalysisContext& Context) override;
+	virtual void OnEvent(uint16 RouteId, const FOnEventContext& Context) override;
+	virtual void OnAnalysisEnd() override;
+
+private:
+	enum : uint16
+	{
+		RouteId_Spec,
+		RouteId_EventBatch,
+	};
+
+	struct FThreadState
+	{
+		uint64 LastCycle = 0;
+	};
+
+	TSharedRef<FThreadState> GetThreadState(uint32 ThreadId);
+
+	Trace::FAnalysisSession& Session;
+	Trace::FCounterProvider& CounterProvider;
+	TMap<uint32, Trace::FCounterInternal*> CountersMap;
+	TMap<uint32, TSharedRef<FThreadState>> ThreadStatesMap;
+};
diff --git a/Engine/Source/Developer/TraceServices/Private/Common/FormatArgs.h b/Engine/Source/Developer/TraceServices/Private/Common/FormatArgs.h
new file mode 100644
--- /dev/null
+++ b/Engine/Source/Developer/TraceServices/Private/Common/FormatArgs.h
@@ -0,0 +1,43 @@
+// Copyright 1998-2019 Epic Games, Inc. All Rights Reserved.
+
+#pragma once
+
+#include "CoreTypes.h"
+
+namespace Trace
+{
+
+struct FFormatArgsHelper
+{
+	static void Format(TCHAR* Out, uint64 MaxOut, const TCHAR* FormatString, const uint8* FormatArgs);
+
+private:
+	struct FFormatArgSpec
+	{
+		uint64 PassthroughLength;
+		TCHAR FormatString[255];
+		uint8 ExpectedTypeCategory;
+		uint8 AdditionalIntegerArgumentCount;
+		bool Valid;
+		bool NothingPrinted;
+	};
+
+	struct FFormatArgsStreamContext
+	{
+		uint8 ArgumentCount;
+		uint8 ArgumentTypeCategory;
+		uint8 ArgumentTypeSize;
+		const uint8* DescriptorPtr;
+		const uint8* PayloadPtr;
+	};
+
+	static const TCHAR* ExtractNextFormatArg(const TCHAR* FormatString, FFormatArgSpec& Spec);
+	static void InitArgumentStream(FFormatArgsStreamContext& Context, const uint8* ArgumentsData);
+	static bool AdvanceArgumentStream(FFormatArgsStreamContext& Context);
+	static uint64 ExtractIntegerArgument(FFormatArgsStreamContext& ArgStream);
+	static double ExtractFloatingPointArgument(FFormatArgsStreamContext& ArgStream);
+	static const TCHAR* ExtractStringArgument(FFormatArgsStreamContext& ArgStream);
+	static int32 FormatArgument(TCHAR* Out, uint64 MaxOut, const FFormatArgSpec& ArgSpec, FFormatArgsStreamContext& ArgStream);
+};
+
+}
\ No newline at end of file
