Updated insights to not use deprecated TabManager API.

diff --git a/Engine/Source/Developer/TraceInsights/Private/Insights/InsightsManager.cpp b/Engine/Source/Developer/TraceInsights/Private/Insights/InsightsManager.cpp
--- a/Engine/Source/Developer/TraceInsights/Private/Insights/InsightsManager.cpp
+++ b/Engine/Source/Developer/TraceInsights/Private/Insights/InsightsManager.cpp
@@ -149,11 +149,11 @@ void FInsightsManager::OnSessionChanged()
 void FInsightsManager::SpawnAndActivateTabs()
 {
 	// Open tabs for profilers.
-	if (FGlobalTabmanager::Get()->CanSpawnTab(FInsightsManagerTabs::TimingProfilerTabId))
+	if (FGlobalTabmanager::Get()->HasTabSpawner(FInsightsManagerTabs::TimingProfilerTabId))
 	{
 		FGlobalTabmanager::Get()->InvokeTab(FInsightsManagerTabs::TimingProfilerTabId);
 	}
-	if (FGlobalTabmanager::Get()->CanSpawnTab(FInsightsManagerTabs::IoProfilerTabId))
+	if (FGlobalTabmanager::Get()->HasTabSpawner(FInsightsManagerTabs::IoProfilerTabId))
 	{
 		FGlobalTabmanager::Get()->InvokeTab(FInsightsManagerTabs::IoProfilerTabId);
 	}
diff --git a/Engine/Source/Developer/TraceInsights/Private/Insights/Widgets/STimingProfilerWindow.cpp b/Engine/Source/Developer/TraceInsights/Private/Insights/Widgets/STimingProfilerWindow.cpp
--- a/Engine/Source/Developer/TraceInsights/Private/Insights/Widgets/STimingProfilerWindow.cpp
+++ b/Engine/Source/Developer/TraceInsights/Private/Insights/Widgets/STimingProfilerWindow.cpp
@@ -370,7 +370,7 @@ void STimingProfilerWindow::FillMenu(FMenuBuilder& MenuBuilder, const TSharedPtr
 
 void STimingProfilerWindow::ShowTab(const FName& TabID)
 {
-	if (TabManager->CanSpawnTab(TabID))
+	if (TabManager->HasTabSpawner(TabID))
 	{
 		TabManager->InvokeTab(TabID);
 	}
