[Insights] Integrated following CLs from //UE4/dev-core: 6893318, 6909107, 6911320, 6913463, 6926018, 6926170, 6953227

- Asset Loading: Fixed crash (or corruption of the Tining view's "Tracks" drop down) when a session is re-loaded (missing re-initialization of Timing view when session changes).
- Timing: Fixed Timers and Stats Counters tabs to correctly re-sync the list of timers / stats after tab is re-opened.
- Timing: Removed Graph track (tab) from the toolbar. Fixes UE-75603.
- Fixed UE-75591 (Timing Insights Toolbar window status does not update when closing Windows via X).
- Frames track: Fixed panning with mouse (to be able to horizontally pan also using left mouse button, not only with right mouse button).
- Asset Loading: Improved table visualization for aggregated stats (Event Aggregation and Object Type Aggregation top10 tables).
- Asset Loading: Minor change to tooltip for hovered asset loading event (renamed "Package Type" and "Export Type" to "Package Event" and "Export Event" respectively).
- Timers / Stats Counters: Fixed default sorting column/mode not working as intended.
- Timers / Stats Counters: Fixed UE-75578 (Timing Insights Timer's & Stats Counters Windows do not use standard UE4 radio button style). Made selected option for radio buttons be enabled all the time.
- Fixed UE-75588 (Timing Insights Timing View spawns an empty box when right clicking).
- Start Page: Updated UI layout:
    * Removed Live and Latest buttons.
    * Removed text info.
    * Added list of available sessions (with "LIVE" status). Double clicking a session will start analysis for it.
    * Added "Open" button to start analysis for selected session. Moved "Load..." under the drop down menu (next to Open) and renamed it as "Open File". Also the drop down now only shows top10 most recently created sessions.
    * Separated "New Connection" panel from "Trace Recorder" panel.
    * Removed control for enabled modules (for connections / live sessions) from Start Page / Trace Recorder.
- Timing view: Sorted the thread groups (from "Tracks" menu) in the same order as the visible tracks. Also added the number of tracks currently available for each group next to group name (also in Tracks menu).
- Timing view: Fixed memory leak regarding aggregated stats for loading tracks.
- Increased version to 0.07
- Fixed UE-76068 (When grouping Timing Insight's Timers by Name, "n" is lowercase)
- Fixed UE-76084 (Selecting a Timer in the Asset Loading Insight changes changes the selected Timer in the Timing Insight's Timers

diff --git a/Engine/Source/Developer/TraceInsights/Private/Insights/Common/PaintUtils.h b/Engine/Source/Developer/TraceInsights/Private/Insights/Common/PaintUtils.h
--- a/Engine/Source/Developer/TraceInsights/Private/Insights/Common/PaintUtils.h
+++ b/Engine/Source/Developer/TraceInsights/Private/Insights/Common/PaintUtils.h
@@ -3,6 +3,8 @@
 #pragma once
 
 #include "CoreMinimal.h"
+#include "Fonts/FontMeasure.h"
+#include "Framework/Application/SlateApplication.h"
 #include "Rendering/DrawElements.h"
 
 #define MAKE_PAINT_GEOMETRY_PT(Geometry, X, Y)       Geometry.ToPaintGeometry(FSlateLayoutTransform(1.0f, FVector2D(X, Y)))
@@ -74,6 +76,22 @@ struct FDrawContext
 		FSlateDrawElement::MakeText(ElementList, InLayer, MAKE_PAINT_GEOMETRY_PT(Geometry, X, Y), Text, StartIndex, EndIndex, Font, DrawEffects, Color);
 	}
 
+	inline void DrawTextAligned(EHorizontalAlignment HAlign, const float X, const float Y, const FString& Text, const FSlateFontInfo& Font, const FLinearColor& Color) const
+	{
+		const TSharedRef<FSlateFontMeasure> FontMeasureService = FSlateApplication::Get().GetRenderer()->GetFontMeasureService();
+		const float TextWidth = FontMeasureService->Measure(Text, Font).X;
+		float TextX = X;
+		if (HAlign == HAlign_Right)
+		{
+			TextX -= TextWidth;
+		}
+		else if (HAlign == HAlign_Center)
+		{
+			TextX -= TextWidth / 2;
+		}
+		FSlateDrawElement::MakeText(ElementList, LayerId, MAKE_PAINT_GEOMETRY_PT(Geometry, TextX, Y), Text, Font, DrawEffects, Color);
+	}
+
 	//inline void IncrementLayer() const { ++LayerId; }
 
 	// Accessors
diff --git a/Engine/Source/Developer/TraceInsights/Private/Insights/InsightsCommands.cpp b/Engine/Source/Developer/TraceInsights/Private/Insights/InsightsCommands.cpp
--- a/Engine/Source/Developer/TraceInsights/Private/Insights/InsightsCommands.cpp
+++ b/Engine/Source/Developer/TraceInsights/Private/Insights/InsightsCommands.cpp
@@ -36,7 +36,6 @@ FInsightsCommands::FInsightsCommands()
 PRAGMA_DISABLE_OPTIMIZATION
 void FInsightsCommands::RegisterCommands()
 {
-	UI_COMMAND(InsightsManager_Live, "Live", "Loads profiler data from a live or last trace session", EUserInterfaceActionType::Button, FInputChord());
 	UI_COMMAND(InsightsManager_Load, "Load...", "Loads profiler data from a trace file", EUserInterfaceActionType::Button, FInputChord(EModifierKey::Control, EKeys::L));
 	UI_COMMAND(ToggleDebugInfo, "Debug", "Toggles the display of debug info", EUserInterfaceActionType::ToggleButton, FInputChord(EModifierKey::Control, EKeys::D));
 	UI_COMMAND(OpenSettings, "Settings", "Opens the Unreal Insights settings", EUserInterfaceActionType::Button, FInputChord(EModifierKey::Control, EKeys::O));
@@ -78,37 +77,6 @@ ECheckBoxState FInsightsActionManager::ToggleDebugInfo_GetCheckState() const
 	return bIsDebugInfoEnabled ? ECheckBoxState::Checked : ECheckBoxState::Unchecked;
 }
 
-////////////////////////////////////////////////////////////////////////////////////////////////////
-// InsightsManager_Live
-////////////////////////////////////////////////////////////////////////////////////////////////////
-
-void FInsightsActionManager::Map_InsightsManager_Live()
-{
-	FUIAction UIAction;
-	UIAction.ExecuteAction = FExecuteAction::CreateRaw(this, &FInsightsActionManager::InsightsManager_Live_Execute);
-	UIAction.CanExecuteAction = FCanExecuteAction::CreateRaw(this, &FInsightsActionManager::InsightsManager_Live_CanExecute);
-
-	This->CommandList->MapAction(This->GetCommands().InsightsManager_Live, UIAction);
-}
-
-////////////////////////////////////////////////////////////////////////////////////////////////////
-
-void FInsightsActionManager::InsightsManager_Live_Execute()
-{
-	This->LoadLastLiveSession();
-	if (!This->GetSession().IsValid())
-	{
-		This->LoadLastSession();
-	}
-}
-
-////////////////////////////////////////////////////////////////////////////////////////////////////
-
-bool FInsightsActionManager::InsightsManager_Live_CanExecute() const
-{
-	return true;
-}
-
 ////////////////////////////////////////////////////////////////////////////////////////////////////
 // InsightsManager_Load
 ////////////////////////////////////////////////////////////////////////////////////////////////////
diff --git a/Engine/Source/Developer/TraceInsights/Private/Insights/InsightsCommands.h b/Engine/Source/Developer/TraceInsights/Private/Insights/InsightsCommands.h
--- a/Engine/Source/Developer/TraceInsights/Private/Insights/InsightsCommands.h
+++ b/Engine/Source/Developer/TraceInsights/Private/Insights/InsightsCommands.h
@@ -27,9 +27,6 @@ public:
 	//     const FUIAction <CommandName>_Custom(...) const;
 	//////////////////////////////////////////////////
 
-	/** Load profiler data from a live trace session. Global version. */
-	TSharedPtr<FUICommandInfo> InsightsManager_Live;
-
 	/** Load profiler data from a trace file. Global version. */
 	TSharedPtr<FUICommandInfo> InsightsManager_Load;
 
@@ -55,15 +52,6 @@ private:
 		: This(Instance)
 	{}
 
-	//////////////////////////////////////////////////
-	// InsightsManager_Live
-
-public:
-	void Map_InsightsManager_Live(); /**< Maps UI command info InsightsManager_Live with the specified UI command list. */
-protected:
-	void InsightsManager_Live_Execute(); /**< Handles FExecuteAction for InsightsManager_Live. */
-	bool InsightsManager_Live_CanExecute() const; /**< Handles FCanExecuteAction for InsightsManager_Live. */
-
 	//////////////////////////////////////////////////
 	// InsightsManager_Load
 
diff --git a/Engine/Source/Developer/TraceInsights/Private/Insights/InsightsManager.cpp b/Engine/Source/Developer/TraceInsights/Private/Insights/InsightsManager.cpp
--- a/Engine/Source/Developer/TraceInsights/Private/Insights/InsightsManager.cpp
+++ b/Engine/Source/Developer/TraceInsights/Private/Insights/InsightsManager.cpp
@@ -8,6 +8,7 @@
 #include "Templates/UniquePtr.h"
 
 // Insights
+#include "Insights/IoProfilerManager.h"
 #include "Insights/TimingProfilerManager.h"
 #include "Insights/Widgets/SStartPageWindow.h"
 #include "Insights/Widgets/STimingProfilerWindow.h"
@@ -53,7 +54,6 @@ void FInsightsManager::PostConstructor()
 
 void FInsightsManager::BindCommands()
 {
-	ActionManager.Map_InsightsManager_Live();
 	ActionManager.Map_InsightsManager_Load();
 	ActionManager.Map_ToggleDebugInfo_Global();
 	ActionManager.Map_OpenSettings_Global();
@@ -141,6 +141,12 @@ void FInsightsManager::OnSessionChanged()
 		TimingProfilerManager->OnSessionChanged();
 	}
 
+	if (TSharedPtr<FIoProfilerManager> IoProfilerManager = FIoProfilerManager::Get())
+	{
+		// FIXME: make IoProfilerManager to register to SessionChangedEvent instead
+		IoProfilerManager->OnSessionChanged();
+	}
+
 	SessionChangedEvent.Broadcast();
 }
 
diff --git a/Engine/Source/Developer/TraceInsights/Private/Insights/InsightsManager.h b/Engine/Source/Developer/TraceInsights/Private/Insights/InsightsManager.h
--- a/Engine/Source/Developer/TraceInsights/Private/Insights/InsightsManager.h
+++ b/Engine/Source/Developer/TraceInsights/Private/Insights/InsightsManager.h
@@ -26,31 +26,6 @@ struct FInsightsManagerTabs
 	static const FName IoProfilerTabId;
 };
 
-////////////////////////////////////////////////////////////////////////////////////////////////////
-/** Enumerates loading a trace file progress states. */
-
-enum class ELoadingProgressState
-{
-	Started,
-	InProgress,
-	Loaded,
-	Failed,
-	Cancelled,
-
-	/** Invalid enum type, may be used as a number of enumerations. */
-	InvalidOrMax,
-};
-
-////////////////////////////////////////////////////////////////////////////////////////////////////
-
-enum class EInsightsNotificationType
-{
-	LoadingTraceFile,
-
-	/** Invalid enum type, may be used as a number of enumerations. */
-	InvalidOrMax,
-};
-
 ////////////////////////////////////////////////////////////////////////////////////////////////////
 /**
  * This class manages following areas:
diff --git a/Engine/Source/Developer/TraceInsights/Private/Insights/InsightsStyle.cpp b/Engine/Source/Developer/TraceInsights/Private/Insights/InsightsStyle.cpp
--- a/Engine/Source/Developer/TraceInsights/Private/Insights/InsightsStyle.cpp
+++ b/Engine/Source/Developer/TraceInsights/Private/Insights/InsightsStyle.cpp
@@ -78,11 +78,11 @@ TSharedRef<FSlateStyleSet> FInsightsStyle::Create()
 	//////////////////////////////////////////////////
 	// Start Page buttons
 
-	Style.Set("Live.Icon.Large", new IMAGE_BRUSH("Icons/Profiler/Profiler_LoadMultiple_Profiler_40x", FVector2D(32.0f, 32.0f)));
-	Style.Set("Live.Icon.Small", new IMAGE_BRUSH("Icons/Profiler/Profiler_Load_Profiler_40x", FVector2D(16.0f, 16.0f)));
+	Style.Set("Open.Icon.Large", new IMAGE_BRUSH("Icons/Profiler/Profiler_LoadMultiple_Profiler_40x", FVector2D(32.0f, 32.0f)));
+	Style.Set("Open.Icon.Small", new IMAGE_BRUSH("Icons/Profiler/Profiler_Load_Profiler_40x", FVector2D(16.0f, 16.0f)));
 
-	Style.Set("Load.Icon.Large", new IMAGE_BRUSH("Icons/LV_Load", FVector2D(32.0f, 32.0f)));
-	Style.Set("Load.Icon.Small", new IMAGE_BRUSH("Icons/LV_Load", FVector2D(16.0f, 16.0f)));
+	Style.Set("OpenFile.Icon.Large", new IMAGE_BRUSH("Icons/LV_Load", FVector2D(32.0f, 32.0f)));
+	Style.Set("OpenFile.Icon.Small", new IMAGE_BRUSH("Icons/LV_Load", FVector2D(16.0f, 16.0f)));
 
 	//////////////////////////////////////////////////
 
diff --git a/Engine/Source/Developer/TraceInsights/Private/Insights/TimingProfilerCommands.cpp b/Engine/Source/Developer/TraceInsights/Private/Insights/TimingProfilerCommands.cpp
--- a/Engine/Source/Developer/TraceInsights/Private/Insights/TimingProfilerCommands.cpp
+++ b/Engine/Source/Developer/TraceInsights/Private/Insights/TimingProfilerCommands.cpp
@@ -9,6 +9,7 @@
 
 // Insights
 #include "Insights/InsightsManager.h"
+#include "Insights/InsightsStyle.h"
 #include "Insights/TimingProfilerManager.h"
 
 ////////////////////////////////////////////////////////////////////////////////////////////////////
@@ -39,7 +40,7 @@ void FTimingProfilerMenuBuilder::AddMenuEntry(FMenuBuilder& MenuBuilder, const T
 FTimingProfilerCommands::FTimingProfilerCommands()
 	: TCommands<FTimingProfilerCommands>(
 		TEXT("TimingProfilerCommand"), // Context name for fast lookup
-		NSLOCTEXT("Contexts", "TimingProfilerCommand", "Timing Profiler Command"), // Localized context name for displaying
+		NSLOCTEXT("Contexts", "TimingProfilerCommand", "Timing Insights"), // Localized context name for displaying
 		NAME_None, // Parent
 		FEditorStyle::GetStyleSetName() // Icon Style Set
 	)
@@ -61,6 +62,31 @@ void FTimingProfilerCommands::RegisterCommands()
 }
 PRAGMA_ENABLE_OPTIMIZATION
 
+////////////////////////////////////////////////////////////////////////////////////////////////////
+// FTimingViewCommands
+////////////////////////////////////////////////////////////////////////////////////////////////////
+
+FTimingViewCommands::FTimingViewCommands()
+	: TCommands<FTimingViewCommands>(
+		TEXT("TimingViewCommand"), // Context name for fast lookup
+		NSLOCTEXT("Contexts", "TimingViewCommand", "Timing Insights"), // Localized context name for displaying
+		NAME_None, // Parent
+		FInsightsStyle::GetStyleSetName() // Icon Style Set
+	)
+{
+}
+
+////////////////////////////////////////////////////////////////////////////////////////////////////
+
+// UI_COMMAND takes long for the compiler to optimize
+PRAGMA_DISABLE_OPTIMIZATION
+void FTimingViewCommands::RegisterCommands()
+{
+	UI_COMMAND(ShowAllGpuTracks, "GPU Track", "Show/hide the GPU track", EUserInterfaceActionType::ToggleButton, FInputChord(EKeys::Y));
+	UI_COMMAND(ShowAllCpuTracks, "CPU Thread Tracks", "Show/hide all CPU tracks (and all CPU thread groups)", EUserInterfaceActionType::ToggleButton, FInputChord(EKeys::U));
+}
+PRAGMA_ENABLE_OPTIMIZATION
+
 ////////////////////////////////////////////////////////////////////////////////////////////////////
 // Toggle Commands
 ////////////////////////////////////////////////////////////////////////////////////////////////////
diff --git a/Engine/Source/Developer/TraceInsights/Private/Insights/TimingProfilerCommands.h b/Engine/Source/Developer/TraceInsights/Private/Insights/TimingProfilerCommands.h
--- a/Engine/Source/Developer/TraceInsights/Private/Insights/TimingProfilerCommands.h
+++ b/Engine/Source/Developer/TraceInsights/Private/Insights/TimingProfilerCommands.h
@@ -48,6 +48,25 @@ public:
 	TSharedPtr<FUICommandInfo> ToggleLogViewVisibility;
 };
 
+////////////////////////////////////////////////////////////////////////////////////////////////////
+
+class FTimingViewCommands : public TCommands<FTimingViewCommands>
+{
+public:
+	/** Default constructor. */
+	FTimingViewCommands();
+
+	/** Initialize commands. */
+	virtual void RegisterCommands() override;
+
+public:
+	/** Toggles visibility for GPU thread track. */
+	TSharedPtr<FUICommandInfo> ShowAllGpuTracks;
+
+	/** Toggles visibility for all CPU thread tracks at once. */
+	TSharedPtr<FUICommandInfo> ShowAllCpuTracks;
+};
+
 ////////////////////////////////////////////////////////////////////////////////////////////////////
 /**
  * Menu builder. Helper class for adding a customized menu entry using the global UI command info.
diff --git a/Engine/Source/Developer/TraceInsights/Private/Insights/ViewModels/TimingEventsTrack.h b/Engine/Source/Developer/TraceInsights/Private/Insights/ViewModels/TimingEventsTrack.h
--- a/Engine/Source/Developer/TraceInsights/Private/Insights/ViewModels/TimingEventsTrack.h
+++ b/Engine/Source/Developer/TraceInsights/Private/Insights/ViewModels/TimingEventsTrack.h
@@ -4,26 +4,27 @@
 
 #include "CoreMinimal.h"
 
+// Insights
 #include "Insights/ViewModels/BaseTimingTrack.h"
 
 ////////////////////////////////////////////////////////////////////////////////////////////////////
 
-static const float RealMinTimelineH = 13.0f;
+struct FTimingEventsTrackLayout
+{
+	static constexpr float RealMinTimelineH = 13.0f;
 
-static const float NormalLayoutEventH = 14.0f;
-static const float NormalLayoutEventDY = 2.0f;
-static const float NormalLayoutTimelineDY = 14.0f;
-static const float NormalLayoutMinTimelineH = 0.0f;
+	static constexpr float NormalLayoutEventH = 14.0f;
+	static constexpr float NormalLayoutEventDY = 2.0f;
+	static constexpr float NormalLayoutTimelineDY = 14.0f;
+	static constexpr float NormalLayoutMinTimelineH = 0.0f;
 
-static const float CompactLayoutEventH = 2.0f;
-static const float CompactLayoutEventDY = 1.0f;
-static const float CompactLayoutTimelineDY = 3.0f;
-static const float CompactLayoutMinTimelineH = 0.0f;
+	static constexpr float CompactLayoutEventH = 2.0f;
+	static constexpr float CompactLayoutEventDY = 1.0f;
+	static constexpr float CompactLayoutTimelineDY = 3.0f;
+	static constexpr float CompactLayoutMinTimelineH = 0.0f;
 
-////////////////////////////////////////////////////////////////////////////////////////////////////
+	//////////////////////////////////////////////////
 
-struct FTimingEventsTrackLayout
-{
 	bool bIsCompactMode;
 
 	float EventH; // height of a timing event, in Slate units
@@ -32,6 +33,8 @@ struct FTimingEventsTrackLayout
 	float MinTimelineH;
 	float TargetMinTimelineH;
 
+	//////////////////////////////////////////////////
+
 	float GetLaneY(uint32 Depth) const { return 1.0f + TimelineDY + Depth * (EventDY + EventH); }
 
 	void ForceNormalMode();
@@ -83,7 +86,7 @@ public:
 	int32 Depth; // number of lanes == Depth + 1
 	bool bIsCollapsed;
 
-	// Cached OnPaint state.
+	// TODO: Cached OnPaint state.
 	//TArray<FEventBoxInfo> Boxes;
 	//TArray<FEventBoxInfo> MergedBorderBoxes;
 	//TArray<FEventBoxInfo> Borders;
diff --git a/Engine/Source/Developer/TraceInsights/Private/Insights/ViewModels/TimingViewDrawHelper.cpp b/Engine/Source/Developer/TraceInsights/Private/Insights/ViewModels/TimingViewDrawHelper.cpp
--- a/Engine/Source/Developer/TraceInsights/Private/Insights/ViewModels/TimingViewDrawHelper.cpp
+++ b/Engine/Source/Developer/TraceInsights/Private/Insights/ViewModels/TimingViewDrawHelper.cpp
@@ -168,9 +168,9 @@ void FTimingViewDrawHelper::EndTimeline(FTimingEventsTrack& Track)
 		// 1.0f is for horizontal line between timelines
 		NewH = 1.0f + Layout.TimelineDY + Layout.EventH * (MaxDepth + 1) + Layout.EventDY * MaxDepth + Layout.TimelineDY;
 
-		if (NewH < RealMinTimelineH)
+		if (NewH < FTimingEventsTrackLayout::RealMinTimelineH)
 		{
-			NewH = RealMinTimelineH;
+			NewH = FTimingEventsTrackLayout::RealMinTimelineH;
 		}
 	}
 
diff --git a/Engine/Source/Developer/TraceInsights/Private/Insights/Widgets/SFrameTrack.cpp b/Engine/Source/Developer/TraceInsights/Private/Insights/Widgets/SFrameTrack.cpp
--- a/Engine/Source/Developer/TraceInsights/Private/Insights/Widgets/SFrameTrack.cpp
+++ b/Engine/Source/Developer/TraceInsights/Private/Insights/Widgets/SFrameTrack.cpp
@@ -72,7 +72,7 @@ void SFrameTrack::Reset()
 	bIsLMB_Pressed = false;
 	bIsRMB_Pressed = false;
 
-	bIsRMB_Scrolling = false;
+	bIsScrolling = false;
 
 	SelectionStartFrameIndex = 0;
 	SelectionEndFrameIndex = 0;
@@ -134,7 +134,7 @@ void SFrameTrack::Tick(const FGeometry& AllottedGeometry, const double InCurrent
 
 	ThisGeometry = AllottedGeometry;
 
-	if (!bIsRMB_Scrolling)
+	if (!bIsScrolling)
 	{
 		// Elastic snap to horizontal limits (+15%).
 
@@ -583,7 +583,12 @@ FReply SFrameTrack::OnMouseButtonUp(const FGeometry& MyGeometry, const FPointerE
 			// Release the mouse.
 			Reply = FReply::Handled().ReleaseMouseCapture();
 
-			if (bIsValidForMouseClick)
+			if (bIsScrolling)
+			{
+				bIsScrolling = false;
+				CursorType = EFrameTrackCursor::Default;
+			}
+			else if (bIsValidForMouseClick)
 			{
 				SelectFrameAtMousePosition(MousePositionOnButtonUp.X, MousePositionOnButtonUp.Y);
 			}
@@ -598,10 +603,10 @@ FReply SFrameTrack::OnMouseButtonUp(const FGeometry& MyGeometry, const FPointerE
 			// Release mouse as we no longer scroll.
 			Reply = FReply::Handled().ReleaseMouseCapture();
 
-			if (bIsRMB_Scrolling)
+			if (bIsScrolling)
 			{
+				bIsScrolling = false;
 				CursorType = EFrameTrackCursor::Default;
-				bIsRMB_Scrolling = false;
 			}
 			else if (bIsValidForMouseClick)
 			{
@@ -633,8 +638,18 @@ FReply SFrameTrack::OnMouseMove(const FGeometry& MyGeometry, const FPointerEvent
 	{
 		if (HasMouseCapture() && !MouseEvent.GetCursorDelta().IsZero())
 		{
-			// Inform other widgets that we have moved the selection box.
-			//SelectionBoxChangedEvent.Broadcast(SelectionBoxFrameStart, SelectionBoxFrameEnd);
+			if (!bIsScrolling)
+			{
+				bIsScrolling = true;
+				CursorType = EFrameTrackCursor::Hand;
+
+				HoveredSample.Reset();
+			}
+
+			float PosX = ViewportPosXOnButtonDown + (MousePositionOnButtonDown.X - MousePosition.X);
+			Viewport.ScrollAtIndex(Viewport.GetIndexAtPosX(PosX));
+			UpdateHorizontalScrollBar();
+			bIsStateDirty = true;
 
 			Reply = FReply::Handled();
 		}
@@ -643,10 +658,13 @@ FReply SFrameTrack::OnMouseMove(const FGeometry& MyGeometry, const FPointerEvent
 	{
 		if (HasMouseCapture() && !MouseEvent.GetCursorDelta().IsZero())
 		{
-			bIsRMB_Scrolling = true;
-			CursorType = EFrameTrackCursor::Hand;
+			if (!bIsScrolling)
+			{
+				bIsScrolling = true;
+				CursorType = EFrameTrackCursor::Hand;
 
-			HoveredSample.Reset();
+				HoveredSample.Reset();
+			}
 
 			float PosX = ViewportPosXOnButtonDown + (MousePositionOnButtonDown.X - MousePosition.X);
 			Viewport.ScrollAtIndex(Viewport.GetIndexAtPosX(PosX));
diff --git a/Engine/Source/Developer/TraceInsights/Private/Insights/Widgets/SFrameTrack.h b/Engine/Source/Developer/TraceInsights/Private/Insights/Widgets/SFrameTrack.h
--- a/Engine/Source/Developer/TraceInsights/Private/Insights/Widgets/SFrameTrack.h
+++ b/Engine/Source/Developer/TraceInsights/Private/Insights/Widgets/SFrameTrack.h
@@ -164,7 +164,7 @@ protected:
 	bool bIsRMB_Pressed;
 
 	/** True, if the user is currently interactively scrolling the view by holding the right mouse button and dragging. */
-	bool bIsRMB_Scrolling;
+	bool bIsScrolling;
 
 	//////////////////////////////////////////////////
 
diff --git a/Engine/Source/Developer/TraceInsights/Private/Insights/Widgets/SStartPageWindow.cpp b/Engine/Source/Developer/TraceInsights/Private/Insights/Widgets/SStartPageWindow.cpp
--- a/Engine/Source/Developer/TraceInsights/Private/Insights/Widgets/SStartPageWindow.cpp
+++ b/Engine/Source/Developer/TraceInsights/Private/Insights/Widgets/SStartPageWindow.cpp
@@ -8,6 +8,7 @@
 #include "Framework/Docking/WorkspaceItem.h"
 #include "Framework/MultiBox/MultiBoxBuilder.h"
 #include "HAL/FileManagerGeneric.h"
+#include "Input/Events.h"
 #include "SlateOptMacros.h"
 #include "Widgets/Docking/SDockTab.h"
 #include "Widgets/Images/SImage.h"
@@ -42,50 +43,12 @@
 #define LOCTEXT_NAMESPACE "SStartPageWindow"
 
 ////////////////////////////////////////////////////////////////////////////////////////////////////
-
-static FText GetTextForNotification(const EInsightsNotificationType NotificatonType, const ELoadingProgressState ProgressState, const FString& Filename, const float ProgressPercent = 0.0f)
-{
-	FText Result;
-
-	if (NotificatonType == EInsightsNotificationType::LoadingTraceFile)
-	{
-		if (ProgressState == ELoadingProgressState::Started)
-		{
-			FFormatNamedArguments Args;
-			Args.Add(TEXT("Filename"), FText::FromString(Filename));
-			Result = FText::Format(LOCTEXT("DescF_OfflineCapture_Started", "Started loading a file {Filename}"), Args);
-		}
-		else if (ProgressState == ELoadingProgressState::InProgress)
-		{
-			FFormatNamedArguments Args;
-			Args.Add(TEXT("Filename"), FText::FromString(Filename));
-			Args.Add(TEXT("LoadingProgressPercent"), FText::AsPercent(ProgressPercent));
-			Result = FText::Format(LOCTEXT("DescF_OfflineCapture_InProgress", "Loading a file {Filename} {LoadingProgressPercent}"), Args);
-		}
-		else if (ProgressState == ELoadingProgressState::Loaded)
-		{
-			FFormatNamedArguments Args;
-			Args.Add(TEXT("Filename"), FText::FromString(Filename));
-			Result = FText::Format(LOCTEXT("DescF_OfflineCapture_Loaded", "File {Filename} has been successfully loaded"), Args);
-		}
-		else if (ProgressState == ELoadingProgressState::Failed)
-		{
-			FFormatNamedArguments Args;
-			Args.Add(TEXT("Filename"), FText::FromString(Filename));
-			Result = FText::Format(LOCTEXT("DescF_OfflineCapture_Failed", "Failed to load file {Filename}"), Args);
-		}
-	}
-
-	return Result;
-}
-
-////////////////////////////////////////////////////////////////////////////////////////////////////
-// SConnectionRow
+// STraceSessionRow
 ////////////////////////////////////////////////////////////////////////////////////////////////////
 
-class SConnectionRow : public SMultiColumnTableRow<TSharedPtr<FRecorderConnection>>
+class STraceSessionRow : public SMultiColumnTableRow<TSharedPtr<FTraceSession>>
 {
-	SLATE_BEGIN_ARGS(SConnectionRow) {}
+	SLATE_BEGIN_ARGS(STraceSessionRow) {}
 	SLATE_END_ARGS()
 
 public:
@@ -93,15 +56,15 @@ public:
 	 * Constructs the widget.
 	 *
 	 * @param InArgs The construction arguments.
-	 * @param InConnection The connection displayed by this row.
+	 * @param InTraceSession The trace session displayed by this row.
 	 * @param InOwnerTableView The table to which the row must be added.
 	 */
-	void Construct(const FArguments& InArgs, TSharedPtr<FRecorderConnection> InConnection, TSharedRef<SStartPageWindow> InParentWidget, const TSharedRef<STableViewBase>& InOwnerTableView)
+	void Construct(const FArguments& InArgs, TSharedPtr<FTraceSession> InTraceSession, TSharedRef<SStartPageWindow> InParentWidget, const TSharedRef<STableViewBase>& InOwnerTableView)
 	{
-		WeakConnection = MoveTemp(InConnection);
+		WeakTraceSession = MoveTemp(InTraceSession);
 		WeakParentWidget = InParentWidget;
 
-		SMultiColumnTableRow<TSharedPtr<FRecorderConnection>>::Construct(FSuperRowType::FArguments(), InOwnerTableView);
+		SMultiColumnTableRow<TSharedPtr<FTraceSession>>::Construct(FSuperRowType::FArguments(), InOwnerTableView);
 	}
 
 public:
@@ -113,8 +76,8 @@ public:
 				.Padding(FMargin(4.0, 0.0))
 				[
 					SNew(STextBlock)
-					.Text(this, &SConnectionRow::GetConnectionName)
-					.ToolTipText(this, &SConnectionRow::GetConnectionName)
+					.Text(this, &STraceSessionRow::GetTraceSessionName)
+					.ToolTipText(this, &STraceSessionRow::GetTraceSessionTooltip)
 				];
 		}
 		else if (ColumnName == FName(TEXT("Uri")))
@@ -123,8 +86,18 @@ public:
 				.Padding(FMargin(4.0, 0.0))
 				[
 					SNew(STextBlock)
-					.Text(this, &SConnectionRow::GetConnectionUri)
-					.ToolTipText(this, &SConnectionRow::GetConnectionUri)
+					.Text(this, &STraceSessionRow::GetTraceSessionUri)
+					.ToolTipText(this, &STraceSessionRow::GetTraceSessionTooltip)
+				];
+		}
+		else if (ColumnName == FName(TEXT("Status")))
+		{
+			return SNew(SBox)
+				.Padding(FMargin(4.0, 0.0))
+				[
+					SNew(STextBlock)
+					.Text(this, &STraceSessionRow::GetTraceSessionstatus)
+					.ToolTipText(this, &STraceSessionRow::GetTraceSessionTooltip)
 				];
 		}
 		else
@@ -133,12 +106,12 @@ public:
 		}
 	}
 
-	FText GetConnectionName() const
+	FText GetTraceSessionName() const
 	{
-		TSharedPtr<FRecorderConnection> ConnectionPin = WeakConnection.Pin();
-		if (ConnectionPin.IsValid())
+		TSharedPtr<FTraceSession> TraceSessionPin = WeakTraceSession.Pin();
+		if (TraceSessionPin.IsValid())
 		{
-			return ConnectionPin->Name;
+			return TraceSessionPin->Name;
 		}
 		else
 		{
@@ -146,12 +119,12 @@ public:
 		}
 	}
 
-	FText GetConnectionUri() const
+	FText GetTraceSessionUri() const
 	{
-		TSharedPtr<FRecorderConnection> ConnectionPin = WeakConnection.Pin();
-		if (ConnectionPin.IsValid())
+		TSharedPtr<FTraceSession> TraceSessionPin = WeakTraceSession.Pin();
+		if (TraceSessionPin.IsValid())
 		{
-			return ConnectionPin->Uri;
+			return TraceSessionPin->Uri;
 		}
 		else
 		{
@@ -159,12 +132,32 @@ public:
 		}
 	}
 
-	FText GetRowToolTip() const
+	FText GetTraceSessionstatus() const
+	{
+		TSharedPtr<FTraceSession> TraceSessionPin = WeakTraceSession.Pin();
+		if (TraceSessionPin.IsValid())
+		{
+			if (TraceSessionPin->bIsLive)
+			{
+				return LOCTEXT("LiveTraceSessionStatus", "LIVE");
+			}
+		}
+		return FText();
+	}
+
+	FText GetTraceSessionTooltip() const
 	{
-		TSharedPtr<FRecorderConnection> ConnectionPin = WeakConnection.Pin();
-		if (ConnectionPin.IsValid())
+		TSharedPtr<FTraceSession> TraceSessionPin = WeakTraceSession.Pin();
+		if (TraceSessionPin.IsValid())
 		{
-			return ConnectionPin->Uri;
+			if (TraceSessionPin->bIsLive)
+			{
+				return FText::Format(LOCTEXT("LiveTraceSessionTooltipFormat", "{0}\n{1}\nLIVE"), TraceSessionPin->Name, TraceSessionPin->Uri);
+			}
+			else
+			{
+				return FText::Format(LOCTEXT("TraceSessionTooltipFormat", "{0}\n{1}"), TraceSessionPin->Name, TraceSessionPin->Uri);
+			}
 		}
 		else
 		{
@@ -173,7 +166,7 @@ public:
 	}
 
 private:
-	TWeakPtr<FRecorderConnection> WeakConnection;
+	TWeakPtr<FTraceSession> WeakTraceSession;
 	TWeakPtr<SStartPageWindow> WeakParentWidget;
 };
 
@@ -183,10 +176,8 @@ private:
 
 SStartPageWindow::SStartPageWindow()
 	: DurationActive(0.0f)
-	, bIsAnyLiveSessionAvailable(false)
-	, LastLiveSessionHandle(static_cast<Trace::FSessionHandle>(0))
-	, bIsAnySessionAvailable(false)
-	, LastSessionHandle(static_cast<Trace::FSessionHandle>(0))
+	, AvailableSessionCount(0)
+	, LiveSessionCount(0)
 {
 }
 
@@ -197,7 +188,7 @@ SStartPageWindow::~SStartPageWindow()
 #if WITH_EDITOR
 	if (DurationActive > 0.0f && FEngineAnalytics::IsAvailable())
 	{
-		FEngineAnalytics::GetProvider().RecordEvent(TEXT("Editor.Usage.Profiler"), FAnalyticsEventAttribute(TEXT("Duration"), DurationActive));
+		FEngineAnalytics::GetProvider().RecordEvent(TEXT("Editor.Usage.Insights.StartPage"), FAnalyticsEventAttribute(TEXT("Duration"), DurationActive));
 	}
 #endif // WITH_EDITOR
 }
@@ -236,13 +227,16 @@ void SStartPageWindow::Construct(const FArguments& InArgs)
 						.HAlign(HAlign_Center)
 						.Padding(3.0f, 3.0f)
 						[
-							SNew(SBorder)
-							.Visibility(this, &SStartPageWindow::IsSessionOverlayVisible)
-							.BorderImage(FEditorStyle::GetBrush("NotificationList.ItemBackground"))
-							.Padding(8.0f)
+							SNew(SBox)
+							.WidthOverride(512.0f)
 							[
-								SNew(STextBlock)
-								.Text(LOCTEXT("SelectTraceOverlayText", "Please select a trace..."))
+								SNew(SBorder)
+								.BorderImage(FEditorStyle::GetBrush("NotificationList.ItemBackground"))
+								.Padding(8.0f)
+								.HAlign(HAlign_Fill)
+								[
+									ConstructSessionsPanel()
+								]
 							]
 						]
 
@@ -251,192 +245,16 @@ void SStartPageWindow::Construct(const FArguments& InArgs)
 						.HAlign(HAlign_Center)
 						.Padding(3.0f, 3.0f)
 						[
-							SNew(SBorder)
-							.BorderImage(FEditorStyle::GetBrush("NotificationList.ItemBackground"))
-							.Padding(8.0f)
+							SNew(SBox)
+							.WidthOverride(512.0f)
 							[
-								SNew(STextBlock)
-								.Text(LOCTEXT("SelectTraceHint", "\
-- Use \"Live\" button to start analysis for the latest available session that is live, if any.\n\
-- Use \"Latest\" button to start analysis for the latest available session, if any.\n\
-- Use \"Load\" button to load a trace file (*.utrace).\n\
-- Use the combo box to choose from available sessions.\n\
-- Drag and drop a *.utrace file over this window."))
-								//.Justification(ETextJustify::Center)
-							]
-						]
-
-					+ SVerticalBox::Slot()
-						.AutoHeight()
-						.HAlign(HAlign_Center)
-						.Padding(3.0f, 3.0f)
-						[
-							SNew(SHorizontalBox)
-
-							+ SHorizontalBox::Slot()
-								.AutoWidth()
-								.Padding(3.0f, 0.0f)
+								SNew(SBorder)
+								.BorderImage(FEditorStyle::GetBrush("NotificationList.ItemBackground"))
+								.Padding(8.0f)
+								.HAlign(HAlign_Fill)
 								[
-									SNew(SButton)
-									.OnClicked(this, &SStartPageWindow::Live_OnClicked)
-									.IsEnabled(this, &SStartPageWindow::Live_IsEnabled)
-									.ToolTipText(LOCTEXT("LiveButtonTooltip", "Start analysis for the latest available session that is live, if any."))
-									.ContentPadding(8.0f)
-									.Content()
-									[
-										SNew(SHorizontalBox)
-
-										+SHorizontalBox::Slot()
-											.AutoWidth()
-											[
-												SNew(SImage)
-												.Image(FInsightsStyle::GetBrush("Live.Icon.Large"))
-											]
-
-										+SHorizontalBox::Slot()
-											.AutoWidth()
-											.VAlign(VAlign_Center)
-											[
-												SNew(SBox)
-												.WidthOverride(100.0f)
-												[
-													SNew(STextBlock)
-													.Font(FCoreStyle::GetDefaultFontStyle("Regular", 18))
-													.Text(LOCTEXT("LiveButtonText", "Live"))
-												]
-											]
-									]
+									ConstructRecoderPanel()
 								]
-
-							+ SHorizontalBox::Slot()
-								.AutoWidth()
-								.Padding(3.0f, 0.0f)
-								[
-									SNew(SButton)
-									.OnClicked(this, &SStartPageWindow::Last_OnClicked)
-									.IsEnabled(this, &SStartPageWindow::Last_IsEnabled)
-									.ToolTipText(LOCTEXT("LastButtonTooltip", "Start analysis for the latest available session, if any."))
-									.ContentPadding(8.0f)
-									.Content()
-									[
-										SNew(SHorizontalBox)
-
-										+SHorizontalBox::Slot()
-											.AutoWidth()
-											[
-												SNew(SImage)
-												.Image(FInsightsStyle::GetBrush("Live.Icon.Large"))
-											]
-
-										+SHorizontalBox::Slot()
-											.AutoWidth()
-											.VAlign(VAlign_Center)
-											[
-												SNew(SBox)
-												.WidthOverride(100.0f)
-												[
-													SNew(STextBlock)
-													.Font(FCoreStyle::GetDefaultFontStyle("Regular", 18))
-													.Text(LOCTEXT("LastButtonText", "Latest"))
-												]
-											]
-									]
-								]
-
-							+ SHorizontalBox::Slot()
-								.AutoWidth()
-								.Padding(3.0f, 0.0f, 0.0f, 0.0f)
-								[
-									SNew(SButton)
-									.OnClicked(this, &SStartPageWindow::Load_OnClicked)
-									.ToolTipText(LOCTEXT("LoadButtonTooltip", "Load a trace file."))
-									.ContentPadding(8.0f)
-									.Content()
-									[
-										SNew(SHorizontalBox)
-
-										+ SHorizontalBox::Slot()
-											.AutoWidth()
-											[
-												SNew(SImage)
-												.Image(FInsightsStyle::GetBrush("Load.Icon.Large"))
-											]
-
-										+ SHorizontalBox::Slot()
-											.AutoWidth()
-											.VAlign(VAlign_Center)
-											[
-												SNew(SBox)
-												.WidthOverride(100.0f)
-												[
-													SNew(STextBlock)
-													.Font(FCoreStyle::GetDefaultFontStyle("Regular", 18))
-													.Text(LOCTEXT("LoadButtonText", "Load..."))
-												]
-											]
-									]
-								]
-
-							+ SHorizontalBox::Slot()
-								.AutoWidth()
-								.Padding(0.0f, 0.0f, 0.0f, 0.0f)
-								[
-									SNew(SComboButton)
-									.ToolTipText(LOCTEXT("SessionList_Tooltip", "Choose from Most Recently Used Sessions or from Available Sessions"))
-									.OnGetMenuContent(this, &SStartPageWindow::MakeSessionListMenu)
-									.HasDownArrow(true)
-									.ContentPadding(FMargin(1.0f, 1.0f, 1.0f, 1.0f))
-								]
-						]
-
-					+ SVerticalBox::Slot()
-						.AutoHeight()
-						.HAlign(HAlign_Center)
-						.Padding(3.0f, 3.0f)
-						[
-							SNew(SBorder)
-							.BorderImage(FEditorStyle::GetBrush("NotificationList.ItemBackground"))
-							.Padding(8.0f)
-							[
-								SNew(SVerticalBox)
-
-								+ SVerticalBox::Slot()
-									.AutoHeight()
-									.HAlign(HAlign_Center)
-									.Padding(0.0, 2.0f)
-									[
-										SNew(STextBlock)
-										.Text(LOCTEXT("LocalSessionDirectoryText", "Local Session Directory:"))
-										.ColorAndOpacity(FLinearColor::Gray)
-									]
-
-								+ SVerticalBox::Slot()
-									.AutoHeight()
-									.HAlign(HAlign_Center)
-									.Padding(0.0, 2.0f)
-									[
-										SNew(SHorizontalBox)
-
-										+ SHorizontalBox::Slot()
-											.AutoWidth()
-											.Padding(2.0f, 0.0f)
-											.VAlign(VAlign_Center)
-											[
-												SNew(STextBlock)
-												.Text(this, &SStartPageWindow::GetLocalSessionDirectory)
-											]
-
-										+ SHorizontalBox::Slot()
-											.AutoWidth()
-											.Padding(2.0f, 0.0f)
-											.VAlign(VAlign_Center)
-											[
-												SNew(SButton)
-												.Text(LOCTEXT("ExploreLocalSessionDirButton", "..."))
-												.ToolTipText(LOCTEXT("ExploreLocalSessionDirButtonToolTip", "Explore the Local Session Directory"))
-												.OnClicked(this, &SStartPageWindow::ExploreLocalSessionDirectory_OnClicked)
-											]
-									]
 							]
 						]
 
@@ -445,11 +263,17 @@ void SStartPageWindow::Construct(const FArguments& InArgs)
 						.HAlign(HAlign_Center)
 						.Padding(3.0f, 3.0f)
 						[
-							SNew(SBorder)
-							.BorderImage(FEditorStyle::GetBrush("NotificationList.ItemBackground"))
-							.Padding(8.0f)
+							SNew(SBox)
+							.WidthOverride(512.0f)
+							.Visibility(this, &SStartPageWindow::StopTraceRecorder_Visibility)
 							[
-								ConstructRecoderPanel()
+								SNew(SBorder)
+								.BorderImage(FEditorStyle::GetBrush("NotificationList.ItemBackground"))
+								.Padding(8.0f)
+								.HAlign(HAlign_Fill)
+								[
+									ConstructConnectPanel()
+								]
 							]
 						]
 				]
@@ -469,6 +293,176 @@ void SStartPageWindow::Construct(const FArguments& InArgs)
 				.VAlign(VAlign_Center)
 				.Expose(OverlaySettingsSlot)
 		];
+
+	FSlateApplication::Get().SetKeyboardFocus(TraceSessionsListView);
+	FSlateApplication::Get().SetUserFocus(0, TraceSessionsListView);
+}
+
+////////////////////////////////////////////////////////////////////////////////////////////////////
+
+TSharedRef<SWidget> SStartPageWindow::ConstructSessionsPanel()
+{
+	TSharedRef<SWidget> Widget = SNew(SVerticalBox)
+
+	+ SVerticalBox::Slot()
+		.AutoHeight()
+		.HAlign(HAlign_Center)
+		.Padding(0.0, 2.0f)
+		[
+			SNew(STextBlock)
+			.Text(LOCTEXT("SessionsPanelTitle", "Trace Sessions"))
+			.Font(FCoreStyle::GetDefaultFontStyle("Bold", 11))
+			.ColorAndOpacity(FLinearColor::Gray)
+		]
+
+	+ SVerticalBox::Slot()
+		.AutoHeight()
+		.HAlign(HAlign_Fill)
+		.Padding(0.0, 1.0f, 0.0f, 2.0f)
+		.MaxHeight(134.0f)
+		[
+			SAssignNew(TraceSessionsListView, SListView<TSharedPtr<FTraceSession>>)
+			.IsFocusable(true)
+			.ItemHeight(20.0f)
+			.SelectionMode(ESelectionMode::Single)
+			.OnSelectionChanged(this, &SStartPageWindow::TraceSessions_OnSelectionChanged)
+			.OnMouseButtonDoubleClick(this, &SStartPageWindow::TraceSessions_OnMouseButtonDoubleClick)
+			.ListItemsSource(&TraceSessions)
+			.OnGenerateRow(this, &SStartPageWindow::TraceSessions_OnGenerateRow)
+			.ConsumeMouseWheel(EConsumeMouseWheel::Always)
+			//.OnContextMenuOpening(FOnContextMenuOpening::CreateSP(this, &SStartPageWindow::TraceSessions_GetContextMenu))
+			.HeaderRow
+			(
+				SNew(SHeaderRow)
+
+				+ SHeaderRow::Column(FName(TEXT("Name")))
+				.FillWidth(0.2f)
+				.DefaultLabel(LOCTEXT("NameColumn", "Name"))
+
+				//+ SHeaderRow::Column(FName(TEXT("Uri")))
+				//.FillWidth(0.8f)
+				//.DefaultLabel(LOCTEXT("UriColumn", "URI"))
+
+				+ SHeaderRow::Column(FName(TEXT("Status")))
+				.FixedWidth(60.0f)
+				.DefaultLabel(LOCTEXT("StatusColumn", "Status"))
+			)
+		]
+
+	+ SVerticalBox::Slot()
+		.AutoHeight()
+		.HAlign(HAlign_Right)
+		.Padding(0.0, 2.0f)
+		[
+			ConstructLoadPanel()
+		]
+
+	+ SVerticalBox::Slot()
+		.AutoHeight()
+		.HAlign(HAlign_Left)
+		.Padding(0.0, 2.0f)
+		[
+			ConstructLocalSessionDirectoryPanel()
+		]
+	;
+
+	return Widget;
+}
+
+////////////////////////////////////////////////////////////////////////////////////////////////////
+
+TSharedRef<SWidget> SStartPageWindow::ConstructLoadPanel()
+{
+	TSharedRef<SWidget> Widget = SNew(SHorizontalBox)
+
+	+ SHorizontalBox::Slot()
+		.AutoWidth()
+		[
+			SNew(SButton)
+			.IsEnabled(this, &SStartPageWindow::Open_IsEnabled)
+			.OnClicked(this, &SStartPageWindow::Open_OnClicked)
+			.ToolTipText(LOCTEXT("OpenButtonTooltip", "Start analysis for selected session."))
+			.ContentPadding(FMargin(4.0f, 1.0f))
+			.Content()
+			[
+				SNew(SHorizontalBox)
+
+				+ SHorizontalBox::Slot()
+					.AutoWidth()
+					[
+						SNew(SImage)
+						.Image(FInsightsStyle::GetBrush("Open.Icon.Small"))
+					]
+
+				+ SHorizontalBox::Slot()
+					.AutoWidth()
+					.VAlign(VAlign_Center)
+					[
+						SNew(STextBlock)
+						.Text(LOCTEXT("OpenButtonText", "Open"))
+					]
+			]
+		]
+
+	+ SHorizontalBox::Slot()
+		.AutoWidth()
+		[
+			SNew(SComboButton)
+			.ToolTipText(LOCTEXT("MRU_Tooltip", "Open a file or choose a session."))
+			.OnGetMenuContent(this, &SStartPageWindow::MakeSessionListMenu)
+			.HasDownArrow(true)
+			.ContentPadding(FMargin(1.0f, 1.0f, 1.0f, 1.0f))
+		]
+	;
+
+	return Widget;
+}
+////////////////////////////////////////////////////////////////////////////////////////////////////
+
+TSharedRef<SWidget> SStartPageWindow::ConstructLocalSessionDirectoryPanel()
+{
+	TSharedRef<SWidget> Widget = SNew(SVerticalBox)
+
+	+ SVerticalBox::Slot()
+		.AutoHeight()
+		.HAlign(HAlign_Left)
+		.Padding(0.0, 0.0f)
+		[
+			SNew(STextBlock)
+			.Text(LOCTEXT("LocalSessionDirectoryText", "Local Session Directory:"))
+			.ColorAndOpacity(FLinearColor::Gray)
+		]
+
+	+ SVerticalBox::Slot()
+		.AutoHeight()
+		.HAlign(HAlign_Left)
+		.Padding(0.0, 0.0f)
+		[
+			SNew(SHorizontalBox)
+
+			+ SHorizontalBox::Slot()
+				.Padding(0.0f, 0.0f)
+				.VAlign(VAlign_Center)
+				[
+					SNew(STextBlock)
+					.Text(this, &SStartPageWindow::GetLocalSessionDirectory)
+					.Justification(ETextJustify::Right)
+				]
+
+			+ SHorizontalBox::Slot()
+				.AutoWidth()
+				.Padding(4.0f, 0.0f, 0.0f, 0.0f)
+				.VAlign(VAlign_Center)
+				[
+					SNew(SButton)
+					.Text(LOCTEXT("ExploreLocalSessionDirButton", "..."))
+					.ToolTipText(LOCTEXT("ExploreLocalSessionDirButtonToolTip", "Explore the Local Session Directory"))
+					.OnClicked(this, &SStartPageWindow::ExploreLocalSessionDirectory_OnClicked)
+				]
+		]
+	;
+
+	return Widget;
 }
 
 ////////////////////////////////////////////////////////////////////////////////////////////////////
@@ -541,8 +535,8 @@ TSharedRef<SWidget> SStartPageWindow::ConstructRecoderPanel()
 
 	+ SVerticalBox::Slot()
 		.AutoHeight()
-		.HAlign(HAlign_Fill)
-		.Padding(0.0, 2.0f)
+		.HAlign(HAlign_Left)
+		.Padding(0.0, 2.0f, 0.0f, 1.0f)
 		[
 			SNew(SHorizontalBox)
 			.Visibility(this, &SStartPageWindow::StopTraceRecorder_Visibility)
@@ -553,36 +547,44 @@ TSharedRef<SWidget> SStartPageWindow::ConstructRecoderPanel()
 				.VAlign(VAlign_Center)
 				[
 					SNew(STextBlock)
-					.Text(LOCTEXT("HostTitle", "Host:"))
+					.Text_Lambda([this]() -> FText
+					{
+						return FText::Format(LOCTEXT("ConnectionCountFormat", "Connections / live sessions: {0}"), FText::AsNumber(LiveSessionCount));
+					})
 					.ColorAndOpacity(FLinearColor::Gray)
 				]
+		]
 
-			+ SHorizontalBox::Slot()
-				.FillWidth(1.0)
-				.VAlign(VAlign_Center)
-				[
-					SAssignNew(HostTextBox, SEditableTextBox)
-				]
+	;
 
-			+ SHorizontalBox::Slot()
-				.AutoWidth()
-				.Padding(2.0f, 0.0f)
-				.VAlign(VAlign_Center)
-				[
-					SNew(SButton)
-					.Text(LOCTEXT("Connect", "Connect"))
-					.ToolTipText(LOCTEXT("ConnectToolTip", "Try connecting to host."))
-					.OnClicked(this, &SStartPageWindow::Connect_OnClicked)
-				]
+	RefreshTraceSessionList();
+
+	return Widget;
+}
+
+////////////////////////////////////////////////////////////////////////////////////////////////////
+
+TSharedRef<SWidget> SStartPageWindow::ConstructConnectPanel()
+{
+	TSharedRef<SWidget> Widget = SNew(SVerticalBox)
+
+	+ SVerticalBox::Slot()
+		.AutoHeight()
+		.HAlign(HAlign_Center)
+		.Padding(0.0, 2.0f)
+		[
+			SNew(STextBlock)
+			.Text(LOCTEXT("ConnectPanelTitle", "New Connection"))
+			.Font(FCoreStyle::GetDefaultFontStyle("Bold", 11))
+			.ColorAndOpacity(FLinearColor::Gray)
 		]
 
 	+ SVerticalBox::Slot()
 		.AutoHeight()
-		.HAlign(HAlign_Left)
-		.Padding(0.0, 2.0f, 0.0f, 1.0f)
+		.HAlign(HAlign_Fill)
+		.Padding(0.0, 2.0f)
 		[
 			SNew(SHorizontalBox)
-			.Visibility(this, &SStartPageWindow::StopTraceRecorder_Visibility)
 
 			+ SHorizontalBox::Slot()
 				.AutoWidth()
@@ -590,176 +592,47 @@ TSharedRef<SWidget> SStartPageWindow::ConstructRecoderPanel()
 				.VAlign(VAlign_Center)
 				[
 					SNew(STextBlock)
-					.Text_Lambda([this]() -> FText
-					{
-						return FText::Format(LOCTEXT("ConnectionListTitle", "Connections (live sessions): {0}"), FText::AsNumber(Connections.Num()));
-					})
+					.Text(LOCTEXT("HostTitle", "Running instance IP:"))
 					.ColorAndOpacity(FLinearColor::Gray)
 				]
 
+			+ SHorizontalBox::Slot()
+				.FillWidth(1.0)
+				.VAlign(VAlign_Center)
+				[
+					SAssignNew(HostTextBox, SEditableTextBox)
+				]
+
 			+ SHorizontalBox::Slot()
 				.AutoWidth()
 				.Padding(2.0f, 0.0f)
 				.VAlign(VAlign_Center)
 				[
 					SNew(SButton)
-					.Text(LOCTEXT("Refresh", "Refresh"))
-					.ToolTipText(LOCTEXT("RefreshToolTip", "Refresh the connection list (live sessions)."))
-					.OnClicked(this, &SStartPageWindow::RefreshConnections_OnClicked)
+					.Text(LOCTEXT("Connect", "Connect"))
+					.ToolTipText(LOCTEXT("ConnectToolTip", "Try connecting to host."))
+					.OnClicked(this, &SStartPageWindow::Connect_OnClicked)
 				]
 		]
-
-	+ SVerticalBox::Slot()
-		.AutoHeight()
-		.HAlign(HAlign_Left)
-		.Padding(0.0, 1.0f, 0.0f, 2.0f)
-		[
-			SNew(SBox)
-			.WidthOverride(640.0f)
-			.MaxDesiredHeight(78.0f)
-			.Visibility(this, &SStartPageWindow::Connections_Visibility)
-			[
-				SAssignNew(ConnectionsListView, SListView<TSharedPtr<FRecorderConnection>>)
-				.ItemHeight(20.0f)
-				.SelectionMode(ESelectionMode::Single)
-				.OnSelectionChanged(this, &SStartPageWindow::Connections_OnSelectionChanged)
-				.ListItemsSource(&Connections)
-				.OnGenerateRow(this, &SStartPageWindow::Connections_OnGenerateRow)
-				.ConsumeMouseWheel(EConsumeMouseWheel::Always)
-				//.OnContextMenuOpening(FOnContextMenuOpening::CreateSP(this, &SStartPageWindow::Connections_GetContextMenu))
-				.HeaderRow
-				(
-					SNew(SHeaderRow)
-
-					+ SHeaderRow::Column(FName(TEXT("Name")))
-					.FillWidth(0.2f)
-					.DefaultLabel(LOCTEXT("NameColumn", "Name"))
-
-					+ SHeaderRow::Column(FName(TEXT("Uri")))
-					.FillWidth(0.8f)
-					.DefaultLabel(LOCTEXT("UriColumn", "URI"))
-				)
-			]
-		]
-
-	+ SVerticalBox::Slot()
-		.AutoHeight()
-		.HAlign(HAlign_Left)
-		.Padding(0.0, 2.0f, 0.0f, 1.0f)
-		[
-			SNew(STextBlock)
-			.Visibility(this, &SStartPageWindow::Modules_Visibility)
-			.Text(LOCTEXT("ModulesListTitle", "Modules for selected connection:"))
-			.ColorAndOpacity(FLinearColor::Gray)
-		]
-
-	+ SVerticalBox::Slot()
-		.AutoHeight()
-		.HAlign(HAlign_Left)
-		.Padding(0.0, 1.0f, 0.0f, 2.0f)
-		[
-			ConstructModuleList()
-		]
 	;
 
-	RefreshConnectionList();
-
 	return Widget;
 }
 
 ////////////////////////////////////////////////////////////////////////////////////////////////////
 
-TSharedRef<SWidget> SStartPageWindow::ConstructModuleList()
-{
-	TSharedRef<SVerticalBox> VerticalBox = SNew(SVerticalBox)
-		.Visibility(this, &SStartPageWindow::Modules_Visibility);
-
-	TSharedRef<Trace::IModuleService> ModuleService = FInsightsManager::Get()->GetModuleService();
-	ModuleService->GetAvailableModules(AvailableModules);
-
-	constexpr bool bDefaultModuleEnableState = false;
-	AvailableModulesEnabledState.Reset();
-
-	for (int32 ModuleIndex = 0; ModuleIndex < AvailableModules.Num(); ++ModuleIndex)
-	{
-		const Trace::FModuleInfo& Module = AvailableModules[ModuleIndex];
-
-		AvailableModulesEnabledState.Add(bDefaultModuleEnableState);
-		ModuleService->SetModuleEnabled(Module.Name, bDefaultModuleEnableState);
-
-		VerticalBox->AddSlot()
-			.AutoHeight()
-			.Padding(0.0f, 1.0f)
-			[
-				SNew(SCheckBox)
-				.IsChecked(this, &SStartPageWindow::Module_IsChecked, ModuleIndex)
-				.OnCheckStateChanged(this, &SStartPageWindow::Module_OnCheckStateChanged, ModuleIndex)
-				[
-					SNew(STextBlock)
-					.Text(FText::FromString(Module.DisplayName))
-				]
-			];
-	}
-
-	return VerticalBox;
-}
-
-////////////////////////////////////////////////////////////////////////////////////////////////////
-
-TSharedRef<ITableRow> SStartPageWindow::Connections_OnGenerateRow(TSharedPtr<FRecorderConnection> InConnection, const TSharedRef<STableViewBase>& OwnerTable)
+TSharedRef<ITableRow> SStartPageWindow::TraceSessions_OnGenerateRow(TSharedPtr<FTraceSession> InTraceSession, const TSharedRef<STableViewBase>& OwnerTable)
 {
-	return SNew(SConnectionRow, InConnection, SharedThis(this), OwnerTable);
+	return SNew(STraceSessionRow, InTraceSession, SharedThis(this), OwnerTable);
 }
 
 END_SLATE_FUNCTION_BUILD_OPTIMIZATION
 
 ////////////////////////////////////////////////////////////////////////////////////////////////////
 
-EVisibility SStartPageWindow::Modules_Visibility() const
-{
-	TSharedRef<Trace::ISessionService> SessionService = FInsightsManager::Get()->GetSessionService();
-	return (SessionService->IsRecorderServerRunning() && Connections.Num() > 0 && SelectedConnection.IsValid()) ? EVisibility::Visible : EVisibility::Collapsed;
-}
-
-////////////////////////////////////////////////////////////////////////////////////////////////////
-
-ECheckBoxState SStartPageWindow::Module_IsChecked(int32 ModuleIndex) const
+FReply SStartPageWindow::RefreshTraceSessions_OnClicked()
 {
-	if (ModuleIndex >= 0 && ModuleIndex < AvailableModulesEnabledState.Num())
-	{
-		const bool bIsEnabled = AvailableModulesEnabledState[ModuleIndex];
-		return bIsEnabled ? ECheckBoxState::Checked : ECheckBoxState::Unchecked;
-	}
-	else
-	{
-		return ECheckBoxState::Undetermined;
-	}
-}
-
-////////////////////////////////////////////////////////////////////////////////////////////////////
-
-void SStartPageWindow::Module_OnCheckStateChanged(ECheckBoxState NewRadioState, int32 ModuleIndex)
-{
-	ensure(AvailableModulesEnabledState.Num() == AvailableModules.Num());
-
-	if (SelectedConnection.IsValid() &&
-		ModuleIndex >= 0 && ModuleIndex < AvailableModules.Num())
-	{
-		bool bIsEnabled = NewRadioState == ECheckBoxState::Checked;
-
-		TSharedRef<Trace::ISessionService> SessionService = FInsightsManager::Get()->GetSessionService();
-		
-		SessionService->SetModuleEnabled(SelectedConnection->SessionHandle, AvailableModules[ModuleIndex].Name, bIsEnabled);
-		//AvailableModulesEnabledState[ModuleIndex] = bIsEnabled;
-		AvailableModulesEnabledState[ModuleIndex] = SessionService->IsModuleEnabled(SelectedConnection->SessionHandle, AvailableModules[ModuleIndex].Name);
-	}
-}
-
-////////////////////////////////////////////////////////////////////////////////////////////////////
-
-FReply SStartPageWindow::RefreshConnections_OnClicked()
-{
-	RefreshConnectionList();
+	RefreshTraceSessionList();
 	return FReply::Handled();
 }
 
@@ -772,7 +645,7 @@ FReply SStartPageWindow::Connect_OnClicked()
 	FText HostText = HostTextBox->GetText();
 	if (HostText.IsEmptyOrWhitespace())
 	{
-		//...
+		// nothing to do
 	}
 	else if (SessionService->ConnectSession(*HostText.ToString()))
 	{
@@ -799,74 +672,98 @@ FReply SStartPageWindow::Connect_OnClicked()
 		ActiveNotifications.Add(TEXT("ConnectFailed"), NotificationItem);
 	}
 
-	RefreshConnectionList();
+	RefreshTraceSessionList();
 	return FReply::Handled();
 }
 
 ////////////////////////////////////////////////////////////////////////////////////////////////////
 
-void SStartPageWindow::RefreshConnectionList()
+void SStartPageWindow::RefreshTraceSessionList()
 {
 	TSharedRef<Trace::ISessionService> SessionService = FInsightsManager::Get()->GetSessionService();
 
-	TArray<Trace::FSessionHandle> LiveSessions;
-	SessionService->GetLiveSessions(LiveSessions);
-
-	TSharedPtr<FRecorderConnection> NewSelectedConnection;
+	TArray<Trace::FSessionHandle> AvailableSessions;
+	SessionService->GetAvailableSessions(AvailableSessions);
 
-	Connections.Reset();
+	AvailableSessionCount = AvailableSessions.Num();
 
-	for (Trace::FSessionHandle SessionHandle : LiveSessions)
+	// Count number of live sessions.
+	int32 OldLiveSessionCount = LiveSessionCount;
+	LiveSessionCount = 0;
+	for (Trace::FSessionHandle SessionHandle : AvailableSessions)
 	{
 		Trace::FSessionInfo SessionInfo;
 		SessionService->GetSessionInfo(SessionHandle, SessionInfo);
+		if (SessionInfo.bIsLive)
+		{
+			++LiveSessionCount;
+		}
+	}
 
-		Connections.Add(MakeShareable(new FRecorderConnection(SessionHandle, SessionInfo)));
+	// If session list has changed on analysis side, recreate the TraceSessions list view widget.
+	//TODO: if (AvailableSessionsChangeNumber != SessionService->GetAvailableSessionsChangeNumber())
+	if (AvailableSessionCount != TraceSessions.Num() ||
+		LiveSessionCount != OldLiveSessionCount)
+	{
+		TSharedPtr<FTraceSession> NewSelectedTraceSession;
+
+		TraceSessions.Reset();
 
-		if (SelectedConnection && SelectedConnection->SessionHandle == SessionHandle)
+		for (Trace::FSessionHandle SessionHandle : AvailableSessions)
 		{
-			NewSelectedConnection = Connections.Last();
+			Trace::FSessionInfo SessionInfo;
+			SessionService->GetSessionInfo(SessionHandle, SessionInfo);
+
+			TraceSessions.Add(MakeShareable(new FTraceSession(SessionHandle, SessionInfo)));
+
+			// Identify the previously selected session (if stil available) to ensure selection remains unchanged.
+			if (SelectedTraceSession && SelectedTraceSession->SessionHandle == SessionHandle)
+			{
+				NewSelectedTraceSession = TraceSessions.Last();
+			}
 		}
-	}
 
-	ConnectionsListView->RebuildList();
+		TraceSessionsListView->RebuildList();
 
-	if (NewSelectedConnection.IsValid())
-	{
-		ConnectionsListView->SetItemSelection(NewSelectedConnection, true);
-		ConnectionsListView->RequestScrollIntoView(NewSelectedConnection);
+		// If no selection, auto select the last (newest) session.
+		if (!NewSelectedTraceSession.IsValid() && TraceSessions.Num() > 0)
+		{
+			NewSelectedTraceSession = TraceSessions.Last();
+		}
+
+		TraceSessionsListView->ScrollToBottom();
+
+		// Restore selection and ensure is visible.
+		if (NewSelectedTraceSession.IsValid())
+		{
+			TraceSessionsListView->SetItemSelection(NewSelectedTraceSession, true);
+			TraceSessionsListView->RequestScrollIntoView(NewSelectedTraceSession);
+		}
 	}
 }
 
 ////////////////////////////////////////////////////////////////////////////////////////////////////
 
-void SStartPageWindow::Connections_OnSelectionChanged(TSharedPtr<FRecorderConnection> Connection, ESelectInfo::Type SelectInfo)
+void SStartPageWindow::TraceSessions_OnSelectionChanged(TSharedPtr<FTraceSession> TraceSession, ESelectInfo::Type SelectInfo)
 {
-	SelectedConnection = Connection;
+	SelectedTraceSession = TraceSession;
+}
 
-	if (Connection.IsValid())
-	{
-		TSharedRef<Trace::ISessionService> SessionService = FInsightsManager::Get()->GetSessionService();
-		for (int32 Index = 0; Index < AvailableModulesEnabledState.Num(); ++Index)
-		{
-			AvailableModulesEnabledState[Index] = SessionService->IsModuleEnabled(SelectedConnection->SessionHandle, AvailableModules[Index].Name);
-		}
-	}
-	else
+////////////////////////////////////////////////////////////////////////////////////////////////////
+
+void SStartPageWindow::TraceSessions_OnMouseButtonDoubleClick(TSharedPtr<FTraceSession> TraceSession)
+{
+	if (TraceSession.IsValid())
 	{
-		for (int32 Index = 0; Index < AvailableModulesEnabledState.Num(); ++Index)
-		{
-			AvailableModulesEnabledState[Index] = false;
-		}
+		FInsightsManager::Get()->LoadSession(TraceSession->SessionHandle);
 	}
 }
 
 ////////////////////////////////////////////////////////////////////////////////////////////////////
 
-EVisibility SStartPageWindow::Connections_Visibility() const
+EVisibility SStartPageWindow::TraceSessions_Visibility() const
 {
-	TSharedRef<Trace::ISessionService> SessionService = FInsightsManager::Get()->GetSessionService();
-	return (SessionService->IsRecorderServerRunning() && Connections.Num() > 0) ? EVisibility::Visible : EVisibility::Collapsed;
+	return (TraceSessions.Num() > 0) ? EVisibility::Visible : EVisibility::Collapsed;
 }
 
 ////////////////////////////////////////////////////////////////////////////////////////////////////
@@ -880,28 +777,7 @@ void SStartPageWindow::Tick(const FGeometry& AllottedGeometry, const double InCu
 	{
 		const uint64 WaitTime = static_cast<uint64>(0.5 / FPlatformTime::GetSecondsPerCycle64()); // 500ms
 		NextTimestamp = Time + WaitTime;
-
-		bIsAnySessionAvailable = FInsightsManager::Get()->IsAnySessionAvailable(LastSessionHandle);
-		//bIsAnyLiveSessionAvailable = FInsightsManager::Get()->IsAnyLiveSessionAvailable(LastLiveSessionHandle);
-
-		TSharedRef<Trace::ISessionService> SessionService = FInsightsManager::Get()->GetSessionService();
-		TArray<Trace::FSessionHandle> LiveSessions;
-		SessionService->GetLiveSessions(LiveSessions);
-
-		if (LiveSessions.Num() > 0)
-		{
-			LastLiveSessionHandle = LiveSessions.Last();
-			bIsAnyLiveSessionAvailable = true;
-		}
-		else
-		{
-			bIsAnyLiveSessionAvailable = false;
-		}
-
-		if (LiveSessions.Num() != Connections.Num())
-		{
-			RefreshConnectionList();
-		}
+		RefreshTraceSessionList();
 	}
 }
 
@@ -921,124 +797,7 @@ EVisibility SStartPageWindow::IsSessionOverlayVisible() const
 
 bool SStartPageWindow::IsSessionValid() const
 {
-	const bool bIsActive = FInsightsManager::Get()->GetSession().IsValid();
-	return bIsActive;
-}
-
-////////////////////////////////////////////////////////////////////////////////////////////////////
-
-void SStartPageWindow::ManageLoadingProgressNotificationState(const FString& Filename, const EInsightsNotificationType NotificatonType, const ELoadingProgressState ProgressState, const float LoadingProgress)
-{
-	const FString BaseFilename = FPaths::GetBaseFilename(Filename);
-
-	if (ProgressState == ELoadingProgressState::Started)
-	{
-		const bool bContains = ActiveNotifications.Contains(Filename);
-		if (!bContains)
-		{
-			FNotificationInfo NotificationInfo(GetTextForNotification(NotificatonType, ProgressState, BaseFilename));
-			NotificationInfo.bFireAndForget = false;
-			NotificationInfo.bUseLargeFont = false;
-
-			// Add two buttons, one for cancel, one for loading the received file.
-			if (NotificatonType == EInsightsNotificationType::LoadingTraceFile)
-			{
-				const FText CancelButtonText = LOCTEXT("CancelButton_Text", "Cancel");
-				const FText CancelButtonTT = LOCTEXT("CancelButton_TTText", "Hides this notification");
-				const FText LoadButtonText = LOCTEXT("LoadButton_Text", "Load file");
-				const FText LoadButtonTT = LOCTEXT("LoadButton_TTText", "Loads the received file and hides this notification");
-
-				NotificationInfo.ButtonDetails.Add(FNotificationButtonInfo(CancelButtonText, CancelButtonTT,
-					FSimpleDelegate::CreateSP(this, &SStartPageWindow::SendingServiceSideCapture_Cancel, Filename), SNotificationItem::CS_Success));
-				NotificationInfo.ButtonDetails.Add(FNotificationButtonInfo(LoadButtonText, LoadButtonTT,
-					FSimpleDelegate::CreateSP(this, &SStartPageWindow::SendingServiceSideCapture_Load, Filename), SNotificationItem::CS_Success));
-			}
-
-			SNotificationItemWeak NotificationItem = NotificationList->AddNotification(NotificationInfo);
-			NotificationItem.Pin()->SetCompletionState(SNotificationItem::CS_Pending);
-			ActiveNotifications.Add(Filename, NotificationItem);
-		}
-	}
-	else if (ProgressState == ELoadingProgressState::InProgress)
-	{
-		const SNotificationItemWeak* LoadingProgressPtr = ActiveNotifications.Find(Filename);
-		if (LoadingProgressPtr)
-		{
-			TSharedPtr<SNotificationItem> LoadingProcessPinned = LoadingProgressPtr->Pin();
-			LoadingProcessPinned->SetText(GetTextForNotification(NotificatonType, ProgressState, BaseFilename, LoadingProgress));
-			LoadingProcessPinned->SetCompletionState(SNotificationItem::CS_Pending);
-		}
-	}
-	else if (ProgressState == ELoadingProgressState::Loaded)
-	{
-		const SNotificationItemWeak* LoadingProgressPtr = ActiveNotifications.Find(Filename);
-		if (LoadingProgressPtr)
-		{
-			TSharedPtr<SNotificationItem> LoadingProcessPinned = LoadingProgressPtr->Pin();
-			LoadingProcessPinned->SetText(GetTextForNotification(NotificatonType, ProgressState, BaseFilename));
-			LoadingProcessPinned->SetCompletionState(SNotificationItem::CS_Success);
-
-			// Notifications for received files are handled by the user.
-			if (NotificatonType == EInsightsNotificationType::LoadingTraceFile)
-			{
-				LoadingProcessPinned->ExpireAndFadeout();
-				ActiveNotifications.Remove(Filename);
-			}
-		}
-	}
-	else if (ProgressState == ELoadingProgressState::Failed)
-	{
-		const SNotificationItemWeak* LoadingProgressPtr = ActiveNotifications.Find(Filename);
-		if (LoadingProgressPtr)
-		{
-			TSharedPtr<SNotificationItem> LoadingProcessPinned = LoadingProgressPtr->Pin();
-			LoadingProcessPinned->SetText(GetTextForNotification(NotificatonType, ProgressState, BaseFilename));
-			LoadingProcessPinned->SetCompletionState(SNotificationItem::CS_Fail);
-
-			LoadingProcessPinned->ExpireAndFadeout();
-			ActiveNotifications.Remove(Filename);
-		}
-	}
-	else if (ProgressState == ELoadingProgressState::Cancelled)
-	{
-		const SNotificationItemWeak* LoadingProgressPtr = ActiveNotifications.Find(Filename);
-		if (LoadingProgressPtr)
-		{
-			TSharedPtr<SNotificationItem> LoadingProcessPinned = LoadingProgressPtr->Pin();
-			LoadingProcessPinned->ExpireAndFadeout();
-			ActiveNotifications.Remove(Filename);
-		}
-	}
-}
-
-////////////////////////////////////////////////////////////////////////////////////////////////////
-
-void SStartPageWindow::SendingServiceSideCapture_Cancel(const FString Filename)
-{
-	const SNotificationItemWeak* LoadingProgressPtr = ActiveNotifications.Find(Filename);
-	if (LoadingProgressPtr)
-	{
-		TSharedPtr<SNotificationItem> LoadingProcessPinned = LoadingProgressPtr->Pin();
-		LoadingProcessPinned->ExpireAndFadeout();
-		ActiveNotifications.Remove(Filename);
-	}
-}
-
-////////////////////////////////////////////////////////////////////////////////////////////////////
-
-void SStartPageWindow::SendingServiceSideCapture_Load(const FString Filename)
-{
-	const SNotificationItemWeak* LoadingProgressPtr = ActiveNotifications.Find(Filename);
-	if (LoadingProgressPtr)
-	{
-		TSharedPtr<SNotificationItem> LoadingProcessPinned = LoadingProgressPtr->Pin();
-		LoadingProcessPinned->ExpireAndFadeout();
-		ActiveNotifications.Remove(Filename);
-
-		const FString PathName = FPaths::ProfilingDir() + TEXT("UnrealStats/Received/");
-		const FString TraceFilepath = PathName + Filename;
-		FInsightsManager::Get()->LoadTraceFile(TraceFilepath);
-	}
+	return FInsightsManager::Get()->GetSession().IsValid();
 }
 
 ////////////////////////////////////////////////////////////////////////////////////////////////////
@@ -1135,39 +894,27 @@ FReply SStartPageWindow::OnDrop(const FGeometry& MyGeometry, const FDragDropEven
 
 ////////////////////////////////////////////////////////////////////////////////////////////////////
 
-bool SStartPageWindow::Live_IsEnabled() const
+bool SStartPageWindow::Open_IsEnabled() const
 {
-	return bIsAnyLiveSessionAvailable;
+	return AvailableSessionCount > 0;
 }
 
 ////////////////////////////////////////////////////////////////////////////////////////////////////
 
-bool SStartPageWindow::Last_IsEnabled() const
+FReply SStartPageWindow::Open_OnClicked()
 {
-	return bIsAnySessionAvailable;
-}
-
-////////////////////////////////////////////////////////////////////////////////////////////////////
-
-FReply SStartPageWindow::Live_OnClicked()
-{
-	FInsightsManager::Get()->LoadLastLiveSession();
-	return FReply::Handled();
-}
-
-////////////////////////////////////////////////////////////////////////////////////////////////////
-
-FReply SStartPageWindow::Last_OnClicked()
-{
-	FInsightsManager::Get()->LoadLastSession();
+	if (SelectedTraceSession.IsValid())
+	{
+		Trace::FSessionHandle SessionHandle = SelectedTraceSession->SessionHandle;
+		FInsightsManager::Get()->LoadSession(SessionHandle);
+	}
 	return FReply::Handled();
 }
 
 ////////////////////////////////////////////////////////////////////////////////////////////////////
 
-FReply SStartPageWindow::Load_OnClicked()
+void SStartPageWindow::OpenFileDialog()
 {
-	//const FString ProfilingDirectory(FPaths::ConvertRelativePathToFull(*FPaths::ProfilingDir()));
 	TSharedRef<Trace::ISessionService> SessionService = FInsightsManager::Get()->GetSessionService();
 	const FString ProfilingDirectory(FPaths::ConvertRelativePathToFull(SessionService->GetLocalSessionDirectory()));
 
@@ -1198,8 +945,6 @@ FReply SStartPageWindow::Load_OnClicked()
 			FInsightsManager::Get()->LoadTraceFile(OutFiles[0]);
 		}
 	}
-
-	return FReply::Handled();
 }
 
 ////////////////////////////////////////////////////////////////////////////////////////////////////
@@ -1220,25 +965,71 @@ void SStartPageWindow::LoadSession(Trace::FSessionHandle SessionHandle)
 
 TSharedRef<SWidget> SStartPageWindow::MakeSessionListMenu()
 {
-	RefreshConnectionList();
+	RefreshTraceSessionList();
 
 	FMenuBuilder MenuBuilder(/*bInShouldCloseWindowAfterMenuSelection=*/true, nullptr);
 
+	MenuBuilder.BeginSection("Misc", LOCTEXT("MiscHeading", "Misc"));
+	{
+		MenuBuilder.AddMenuEntry(
+			LOCTEXT("OpenFileButtonLabel", "Open File..."),
+			LOCTEXT("OpenFileButtonTooltip", "Start analysis for a specified trace file."),
+			FSlateIcon(FInsightsStyle::GetStyleSetName(), "OpenFile.Icon.Small"),
+			FUIAction(FExecuteAction::CreateSP(this, &SStartPageWindow::OpenFileDialog)),
+			NAME_None,
+			EUserInterfaceActionType::Button
+		);
+	}
+	MenuBuilder.EndSection();
+
+	/* TODO: persistent MRU
 	MenuBuilder.BeginSection("MostRecentlyUsedSessions", LOCTEXT("MostRecentlyUsedSessionsHeading", "Most Recently Used Sessions"));
 	{
-		//TODO: MRU
+		TSharedRef<Trace::ISessionService> SessionService = FInsightsManager::Get()->GetSessionService();
+
+		TArray<Trace::FSessionHandle> MruSessionList;
+
+		//TODO: real mru list
+		SessionService->GetAvailableSessions(MruSessionList); // mock the list using available sessions
+		int32 NumSessions = FMath::Min(3, MruSessionList.Num());
+
+		for (int32 SessionIndex = 0; SessionIndex < NumSessions; ++SessionIndex)
+		{
+			Trace::FSessionHandle SessionHandle = MruSessionList[SessionIndex];
+
+			Trace::FSessionInfo SessionInfo;
+			SessionService->GetSessionInfo(SessionHandle, SessionInfo);
+
+			FText Label = FText::FromString(SessionInfo.Name);
+			if (SessionInfo.bIsLive)
+			{
+				Label = FText::Format(LOCTEXT("LiveSessionTextFmt", "{0} (LIVE!)"), Label);
+			}
+
+			MenuBuilder.AddMenuEntry(
+				Label,
+				TAttribute<FText>(), // no tooltip
+				FSlateIcon(),
+				FUIAction(FExecuteAction::CreateSP(this, &SStartPageWindow::LoadSession, SessionHandle)),
+				NAME_None,
+				EUserInterfaceActionType::Button
+			);
+		}
 	}
 	MenuBuilder.EndSection();
+	*/
 
-	MenuBuilder.BeginSection("AvailableSessions", LOCTEXT("AvailableSessionsHeading", "Available Sessions"));
+	MenuBuilder.BeginSection("AvailableSessions", LOCTEXT("AvailableSessionsHeading", "Top 10 Most Recently Created Sessions"));
 	{
 		TSharedRef<Trace::ISessionService> SessionService = FInsightsManager::Get()->GetSessionService();
 
 		TArray<Trace::FSessionHandle> AvailableSessions;
 		SessionService->GetAvailableSessions(AvailableSessions);
 
+		int32 SessionCountLimit = 10; // top 10
+
 		// Iterate in reverse order as we want most recent sessions first.
-		for (int32 SessionIndex = AvailableSessions.Num() - 1; SessionIndex >= 0; --SessionIndex)
+		for (int32 SessionIndex = AvailableSessions.Num() - 1; SessionIndex >= 0 && SessionCountLimit > 0; --SessionIndex, --SessionCountLimit)
 		{
 			Trace::FSessionHandle SessionHandle = AvailableSessions[SessionIndex];
 
@@ -1321,7 +1112,7 @@ FReply SStartPageWindow::StartTraceRecorder_OnClicked()
 {
 	TSharedRef<Trace::ISessionService> SessionService = FInsightsManager::Get()->GetSessionService();
 	SessionService->StartRecorderServer();
-	RefreshConnectionList();
+	RefreshTraceSessionList();
 	return FReply::Handled();
 }
 
@@ -1331,7 +1122,7 @@ FReply SStartPageWindow::StopTraceRecorder_OnClicked()
 {
 	TSharedRef<Trace::ISessionService> SessionService = FInsightsManager::Get()->GetSessionService();
 	SessionService->StopRecorderServer();
-	RefreshConnectionList();
+	RefreshTraceSessionList();
 	return FReply::Handled();
 }
 
diff --git a/Engine/Source/Developer/TraceInsights/Private/Insights/Widgets/SStartPageWindow.h b/Engine/Source/Developer/TraceInsights/Private/Insights/Widgets/SStartPageWindow.h
--- a/Engine/Source/Developer/TraceInsights/Private/Insights/Widgets/SStartPageWindow.h
+++ b/Engine/Source/Developer/TraceInsights/Private/Insights/Widgets/SStartPageWindow.h
@@ -39,16 +39,18 @@ typedef TWeakPtr<class SNotificationItem> SNotificationItemWeak;
 
 ////////////////////////////////////////////////////////////////////////////////////////////////////
 
-struct FRecorderConnection
+struct FTraceSession
 {
 	Trace::FSessionHandle SessionHandle;
 	FText Name;
 	FText Uri;
+	bool bIsLive;
 
-	FRecorderConnection(const Trace::FSessionHandle InSessionHandle, const Trace::FSessionInfo& InSessionInfo)
+	FTraceSession(const Trace::FSessionHandle InSessionHandle, const Trace::FSessionInfo& InSessionInfo)
 		: SessionHandle(InSessionHandle)
 		, Name(FText::FromString(InSessionInfo.Name))
 		, Uri(FText::FromString(InSessionInfo.Uri))
+		, bIsLive(InSessionInfo.bIsLive)
 	{}
 };
 
@@ -68,38 +70,39 @@ public:
 	/** Constructs this widget. */
 	void Construct(const FArguments& InArgs);
 
-	void ManageLoadingProgressNotificationState(const FString& Filename, const EInsightsNotificationType NotificatonType, const ELoadingProgressState ProgressState, const float LoadingProgress);
-
 	void OpenSettings();
 	void CloseSettings();
 
 private:
+	TSharedRef<SWidget> ConstructSessionsPanel();
+	TSharedRef<SWidget> ConstructLoadPanel();
+	TSharedRef<SWidget> ConstructLocalSessionDirectoryPanel();
 	TSharedRef<SWidget> ConstructRecoderPanel();
-	TSharedRef<SWidget> ConstructModuleList();
+	TSharedRef<SWidget> ConstructConnectPanel();
 
-	/** Generate a new row for the Connections list view. */
-	TSharedRef<ITableRow> Connections_OnGenerateRow(TSharedPtr<FRecorderConnection> InConnection, const TSharedRef<STableViewBase>& OwnerTable);
+	/** Generate a new row for the TraceSessions list view. */
+	TSharedRef<ITableRow> TraceSessions_OnGenerateRow(TSharedPtr<FTraceSession> InConnection, const TSharedRef<STableViewBase>& OwnerTable);
 
 	/** Callback for determining the visibility of the 'Please select a trace' overlay. */
 	EVisibility IsSessionOverlayVisible() const;
 
 	bool IsSessionValid() const;
 
-	void SendingServiceSideCapture_Cancel(const FString Filename);
-	void SendingServiceSideCapture_Load(const FString Filename);
-
-	bool Live_IsEnabled() const;
-	bool Last_IsEnabled() const;
-
-	FReply Live_OnClicked();
-	FReply Last_OnClicked();
-	FReply Load_OnClicked();
+	bool Open_IsEnabled() const;
+	FReply Open_OnClicked();
 
+	void OpenFileDialog();
 	void LoadTraceFile(const TCHAR* TraceFile);
 	void LoadSession(Trace::FSessionHandle SessionHandle);
 
 	TSharedRef<SWidget> MakeSessionListMenu();
 
+	void RefreshTraceSessionList();
+	void TraceSessions_OnSelectionChanged(TSharedPtr<FTraceSession> TraceSession, ESelectInfo::Type SelectInfo);
+	void TraceSessions_OnMouseButtonDoubleClick(TSharedPtr<FTraceSession> TraceSession);
+	EVisibility TraceSessions_Visibility() const;
+	FReply RefreshTraceSessions_OnClicked();
+
 	FText GetLocalSessionDirectory() const;
 	FReply ExploreLocalSessionDirectory_OnClicked();
 
@@ -109,15 +112,11 @@ private:
 	FReply StartTraceRecorder_OnClicked();
 	FReply StopTraceRecorder_OnClicked();
 
-	EVisibility Modules_Visibility() const;
-	ECheckBoxState Module_IsChecked(int32 ModuleIndex) const;
-	void Module_OnCheckStateChanged(ECheckBoxState NewRadioState, int32 ModuleIndex);
+	//EVisibility Modules_Visibility() const;
+	//ECheckBoxState Module_IsChecked(int32 ModuleIndex) const;
+	//void Module_OnCheckStateChanged(ECheckBoxState NewRadioState, int32 ModuleIndex);
 
-	FReply RefreshConnections_OnClicked();
 	FReply Connect_OnClicked();
-	void RefreshConnectionList();
-	void Connections_OnSelectionChanged(TSharedPtr<FRecorderConnection> Connection, ESelectInfo::Type SelectInfo);
-	EVisibility Connections_Visibility() const;
 
 	/** Updates the amount of time the profiler has been active. */
 	EActiveTimerReturnType UpdateActiveDuration(double InCurrentTime, float InDeltaTime);
@@ -196,16 +195,14 @@ private:
 	/** Holds all widgets for the profiler window like menu bar, toolbar and tabs. */
 	TSharedPtr<SVerticalBox> MainContentPanel;
 
-	bool bIsAnyLiveSessionAvailable;
-	Trace::FSessionHandle LastLiveSessionHandle;
-
-	bool bIsAnySessionAvailable;
-	Trace::FSessionHandle LastSessionHandle;
+	int32 AvailableSessionCount;
+	int32 LiveSessionCount;
 
-	TSharedPtr<SListView<TSharedPtr<FRecorderConnection>>> ConnectionsListView;
-	TArray<TSharedPtr<FRecorderConnection>> Connections;
+	TSharedPtr<SListView<TSharedPtr<FTraceSession>>> TraceSessionsListView;
+	TArray<TSharedPtr<FTraceSession>> TraceSessions;
 	TSharedPtr<SEditableTextBox> HostTextBox;
-	TSharedPtr<FRecorderConnection> SelectedConnection;
-	TArray<Trace::FModuleInfo> AvailableModules;
-	TArray<bool> AvailableModulesEnabledState;
+	TSharedPtr<FTraceSession> SelectedTraceSession;
+
+	//TArray<Trace::FModuleInfo> AvailableModules;
+	//TArray<bool> AvailableModulesEnabledState;
 };
diff --git a/Engine/Source/Developer/TraceInsights/Private/Insights/Widgets/SStatsView.cpp b/Engine/Source/Developer/TraceInsights/Private/Insights/Widgets/SStatsView.cpp
--- a/Engine/Source/Developer/TraceInsights/Private/Insights/Widgets/SStatsView.cpp
+++ b/Engine/Source/Developer/TraceInsights/Private/Insights/Widgets/SStatsView.cpp
@@ -27,16 +27,25 @@
 
 ////////////////////////////////////////////////////////////////////////////////////////////////////
 
-const EColumnSortMode::Type SStatsView::DefaultColumnSortMode(EColumnSortMode::Ascending);
-const FName SStatsView::DefaultColumnBeingSorted(FStatsViewColumns::NameColumnID);
+const EColumnSortMode::Type SStatsView::GetDefaultColumnSortMode()
+{
+	return EColumnSortMode::Ascending;
+}
+
+////////////////////////////////////////////////////////////////////////////////////////////////////
+
+const FName SStatsView::GetDefaultColumnBeingSorted()
+{
+	return FStatsViewColumns::NameColumnID;
+}
 
 ////////////////////////////////////////////////////////////////////////////////////////////////////
 
 SStatsView::SStatsView()
 	: bExpansionSaved(false)
 	, GroupingMode(EStatsGroupingMode::Flat)
-	, ColumnSortMode(DefaultColumnSortMode)
-	, ColumnBeingSorted(DefaultColumnBeingSorted)
+	, ColumnSortMode(GetDefaultColumnSortMode())
+	, ColumnBeingSorted(GetDefaultColumnBeingSorted())
 {
 	FMemory::Memset(bStatsNodeIsVisible, 1);
 }
@@ -208,6 +217,9 @@ void SStatsView::Construct(const FArguments& InArgs)
 
 	// Register ourselves with the Insights manager.
 	FInsightsManager::Get()->GetSessionChangedEvent().AddSP(this, &SStatsView::InsightsManager_OnSessionChanged);
+
+	// Update the Session (i.e. when analysis session was already started).
+	InsightsManager_OnSessionChanged();
 }
 END_SLATE_FUNCTION_BUILD_OPTIMIZATION
 
@@ -995,7 +1007,7 @@ void SStatsView::CreateGroups()
 	{
 		for (const FStatsNodePtr& StatsNodePtr : StatsNodes)
 		{
-			const FName GroupName = *StatsNodePtr->GetName().GetPlainNameString().Left(1);
+			const FName GroupName = *StatsNodePtr->GetName().GetPlainNameString().Left(1).ToUpper();
 
 			FStatsNodePtr* GroupPtr = GroupNodeSet.Find(GroupName);
 			if (!GroupPtr)
@@ -1160,12 +1172,7 @@ bool SStatsView::HeaderMenu_SortMode_IsChecked(const FName ColumnId, const EColu
 bool SStatsView::HeaderMenu_SortMode_CanExecute(const FName ColumnId, const EColumnSortMode::Type InSortMode) const
 {
 	const FStatsViewColumn& Column = TreeViewHeaderColumns.FindChecked(ColumnId);
-	const bool bIsValid = Column.bCanBeSorted();
-
-	bool bCanExecute = ColumnBeingSorted != ColumnId ? true : ColumnSortMode != InSortMode;
-	bCanExecute = bCanExecute && bIsValid;
-
-	return bCanExecute;
+	return Column.bCanBeSorted();
 }
 
 ////////////////////////////////////////////////////////////////////////////////////////////////////
@@ -1189,7 +1196,7 @@ bool SStatsView::ContextMenu_SortMode_IsChecked(const EColumnSortMode::Type InSo
 
 bool SStatsView::ContextMenu_SortMode_CanExecute(const EColumnSortMode::Type InSortMode) const
 {
-	return ColumnSortMode != InSortMode;
+	return true; //ColumnSortMode != InSortMode;
 }
 
 ////////////////////////////////////////////////////////////////////////////////////////////////////
@@ -1213,7 +1220,7 @@ bool SStatsView::ContextMenu_SortByColumn_IsChecked(const FName ColumnId)
 
 bool SStatsView::ContextMenu_SortByColumn_CanExecute(const FName ColumnId) const
 {
-	return ColumnId != ColumnBeingSorted;
+	return true; //ColumnId != ColumnBeingSorted;
 }
 
 ////////////////////////////////////////////////////////////////////////////////////////////////////
@@ -1289,8 +1296,8 @@ bool SStatsView::ContextMenu_ShowAllColumns_CanExecute() const
 
 void SStatsView::ContextMenu_ShowAllColumns_Execute()
 {
-	ColumnSortMode = DefaultColumnSortMode;
-	ColumnBeingSorted = DefaultColumnBeingSorted;
+	ColumnSortMode = GetDefaultColumnSortMode();
+	ColumnBeingSorted = GetDefaultColumnBeingSorted();
 
 	const int32 NumColumns = FStatsViewColumnFactory::Get().Collection.Num();
 	for (int32 ColumnIndex = 0; ColumnIndex < NumColumns; ColumnIndex++)
@@ -1367,8 +1374,8 @@ bool SStatsView::ContextMenu_ResetColumns_CanExecute() const
 
 void SStatsView::ContextMenu_ResetColumns_Execute()
 {
-	ColumnSortMode = DefaultColumnSortMode;
-	ColumnBeingSorted = DefaultColumnBeingSorted;
+	ColumnSortMode = GetDefaultColumnSortMode();
+	ColumnBeingSorted = GetDefaultColumnBeingSorted();
 
 	const int32 NumColumns = FStatsViewColumnFactory::Get().Collection.Num();
 	for (int32 ColumnIndex = 0; ColumnIndex < NumColumns; ColumnIndex++)
diff --git a/Engine/Source/Developer/TraceInsights/Private/Insights/Widgets/SStatsView.h b/Engine/Source/Developer/TraceInsights/Private/Insights/Widgets/SStatsView.h
--- a/Engine/Source/Developer/TraceInsights/Private/Insights/Widgets/SStatsView.h
+++ b/Engine/Source/Developer/TraceInsights/Private/Insights/Widgets/SStatsView.h
@@ -187,6 +187,9 @@ protected:
 	////////////////////////////////////////////////////////////////////////////////////////////////////
 	// Sorting
 
+	static const EColumnSortMode::Type GetDefaultColumnSortMode();
+	static const FName GetDefaultColumnBeingSorted();
+
 	void SortStats();
 
 	EColumnSortMode::Type GetSortModeForColumn(const FName ColumnId) const;
@@ -329,9 +332,6 @@ protected:
 	/** Name of the column currently being sorted, NAME_None if sorting is disabled. */
 	FName ColumnBeingSorted;
 
-	static const EColumnSortMode::Type DefaultColumnSortMode;
-	static const FName DefaultColumnBeingSorted;
-
 	//////////////////////////////////////////////////
 
 	double StatsStartTime;
diff --git a/Engine/Source/Developer/TraceInsights/Private/Insights/Widgets/STimersView.cpp b/Engine/Source/Developer/TraceInsights/Private/Insights/Widgets/STimersView.cpp
--- a/Engine/Source/Developer/TraceInsights/Private/Insights/Widgets/STimersView.cpp
+++ b/Engine/Source/Developer/TraceInsights/Private/Insights/Widgets/STimersView.cpp
@@ -29,16 +29,25 @@
 
 ////////////////////////////////////////////////////////////////////////////////////////////////////
 
-const EColumnSortMode::Type STimersView::DefaultColumnSortMode(EColumnSortMode::Descending);
-const FName STimersView::DefaultColumnBeingSorted(FTimersViewColumns::TotalInclusiveTimeColumnID);
+const EColumnSortMode::Type STimersView::GetDefaultColumnSortMode()
+{
+	return EColumnSortMode::Descending;
+}
+
+////////////////////////////////////////////////////////////////////////////////////////////////////
+
+const FName STimersView::GetDefaultColumnBeingSorted()
+{
+	return FTimersViewColumns::TotalInclusiveTimeColumnID;
+}
 
 ////////////////////////////////////////////////////////////////////////////////////////////////////
 
 STimersView::STimersView()
 	: bExpansionSaved(false)
 	, GroupingMode(ETimerGroupingMode::ByType)
-	, ColumnSortMode(DefaultColumnSortMode)
-	, ColumnBeingSorted(DefaultColumnBeingSorted)
+	, ColumnSortMode(GetDefaultColumnSortMode())
+	, ColumnBeingSorted(GetDefaultColumnBeingSorted())
 {
 	FMemory::Memset(bTimerNodeIsVisible, 1);
 }
@@ -217,6 +226,9 @@ void STimersView::Construct(const FArguments& InArgs)
 
 	// Register ourselves with the Insights manager.
 	FInsightsManager::Get()->GetSessionChangedEvent().AddSP(this, &STimersView::InsightsManager_OnSessionChanged);
+
+	// Update the Session (i.e. when analysis session was already started).
+	InsightsManager_OnSessionChanged();
 }
 END_SLATE_FUNCTION_BUILD_OPTIMIZATION
 
@@ -1004,7 +1016,7 @@ void STimersView::CreateGroups()
 	{
 		for (const FTimerNodePtr& TimerNodePtr : TimerNodes)
 		{
-			const FName GroupName = *TimerNodePtr->GetName().GetPlainNameString().Left(1);
+			const FName GroupName = *TimerNodePtr->GetName().GetPlainNameString().Left(1).ToUpper();
 
 			FTimerNodePtr* GroupPtr = GroupNodeSet.Find(GroupName);
 			if (!GroupPtr)
@@ -1189,12 +1201,7 @@ bool STimersView::HeaderMenu_SortMode_IsChecked(const FName ColumnId, const ECol
 bool STimersView::HeaderMenu_SortMode_CanExecute(const FName ColumnId, const EColumnSortMode::Type InSortMode) const
 {
 	const FTimersViewColumn& Column = TreeViewHeaderColumns.FindChecked(ColumnId);
-	const bool bIsValid = Column.bCanBeSorted();
-
-	bool bCanExecute = ColumnBeingSorted != ColumnId ? true : ColumnSortMode != InSortMode;
-	bCanExecute = bCanExecute && bIsValid;
-
-	return bCanExecute;
+	return Column.bCanBeSorted();
 }
 
 ////////////////////////////////////////////////////////////////////////////////////////////////////
@@ -1218,7 +1225,7 @@ bool STimersView::ContextMenu_SortMode_IsChecked(const EColumnSortMode::Type InS
 
 bool STimersView::ContextMenu_SortMode_CanExecute(const EColumnSortMode::Type InSortMode) const
 {
-	return ColumnSortMode != InSortMode;
+	return true; //ColumnSortMode != InSortMode;
 }
 
 ////////////////////////////////////////////////////////////////////////////////////////////////////
@@ -1242,7 +1249,7 @@ bool STimersView::ContextMenu_SortByColumn_IsChecked(const FName ColumnId)
 
 bool STimersView::ContextMenu_SortByColumn_CanExecute(const FName ColumnId) const
 {
-	return ColumnId != ColumnBeingSorted;
+	return true; //ColumnId != ColumnBeingSorted;
 }
 
 ////////////////////////////////////////////////////////////////////////////////////////////////////
@@ -1318,8 +1325,8 @@ bool STimersView::ContextMenu_ShowAllColumns_CanExecute() const
 
 void STimersView::ContextMenu_ShowAllColumns_Execute()
 {
-	ColumnSortMode = DefaultColumnSortMode;
-	ColumnBeingSorted = DefaultColumnBeingSorted;
+	ColumnSortMode = GetDefaultColumnSortMode();
+	ColumnBeingSorted = GetDefaultColumnBeingSorted();
 
 	const int32 NumColumns = FTimersViewColumnFactory::Get().Collection.Num();
 	for (int32 ColumnIndex = 0; ColumnIndex < NumColumns; ColumnIndex++)
@@ -1399,8 +1406,8 @@ bool STimersView::ContextMenu_ResetColumns_CanExecute() const
 
 void STimersView::ContextMenu_ResetColumns_Execute()
 {
-	ColumnSortMode = DefaultColumnSortMode;
-	ColumnBeingSorted = DefaultColumnBeingSorted;
+	ColumnSortMode = GetDefaultColumnSortMode();
+	ColumnBeingSorted = GetDefaultColumnBeingSorted();
 
 	const int32 NumColumns = FTimersViewColumnFactory::Get().Collection.Num();
 	for (int32 ColumnIndex = 0; ColumnIndex < NumColumns; ColumnIndex++)
diff --git a/Engine/Source/Developer/TraceInsights/Private/Insights/Widgets/STimersView.h b/Engine/Source/Developer/TraceInsights/Private/Insights/Widgets/STimersView.h
--- a/Engine/Source/Developer/TraceInsights/Private/Insights/Widgets/STimersView.h
+++ b/Engine/Source/Developer/TraceInsights/Private/Insights/Widgets/STimersView.h
@@ -167,6 +167,9 @@ protected:
 	////////////////////////////////////////////////////////////////////////////////////////////////////
 	// Sorting
 
+	static const EColumnSortMode::Type GetDefaultColumnSortMode();
+	static const FName GetDefaultColumnBeingSorted();
+
 	void SortTimers();
 
 	EColumnSortMode::Type GetSortModeForColumn(const FName ColumnId) const;
@@ -309,9 +312,6 @@ protected:
 	/** Name of the column currently being sorted, NAME_None if sorting is disabled. */
 	FName ColumnBeingSorted;
 
-	static const EColumnSortMode::Type DefaultColumnSortMode;
-	static const FName DefaultColumnBeingSorted;
-
 	//////////////////////////////////////////////////
 
 	double StatsStartTime;
diff --git a/Engine/Source/Developer/TraceInsights/Private/Insights/Widgets/STimingProfilerToolbar.cpp b/Engine/Source/Developer/TraceInsights/Private/Insights/Widgets/STimingProfilerToolbar.cpp
--- a/Engine/Source/Developer/TraceInsights/Private/Insights/Widgets/STimingProfilerToolbar.cpp
+++ b/Engine/Source/Developer/TraceInsights/Private/Insights/Widgets/STimingProfilerToolbar.cpp
@@ -52,7 +52,7 @@ void STimingProfilerToolbar::Construct(const FArguments& InArgs)
 			ToolbarBuilder.BeginSection("View");
 			{
 				ToolbarBuilder.AddToolBarButton(FTimingProfilerCommands::Get().ToggleFramesTrackVisibility);
-				ToolbarBuilder.AddToolBarButton(FTimingProfilerCommands::Get().ToggleGraphTrackVisibility);
+				//ToolbarBuilder.AddToolBarButton(FTimingProfilerCommands::Get().ToggleGraphTrackVisibility);
 				ToolbarBuilder.AddToolBarButton(FTimingProfilerCommands::Get().ToggleTimingViewVisibility);
 				ToolbarBuilder.AddToolBarButton(FTimingProfilerCommands::Get().ToggleTimersViewVisibility);
 				ToolbarBuilder.AddToolBarButton(FTimingProfilerCommands::Get().ToggleStatsCountersViewVisibility);
diff --git a/Engine/Source/Developer/TraceInsights/Private/Insights/Widgets/STimingProfilerWindow.cpp b/Engine/Source/Developer/TraceInsights/Private/Insights/Widgets/STimingProfilerWindow.cpp
--- a/Engine/Source/Developer/TraceInsights/Private/Insights/Widgets/STimingProfilerWindow.cpp
+++ b/Engine/Source/Developer/TraceInsights/Private/Insights/Widgets/STimingProfilerWindow.cpp
@@ -75,13 +75,23 @@ BEGIN_SLATE_FUNCTION_BUILD_OPTIMIZATION
 
 TSharedRef<SDockTab> STimingProfilerWindow::SpawnTab_Toolbar(const FSpawnTabArgs& Args)
 {
-	return SNew(SDockTab)
+	const TSharedRef<SDockTab> DockTab = SNew(SDockTab)
 		.ShouldAutosize(true)
 		.TabRole(ETabRole::PanelTab)
 		//.IsEnabled(this, &STimingProfilerWindow::IsProfilerEnabled)
 		[
 			SNew(STimingProfilerToolbar)
 		];
+
+	DockTab->SetOnTabClosed(SDockTab::FOnTabClosedCallback::CreateRaw(this, &STimingProfilerWindow::OnToolbarTabClosed));
+
+	return DockTab;
+}
+
+////////////////////////////////////////////////////////////////////////////////////////////////////
+
+void STimingProfilerWindow::OnToolbarTabClosed(TSharedRef<SDockTab> TabBeingClosed)
+{
 }
 
 ////////////////////////////////////////////////////////////////////////////////////////////////////
@@ -90,13 +100,24 @@ TSharedRef<SDockTab> STimingProfilerWindow::SpawnTab_FramesTrack(const FSpawnTab
 {
 	FTimingProfilerManager::Get()->SetFramesTrackVisible(true);
 
-	return SNew(SDockTab)
+	const TSharedRef<SDockTab> DockTab = SNew(SDockTab)
 		.ShouldAutosize(false)
 		.TabRole(ETabRole::PanelTab)
 		//.IsEnabled(this, &STimingProfilerWindow::IsProfilerEnabled)
 		[
 			SAssignNew(FrameTrack, SFrameTrack)
 		];
+
+	DockTab->SetOnTabClosed(SDockTab::FOnTabClosedCallback::CreateRaw(this, &STimingProfilerWindow::OnFramesTrackTabClosed));
+
+	return DockTab;
+}
+
+////////////////////////////////////////////////////////////////////////////////////////////////////
+
+void STimingProfilerWindow::OnFramesTrackTabClosed(TSharedRef<SDockTab> TabBeingClosed)
+{
+	FTimingProfilerManager::Get()->SetFramesTrackVisible(false);
 }
 
 ////////////////////////////////////////////////////////////////////////////////////////////////////
@@ -105,13 +126,24 @@ TSharedRef<SDockTab> STimingProfilerWindow::SpawnTab_GraphTrack(const FSpawnTabA
 {
 	FTimingProfilerManager::Get()->SetGraphTrackVisible(true);
 
-	return SNew(SDockTab)
+	const TSharedRef<SDockTab> DockTab = SNew(SDockTab)
 		.ShouldAutosize(false)
 		.TabRole(ETabRole::PanelTab)
 		//.IsEnabled(this, &STimingProfilerWindow::IsProfilerEnabled)
 		[
 			SAssignNew(GraphTrack, SGraphTrack)
 		];
+
+	DockTab->SetOnTabClosed(SDockTab::FOnTabClosedCallback::CreateRaw(this, &STimingProfilerWindow::OnGraphTrackTabClosed));
+
+	return DockTab;
+}
+
+////////////////////////////////////////////////////////////////////////////////////////////////////
+
+void STimingProfilerWindow::OnGraphTrackTabClosed(TSharedRef<SDockTab> TabBeingClosed)
+{
+	FTimingProfilerManager::Get()->SetGraphTrackVisible(false);
 }
 
 ////////////////////////////////////////////////////////////////////////////////////////////////////
@@ -120,13 +152,24 @@ TSharedRef<SDockTab> STimingProfilerWindow::SpawnTab_TimingView(const FSpawnTabA
 {
 	FTimingProfilerManager::Get()->SetTimingViewVisible(true);
 
-	return SNew(SDockTab)
+	const TSharedRef<SDockTab> DockTab = SNew(SDockTab)
 		.ShouldAutosize(false)
 		.TabRole(ETabRole::PanelTab)
 		//.IsEnabled(this, &STimingProfilerWindow::IsProfilerEnabled)
 		[
 			SAssignNew(TimingView, STimingView)
 		];
+
+	DockTab->SetOnTabClosed(SDockTab::FOnTabClosedCallback::CreateRaw(this, &STimingProfilerWindow::OnTimingViewTabClosed));
+
+	return DockTab;
+}
+
+////////////////////////////////////////////////////////////////////////////////////////////////////
+
+void STimingProfilerWindow::OnTimingViewTabClosed(TSharedRef<SDockTab> TabBeingClosed)
+{
+	FTimingProfilerManager::Get()->SetTimingViewVisible(false);
 }
 
 ////////////////////////////////////////////////////////////////////////////////////////////////////
@@ -135,13 +178,24 @@ TSharedRef<SDockTab> STimingProfilerWindow::SpawnTab_Timers(const FSpawnTabArgs&
 {
 	FTimingProfilerManager::Get()->SetTimersViewVisible(true);
 
-	return SNew(SDockTab)
+	const TSharedRef<SDockTab> DockTab = SNew(SDockTab)
 		.ShouldAutosize(false)
 		.TabRole(ETabRole::PanelTab)
 		//.IsEnabled(this, &STimingProfilerWindow::IsProfilerEnabled)
 		[
 			SAssignNew(TimersView, STimersView)
 		];
+
+	DockTab->SetOnTabClosed(SDockTab::FOnTabClosedCallback::CreateRaw(this, &STimingProfilerWindow::OnTimersTabClosed));
+
+	return DockTab;
+}
+
+////////////////////////////////////////////////////////////////////////////////////////////////////
+
+void STimingProfilerWindow::OnTimersTabClosed(TSharedRef<SDockTab> TabBeingClosed)
+{
+	FTimingProfilerManager::Get()->SetTimersViewVisible(false);
 }
 
 ////////////////////////////////////////////////////////////////////////////////////////////////////
@@ -150,13 +204,24 @@ TSharedRef<SDockTab> STimingProfilerWindow::SpawnTab_StatsCounters(const FSpawnT
 {
 	FTimingProfilerManager::Get()->SetStatsCountersViewVisible(true);
 
-	return SNew(SDockTab)
+	const TSharedRef<SDockTab> DockTab = SNew(SDockTab)
 		.ShouldAutosize(false)
 		.TabRole(ETabRole::PanelTab)
 		//.IsEnabled(this, &STimingProfilerWindow::IsProfilerEnabled)
 		[
 			SAssignNew(StatsView, SStatsView)
 		];
+
+	DockTab->SetOnTabClosed(SDockTab::FOnTabClosedCallback::CreateRaw(this, &STimingProfilerWindow::OnStatsCountersTabClosed));
+
+	return DockTab;
+}
+
+////////////////////////////////////////////////////////////////////////////////////////////////////
+
+void STimingProfilerWindow::OnStatsCountersTabClosed(TSharedRef<SDockTab> TabBeingClosed)
+{
+	FTimingProfilerManager::Get()->SetStatsCountersViewVisible(false);
 }
 
 ////////////////////////////////////////////////////////////////////////////////////////////////////
@@ -165,13 +230,24 @@ TSharedRef<SDockTab> STimingProfilerWindow::SpawnTab_LogView(const FSpawnTabArgs
 {
 	FTimingProfilerManager::Get()->SetLogViewVisible(true);
 
-	return SNew(SDockTab)
+	const TSharedRef<SDockTab> DockTab = SNew(SDockTab)
 		.ShouldAutosize(false)
 		.TabRole(ETabRole::PanelTab)
 		//.IsEnabled(this, &STimingProfilerWindow::IsProfilerEnabled)
 		[
 			SAssignNew(LogView, SLogView)
 		];
+
+	DockTab->SetOnTabClosed(SDockTab::FOnTabClosedCallback::CreateRaw(this, &STimingProfilerWindow::OnLogViewTabClosed));
+
+	return DockTab;
+}
+
+////////////////////////////////////////////////////////////////////////////////////////////////////
+
+void STimingProfilerWindow::OnLogViewTabClosed(TSharedRef<SDockTab> TabBeingClosed)
+{
+	FTimingProfilerManager::Get()->SetLogViewVisible(false);
 }
 
 ////////////////////////////////////////////////////////////////////////////////////////////////////
@@ -192,10 +268,10 @@ void STimingProfilerWindow::Construct(const FArguments& InArgs, const TSharedRef
 		.SetIcon(FSlateIcon(FInsightsStyle::GetStyleSetName(), "FramesTrack.Icon.Small"))
 		.SetGroup(AppMenuGroup);
 
-	TabManager->RegisterTabSpawner(FTimingProfilerTabs::GraphTrackID, FOnSpawnTab::CreateRaw(this, &STimingProfilerWindow::SpawnTab_GraphTrack))
-		.SetDisplayName(LOCTEXT("GraphTrackTabTitle", "Graph"))
-		.SetIcon(FSlateIcon(FInsightsStyle::GetStyleSetName(), "GraphTrack.Icon.Small"))
-		.SetGroup(AppMenuGroup);
+	//TabManager->RegisterTabSpawner(FTimingProfilerTabs::GraphTrackID, FOnSpawnTab::CreateRaw(this, &STimingProfilerWindow::SpawnTab_GraphTrack))
+	//	.SetDisplayName(LOCTEXT("GraphTrackTabTitle", "Graph"))
+	//	.SetIcon(FSlateIcon(FInsightsStyle::GetStyleSetName(), "GraphTrack.Icon.Small"))
+	//	.SetGroup(AppMenuGroup);
 
 	TabManager->RegisterTabSpawner(FTimingProfilerTabs::TimingViewID, FOnSpawnTab::CreateRaw(this, &STimingProfilerWindow::SpawnTab_TimingView))
 		.SetDisplayName(LOCTEXT("TimingViewTabTitle", "Timing View"))
@@ -284,7 +360,7 @@ void STimingProfilerWindow::Construct(const FArguments& InArgs, const TSharedRef
 
 	// Create & initialize main menu.
 	FMenuBarBuilder MenuBarBuilder = FMenuBarBuilder(TSharedPtr<FUICommandList>());
-	
+
 	MenuBarBuilder.AddPullDownMenu(
 		LOCTEXT("MenuLabel", "MENU"),
 		FText::GetEmpty(),
diff --git a/Engine/Source/Developer/TraceInsights/Private/Insights/Widgets/STimingProfilerWindow.h b/Engine/Source/Developer/TraceInsights/Private/Insights/Widgets/STimingProfilerWindow.h
--- a/Engine/Source/Developer/TraceInsights/Private/Insights/Widgets/STimingProfilerWindow.h
+++ b/Engine/Source/Developer/TraceInsights/Private/Insights/Widgets/STimingProfilerWindow.h
@@ -66,12 +66,25 @@ public:
 
 private:
 	TSharedRef<SDockTab> SpawnTab_Toolbar(const FSpawnTabArgs& Args);
+	void OnToolbarTabClosed(TSharedRef<SDockTab> TabBeingClosed);
+	
 	TSharedRef<SDockTab> SpawnTab_FramesTrack(const FSpawnTabArgs& Args);
+	void OnFramesTrackTabClosed(TSharedRef<SDockTab> TabBeingClosed);
+	
 	TSharedRef<SDockTab> SpawnTab_GraphTrack(const FSpawnTabArgs& Args);
+	void OnGraphTrackTabClosed(TSharedRef<SDockTab> TabBeingClosed);
+	
 	TSharedRef<SDockTab> SpawnTab_TimingView(const FSpawnTabArgs& Args);
+	void OnTimingViewTabClosed(TSharedRef<SDockTab> TabBeingClosed);
+	
 	TSharedRef<SDockTab> SpawnTab_Timers(const FSpawnTabArgs& Args);
+	void OnTimersTabClosed(TSharedRef<SDockTab> TabBeingClosed);
+	
 	TSharedRef<SDockTab> SpawnTab_StatsCounters(const FSpawnTabArgs& Args);
+	void OnStatsCountersTabClosed(TSharedRef<SDockTab> TabBeingClosed);
+
 	TSharedRef<SDockTab> SpawnTab_LogView(const FSpawnTabArgs& Args);
+	void OnLogViewTabClosed(TSharedRef<SDockTab> TabBeingClosed);
 
 	/**
 	 * Fill the main menu with menu items.
diff --git a/Engine/Source/Developer/TraceInsights/Private/Insights/Widgets/STimingView.cpp b/Engine/Source/Developer/TraceInsights/Private/Insights/Widgets/STimingView.cpp
--- a/Engine/Source/Developer/TraceInsights/Private/Insights/Widgets/STimingView.cpp
+++ b/Engine/Source/Developer/TraceInsights/Private/Insights/Widgets/STimingView.cpp
@@ -33,6 +33,7 @@
 #include "Insights/Common/Stopwatch.h"
 #include "Insights/Common/TimeUtils.h"
 #include "Insights/InsightsManager.h"
+#include "Insights/InsightsStyle.h"
 #include "Insights/TimingProfilerCommon.h"
 #include "Insights/TimingProfilerManager.h"
 #include "Insights/ViewModels/BaseTimingTrack.h"
@@ -218,8 +219,10 @@ void STimingView::Reset()
 
 	LoadingGetEventNameFn = FLoadingTrackGetEventNameDelegate::CreateRaw(this, &STimingView::GetLoadTimeProfilerEventName);
 
-	EventAggregationStr.Reset();
-	ObjectTypeAggregationStr.Reset();
+	EventAggregationTotalCount = 0;
+	EventAggregation.Reset();
+	ObjectTypeAggregationTotalCount = 0;
+	ObjectTypeAggregation.Reset();
 
 	IoOverviewTrack = nullptr;
 	IoActivityTrack = nullptr;
@@ -505,11 +508,13 @@ void STimingView::Tick(const FGeometry& AllottedGeometry, const double InCurrent
 				{
 					if (!ThreadGroups.Contains(GroupName))
 					{
-						ThreadGroups.Add(GroupName, { GroupName, bIsGroupVisible, 0 });
+						ThreadGroups.Add(GroupName, { GroupName, bIsGroupVisible, 0, Order });
 					}
 					else
 					{
-						bIsGroupVisible = ThreadGroups[GroupName].bIsVisible;
+						FThreadGroup& ThreadGroup = ThreadGroups[GroupName];
+						bIsGroupVisible = ThreadGroup.bIsVisible;
+						ThreadGroup.Order = Order;
 					}
 				}
 
@@ -521,7 +526,7 @@ void STimingView::Tick(const FGeometry& AllottedGeometry, const double InCurrent
 
 					if (!CachedTimelines.Contains(TrackId))
 					{
-						FString TrackName(ThreadInfo.Name && *ThreadInfo.Name ? ThreadInfo.Name : FString::Printf(TEXT("Thread %u")));
+						FString TrackName(ThreadInfo.Name && *ThreadInfo.Name ? ThreadInfo.Name : FString::Printf(TEXT("Thread %u"), ThreadInfo.Id));
 
 						// Create new Timing Events track for the CPU thread.
 						Track = AddTimingEventsTrack(TrackId, ETimingEventsTrackType::Cpu, TrackName, GroupName, Order);
@@ -575,10 +580,7 @@ void STimingView::Tick(const FGeometry& AllottedGeometry, const double InCurrent
 		if (bIsTimingEventsTrackDirty)
 		{
 			// The list has changed. Sort the list again.
-			TimingEventsTracks.Sort([this](const FTimingEventsTrack& A, const FTimingEventsTrack& B) -> bool
-			{
-				return A.GetOrder() < B.GetOrder();
-			});
+			Algo::SortBy(TimingEventsTracks, &FTimingEventsTrack::GetOrder);
 		}
 	}
 
@@ -1152,7 +1154,15 @@ int32 STimingView::OnPaint(const FPaintArgs& Args, const FGeometry& AllottedGeom
 				const float ValueX = X + ValueOffsetX;
 
 				constexpr float LineH = 14.0f;
-				constexpr float H = 10 * LineH;
+				float H = 5 * LineH;
+				if (Package)
+				{
+					H += 3 * LineH;
+				}
+				if (Export)
+				{
+					H += 2 * LineH;
+				}
 				const float MaxY = FMath::Max(0.0f, Viewport.Height - H - 18.0f);
 				float Y = FMath::Clamp<float>(MousePosition.Y + 18.0f, 0.0f, MaxY);
 
@@ -1187,7 +1197,7 @@ int32 STimingView::OnPaint(const FPaintArgs& Args, const FGeometry& AllottedGeom
 				DrawContext.DrawText(ValueX, Y, DepthStr, MainFont, ValueColor);
 				Y += LineH;
 
-				DrawContext.DrawText(X, Y, TEXT("Package Type:"), MainFont, TextColor);
+				DrawContext.DrawText(X, Y, TEXT("Package Event:"), MainFont, TextColor);
 				DrawContext.DrawText(ValueX, Y, GetName(HoveredTimingEvent.LoadingInfo.PackageEventType), MainFont, ValueColor);
 				Y += LineH;
 
@@ -1209,7 +1219,7 @@ int32 STimingView::OnPaint(const FPaintArgs& Args, const FGeometry& AllottedGeom
 				}
 
 				{
-					DrawContext.DrawText(X, Y, TEXT("Export Type:"), MainFont, TextColor);
+					DrawContext.DrawText(X, Y, TEXT("Export Event:"), MainFont, TextColor);
 					FString ExportTypeStr = FString::Printf(TEXT("%s%s"), GetName(HoveredTimingEvent.LoadingInfo.ExportEventType), Export && Export->IsAsset ? TEXT(" [asset]") : TEXT(""));
 					DrawContext.DrawText(ValueX, Y, ExportTypeStr, MainFont, ValueColor);
 					Y += LineH;
@@ -1240,36 +1250,22 @@ int32 STimingView::OnPaint(const FPaintArgs& Args, const FGeometry& AllottedGeom
 
 	if (bAssetLoadingMode)
 	{
-		const FLinearColor BackgroundColor(0.01f, 0.01f, 0.01f, 0.9f);
-		const FLinearColor TextColor(1.0f, 1.0f, 1.0f, 0.9f);
-
 		constexpr float MarginX = 8.0f;
 		constexpr float MarginY = 8.0f;
 
-		constexpr float BorderX = 4.0f;
-		constexpr float BorderY = 4.0f;
+		const float X = ViewWidth - MarginX;
 		float Y = ViewHeight;
 
-		if (ObjectTypeAggregationStr.Len() > 0)
+		if (ObjectTypeAggregation.Num() > 0)
 		{
-			FVector2D TextSize = FontMeasureService->Measure(ObjectTypeAggregationStr, MainFont);
-			const float X = ViewWidth - MarginX - TextSize.X - 2 * BorderX;
-			Y -= MarginY + TextSize.Y + 2 * BorderY;
-			DrawContext.DrawBox(X, Y, TextSize.X + 2 * BorderX, TextSize.Y + 2 * BorderY, WhiteBrush, BackgroundColor);
-			DrawContext.LayerId++;
-			DrawContext.DrawText(X + BorderX, Y + BorderY, ObjectTypeAggregationStr, MainFont, TextColor);
-			DrawContext.LayerId++;
+			Y -= MarginY;
+			Y -= DrawAssetLoadingAggregationTable(DrawContext, X, Y, TEXT("Object Type Aggregation"), ObjectTypeAggregation, ObjectTypeAggregationTotalCount);
 		}
 
-		if (EventAggregationStr.Len() > 0)
+		if (EventAggregation.Num() > 0)
 		{
-			FVector2D TextSize = FontMeasureService->Measure(EventAggregationStr, MainFont);
-			const float X = ViewWidth - MarginX - TextSize.X - 2 * BorderX;
-			Y -= MarginY + TextSize.Y + 2 * BorderY;
-			DrawContext.DrawBox(X, Y, TextSize.X + 2 * BorderX, TextSize.Y + 2 * BorderY, WhiteBrush, BackgroundColor);
-			DrawContext.LayerId++;
-			DrawContext.DrawText(X + BorderX, Y + BorderY, EventAggregationStr, MainFont, TextColor);
-			DrawContext.LayerId++;
+			Y -= MarginY;
+			Y -= DrawAssetLoadingAggregationTable(DrawContext, X, Y, TEXT("Event Aggregation"), EventAggregation, EventAggregationTotalCount);
 		}
 	}
 
@@ -1278,6 +1274,150 @@ int32 STimingView::OnPaint(const FPaintArgs& Args, const FGeometry& AllottedGeom
 
 ////////////////////////////////////////////////////////////////////////////////////////////////////
 
+float STimingView::DrawAssetLoadingAggregationTable(FDrawContext& DrawContext, float RightX, float BottomY, const TCHAR* TableName, const TArray<FAssetLoadingEventAggregationRow>& Aggregation, int32 TotalRowCount) const
+{
+	const FLinearColor BackgroundColor(0.01f, 0.01f, 0.01f, 0.9f);
+	const FLinearColor TextColorHeader(0.5f, 0.5f, 0.5f, 0.9f);
+	const FLinearColor TextColor(1.0f, 1.0f, 1.0f, 0.9f);
+
+	constexpr float BorderX = 4.0f;
+	constexpr float BorderY = 4.0f;
+
+	constexpr float LineH = 14.0f;
+
+	float TableHeight = (2 + Aggregation.Num()) * LineH;
+	if (Aggregation.Num() < TotalRowCount)
+	{
+		TableHeight += LineH; // for the "[...]" line
+	}
+
+	const TSharedRef<FSlateFontMeasure> FontMeasureService = FSlateApplication::Get().GetRenderer()->GetFontMeasureService();
+	float MaxNameWidth = 140.0f;
+	for (const FAssetLoadingEventAggregationRow& Row : Aggregation)
+	{
+		float NameWidth = FontMeasureService->Measure(Row.Name, MainFont).X;
+		if (NameWidth > MaxNameWidth)
+		{
+			MaxNameWidth = NameWidth;
+		}
+	}
+	float NameColumnWidth = MaxNameWidth + 5.0f;
+
+	constexpr float RowNumberColumnWidth = 20.0f;
+	constexpr float CountColumnWidth = 50.0f;
+	constexpr float TotalColumnWidth = 65.0f;
+	constexpr float MinColumnWidth = 60.0f;
+	constexpr float MaxColumnWidth = 60.0f;
+	constexpr float AvgColumnWidth = 60.0f;
+	constexpr float MedColumnWidth = 60.0f;
+
+	const float TableWidth = RowNumberColumnWidth
+						   + NameColumnWidth
+						   + CountColumnWidth
+						   + TotalColumnWidth
+						   + MaxColumnWidth
+						   + AvgColumnWidth
+						   + MedColumnWidth
+						   + MinColumnWidth;
+
+	float W = TableWidth + 2 * BorderX;
+	float H = TableHeight + 2 * BorderY;
+
+	float Y = BottomY - H;
+
+	DrawContext.DrawBox(RightX - W, Y, W, H, WhiteBrush, BackgroundColor);
+	DrawContext.LayerId++;
+	Y += BorderY;
+
+	const float TableX = RightX - W + BorderX;
+
+	// Table name
+	{
+		FString TableNameRow = FString::Printf(TEXT("%s (%d records, sorted by Total time)"), TableName, TotalRowCount);
+		DrawContext.DrawText(TableX, Y, TableNameRow, MainFont, TextColor);
+		Y += LineH;
+	}
+
+	// Column names
+	{
+		float X = TableX;
+
+		X += RowNumberColumnWidth;
+
+		DrawContext.DrawText(X, Y, TEXT("Name"), MainFont, TextColorHeader);
+		X += NameColumnWidth;
+
+		X += CountColumnWidth;
+		DrawContext.DrawTextAligned(HAlign_Right, X, Y, TEXT("Count"), MainFont, TextColorHeader);
+
+		X += TotalColumnWidth;
+		DrawContext.DrawTextAligned(HAlign_Right, X, Y, TEXT("Total"), MainFont, TextColorHeader);
+
+		X += MaxColumnWidth;
+		DrawContext.DrawTextAligned(HAlign_Right, X, Y, TEXT("Max [ms]"), MainFont, TextColorHeader);
+
+		X += AvgColumnWidth;
+		DrawContext.DrawTextAligned(HAlign_Right, X, Y, TEXT("Avg [ms]"), MainFont, TextColorHeader);
+
+		X += MedColumnWidth;
+		DrawContext.DrawTextAligned(HAlign_Right, X, Y, TEXT("Med [ms]"), MainFont, TextColorHeader);
+
+		X += MinColumnWidth;
+		DrawContext.DrawTextAligned(HAlign_Right, X, Y, TEXT("Min [ms]"), MainFont, TextColorHeader);
+
+		Y += LineH;
+	}
+
+	// Records
+	for (int32 Index = 0; Index < Aggregation.Num(); ++Index)
+	{
+		constexpr int32 NumDigits = 2;
+		constexpr bool bAddTimeUnit = false;
+
+		const FAssetLoadingEventAggregationRow& Row = Aggregation[Index];
+
+		float X = TableX;
+
+		X += RowNumberColumnWidth;
+		DrawContext.DrawTextAligned(HAlign_Right, X - 4.0f, Y, FString::Printf(TEXT("%d."), Index + 1), MainFont, TextColorHeader);
+
+		DrawContext.DrawText(X, Y, Row.Name, MainFont, TextColor);
+		X += NameColumnWidth;
+
+		X += CountColumnWidth;
+		DrawContext.DrawTextAligned(HAlign_Right, X, Y, FText::AsNumber(Row.Count).ToString(), MainFont, TextColor);
+
+		X += TotalColumnWidth;
+		DrawContext.DrawTextAligned(HAlign_Right, X, Y, TimeUtils::FormatTimeAuto(Row.Total), MainFont, TextColor);
+
+		X += MaxColumnWidth;
+		DrawContext.DrawTextAligned(HAlign_Right, X, Y, TimeUtils::FormatTimeMs(Row.Max, NumDigits, bAddTimeUnit), MainFont, TextColor);
+
+		X += AvgColumnWidth;
+		DrawContext.DrawTextAligned(HAlign_Right, X, Y, TimeUtils::FormatTimeMs(Row.Avg, NumDigits, bAddTimeUnit), MainFont, TextColor);
+
+		X += MedColumnWidth;
+		DrawContext.DrawTextAligned(HAlign_Right, X, Y, TimeUtils::FormatTimeMs(Row.Med, NumDigits, bAddTimeUnit), MainFont, TextColor);
+
+		X += MinColumnWidth;
+		DrawContext.DrawTextAligned(HAlign_Right, X, Y, TimeUtils::FormatTimeMs(Row.Min, NumDigits, bAddTimeUnit), MainFont, TextColor);
+
+		Y += LineH;
+	}
+
+	if (Aggregation.Num() < TotalRowCount)
+	{
+		DrawContext.DrawText(TableX, Y, TEXT("[...]"), MainFont, TextColorHeader);
+		Y += LineH;
+	}
+
+	DrawContext.LayerId++;
+
+	return H;
+}
+
+////////////////////////////////////////////////////////////////////////////////////////////////////
+
 const TCHAR* GetName(ETimingEventsTrackType Type)
 {
 	switch (Type)
@@ -2027,8 +2167,8 @@ FReply STimingView::OnMouseButtonDown(const FGeometry& MyGeometry, const FPointe
 				SelectionEndTime = SelectionStartTime;
 				LastSelectionType = ESelectionType::None;
 				//TODO: SelectionChangingEvent.Broadcast(SelectionStartTime, SelectionEndTime);
-				EventAggregationStr.Reset();
-				ObjectTypeAggregationStr.Reset();
+				EventAggregation.Reset();
+				ObjectTypeAggregation.Reset();
 			}
 
 			// Capture mouse, so we can drag outside this widget.
@@ -2116,8 +2256,8 @@ FReply STimingView::OnMouseButtonUp(const FGeometry& MyGeometry, const FPointerE
 					SelectionEndTime = SelectionStartTime = 0.0;
 					LastSelectionType = ESelectionType::None;
 					//TODO: SelectionChangedEvent.Broadcast(SelectionStartTime, SelectionEndTime);
-					EventAggregationStr.Reset();
-					ObjectTypeAggregationStr.Reset();
+					EventAggregation.Reset();
+					ObjectTypeAggregation.Reset();
 				}
 			}
 
@@ -2227,8 +2367,8 @@ FReply STimingView::OnMouseMove(const FGeometry& MyGeometry, const FPointerEvent
 			}
 			LastSelectionType = ESelectionType::TimeRange;
 			//TODO: SelectionChangingEvent.Broadcast(SelectionStartTime, SelectionEndTime);
-			EventAggregationStr.Reset();
-			ObjectTypeAggregationStr.Reset();
+			EventAggregation.Reset();
+			ObjectTypeAggregation.Reset();
 		}
 
 		Reply = FReply::Handled();
@@ -2653,17 +2793,20 @@ FReply STimingView::OnKeyUp(const FGeometry& MyGeometry, const FKeyEvent& InKeyE
 
 void STimingView::ShowContextMenu(const FVector2D& ScreenSpacePosition, const FPointerEvent& MouseEvent)
 {
+	/* TODO
 	TSharedPtr<FUICommandList> ProfilerCommandList = FTimingProfilerManager::Get()->GetCommandList();
 	const FTimingProfilerCommands& ProfilerCommands = FTimingProfilerManager::GetCommands();
-	const FTimingProfilerActionManager& ProfilerActionManager = FTimingProfilerManager::GetActionManager();
+	const FTimingViewCommands& Commands = FTimingViewCommands::Get();
 
 	const bool bShouldCloseWindowAfterMenuSelection = true;
 	FMenuBuilder MenuBuilder(bShouldCloseWindowAfterMenuSelection, ProfilerCommandList);
 
 	MenuBuilder.BeginSection(TEXT("Misc"), LOCTEXT("Miscellaneous", "Miscellaneous"));
 	{
-		//MenuBuilder.AddMenuEntry(FTimingProfilerManager::GetCommands().EventGraph_SelectAllFrames);
-		//MenuBuilder.AddMenuEntry(FTimingProfilerManager::GetCommands().ProfilerManager_ToggleLivePreview);
+		MenuBuilder.AddMenuEntry(Commands.ShowAllGpuTracks);
+		MenuBuilder.AddMenuEntry(Commands.ShowAllCpuTracks);
+		MenuBuilder.AddMenuEntry(ProfilerCommands.ToggleTimersViewVisibility);
+		MenuBuilder.AddMenuEntry(ProfilerCommands.ToggleStatsCountersViewVisibility);
 	}
 	MenuBuilder.EndSection();
 
@@ -2671,14 +2814,29 @@ void STimingView::ShowContextMenu(const FVector2D& ScreenSpacePosition, const FP
 
 	FWidgetPath EventPath = MouseEvent.GetEventPath() != nullptr ? *MouseEvent.GetEventPath() : FWidgetPath();
 	FSlateApplication::Get().PushMenu(SharedThis(this), EventPath, MenuWidget, ScreenSpacePosition, FPopupTransitionEffect::ContextMenu);
+	*/
 }
 
 ////////////////////////////////////////////////////////////////////////////////////////////////////
 
 void STimingView::BindCommands()
 {
-	//TSharedPtr<FUICommandList> CommandList = FTimingProfilerManager::Get()->GetCommandList();
-	//const FTimingProfilerCommands& Commands = FTimingProfilerManager::GetCommands();
+	FTimingViewCommands::Register();
+	const FTimingViewCommands& Commands = FTimingViewCommands::Get();
+
+	TSharedPtr<FUICommandList> CommandList = FTimingProfilerManager::Get()->GetCommandList();
+
+	CommandList->MapAction(
+		Commands.ShowAllGpuTracks,
+		FExecuteAction::CreateSP(this, &STimingView::ShowHideAllGpuTracks_Execute),
+		FCanExecuteAction(), //FCanExecuteAction::CreateLambda([] { return true; }),
+		FIsActionChecked::CreateSP(this, &STimingView::ShowHideAllGpuTracks_IsChecked));
+
+	CommandList->MapAction(
+		Commands.ShowAllCpuTracks,
+		FExecuteAction::CreateSP(this, &STimingView::ShowHideAllCpuTracks_Execute),
+		FCanExecuteAction(), //FCanExecuteAction::CreateLambda([] { return true; }),
+		FIsActionChecked::CreateSP(this, &STimingView::ShowHideAllCpuTracks_IsChecked));
 }
 
 ////////////////////////////////////////////////////////////////////////////////////////////////////
@@ -2827,44 +2985,33 @@ void STimingView::UpdateAggregatedStats()
 		TSharedPtr<const Trace::IAnalysisSession> Session = FInsightsManager::Get()->GetSession();
 		if (Session.IsValid() && Trace::ReadLoadTimeProfilerProvider(*Session.Get()))
 		{
-			Trace::FAnalysisSessionReadScope SessionReadScope(*Session.Get());
-
-			const Trace::ILoadTimeProfilerProvider& LoadTimeProfilerProvider = *Trace::ReadLoadTimeProfilerProvider(*Session.Get());
-
-			Trace::ITable<Trace::FLoadTimeProfilerAggregatedStats>* EventAggregationTable = LoadTimeProfilerProvider.CreateEventAggregation(SelectionStartTime, SelectionEndTime);
-			Trace::ITable<Trace::FLoadTimeProfilerAggregatedStats>* ObjectTypeAggregationTable = LoadTimeProfilerProvider.CreateObjectTypeAggregation(SelectionStartTime, SelectionEndTime);
-
-			auto TableToString = [](const TCHAR* TableName, Trace::ITable<Trace::FLoadTimeProfilerAggregatedStats>* Table) -> FString
+			auto ReadTable = [](Trace::ITable<Trace::FLoadTimeProfilerAggregatedStats>* Table, TArray<FAssetLoadingEventAggregationRow>& Aggregation, int32& TotalRowCount)
 			{
-				FString Str;
-
-				const uint32 TotalRowCount = static_cast<uint32>(Table->GetRowCount());
+				TotalRowCount = Table->GetRowCount();
 
-				Str.Append(TableName).Append(": ").Append(FText::AsNumber(TotalRowCount).ToString()).Append(" records, sorted by Total time");
-
-				constexpr uint32 MaxRowCount = 10; // only get first 10 records
-				const uint32 RowCount = FMath::Min<uint32>(TotalRowCount, MaxRowCount);
+				constexpr int32 MaxRowCount = 10; // only get first 10 records
+				const int32 RowCount = FMath::Min(TotalRowCount, MaxRowCount);
 
 				const Trace::ITableLayout& TableLayout = Table->GetLayout();
 				const int32 ColumnCount = TableLayout.GetColumnCount();
 
-				auto Reader = Table->CreateReader();
+				Trace::ITableReader<Trace::FLoadTimeProfilerAggregatedStats>* Reader = Table->CreateReader();
 
 				//////////////////////////////////////////////////
 
 				struct FSortedIndexEntry
 				{
-					uint32 RowIndex;
+					int32 RowIndex;
 					double Value;
 				};
 
 				TArray<FSortedIndexEntry> SortedIndex;
 				SortedIndex.Reserve(TotalRowCount);
 
-				constexpr uint32 TotalTimeColumnIndex = 2;
+				constexpr int32 TotalTimeColumnIndex = 2;
 				ensure(TableLayout.GetColumnType(TotalTimeColumnIndex) == Trace::ETableColumnType::TableColumnType_Double);
 
-				for (uint32 RowIndex = 0; RowIndex < TotalRowCount; ++RowIndex)
+				for (int32 RowIndex = 0; RowIndex < TotalRowCount; ++RowIndex)
 				{
 					Reader->SetRowIndex(RowIndex);
 					double Value = Reader->GetValueDouble(TotalTimeColumnIndex);
@@ -2875,49 +3022,60 @@ void STimingView::UpdateAggregatedStats()
 
 				//////////////////////////////////////////////////
 
-				for (uint32 Index = 0; Index < RowCount; ++Index)
+				Aggregation.Reset();
+				Aggregation.AddDefaulted(RowCount);
+
+				for (int32 Index = 0; Index < RowCount; ++Index)
 				{
-					uint32 RowIndex = SortedIndex[Index].RowIndex;
-					Reader->SetRowIndex(RowIndex);
+					FAssetLoadingEventAggregationRow& Row = Aggregation[Index];
 
-					Str.Append("\n").Append(FText::AsNumber(Index + 1).ToString()).Append(". ");
+					int32 RowIndex = SortedIndex[Index].RowIndex;
+					Reader->SetRowIndex(RowIndex);
 
+					ensure(ColumnCount == 7);
 					for (int32 ColumnIndex = 0; ColumnIndex < ColumnCount; ++ColumnIndex)
 					{
-						Str.Append(TableLayout.GetColumnName(ColumnIndex)).Append("=");
+						ensure(TableLayout.GetColumnType(0) == Trace::ETableColumnType::TableColumnType_CString);
+						Row.Name = Reader->GetValueCString(0);
 
-						Trace::ETableColumnType ColumnType = TableLayout.GetColumnType(ColumnIndex);
-						switch (ColumnType)
-						{
-						case Trace::ETableColumnType::TableColumnType_Bool:
-							Str.Append(Reader->GetValueBool(ColumnIndex) ? TEXT("True") : TEXT("False")).Append(", ");
-							break;
-						case Trace::ETableColumnType::TableColumnType_Int:
-							Str.Append(FText::AsNumber(Reader->GetValueInt(ColumnIndex)).ToString()).Append(", ");
-							break;
-						case Trace::ETableColumnType::TableColumnType_Float:
-							Str.Append(FText::AsNumber(Reader->GetValueFloat(ColumnIndex)).ToString()).Append(", ");
-							break;
-						case Trace::ETableColumnType::TableColumnType_Double:
-							Str.Append(FText::AsNumber(Reader->GetValueDouble(ColumnIndex)).ToString()).Append(", ");
-							break;
-						case Trace::ETableColumnType::TableColumnType_CString:
-							Str.Append(Reader->GetValueCString(ColumnIndex)).Append(", ");
-							break;
-						}
-					}
-				}
+						ensure(TableLayout.GetColumnType(1) == Trace::ETableColumnType::TableColumnType_Int);
+						Row.Count = Reader->GetValueInt(1);
 
-				if (RowCount != TotalRowCount)
-				{
-					Str.Append("\n[...]");
+						ensure(TableLayout.GetColumnType(2) == Trace::ETableColumnType::TableColumnType_Double);
+						Row.Total = Reader->GetValueDouble(2);
+
+						ensure(TableLayout.GetColumnType(3) == Trace::ETableColumnType::TableColumnType_Double);
+						Row.Min = Reader->GetValueDouble(3);
+
+						ensure(TableLayout.GetColumnType(4) == Trace::ETableColumnType::TableColumnType_Double);
+						Row.Max = Reader->GetValueDouble(4);
+
+						ensure(TableLayout.GetColumnType(5) == Trace::ETableColumnType::TableColumnType_Double);
+						Row.Avg = Reader->GetValueDouble(5);
+
+						ensure(TableLayout.GetColumnType(6) == Trace::ETableColumnType::TableColumnType_Double);
+						Row.Med = Reader->GetValueDouble(6);
+					}
 				}
 
-				return Str;
+				delete Reader;
 			};
 
-			EventAggregationStr = TableToString(TEXT("Event Aggregation"), EventAggregationTable);
-			ObjectTypeAggregationStr = TableToString(TEXT("Object Type Aggregation"), ObjectTypeAggregationTable);
+			Trace::FAnalysisSessionReadScope SessionReadScope(*Session.Get());
+
+			const Trace::ILoadTimeProfilerProvider& LoadTimeProfilerProvider = *Trace::ReadLoadTimeProfilerProvider(*Session.Get());
+
+			{
+				Trace::ITable<Trace::FLoadTimeProfilerAggregatedStats>* EventAggregationTable = LoadTimeProfilerProvider.CreateEventAggregation(SelectionStartTime, SelectionEndTime);
+				ReadTable(EventAggregationTable, EventAggregation, EventAggregationTotalCount);
+				delete EventAggregationTable;
+			}
+
+			{
+				Trace::ITable<Trace::FLoadTimeProfilerAggregatedStats>* ObjectTypeAggregationTable = LoadTimeProfilerProvider.CreateObjectTypeAggregation(SelectionStartTime, SelectionEndTime);
+				ReadTable(ObjectTypeAggregationTable, ObjectTypeAggregation, ObjectTypeAggregationTotalCount);
+				delete ObjectTypeAggregationTable;
+			}
 		}
 	}
 }
@@ -3245,10 +3403,12 @@ bool STimingView::SearchTimingEvent(const double InStartTime,
 void STimingView::OnSelectedTimingEventChanged()
 {
 	// Select the timer node coresponding to timing event type of selected timing event.
-	if (SelectedTimingEvent.IsValid() &&
+	if (!bAssetLoadingMode &&
+		SelectedTimingEvent.IsValid() &&
 		(SelectedTimingEvent.Track->GetType() == ETimingEventsTrackType::Cpu ||
 		 SelectedTimingEvent.Track->GetType() == ETimingEventsTrackType::Gpu))
 	{
+		//TODO: make this a more generic (i.e. no hardocdings on TimingProfilerManager)
 		TSharedPtr<STimingProfilerWindow> Wnd = FTimingProfilerManager::Get()->GetProfilerWindow();
 		if (Wnd)
 		{
@@ -3477,57 +3637,64 @@ TSharedRef<SWidget> STimingView::MakeTracksFilterMenu()
 
 	MenuBuilder.BeginSection("QuickFilter", LOCTEXT("TracksFilterHeading", "Quick Filter"));
 	{
+		const FTimingViewCommands& Commands = FTimingViewCommands::Get();
+
+		//TODO: MenuBuilder.AddMenuEntry(Commands.ShowAllGpuTracks);
 		MenuBuilder.AddMenuEntry(
-			LOCTEXT("ShowAllGpuTracks", "GPU Track (Y)"),
+			LOCTEXT("ShowAllGpuTracks", "GPU Track - Y"),
 			LOCTEXT("ShowAllGpuTracks_Tooltip", "Show/hide the GPU track"),
 			FSlateIcon(),
 			FUIAction(FExecuteAction::CreateSP(this, &STimingView::ShowHideAllGpuTracks_Execute),
-				FCanExecuteAction::CreateLambda([] { return true; }),
-				FIsActionChecked::CreateSP(this, &STimingView::ShowHideAllGpuTracks_IsChecked)),
+					  FCanExecuteAction(),
+					  FIsActionChecked::CreateSP(this, &STimingView::ShowHideAllGpuTracks_IsChecked)),
 			NAME_None,
 			EUserInterfaceActionType::ToggleButton
 		);
 
+		//TODO: MenuBuilder.AddMenuEntry(Commands.ShowAllCpuTracks);
 		MenuBuilder.AddMenuEntry(
-			LOCTEXT("ShowAllCpuTracks", "CPU Thread Tracks (U)"),
+			LOCTEXT("ShowAllCpuTracks", "CPU Thread Tracks - U"),
 			LOCTEXT("ShowAllCpuTracks_Tooltip", "Show/hide all CPU tracks (and all CPU thread groups)"),
 			FSlateIcon(),
 			FUIAction(FExecuteAction::CreateSP(this, &STimingView::ShowHideAllCpuTracks_Execute),
-				FCanExecuteAction::CreateLambda([] { return true; }),
-				FIsActionChecked::CreateSP(this, &STimingView::ShowHideAllCpuTracks_IsChecked)),
+					  FCanExecuteAction(),
+					  FIsActionChecked::CreateSP(this, &STimingView::ShowHideAllCpuTracks_IsChecked)),
 			NAME_None,
 			EUserInterfaceActionType::ToggleButton
 		);
 
+		//TODO: MenuBuilder.AddMenuEntry(Commands.ShowAllLoadingTracks);
 		MenuBuilder.AddMenuEntry(
-			LOCTEXT("ShowAllLoadingTracks", "Asset Loading Tracks (L)"),
+			LOCTEXT("ShowAllLoadingTracks", "Asset Loading Tracks - L"),
 			LOCTEXT("ShowAllLoadingTracks_Tooltip", "Show/hide the Asset Loading tracks"),
 			FSlateIcon(),
 			FUIAction(FExecuteAction::CreateSP(this, &STimingView::ShowHideAllLoadingTracks_Execute),
-				FCanExecuteAction::CreateLambda([] { return true; }),
-				FIsActionChecked::CreateSP(this, &STimingView::ShowHideAllLoadingTracks_IsChecked)),
+					  FCanExecuteAction(),
+					  FIsActionChecked::CreateSP(this, &STimingView::ShowHideAllLoadingTracks_IsChecked)),
 			NAME_None,
 			EUserInterfaceActionType::ToggleButton
 		);
 
+		//TODO: MenuBuilder.AddMenuEntry(Commands.ShowAllIoTracks);
 		MenuBuilder.AddMenuEntry(
-			LOCTEXT("ShowAllIoTracks", "I/O Tracks (I)"),
+			LOCTEXT("ShowAllIoTracks", "I/O Tracks - I"),
 			LOCTEXT("ShowAllIoTracks_Tooltip", "Show/hide the I/O (File Activity) tracks"),
 			FSlateIcon(),
 			FUIAction(FExecuteAction::CreateSP(this, &STimingView::ShowHideAllIoTracks_Execute),
-				FCanExecuteAction::CreateLambda([] { return true; }),
-				FIsActionChecked::CreateSP(this, &STimingView::ShowHideAllIoTracks_IsChecked)),
+					  FCanExecuteAction(),
+					  FIsActionChecked::CreateSP(this, &STimingView::ShowHideAllIoTracks_IsChecked)),
 			NAME_None,
 			EUserInterfaceActionType::ToggleButton
 		);
 
+		//TODO: MenuBuilder.AddMenuEntry(Commands.AutoHideEmptyTracks);
 		MenuBuilder.AddMenuEntry(
-			LOCTEXT("AutoHideEmptyTracks", "Auto Hide Empty Tracks (V)"),
+			LOCTEXT("AutoHideEmptyTracks", "Auto Hide Empty Tracks - V"),
 			LOCTEXT("AutoHideEmptyTracks_Tooltip", "Auto hide empty tracks (ones without timing events in current viewport)"),
 			FSlateIcon(),
 			FUIAction(FExecuteAction::CreateSP(this, &STimingView::ToggleAutoHideEmptyTracks),
-				FCanExecuteAction::CreateLambda([] { return true; }),
-				FIsActionChecked::CreateSP(this, &STimingView::IsAutoHideEmptyTracksEnabled)),
+					  FCanExecuteAction(),
+					  FIsActionChecked::CreateSP(this, &STimingView::IsAutoHideEmptyTracksEnabled)),
 			NAME_None,
 			EUserInterfaceActionType::ToggleButton
 		);
@@ -3549,13 +3716,23 @@ TSharedRef<SWidget> STimingView::MakeTracksFilterMenu()
 
 void STimingView::CreateThreadGroupsMenu(FMenuBuilder& MenuBuilder)
 {
+	// Sort the list of thread groups.
+	TArray<const FThreadGroup*> SortedThreadGroups;
+	SortedThreadGroups.Reserve(ThreadGroups.Num());
 	for (const auto& KV : ThreadGroups)
 	{
-		const FThreadGroup& ThreadGroup = KV.Value;
+		SortedThreadGroups.Add(&KV.Value);
+	}
+	Algo::SortBy(SortedThreadGroups, &FThreadGroup::GetOrder);
+
+	for (const FThreadGroup* ThreadGroupPtr : SortedThreadGroups)
+	{
+		const FThreadGroup& ThreadGroup = *ThreadGroupPtr;
 		if (ThreadGroup.NumTimelines > 0)
 		{
 			MenuBuilder.AddMenuEntry(
-				FText::FromString(ThreadGroup.Name),
+				//FText::FromString(ThreadGroup.Name),
+				FText::Format(LOCTEXT("ThreadGroupFmt", "{0} ({1})"), FText::FromString(ThreadGroup.Name), ThreadGroup.NumTimelines),
 				TAttribute<FText>(), // no tooltip
 				FSlateIcon(),
 				FUIAction(FExecuteAction::CreateSP(this, &STimingView::ToggleTrackVisibilityByGroup_Execute, ThreadGroup.Name),
@@ -3720,7 +3897,7 @@ void STimingView::ToggleAutoHideEmptyTracks()
 {
 	if (Layout.TargetMinTimelineH == 0.0f)
 	{
-		Layout.TargetMinTimelineH = RealMinTimelineH;
+		Layout.TargetMinTimelineH = FTimingEventsTrackLayout::RealMinTimelineH;
 	}
 	else
 	{
diff --git a/Engine/Source/Developer/TraceInsights/Private/Insights/Widgets/STimingView.h b/Engine/Source/Developer/TraceInsights/Private/Insights/Widgets/STimingView.h
--- a/Engine/Source/Developer/TraceInsights/Private/Insights/Widgets/STimingView.h
+++ b/Engine/Source/Developer/TraceInsights/Private/Insights/Widgets/STimingView.h
@@ -44,6 +44,18 @@ class STimingView : public SCompoundWidget
 private:
 	static constexpr float MinTooltipWidth = 110.0f;
 
+protected:
+	struct FAssetLoadingEventAggregationRow
+	{
+		FString Name;
+		int32 Count;
+		double Total;
+		double Min;
+		double Max;
+		double Avg;
+		double Med;
+	};
+
 public:
 	/** Default constructor. */
 	STimingView();
@@ -276,7 +288,7 @@ protected:
 	void DrawIoOverviewTrack(FTimingViewDrawHelper& Helper, FTimingEventsTrack& Track) const;
 	void DrawIoActivityTrack(FTimingViewDrawHelper& Helper, FTimingEventsTrack& Track) const;
 
-	void DrawTimeRangeSelection(FDrawContext& DH) const;
+	void DrawTimeRangeSelection(FDrawContext& DrawContext) const;
 
 	void ShowContextMenu(const FVector2D& ScreenSpacePosition, const FPointerEvent& MouseEvent);
 
@@ -303,6 +315,8 @@ protected:
 
 	////////////////////////////////////////////////////////////////////////////////////////////////////
 
+	float DrawAssetLoadingAggregationTable(FDrawContext& DrawContext, float RightX, float BottomY, const TCHAR* TableName, const TArray<FAssetLoadingEventAggregationRow>& Aggregation, int32 TotalRowCount) const;
+
 	void UpdateAggregatedStats();
 
 	void UpdateHoveredTimingEvent(float MX, float MY);
@@ -390,9 +404,12 @@ protected:
 
 	struct FThreadGroup
 	{
-		const TCHAR* Name; /**< the thread group name; pointer to string owned by ThreadProvider */
-		bool bIsVisible;
-		uint32 NumTimelines;
+		const TCHAR* Name; /**< The thread group name; pointer to string owned by ThreadProvider. */
+		bool bIsVisible;  /**< Toggle to show/hide all thread timelines associated with this group at once. Used also as default for new thread timelines. */
+		uint32 NumTimelines; /**< Number of thread timelines associated with this group. */
+		int32 Order; //**< Order index used for sorting. Inherited from last thread timeline associated with this group. **/
+
+		int32 GetOrder() const { return Order; }
 	};
 
 	TMap<const TCHAR*, FThreadGroup> ThreadGroups; /**< maps thread group name to thread group info */
@@ -410,8 +427,11 @@ protected:
 
 	bool bAssetLoadingMode;
 
-	FString EventAggregationStr;
-	FString ObjectTypeAggregationStr;
+	int32 EventAggregationTotalCount;
+	TArray<FAssetLoadingEventAggregationRow> EventAggregation;
+
+	int32 ObjectTypeAggregationTotalCount;
+	TArray<FAssetLoadingEventAggregationRow> ObjectTypeAggregation;
 
 	////////////////////////////////////////////////////////////
 	// File activity (I/O)
diff --git a/Engine/Source/Developer/TraceInsights/Public/Insights/Version.h b/Engine/Source/Developer/TraceInsights/Public/Insights/Version.h
--- a/Engine/Source/Developer/TraceInsights/Public/Insights/Version.h
+++ b/Engine/Source/Developer/TraceInsights/Public/Insights/Version.h
@@ -5,10 +5,10 @@
 // See https://semver.org/
 
 #define UNREAL_INSIGHTS_VERSION_MAJOR 0
-#define UNREAL_INSIGHTS_VERSION_MINOR 6
+#define UNREAL_INSIGHTS_VERSION_MINOR 7
 #define UNREAL_INSIGHTS_VERSION_PATCH 0
 
-#define UNREAL_INSIGHTS_VERSION_STRING "0.06"
+#define UNREAL_INSIGHTS_VERSION_STRING "0.07"
 
 /** Extra identifier string (like: "alpha", "beta.1", "beta.2", "test", "featureX", etc.). */
 //#define UNREAL_INSIGHTS_VERSION_ID_STRING "alpha"
