Copy up of Unreal Insights from //Tasks/UE4/Dev-Core-Trace.

diff --git a/Engine/Source/Developer/TraceInsights/Private/Insights/IoProfilerCommands.h b/Engine/Source/Developer/TraceInsights/Private/Insights/IoProfilerCommands.h
new file mode 100644
--- /dev/null
+++ b/Engine/Source/Developer/TraceInsights/Private/Insights/IoProfilerCommands.h
@@ -0,0 +1,115 @@
+// Copyright 1998-2019 Epic Games, Inc. All Rights Reserved.
+
+#pragma once
+
+#include "CoreMinimal.h"
+#include "Framework/Commands/UIAction.h"
+#include "Framework/Commands/Commands.h"
+
+class FMenuBuilder;
+
+////////////////////////////////////////////////////////////////////////////////////////////////////
+/// Class that holds all profiler commands.
+
+class FIoProfilerCommands
+	: public TCommands<FIoProfilerCommands>
+{
+public:
+	/// Default constructor.
+	FIoProfilerCommands();
+
+	/// Initialize commands.
+	virtual void RegisterCommands() override;
+
+public:
+	//////////////////////////////////////////////////
+	// Global commands need to implement following method:
+	//     void Map_<CommandName>_Global();
+	// Custom commands needs to implement also the following method:
+	//     const FUIAction <CommandName>_Custom(...) const;
+	//////////////////////////////////////////////////
+
+	/// Toggles visibility for the Frames Track. Global and custom command.
+	TSharedPtr<FUICommandInfo> ToggleFramesTrackVisibility;
+
+	/// Toggles visibility for the Graph Track. Global and custom command.
+	TSharedPtr<FUICommandInfo> ToggleGraphTrackVisibility;
+
+	/// Toggles visibility for the Timing Track. Global and custom command.
+	TSharedPtr<FUICommandInfo> ToggleTimingTrackVisibility;
+
+	/// Toggles visibility for the Timers View. Global and custom command.
+	TSharedPtr<FUICommandInfo> ToggleTimersViewVisibility;
+
+	/// Toggles visibility for the Log View. Global and custom command.
+	TSharedPtr<FUICommandInfo> ToggleLogViewVisibility;
+};
+
+////////////////////////////////////////////////////////////////////////////////////////////////////
+
+class FIoProfilerMenuBuilder
+{
+public:
+	/// Helper method for adding a customized menu entry using the global UI command info.
+	/// FUICommandInfo cannot be executed with custom parameters, so we need to create a custom FUIAction,
+	/// but sometime we have global and local version for the UI command, so reuse data from the global UI command info.
+	/// Ex:
+	///     SessionInstance_ToggleCapture          - Global version will toggle capture process for all active session instances
+	///     SessionInstance_ToggleCapture_OneParam - Local version will toggle capture process only for the specified session instance
+	///
+	/// @param MenuBuilder The menu to add items to
+	/// @param FUICommandInfo A shared pointer to the UI command info
+	/// @param UIAction Customized version of the UI command info stored in an UI action
+	///
+	static void AddMenuEntry(FMenuBuilder& MenuBuilder, const TSharedPtr<FUICommandInfo>& UICommandInfo, const FUIAction& UIAction);
+};
+
+////////////////////////////////////////////////////////////////////////////////////////////////////
+/// Class that provides helper functions for the commands to avoid cluttering profiler manager with many small functions.
+/// Can't contain any variables. Directly operates on the profiler manager instance.
+
+class FIoProfilerActionManager
+{
+	friend class FIoProfilerManager;
+
+private:
+	/// Private constructor.
+	FIoProfilerActionManager(class FIoProfilerManager* Instance)
+		: This(Instance)
+	{}
+
+	//////////////////////////////////////////////////
+	// Toggle Commands
+
+#define DECLARE_TOGGLE_COMMAND(CmdName)\
+public:\
+	void Map_##CmdName##_Global(); /**< Maps UI command info CmdName with the specified UI command list. */\
+	const FUIAction CmdName##_Custom() const; /**< UI action for CmdName command. */\
+protected:\
+	void CmdName##_Execute(); /**< Handles FExecuteAction for CmdName. */\
+	bool CmdName##_CanExecute() const; /**< Handles FCanExecuteAction for CmdName. */\
+	ECheckBoxState CmdName##_GetCheckState() const; /**< Handles FGetActionCheckState for CmdName. */
+
+	DECLARE_TOGGLE_COMMAND(ToggleFramesTrackVisibility)
+	DECLARE_TOGGLE_COMMAND(ToggleGraphTrackVisibility)
+	DECLARE_TOGGLE_COMMAND(ToggleTimingTrackVisibility)
+	DECLARE_TOGGLE_COMMAND(ToggleTimersViewVisibility)
+	DECLARE_TOGGLE_COMMAND(ToggleLogViewVisibility)
+#undef DECLARE_TOGGLE_COMMAND
+
+	//////////////////////////////////////////////////
+	// OpenSettings
+
+public:
+	void Map_OpenSettings_Global(); ///< Maps UI command info OpenSettings with the specified UI command list.
+	const FUIAction OpenSettings_Custom() const; ///< UI action for OpenSettings command.
+protected:
+	void OpenSettings_Execute(); ///< Handles FExecuteAction for OpenSettings.
+	bool OpenSettings_CanExecute() const; ///< Handles FCanExecuteAction for OpenSettings.
+
+	//////////////////////////////////////////////////
+
+protected:
+	/// Reference to the global instance of the profiler manager.
+	class FIoProfilerManager* This;
+};
diff --git a/Engine/Source/Developer/TraceInsights/Private/Insights/IoProfilerManager.cpp b/Engine/Source/Developer/TraceInsights/Private/Insights/IoProfilerManager.cpp
new file mode 100644
--- /dev/null
+++ b/Engine/Source/Developer/TraceInsights/Private/Insights/IoProfilerManager.cpp
@@ -0,0 +1,154 @@
+// Copyright 1998-2019 Epic Games, Inc. All Rights Reserved.
+
+#include "IoProfilerManager.h"
+
+#include "Modules/ModuleManager.h"
+#include "Templates/ScopedPointer.h"
+#include "Templates/UniquePtr.h"
+#include "TraceServices/AnalysisService.h"
+#include "TraceServices/SessionService.h"
+
+// Insights
+#include "Insights/InsightsManager.h"
+//#include "Insights/IoProfilerCommon.h"
+#include "Insights/Widgets/SFrameTrack.h"
+#include "Insights/Widgets/SGraphTrack.h"
+#include "Insights/Widgets/SLogView.h"
+#include "Insights/Widgets/STimersView.h"
+#include "Insights/Widgets/STimingView.h"
+#include "Insights/Widgets/SIoProfilerWindow.h"
+
+
+////////////////////////////////////////////////////////////////////////////////////////////////////
+
+#define LOCTEXT_NAMESPACE "IoProfilerManager"
+
+//DEFINE_LOG_CATEGORY(IoProfiler);
+//
+//DEFINE_STAT(STAT_FT_OnPaint);
+//DEFINE_STAT(STAT_GT_OnPaint);
+//DEFINE_STAT(STAT_TT_OnPaint);
+//DEFINE_STAT(STAT_IOPM_Tick);
+
+TSharedPtr<FIoProfilerManager> FIoProfilerManager::Instance = nullptr;
+
+////////////////////////////////////////////////////////////////////////////////////////////////////
+
+FIoProfilerManager::FIoProfilerManager(TSharedRef<FUICommandList> InCommandList)
+	: CommandList(InCommandList)
+	, ActionManager(this)
+	, ProfilerWindow(nullptr)
+	, bIsFramesTrackVisible(true)
+	, bIsGraphTrackVisible(false)
+	, bIsTimingTrackVisible(true)
+	, bIsTimersViewVisible(true)
+	, bIsLogViewVisible(true)
+{
+}
+
+////////////////////////////////////////////////////////////////////////////////////////////////////
+
+void FIoProfilerManager::PostConstructor()
+{
+	// Register tick functions.
+	//OnTick = FTickerDelegate::CreateSP(this, &FIoProfilerManager::Tick);
+	//OnTickHandle = FTicker::GetCoreTicker().AddTicker(OnTick, 1.0f);
+
+	FIoProfilerCommands::Register();
+	BindCommands();
+}
+
+////////////////////////////////////////////////////////////////////////////////////////////////////
+
+void FIoProfilerManager::BindCommands()
+{
+	ActionManager.Map_ToggleFramesTrackVisibility_Global();
+	ActionManager.Map_ToggleGraphTrackVisibility_Global();
+	ActionManager.Map_ToggleTimingTrackVisibility_Global();
+	ActionManager.Map_ToggleTimersViewVisibility_Global();
+	ActionManager.Map_ToggleLogViewVisibility_Global();
+}
+
+////////////////////////////////////////////////////////////////////////////////////////////////////
+
+FIoProfilerManager::~FIoProfilerManager()
+{
+	FIoProfilerCommands::Unregister();
+
+	// Unregister tick function.
+	//FTicker::GetCoreTicker().RemoveTicker(OnTickHandle);
+}
+
+////////////////////////////////////////////////////////////////////////////////////////////////////
+
+TSharedPtr<FIoProfilerManager> FIoProfilerManager::Get()
+{
+	return FIoProfilerManager::Instance;
+}
+
+////////////////////////////////////////////////////////////////////////////////////////////////////
+
+const TSharedRef<FUICommandList> FIoProfilerManager::GetCommandList() const
+{
+	return CommandList;
+}
+
+////////////////////////////////////////////////////////////////////////////////////////////////////
+
+const FIoProfilerCommands& FIoProfilerManager::GetCommands()
+{
+	return FIoProfilerCommands::Get();
+}
+
+////////////////////////////////////////////////////////////////////////////////////////////////////
+
+FIoProfilerActionManager& FIoProfilerManager::GetActionManager()
+{
+	return FIoProfilerManager::Instance->ActionManager;
+}
+
+////////////////////////////////////////////////////////////////////////////////////////////////////
+
+bool FIoProfilerManager::Tick(float DeltaTime)
+{
+	//SCOPE_CYCLE_COUNTER(STAT_IOPM_Tick);
+
+	return true;
+}
+
+////////////////////////////////////////////////////////////////////////////////////////////////////
+
+void FIoProfilerManager::OnSessionChanged()
+{
+	TSharedPtr<SIoProfilerWindow> Wnd = GetProfilerWindow();
+	if (Wnd.IsValid())
+	{
+		Wnd->FrameTrack->Reset();
+		Wnd->GraphTrack->Reset();
+		Wnd->TimingView->Reset();
+		Wnd->TimersView->RebuildTree();
+		Wnd->LogView->Reset();
+	}
+}
+
+////////////////////////////////////////////////////////////////////////////////////////////////////
+
+void FIoProfilerManager::SetTimersViewVisible(const bool bTimersViewVisibleState)
+{
+	bool bWasTimersViewVisible = bIsTimersViewVisible;
+
+	bIsTimersViewVisible = bTimersViewVisibleState;
+
+	if (bIsTimersViewVisible && !bWasTimersViewVisible && FInsightsManager::Get()->GetSession().IsValid())
+	{
+		TSharedPtr<SIoProfilerWindow> Wnd = GetProfilerWindow();
+		if (Wnd.IsValid())
+		{
+			Wnd->TimersView->RebuildTree();
+		}
+	}
+}
+
+////////////////////////////////////////////////////////////////////////////////////////////////////
+
+#undef LOCTEXT_NAMESPACE
diff --git a/Engine/Source/Developer/TraceInsights/Private/Insights/IoProfilerManager.h b/Engine/Source/Developer/TraceInsights/Private/Insights/IoProfilerManager.h
new file mode 100644
--- /dev/null
+++ b/Engine/Source/Developer/TraceInsights/Private/Insights/IoProfilerManager.h
@@ -0,0 +1,157 @@
+// Copyright 1998-2019 Epic Games, Inc. All Rights Reserved.
+
+#pragma once
+
+#include "Containers/Ticker.h"
+#include "CoreMinimal.h"
+#include "Framework/Commands/UICommandList.h"
+
+// Insights
+#include "Insights/InsightsManager.h"
+#include "Insights/IoProfilerCommands.h"
+
+class SIoProfilerWindow;
+
+namespace Trace
+{
+	class ISessionService;
+	class IAnalysisService;
+	class IAnalysisSession;
+}
+
+////////////////////////////////////////////////////////////////////////////////////////////////////
+/// This class manages the Timing Profiler state and settings.
+
+class FIoProfilerManager
+	: public TSharedFromThis<FIoProfilerManager>
+{
+	friend class FIoProfilerActionManager;
+
+public:
+	/// Creates the Timing Profiler manager, only one instance can exist.
+	FIoProfilerManager(TSharedRef<FUICommandList> InCommandList);
+
+	/// Virtual destructor.
+	virtual ~FIoProfilerManager();
+
+	/// Creates an instance of the profiler manager.
+	static TSharedPtr<FIoProfilerManager> Initialize()
+	{
+		if (FIoProfilerManager::Instance.IsValid())
+		{
+			FIoProfilerManager::Instance.Reset();
+		}
+
+		FIoProfilerManager::Instance = MakeShareable(new FIoProfilerManager(FInsightsManager::Get()->GetCommandList()));
+		FIoProfilerManager::Instance->PostConstructor();
+
+		return FIoProfilerManager::Instance;
+	}
+
+	/// Shutdowns the Timing Profiler manager.
+	void Shutdown()
+	{
+		FIoProfilerManager::Instance.Reset();
+	}
+
+protected:
+	/// Finishes initialization of the profiler manager
+	void PostConstructor();
+
+	/// Binds our UI commands to delegates.
+	void BindCommands();
+
+public:
+	/// @return the global instance of the Timing Profiler manager (FIoProfilerManager).
+	/// This is an internal singleton and cannot be used outside ProfilerModule.
+	/// For external use:
+	///     IProfilerModule& ProfilerModule = FModuleManager::Get().LoadModuleChecked<IProfilerModule>("Profiler");
+	///     ProfilerModule.GetProfilerManager();
+	static TSharedPtr<FIoProfilerManager> Get();
+
+	/// @returns UI command list for the Timing Profiler manager.
+	const TSharedRef<FUICommandList> GetCommandList() const;
+
+	/// @return an instance of the Timing Profiler commands.
+	static const FIoProfilerCommands& GetCommands();
+
+	/// @return an instance of the Timing Profiler action manager.
+	static FIoProfilerActionManager& GetActionManager();
+
+	void AssignProfilerWindow(const TSharedRef<SIoProfilerWindow>& InProfilerWindow)
+	{
+		ProfilerWindow = InProfilerWindow;
+	}
+
+	/// Converts profiler window weak pointer to a shared pointer and returns it.
+	/// Make sure the returned pointer is valid before trying to dereference it.
+	TSharedPtr<class SIoProfilerWindow> GetProfilerWindow() const
+	{
+		return ProfilerWindow.Pin();
+	}
+
+	////////////////////////////////////////////////////////////////////////////////////////////////////
+	// Getters and setters used by Toggle Commands.
+
+	/// @return true, if Frames Track is visible
+	const bool IsFramesTrackVisible() const { return bIsFramesTrackVisible; }
+	void SetFramesTrackVisible(const bool bFramesTrackVisibleState) { bIsFramesTrackVisible = bFramesTrackVisibleState; }
+
+	/// @return true, if Graph Track is visible
+	const bool IsGraphTrackVisible() const { return bIsGraphTrackVisible; }
+	void SetGraphTrackVisible(const bool bGraphTrackVisibleState) { bIsGraphTrackVisible = bGraphTrackVisibleState; }
+
+	/// @return true, if Timing Track is visible
+	const bool IsTimingTrackVisible() const { return bIsTimingTrackVisible; }
+	void SetTimingTrackVisible(const bool bTimingTrackVisibleState) { bIsTimingTrackVisible = bTimingTrackVisibleState; }
+
+	/// @return true, if Timers View is visible
+	const bool IsTimersViewVisible() const { return bIsTimersViewVisible; }
+	void SetTimersViewVisible(const bool bTimersViewVisibleState);
+
+	/// @return true, if Log View is visible
+	const bool IsLogViewVisible() const { return bIsLogViewVisible; }
+	void SetLogViewVisible(const bool bLogViewVisibleState) { bIsLogViewVisible = bLogViewVisibleState; }
+
+	////////////////////////////////////////////////////////////////////////////////////////////////////
+
+	void OnSessionChanged();
+
+protected:
+	/// Updates this manager, done through FCoreTicker.
+	bool Tick(float DeltaTime);
+
+protected:
+	/// The delegate to be invoked when this manager ticks.
+	FTickerDelegate OnTick;
+
+	/// Handle to the registered OnTick.
+	FDelegateHandle OnTickHandle;
+
+	/// List of UI commands for this manager. This will be filled by this and corresponding classes.
+	TSharedRef<FUICommandList> CommandList;
+
+	/// An instance of the Timing Manager action manager.
+	FIoProfilerActionManager ActionManager;
+
+	/// A weak pointer to the Timing Profiler window.
+	TWeakPtr<class SIoProfilerWindow> ProfilerWindow;
+
+	/// If the Frames Track is visible or hidden.
+	bool bIsFramesTrackVisible;
+
+	/// If the Graph Track is visible or hidden.
+	bool bIsGraphTrackVisible;
+
+	/// If the Timing Track is visible or hidden.
+	bool bIsTimingTrackVisible;
+
+	/// If the Timers View is visible or hidden.
+	bool bIsTimersViewVisible;
+
+	/// If the Log View is visible or hidden.
+	bool bIsLogViewVisible;
+
+	/// A shared pointer to the global instance of the profiler manager.
+	static TSharedPtr<FIoProfilerManager> Instance;
+};
diff --git a/Engine/Source/Developer/TraceInsights/Private/Insights/TimingProfilerCommands.cpp b/Engine/Source/Developer/TraceInsights/Private/Insights/TimingProfilerCommands.cpp
new file mode 100644
--- /dev/null
+++ b/Engine/Source/Developer/TraceInsights/Private/Insights/TimingProfilerCommands.cpp
@@ -0,0 +1,110 @@
+// Copyright 1998-2019 Epic Games, Inc. All Rights Reserved.
+
+#include "Insights/TimingProfilerCommands.h"
+
+#include "DesktopPlatformModule.h"
+#include "EditorStyleSet.h"
+#include "Framework/Application/SlateApplication.h"
+#include "Framework/MultiBox/MultiBoxBuilder.h"
+
+// Insights
+#include "Insights/InsightsManager.h"
+#include "Insights/TimingProfilerManager.h"
+
+////////////////////////////////////////////////////////////////////////////////////////////////////
+
+#define LOCTEXT_NAMESPACE "FTimingProfilerCommands"
+
+////////////////////////////////////////////////////////////////////////////////////////////////////
+// FTimingProfilerMenuBuilder
+////////////////////////////////////////////////////////////////////////////////////////////////////
+
+void FTimingProfilerMenuBuilder::AddMenuEntry(FMenuBuilder& MenuBuilder, const TSharedPtr< FUICommandInfo >& UICommandInfo, const FUIAction& UIAction)
+{
+	MenuBuilder.AddMenuEntry
+	(
+		UICommandInfo->GetLabel(),
+		UICommandInfo->GetDescription(),
+		UICommandInfo->GetIcon(),
+		UIAction,
+		NAME_None,
+		UICommandInfo->GetUserInterfaceType()
+	);
+}
+
+////////////////////////////////////////////////////////////////////////////////////////////////////
+// FTimingProfilerCommands
+////////////////////////////////////////////////////////////////////////////////////////////////////
+
+FTimingProfilerCommands::FTimingProfilerCommands()
+	: TCommands<FTimingProfilerCommands>(
+		TEXT("TimingProfilerCommand"), // Context name for fast lookup
+		NSLOCTEXT("Contexts", "TimingProfilerCommand", "Timing Profiler Command"), // Localized context name for displaying
+		NAME_None, // Parent
+		FEditorStyle::GetStyleSetName() // Icon Style Set
+	)
+{
+}
+
+////////////////////////////////////////////////////////////////////////////////////////////////////
+
+// UI_COMMAND takes long for the compile to optimize
+PRAGMA_DISABLE_OPTIMIZATION
+void FTimingProfilerCommands::RegisterCommands()
+{
+	UI_COMMAND(ToggleFramesTrackVisibility, "Frames", "Toggles the visibility of the Frames track", EUserInterfaceActionType::ToggleButton, FInputChord(EModifierKey::Control, EKeys::F));
+	UI_COMMAND(ToggleGraphTrackVisibility, "Graph", "Toggles the visibility of the Overview Graph track", EUserInterfaceActionType::ToggleButton, FInputChord(EModifierKey::Control, EKeys::G));
+	UI_COMMAND(ToggleTimingTrackVisibility, "Timing", "Toggles the visibility of the main Timing view", EUserInterfaceActionType::ToggleButton, FInputChord(EModifierKey::Control, EKeys::T));
+	UI_COMMAND(ToggleTimersViewVisibility, "Timers", "Toggles the visibility of the Timers view", EUserInterfaceActionType::ToggleButton, FInputChord(EKeys::T));
+	UI_COMMAND(ToggleLogViewVisibility, "Log", "Toggles the visibility of the Log view", EUserInterfaceActionType::ToggleButton, FInputChord(EKeys::L));
+}
+PRAGMA_ENABLE_OPTIMIZATION
+
+////////////////////////////////////////////////////////////////////////////////////////////////////
+// Toggle Commands
+////////////////////////////////////////////////////////////////////////////////////////////////////
+
+#define IMPLEMENT_TOGGLE_COMMAND(CmdName, IsEnabled, SetIsEnabled) \
+	\
+	void FTimingProfilerActionManager::Map_##CmdName##_Global()\
+	{\
+		This->CommandList->MapAction(This->GetCommands().CmdName, CmdName##_Custom());\
+	}\
+	\
+	const FUIAction FTimingProfilerActionManager::CmdName##_Custom() const\
+	{\
+		FUIAction UIAction;\
+		UIAction.ExecuteAction = FExecuteAction::CreateRaw(this, &FTimingProfilerActionManager::CmdName##_Execute);\
+		UIAction.CanExecuteAction = FCanExecuteAction::CreateRaw(this, &FTimingProfilerActionManager::CmdName##_CanExecute);\
+		UIAction.GetActionCheckState = FGetActionCheckState::CreateRaw(this, &FTimingProfilerActionManager::CmdName##_GetCheckState);\
+		return UIAction;\
+	}\
+	\
+	void FTimingProfilerActionManager::CmdName##_Execute()\
+	{\
+		const bool b##IsEnabled = !This->IsEnabled();\
+		This->SetIsEnabled(b##IsEnabled);\
+	}\
+	\
+	bool FTimingProfilerActionManager::CmdName##_CanExecute() const\
+	{\
+		return FInsightsManager::Get()->GetSession().IsValid();\
+	}\
+	\
+	ECheckBoxState FTimingProfilerActionManager::CmdName##_GetCheckState() const\
+	{\
+		const bool b##IsEnabled = This->IsEnabled();\
+		return b##IsEnabled ? ECheckBoxState::Checked : ECheckBoxState::Unchecked;\
+	}
+
+IMPLEMENT_TOGGLE_COMMAND(ToggleFramesTrackVisibility, IsFramesTrackVisible, SetFramesTrackVisible)
+IMPLEMENT_TOGGLE_COMMAND(ToggleGraphTrackVisibility, IsGraphTrackVisible, SetGraphTrackVisible)
+IMPLEMENT_TOGGLE_COMMAND(ToggleTimingTrackVisibility, IsTimingTrackVisible, SetTimingTrackVisible)
+IMPLEMENT_TOGGLE_COMMAND(ToggleTimersViewVisibility, IsTimersViewVisible, SetTimersViewVisible)
+IMPLEMENT_TOGGLE_COMMAND(ToggleLogViewVisibility, IsLogViewVisible, SetLogViewVisible)
+
+#undef IMPLEMENT_TOGGLE_COMMAND
+
+////////////////////////////////////////////////////////////////////////////////////////////////////
+
+#undef LOCTEXT_NAMESPACE
diff --git a/Engine/Source/Developer/TraceInsights/Private/Insights/TimingProfilerCommands.h b/Engine/Source/Developer/TraceInsights/Private/Insights/TimingProfilerCommands.h
new file mode 100644
--- /dev/null
+++ b/Engine/Source/Developer/TraceInsights/Private/Insights/TimingProfilerCommands.h
@@ -0,0 +1,115 @@
+// Copyright 1998-2019 Epic Games, Inc. All Rights Reserved.
+
+#pragma once
+
+#include "CoreMinimal.h"
+#include "Framework/Commands/UIAction.h"
+#include "Framework/Commands/Commands.h"
+
+class FMenuBuilder;
+
+////////////////////////////////////////////////////////////////////////////////////////////////////
+/// Class that holds all profiler commands.
+
+class FTimingProfilerCommands
+	: public TCommands<FTimingProfilerCommands>
+{
+public:
+	/// Default constructor.
+	FTimingProfilerCommands();
+
+	/// Initialize commands.
+	virtual void RegisterCommands() override;
+
+public:
+	//////////////////////////////////////////////////
+	// Global commands need to implement following method:
+	//     void Map_<CommandName>_Global();
+	// Custom commands needs to implement also the following method:
+	//     const FUIAction <CommandName>_Custom(...) const;
+	//////////////////////////////////////////////////
+
+	/// Toggles visibility for the Frames Track. Global and custom command.
+	TSharedPtr<FUICommandInfo> ToggleFramesTrackVisibility;
+
+	/// Toggles visibility for the Graph Track. Global and custom command.
+	TSharedPtr<FUICommandInfo> ToggleGraphTrackVisibility;
+
+	/// Toggles visibility for the Timing Track. Global and custom command.
+	TSharedPtr<FUICommandInfo> ToggleTimingTrackVisibility;
+
+	/// Toggles visibility for the Timers View. Global and custom command.
+	TSharedPtr<FUICommandInfo> ToggleTimersViewVisibility;
+
+	/// Toggles visibility for the Log View. Global and custom command.
+	TSharedPtr<FUICommandInfo> ToggleLogViewVisibility;
+};
+
+////////////////////////////////////////////////////////////////////////////////////////////////////
+
+class FTimingProfilerMenuBuilder
+{
+public:
+	/// Helper method for adding a customized menu entry using the global UI command info.
+	/// FUICommandInfo cannot be executed with custom parameters, so we need to create a custom FUIAction,
+	/// but sometime we have global and local version for the UI command, so reuse data from the global UI command info.
+	/// Ex:
+	///     SessionInstance_ToggleCapture          - Global version will toggle capture process for all active session instances
+	///     SessionInstance_ToggleCapture_OneParam - Local version will toggle capture process only for the specified session instance
+	///
+	/// @param MenuBuilder The menu to add items to
+	/// @param FUICommandInfo A shared pointer to the UI command info
+	/// @param UIAction Customized version of the UI command info stored in an UI action
+	///
+	static void AddMenuEntry(FMenuBuilder& MenuBuilder, const TSharedPtr<FUICommandInfo>& UICommandInfo, const FUIAction& UIAction);
+};
+
+////////////////////////////////////////////////////////////////////////////////////////////////////
+/// Class that provides helper functions for the commands to avoid cluttering profiler manager with many small functions.
+/// Can't contain any variables. Directly operates on the profiler manager instance.
+
+class FTimingProfilerActionManager
+{
+	friend class FTimingProfilerManager;
+
+private:
+	/// Private constructor.
+	FTimingProfilerActionManager(class FTimingProfilerManager* Instance)
+		: This(Instance)
+	{}
+
+	//////////////////////////////////////////////////
+	// Toggle Commands
+
+#define DECLARE_TOGGLE_COMMAND(CmdName)\
+public:\
+	void Map_##CmdName##_Global(); /**< Maps UI command info CmdName with the specified UI command list. */\
+	const FUIAction CmdName##_Custom() const; /**< UI action for CmdName command. */\
+protected:\
+	void CmdName##_Execute(); /**< Handles FExecuteAction for CmdName. */\
+	bool CmdName##_CanExecute() const; /**< Handles FCanExecuteAction for CmdName. */\
+	ECheckBoxState CmdName##_GetCheckState() const; /**< Handles FGetActionCheckState for CmdName. */
+
+	DECLARE_TOGGLE_COMMAND(ToggleFramesTrackVisibility)
+	DECLARE_TOGGLE_COMMAND(ToggleGraphTrackVisibility)
+	DECLARE_TOGGLE_COMMAND(ToggleTimingTrackVisibility)
+	DECLARE_TOGGLE_COMMAND(ToggleTimersViewVisibility)
+	DECLARE_TOGGLE_COMMAND(ToggleLogViewVisibility)
+#undef DECLARE_TOGGLE_COMMAND
+
+	//////////////////////////////////////////////////
+	// OpenSettings
+
+public:
+	void Map_OpenSettings_Global(); ///< Maps UI command info OpenSettings with the specified UI command list.
+	const FUIAction OpenSettings_Custom() const; ///< UI action for OpenSettings command.
+protected:
+	void OpenSettings_Execute(); ///< Handles FExecuteAction for OpenSettings.
+	bool OpenSettings_CanExecute() const; ///< Handles FCanExecuteAction for OpenSettings.
+
+	//////////////////////////////////////////////////
+
+protected:
+	/// Reference to the global instance of the profiler manager.
+	class FTimingProfilerManager* This;
+};
diff --git a/Engine/Source/Developer/TraceInsights/Private/Insights/TimingProfilerManager.cpp b/Engine/Source/Developer/TraceInsights/Private/Insights/TimingProfilerManager.cpp
new file mode 100644
--- /dev/null
+++ b/Engine/Source/Developer/TraceInsights/Private/Insights/TimingProfilerManager.cpp
@@ -0,0 +1,153 @@
+// Copyright 1998-2019 Epic Games, Inc. All Rights Reserved.
+
+#include "TimingProfilerManager.h"
+
+#include "Modules/ModuleManager.h"
+#include "Templates/ScopedPointer.h"
+#include "Templates/UniquePtr.h"
+#include "TraceServices/AnalysisService.h"
+#include "TraceServices/SessionService.h"
+
+// Insights
+#include "Insights/InsightsManager.h"
+#include "Insights/TimingProfilerCommon.h"
+#include "Insights/Widgets/SFrameTrack.h"
+#include "Insights/Widgets/SGraphTrack.h"
+#include "Insights/Widgets/SLogView.h"
+#include "Insights/Widgets/STimersView.h"
+#include "Insights/Widgets/STimingView.h"
+#include "Insights/Widgets/STimingProfilerWindow.h"
+
+////////////////////////////////////////////////////////////////////////////////////////////////////
+
+#define LOCTEXT_NAMESPACE "TimingProfilerManager"
+
+DEFINE_LOG_CATEGORY(TimingProfiler);
+
+DEFINE_STAT(STAT_FT_OnPaint);
+DEFINE_STAT(STAT_GT_OnPaint);
+DEFINE_STAT(STAT_TT_OnPaint);
+DEFINE_STAT(STAT_TPM_Tick);
+
+TSharedPtr<FTimingProfilerManager> FTimingProfilerManager::Instance = nullptr;
+
+////////////////////////////////////////////////////////////////////////////////////////////////////
+
+FTimingProfilerManager::FTimingProfilerManager(TSharedRef<FUICommandList> InCommandList)
+	: CommandList(InCommandList)
+	, ActionManager(this)
+	, ProfilerWindow(nullptr)
+	, bIsFramesTrackVisible(true)
+	, bIsGraphTrackVisible(false)
+	, bIsTimingTrackVisible(true)
+	, bIsTimersViewVisible(true)
+	, bIsLogViewVisible(true)
+{
+}
+
+////////////////////////////////////////////////////////////////////////////////////////////////////
+
+void FTimingProfilerManager::PostConstructor()
+{
+	// Register tick functions.
+	//OnTick = FTickerDelegate::CreateSP(this, &FTimingProfilerManager::Tick);
+	//OnTickHandle = FTicker::GetCoreTicker().AddTicker(OnTick, 1.0f);
+
+	FTimingProfilerCommands::Register();
+	BindCommands();
+}
+
+////////////////////////////////////////////////////////////////////////////////////////////////////
+
+void FTimingProfilerManager::BindCommands()
+{
+	ActionManager.Map_ToggleFramesTrackVisibility_Global();
+	ActionManager.Map_ToggleGraphTrackVisibility_Global();
+	ActionManager.Map_ToggleTimingTrackVisibility_Global();
+	ActionManager.Map_ToggleTimersViewVisibility_Global();
+	ActionManager.Map_ToggleLogViewVisibility_Global();
+}
+
+////////////////////////////////////////////////////////////////////////////////////////////////////
+
+FTimingProfilerManager::~FTimingProfilerManager()
+{
+	FTimingProfilerCommands::Unregister();
+
+	// Unregister tick function.
+	//FTicker::GetCoreTicker().RemoveTicker(OnTickHandle);
+}
+
+////////////////////////////////////////////////////////////////////////////////////////////////////
+
+TSharedPtr<FTimingProfilerManager> FTimingProfilerManager::Get()
+{
+	return FTimingProfilerManager::Instance;
+}
+
+////////////////////////////////////////////////////////////////////////////////////////////////////
+
+const TSharedRef<FUICommandList> FTimingProfilerManager::GetCommandList() const
+{
+	return CommandList;
+}
+
+////////////////////////////////////////////////////////////////////////////////////////////////////
+
+const FTimingProfilerCommands& FTimingProfilerManager::GetCommands()
+{
+	return FTimingProfilerCommands::Get();
+}
+
+////////////////////////////////////////////////////////////////////////////////////////////////////
+
+FTimingProfilerActionManager& FTimingProfilerManager::GetActionManager()
+{
+	return FTimingProfilerManager::Instance->ActionManager;
+}
+
+////////////////////////////////////////////////////////////////////////////////////////////////////
+
+bool FTimingProfilerManager::Tick(float DeltaTime)
+{
+	SCOPE_CYCLE_COUNTER(STAT_TPM_Tick);
+
+	return true;
+}
+
+////////////////////////////////////////////////////////////////////////////////////////////////////
+
+void FTimingProfilerManager::OnSessionChanged()
+{
+	TSharedPtr<STimingProfilerWindow> Wnd = GetProfilerWindow();
+	if (Wnd.IsValid())
+	{
+		Wnd->FrameTrack->Reset();
+		Wnd->GraphTrack->Reset();
+		Wnd->TimingView->Reset();
+		Wnd->TimersView->RebuildTree();
+		Wnd->LogView->Reset();
+	}
+}
+
+////////////////////////////////////////////////////////////////////////////////////////////////////
+
+void FTimingProfilerManager::SetTimersViewVisible(const bool bTimersViewVisibleState)
+{
+	bool bWasTimersViewVisible = bIsTimersViewVisible;
+
+	bIsTimersViewVisible = bTimersViewVisibleState;
+
+	if (bIsTimersViewVisible && !bWasTimersViewVisible && FInsightsManager::Get()->GetSession().IsValid())
+	{
+		TSharedPtr<STimingProfilerWindow> Wnd = GetProfilerWindow();
+		if (Wnd.IsValid())
+		{
+			Wnd->TimersView->RebuildTree();
+		}
+	}
+}
+
+////////////////////////////////////////////////////////////////////////////////////////////////////
+
+#undef LOCTEXT_NAMESPACE
diff --git a/Engine/Source/Developer/TraceInsights/Private/Insights/TimingProfilerManager.h b/Engine/Source/Developer/TraceInsights/Private/Insights/TimingProfilerManager.h
new file mode 100644
--- /dev/null
+++ b/Engine/Source/Developer/TraceInsights/Private/Insights/TimingProfilerManager.h
@@ -0,0 +1,157 @@
+// Copyright 1998-2019 Epic Games, Inc. All Rights Reserved.
+
+#pragma once
+
+#include "Containers/Ticker.h"
+#include "CoreMinimal.h"
+#include "Framework/Commands/UICommandList.h"
+
+// Insights
+#include "Insights/InsightsManager.h"
+#include "Insights/TimingProfilerCommands.h"
+
+class STimingProfilerWindow;
+
+namespace Trace
+{
+	class ISessionService;
+	class IAnalysisService;
+	class IAnalysisSession;
+}
+
+////////////////////////////////////////////////////////////////////////////////////////////////////
+/// This class manages the Timing Profiler state and settings.
+
+class FTimingProfilerManager
+	: public TSharedFromThis<FTimingProfilerManager>
+{
+	friend class FTimingProfilerActionManager;
+
+public:
+	/// Creates the Timing Profiler manager, only one instance can exist.
+	FTimingProfilerManager(TSharedRef<FUICommandList> InCommandList);
+
+	/// Virtual destructor.
+	virtual ~FTimingProfilerManager();
+
+	/// Creates an instance of the profiler manager.
+	static TSharedPtr<FTimingProfilerManager> Initialize()
+	{
+		if (FTimingProfilerManager::Instance.IsValid())
+		{
+			FTimingProfilerManager::Instance.Reset();
+		}
+
+		FTimingProfilerManager::Instance = MakeShareable(new FTimingProfilerManager(FInsightsManager::Get()->GetCommandList()));
+		FTimingProfilerManager::Instance->PostConstructor();
+
+		return FTimingProfilerManager::Instance;
+	}
+
+	/// Shutdowns the Timing Profiler manager.
+	void Shutdown()
+	{
+		FTimingProfilerManager::Instance.Reset();
+	}
+
+protected:
+	/// Finishes initialization of the profiler manager
+	void PostConstructor();
+
+	/// Binds our UI commands to delegates.
+	void BindCommands();
+
+public:
+	/// @return the global instance of the Timing Profiler manager (FTimingProfilerManager).
+	/// This is an internal singleton and cannot be used outside ProfilerModule.
+	/// For external use:
+	///     IProfilerModule& ProfilerModule = FModuleManager::Get().LoadModuleChecked<IProfilerModule>("Profiler");
+	///     ProfilerModule.GetProfilerManager();
+	static TSharedPtr<FTimingProfilerManager> Get();
+
+	/// @returns UI command list for the Timing Profiler manager.
+	const TSharedRef<FUICommandList> GetCommandList() const;
+
+	/// @return an instance of the Timing Profiler commands.
+	static const FTimingProfilerCommands& GetCommands();
+
+	/// @return an instance of the Timing Profiler action manager.
+	static FTimingProfilerActionManager& GetActionManager();
+
+	void AssignProfilerWindow(const TSharedRef<STimingProfilerWindow>& InProfilerWindow)
+	{
+		ProfilerWindow = InProfilerWindow;
+	}
+
+	/// Converts profiler window weak pointer to a shared pointer and returns it.
+	/// Make sure the returned pointer is valid before trying to dereference it.
+	TSharedPtr<class STimingProfilerWindow> GetProfilerWindow() const
+	{
+		return ProfilerWindow.Pin();
+	}
+
+	////////////////////////////////////////////////////////////////////////////////////////////////////
+	// Getters and setters used by Toggle Commands.
+
+	/// @return true, if Frames Track is visible
+	const bool IsFramesTrackVisible() const { return bIsFramesTrackVisible; }
+	void SetFramesTrackVisible(const bool bFramesTrackVisibleState) { bIsFramesTrackVisible = bFramesTrackVisibleState; }
+
+	/// @return true, if Graph Track is visible
+	const bool IsGraphTrackVisible() const { return bIsGraphTrackVisible; }
+	void SetGraphTrackVisible(const bool bGraphTrackVisibleState) { bIsGraphTrackVisible = bGraphTrackVisibleState; }
+
+	/// @return true, if Timing Track is visible
+	const bool IsTimingTrackVisible() const { return bIsTimingTrackVisible; }
+	void SetTimingTrackVisible(const bool bTimingTrackVisibleState) { bIsTimingTrackVisible = bTimingTrackVisibleState; }
+
+	/// @return true, if Timers View is visible
+	const bool IsTimersViewVisible() const { return bIsTimersViewVisible; }
+	void SetTimersViewVisible(const bool bTimersViewVisibleState);
+
+	/// @return true, if Log View is visible
+	const bool IsLogViewVisible() const { return bIsLogViewVisible; }
+	void SetLogViewVisible(const bool bLogViewVisibleState) { bIsLogViewVisible = bLogViewVisibleState; }
+
+	////////////////////////////////////////////////////////////////////////////////////////////////////
+
+	void OnSessionChanged();
+
+protected:
+	/// Updates this manager, done through FCoreTicker.
+	bool Tick(float DeltaTime);
+
+protected:
+	/// The delegate to be invoked when this manager ticks.
+	FTickerDelegate OnTick;
+
+	/// Handle to the registered OnTick.
+	FDelegateHandle OnTickHandle;
+
+	/// List of UI commands for this manager. This will be filled by this and corresponding classes.
+	TSharedRef<FUICommandList> CommandList;
+
+	/// An instance of the Timing Manager action manager.
+	FTimingProfilerActionManager ActionManager;
+
+	/// A weak pointer to the Timing Profiler window.
+	TWeakPtr<class STimingProfilerWindow> ProfilerWindow;
+
+	/// If the Frames Track is visible or hidden.
+	bool bIsFramesTrackVisible;
+
+	/// If the Graph Track is visible or hidden.
+	bool bIsGraphTrackVisible;
+
+	/// If the Timing Track is visible or hidden.
+	bool bIsTimingTrackVisible;
+
+	/// If the Timers View is visible or hidden.
+	bool bIsTimersViewVisible;
+
+	/// If the Log View is visible or hidden.
+	bool bIsLogViewVisible;
+
+	/// A shared pointer to the global instance of the profiler manager.
+	static TSharedPtr<FTimingProfilerManager> Instance;
+};
diff --git a/Engine/Source/Developer/TraceInsights/Private/Insights/TraceInsightsModule.cpp b/Engine/Source/Developer/TraceInsights/Private/Insights/TraceInsightsModule.cpp
new file mode 100644
--- /dev/null
+++ b/Engine/Source/Developer/TraceInsights/Private/Insights/TraceInsightsModule.cpp
@@ -0,0 +1,257 @@
+// Copyright 1998-2019 Epic Games, Inc. All Rights Reserved.
+
+#include "CoreMinimal.h"
+#include "EditorStyleSet.h"
+#include "HAL/PlatformApplicationMisc.h"
+#include "Modules/ModuleManager.h"
+#include "TraceServices/ITraceServicesModule.h"
+#include "TraceServices/SessionService.h"
+#include "Widgets/DeclarativeSyntaxSupport.h"
+#include "Widgets/Docking/SDockTab.h"
+#include "Widgets/SWidget.h"
+#include "Widgets/Text/STextBlock.h"
+//#include "WorkspaceMenuStructure.h"
+//#include "WorkspaceMenuStructureModule.h"
+
+// Insights
+#include "Insights/IUnrealInsightsModule.h"
+#include "Insights/InsightsManager.h"
+#include "Insights/IoProfilerManager.h"
+#include "Insights/TimingProfilerManager.h"
+#include "Insights/Widgets/SIoProfilerWindow.h"
+#include "Insights/Widgets/STimingProfilerWindow.h"
+
+////////////////////////////////////////////////////////////////////////////////////////////////////
+
+static const FName StartPageTabName("StartPage");
+static const FName TimingProfilerTabName("TimingProfiler");
+static const FName IoProfilerTabName("IoProfiler");
+//static const FName MemoryProfilerTabName("MemoryProfiler");
+
+////////////////////////////////////////////////////////////////////////////////////////////////////
+///
+/// Implements the Trace Insights module (TimingProfiler, IoProfiler, etc.).
+///
+class FTraceInsightsModule : public IUnrealInsightsModule
+{
+public:
+	virtual void StartupModule() override;
+	virtual void ShutdownModule() override;
+
+	virtual bool SupportsDynamicReloading() override
+	{
+		return false;
+	}
+
+	virtual void OnNewLayout(TSharedRef<FTabManager::FLayout> NewLayout) override;
+	virtual void OnLayoutRestored(TSharedPtr<FTabManager> TabManager) override;
+
+protected:
+	/// Callback called when a major tab is closed.
+	void OnTabBeingClosed(TSharedRef<SDockTab> TabBeingClosed);
+
+	/// Callback called when the Timing Profiler major tab is closed.
+	void OnTimingProfilerTabBeingClosed(TSharedRef<SDockTab> TabBeingClosed);
+
+	/// Callback called when the Io Profiler major tab is closed.
+	void OnIoProfilerTabBeingClosed(TSharedRef<SDockTab> TabBeingClosed);
+
+	/// Start Page
+	TSharedRef<SDockTab> SpawnStartPageTab(const FSpawnTabArgs& Args);
+
+	/// Timing Profiler
+	TSharedRef<SDockTab> SpawnTimingProfilerTab(const FSpawnTabArgs& Args);
+
+	/// I/O Profiler
+	TSharedRef<SDockTab> SpawnIoProfilerTab(const FSpawnTabArgs& Args);
+
+protected:
+	TSharedPtr<Trace::IAnalysisService> TraceAnalysisService;
+	TSharedPtr<Trace::ISessionService> TraceSessionService;
+};
+
+////////////////////////////////////////////////////////////////////////////////////////////////////
+
+IMPLEMENT_MODULE(FTraceInsightsModule, TraceInsights);
+
+////////////////////////////////////////////////////////////////////////////////////////////////////
+// FTraceInsightsModule
+////////////////////////////////////////////////////////////////////////////////////////////////////
+
+void FTraceInsightsModule::StartupModule()
+{
+	ITraceServicesModule& TraceServicesModule = FModuleManager::LoadModuleChecked<ITraceServicesModule>("TraceServices");
+	TraceAnalysisService = TraceServicesModule.GetAnalysisService();
+	TraceSessionService = TraceServicesModule.GetSessionService();
+	TraceSessionService->StartRecorderServer();
+
+	FInsightsManager::Initialize(TraceAnalysisService.ToSharedRef(), TraceSessionService.ToSharedRef());
+	FTimingProfilerManager::Initialize();
+	FIoProfilerManager::Initialize();
+
+	// Register tab spawner for the Start Page.
+	auto& StartPageTabSpawnerEntry = FGlobalTabmanager::Get()->RegisterTabSpawner(StartPageTabName,
+		FOnSpawnTab::CreateRaw(this, &FTraceInsightsModule::SpawnStartPageTab))
+		.SetDisplayName(NSLOCTEXT("FTraceInsightsModule", "StartPageTabTitle", "Start Page"))
+		.SetTooltipText(NSLOCTEXT("FTraceInsightsModule", "StartPageTooltipText", "Open the Unreal Insights start page."));
+		//.SetIcon(FSlateIcon(FEditorStyle::GetStyleSetName(), "Profiler.Tab")); //im:TODO
+
+	// Register tab spawner for the Timing Profiler.
+	auto& TimingProfilerTabSpawnerEntry = FGlobalTabmanager::Get()->RegisterTabSpawner(TimingProfilerTabName,
+		FOnSpawnTab::CreateRaw(this, &FTraceInsightsModule::SpawnTimingProfilerTab))
+		.SetDisplayName(NSLOCTEXT("FTraceInsightsModule", "TimingProfilerTabTitle", "Timing Profiler"))
+		.SetTooltipText(NSLOCTEXT("FTraceInsightsModule", "TimingProfilerTooltipText", "Open the Timing Profiler tab."));
+		//.SetIcon(FSlateIcon(FEditorStyle::GetStyleSetName(), "Profiler.Tab")); //im:TODO
+
+//#if WITH_EDITOR
+//	TimingProfilerTabSpawnerEntry.SetGroup(WorkspaceMenu::GetMenuStructure().GetDeveloperToolsMiscCategory());
+//#else
+//	TimingProfilerTabSpawnerEntry.SetGroup(WorkspaceMenu::GetMenuStructure().GetToolsCategory());
+//#endif
+
+	// Register tab spawner for the I/O Profiler.
+	auto& IoProfilerTabSpawnerEntry = FGlobalTabmanager::Get()->RegisterTabSpawner(IoProfilerTabName,
+		FOnSpawnTab::CreateRaw(this, &FTraceInsightsModule::SpawnIoProfilerTab))
+		.SetDisplayName(NSLOCTEXT("FTraceInsightsModule", "IoProfilerTabTitle", "I/O Profiler"))
+		.SetTooltipText(NSLOCTEXT("FTraceInsightsModule", "IoProfilerTooltipText", "Open the Timing Profiler tab."));
+		//.SetIcon(FSlateIcon(FEditorStyle::GetStyleSetName(), "Profiler.Tab")); //im:TODO
+
+//#if WITH_EDITOR
+//	IoProfilerTabSpawnerEntry.SetGroup(WorkspaceMenu::GetMenuStructure().GetDeveloperToolsMiscCategory());
+//#else
+//	IoProfilerTabSpawnerEntry.SetGroup(WorkspaceMenu::GetMenuStructure().GetToolsCategory());
+//#endif
+}
+
+////////////////////////////////////////////////////////////////////////////////////////////////////
+
+void FTraceInsightsModule::ShutdownModule()
+{
+	FGlobalTabmanager::Get()->UnregisterNomadTabSpawner(StartPageTabName);
+	FGlobalTabmanager::Get()->UnregisterNomadTabSpawner(TimingProfilerTabName);
+	FGlobalTabmanager::Get()->UnregisterNomadTabSpawner(IoProfilerTabName);
+
+	if (FInsightsManager::Get().IsValid())
+	{
+		// Shutdown the main manager.
+		FInsightsManager::Get()->Shutdown();
+	}
+}
+
+////////////////////////////////////////////////////////////////////////////////////////////////////
+
+void FTraceInsightsModule::OnNewLayout(TSharedRef<FTabManager::FLayout> NewLayout)
+{
+	const float DPIScaleFactor = FPlatformApplicationMisc::GetDPIScaleFactorAtPoint(10.0f, 10.0f);
+	NewLayout->AddArea
+	(
+		FTabManager::NewArea(1280.f * DPIScaleFactor, 720.0f * DPIScaleFactor)
+		->Split
+		(
+			FTabManager::NewStack()
+			->AddTab(StartPageTabName, ETabState::OpenedTab)
+			->AddTab(TimingProfilerTabName, ETabState::OpenedTab)
+			->AddTab(IoProfilerTabName, ETabState::OpenedTab)
+			//->AddTab(FName("MemoryProfiler"), ETabState::OpenedTab)
+			->SetForegroundTab(FTabId(TimingProfilerTabName))
+		)
+	);
+}
+
+////////////////////////////////////////////////////////////////////////////////////////////////////
+
+void FTraceInsightsModule::OnLayoutRestored(TSharedPtr<FTabManager> TabManager)
+{
+	// Set the Start Page as the main tab.
+	//TSharedPtr<SDockTab> StartPageDockTab = TabManager->FindExistingLiveTab(FTabId(FName("StartPage")));
+	//if (StartPageDockTab.IsValid())
+	//{
+	//	TabManager->SetMainTab(StaticCastSharedRef<SDockTab>(StartPageDockTab->AsShared()));
+	//}
+
+	// Activate the Timing Profiler tab.
+	//TSharedPtr<SDockTab> TimingProfilerDockTab = TabManager->FindExistingLiveTab(FTabId(TimingProfilerTabName));
+	//if (TimingProfilerDockTab.IsValid())
+	//{
+	//	TabManager->SetActiveTab(TimingProfilerDockTab);
+	//}
+}
+
+////////////////////////////////////////////////////////////////////////////////////////////////////
+
+TSharedRef<SDockTab> FTraceInsightsModule::SpawnStartPageTab(const FSpawnTabArgs& Args)
+{
+	const TSharedRef<SDockTab> DockTab = SNew(SDockTab)
+		//.Icon(FEditorStyle::GetBrush("Launcher.TabIcon")) //TODO
+		.TabRole(ETabRole::NomadTab);
+
+	// Create the Start Page widget.
+	TSharedRef<STextBlock> ProfilerWindow = SNew(STextBlock)
+		.Text(NSLOCTEXT("FTraceInsightsModule", "StartPageContent", "TODO"))
+		.ColorAndOpacity(FLinearColor(0.0f, 0.0f, 0.0f, 1.0f));
+	DockTab->SetContent(ProfilerWindow);
+
+	return DockTab;
+}
+
+////////////////////////////////////////////////////////////////////////////////////////////////////
+
+TSharedRef<SDockTab> FTraceInsightsModule::SpawnTimingProfilerTab(const FSpawnTabArgs& Args)
+{
+	const TSharedRef<SDockTab> DockTab = SNew(SDockTab)
+		//.Icon(FEditorStyle::GetBrush("Profiler.Tab")) //TODO
+		.TabRole(ETabRole::NomadTab);
+
+	// Register OnTabClosed to handle Timing profiler manager shutdown.
+	DockTab->SetOnTabClosed(SDockTab::FOnTabClosedCallback::CreateRaw(this, &FTraceInsightsModule::OnTimingProfilerTabBeingClosed));
+
+	// Create the STimingProfilerWindow widget.
+	TSharedRef<STimingProfilerWindow> ProfilerWindow = SNew(STimingProfilerWindow);
+	FTimingProfilerManager::Get()->AssignProfilerWindow(ProfilerWindow);
+	DockTab->SetContent(ProfilerWindow);
+
+	return DockTab;
+}
+
+////////////////////////////////////////////////////////////////////////////////////////////////////
+
+void FTraceInsightsModule::OnTimingProfilerTabBeingClosed(TSharedRef<SDockTab> TabBeingClosed)
+{
+	// Disable TabClosed delegate.
+	TabBeingClosed->SetOnTabClosed(SDockTab::FOnTabClosedCallback());
+
+	// Shutdown the Timing Profiler manager.
+	FTimingProfilerManager::Get()->Shutdown();
+}
+
+////////////////////////////////////////////////////////////////////////////////////////////////////
+
+TSharedRef<SDockTab> FTraceInsightsModule::SpawnIoProfilerTab(const FSpawnTabArgs& Args)
+{
+	const TSharedRef<SDockTab> DockTab = SNew(SDockTab)
+		//.Icon(FEditorStyle::GetBrush("Profiler.Tab")) //TODO
+		.TabRole(ETabRole::NomadTab);
+
+	// Register OnTabClosed to handle I/O profiler manager shutdown.
+	DockTab->SetOnTabClosed(SDockTab::FOnTabClosedCallback::CreateRaw(this, &FTraceInsightsModule::OnIoProfilerTabBeingClosed));
+
+	// Create the SIoProfilerWindow widget.
+	TSharedRef<SIoProfilerWindow> ProfilerWindow = SNew(SIoProfilerWindow);
+	FIoProfilerManager::Get()->AssignProfilerWindow(ProfilerWindow);
+	DockTab->SetContent(ProfilerWindow);
+
+	return DockTab;
+}
+
+////////////////////////////////////////////////////////////////////////////////////////////////////
+
+void FTraceInsightsModule::OnIoProfilerTabBeingClosed(TSharedRef<SDockTab> TabBeingClosed)
+{
+	// Disable TabClosed delegate.
+	TabBeingClosed->SetOnTabClosed(SDockTab::FOnTabClosedCallback());
+
+	// Shutdown the I/O manager.
+	FIoProfilerManager::Get()->Shutdown();
+}
+
+////////////////////////////////////////////////////////////////////////////////////////////////////
diff --git a/Engine/Source/Developer/TraceInsights/Private/Insights/ViewModels/BaseTimingTrack.h b/Engine/Source/Developer/TraceInsights/Private/Insights/ViewModels/BaseTimingTrack.h
new file mode 100644
--- /dev/null
+++ b/Engine/Source/Developer/TraceInsights/Private/Insights/ViewModels/BaseTimingTrack.h
@@ -0,0 +1,78 @@
+// Copyright 1998-2019 Epic Games, Inc. All Rights Reserved.
+
+#pragma once
+
+#include "CoreMinimal.h"
+
+////////////////////////////////////////////////////////////////////////////////////////////////////
+
+class FTimingTrackViewport;
+
+class FBaseTimingTrack
+{
+	friend class FTimingViewDrawHelper;
+
+protected:
+	FBaseTimingTrack(uint64 InId)
+		: Id(InId)
+		, Y(0.0f)
+		, H(0.0f)
+		, bIsVisible(true)
+		, bIsSelected(false)
+		, bIsHovered(false)
+	{}
+
+	virtual ~FBaseTimingTrack()
+	{}
+
+public:
+	virtual void Reset()
+	{
+		Id = 0;
+		Y = 0.0f;
+		H = 0.0f;
+		bIsVisible = true;
+		bIsSelected = false;
+		bIsHovered = false;
+	}
+
+	uint64 GetId() const { return Id; }
+
+	float GetPosY() const { return Y; }
+	virtual void SetPosY(float InPosY) { Y = InPosY; }
+
+	float GetHeight() const { return H; }
+	virtual void SetHeight(float InHeight) { H = InHeight; }
+
+	bool IsVisible() const { return bIsVisible; }
+	void Show() { SetVisibilityFlag(true); }
+	void Hide() { SetVisibilityFlag(false); }
+	virtual void SetVisibilityFlag(bool bInIsVisible) { bIsVisible = bInIsVisible; }
+
+	bool IsSelected() const { return bIsSelected; }
+	void Select() { SetSelectedFlag(true); }
+	void Unselect() { SetSelectedFlag(false); }
+	virtual void SetSelectedFlag(bool bInIsSelected) { bIsSelected = bInIsSelected; }
+
+	bool IsHovered() const { return bIsHovered; }
+	void SetHoveredState(bool bInIsHovered) { bIsHovered = bInIsHovered; }
+	bool IsHeaderHovered() const { return bIsHovered && bIsHeaderHovered; }
+	void SetHeaderHoveredState(bool bInIsHeaderHovered) { bIsHeaderHovered = bInIsHeaderHovered; }
+	virtual void UpdateHoveredState(float MX, float MY, const FTimingTrackViewport& Viewport) = 0;
+
+	virtual void Tick(const double InCurrentTime, const float InDeltaTime) {}
+	virtual void Update(const FTimingTrackViewport& InViewport) {}
+
+protected:
+	uint64 Id;
+
+	float Y; // y position, in Slate units
+	float H; // height, in Slate units
+
+	bool bIsVisible;
+	bool bIsSelected;
+	bool bIsHovered;
+	bool bIsHeaderHovered;
+};
+
+////////////////////////////////////////////////////////////////////////////////////////////////////
diff --git a/Engine/Source/Developer/TraceInsights/Private/Insights/ViewModels/FrameTrackHelper.cpp b/Engine/Source/Developer/TraceInsights/Private/Insights/ViewModels/FrameTrackHelper.cpp
new file mode 100644
--- /dev/null
+++ b/Engine/Source/Developer/TraceInsights/Private/Insights/ViewModels/FrameTrackHelper.cpp
@@ -0,0 +1,391 @@
+// Copyright 1998-2019 Epic Games, Inc. All Rights Reserved.
+
+#include "FrameTrackHelper.h"
+
+#include "Brushes/SlateBorderBrush.h"
+#include "Brushes/SlateColorBrush.h"
+#include "EditorStyleSet.h"
+#include "Fonts/FontMeasure.h"
+#include "Fonts/SlateFontInfo.h"
+#include "Framework/Application/SlateApplication.h"
+#include "Rendering/DrawElements.h"
+#include "TraceServices/AnalysisService.h"
+
+// Insights
+#include "Insights/Common/PaintUtils.h"
+#include "Insights/ViewModels/FrameTrackViewport.h"
+
+#include <limits>
+
+////////////////////////////////////////////////////////////////////////////////////////////////////
+// FFrameTrackCacheContext
+////////////////////////////////////////////////////////////////////////////////////////////////////
+
+FFrameTrackCacheContext::FFrameTrackCacheContext(const FFrameTrackViewport& InViewport)
+	: Viewport(InViewport)
+{
+}
+
+////////////////////////////////////////////////////////////////////////////////////////////////////
+
+void FFrameTrackCacheContext::Begin()
+{
+	SampleW = Viewport.GetSampleWidth();
+	FramesPerSample = Viewport.GetNumFramesPerSample();
+	NumSamples = FMath::CeilToInt(Viewport.Width / SampleW);
+	FirstFrameIndex = Viewport.GetFirstFrameIndex();
+
+	NumFrames = 0;
+}
+
+////////////////////////////////////////////////////////////////////////////////////////////////////
+
+bool FFrameTrackCacheContext::BeginTimeline(FFrameTrackTimeline& CachedTimeline)
+{
+	CurrentCachedTimeline = &CachedTimeline;
+
+	CurrentCachedTimeline->NumAggregatedFrames = 0;
+	CurrentCachedTimeline->Samples.Reset(NumSamples);
+	CurrentCachedTimeline->Samples.AddDefaulted(NumSamples);
+
+	return true;
+}
+
+////////////////////////////////////////////////////////////////////////////////////////////////////
+
+void FFrameTrackCacheContext::AddEvent(const Trace::FFrame& Event)
+{
+	NumFrames++;
+
+	const int32 FrameIndex = Event.Index;
+
+	int SampleIndex = (FrameIndex - FirstFrameIndex) / FramesPerSample;
+	if (SampleIndex >= 0 && SampleIndex < NumSamples)
+	{
+		FFrameTrackSample& Sample = CurrentCachedTimeline->Samples[SampleIndex];
+		Sample.NumFrames++;
+
+		double Duration = Event.EndTime - Event.StartTime;
+		Sample.TotalDuration += Duration;
+		if (Event.StartTime < Sample.StartTime)
+		{
+			Sample.StartTime = Event.StartTime;
+		}
+		if (Event.EndTime > Sample.EndTime)
+		{
+			Sample.EndTime = Event.EndTime;
+		}
+		if (Duration > Sample.LargestFrameDuration)
+		{
+			Sample.LargestFrameIndex = FrameIndex;
+			Sample.LargestFrameStartTime = Event.StartTime;
+			Sample.LargestFrameDuration = Duration;
+		}
+
+		CurrentCachedTimeline->NumAggregatedFrames++;
+	}
+}
+
+////////////////////////////////////////////////////////////////////////////////////////////////////
+
+void FFrameTrackCacheContext::EndTimeline()
+{
+	//CurrentCachedTimeline = nullptr;
+}
+
+////////////////////////////////////////////////////////////////////////////////////////////////////
+
+void FFrameTrackCacheContext::End()
+{
+}
+
+////////////////////////////////////////////////////////////////////////////////////////////////////
+// FFrameTrackDrawContext
+////////////////////////////////////////////////////////////////////////////////////////////////////
+
+FFrameTrackDrawContext::FFrameTrackDrawContext(const FFrameTrackViewport& InViewport, const FGeometry& InAllottedGeometry, int32& InLayerId, FSlateWindowElementList& InElementList)
+	: Viewport(InViewport)
+	, AllottedGeometry(InAllottedGeometry)
+	, LayerId(InLayerId)
+	, ElementList(InElementList)
+	, DrawEffects(ESlateDrawEffect::None)
+	, WhiteBrush(FCoreStyle::Get().GetBrush("WhiteBrush"))
+	, BorderBrush(FEditorStyle::GetBrush("PlainBorder"))
+	, NumFrames(0)
+	, NumDrawSamples(0)
+{
+}
+
+////////////////////////////////////////////////////////////////////////////////////////////////////
+
+void FFrameTrackDrawContext::DrawBackground(const FWidgetStyle& InWidgetStyle) const
+{
+	const FSlateBrush* AreaBrush = WhiteBrush;
+	const FLinearColor ValidAreaColor(0.07f, 0.07f, 0.07f, 1.0f);
+	const FLinearColor InvalidAreaColor(0.1f, 0.07f, 0.07f, 1.0f);
+
+	float X0 = Viewport.MinX - Viewport.PosX;
+	float X1 = Viewport.MaxX - Viewport.PosX;
+	const float W = FMath::CeilToFloat(Viewport.Width);
+	const float H = FMath::CeilToFloat(Viewport.Height);
+
+	if (X0 >= W || X1 <= 0.0f)
+	{
+		// Draw invalid area (entire view).
+		FSlateDrawElement::MakeBox
+		(
+			ElementList,
+			LayerId,
+			MAKE_PAINT_GEOMETRY_RC(AllottedGeometry, 0.0f, 0.0f, W, H),
+			AreaBrush,
+			DrawEffects,
+			InvalidAreaColor
+		);
+	}
+	else // X0 < W && X1 > 0
+	{
+		if (X0 > 0.0f)
+		{
+			// Draw invalid area (left).
+			FSlateDrawElement::MakeBox
+			(
+				ElementList,
+				LayerId,
+				MAKE_PAINT_GEOMETRY_RC(AllottedGeometry, 0.0f, 0.0f, X0, H),
+				AreaBrush,
+				DrawEffects,
+				InvalidAreaColor
+			);
+		}
+
+		if (X1 < W)
+		{
+			// Draw invalid area (right).
+			FSlateDrawElement::MakeBox
+			(
+				ElementList,
+				LayerId,
+				MAKE_PAINT_GEOMETRY_RC(AllottedGeometry, X1, 0.0f, W - X1, H),
+				AreaBrush,
+				DrawEffects,
+				InvalidAreaColor
+			);
+		}
+
+		X0 = FMath::Max(X0, 0.0f);
+		X1 = FMath::Min(X1, W);
+
+		if (X1 > X0)
+		{
+			// Draw valid area.
+			FSlateDrawElement::MakeBox
+			(
+				ElementList,
+				LayerId,
+				MAKE_PAINT_GEOMETRY_RC(AllottedGeometry, X0, 0.0f, X1 - X0, H),
+				AreaBrush,
+				DrawEffects,
+				ValidAreaColor
+			);
+		}
+	}
+
+	LayerId++;
+}
+
+////////////////////////////////////////////////////////////////////////////////////////////////////
+
+FLinearColor FFrameTrackDrawContext::GetColorById(int Id) const
+{
+	const float Alpha = 0.9;
+	switch (Id)
+	{
+	case 0: // Game Frames
+		return FLinearColor(0.75, 1.0, 1.0, Alpha);
+
+	case 1: // Render Frames
+		return FLinearColor(1.0, 0.75, 0.75, Alpha);
+
+	case 2:
+		return FLinearColor(0.75, 0.75, 1.0, Alpha);
+
+	default:
+		return FLinearColor(1.0, 1.0, 1.0, Alpha);
+	}
+}
+
+////////////////////////////////////////////////////////////////////////////////////////////////////
+
+void FFrameTrackDrawContext::DrawCached(const FFrameTrackTimeline& CachedTimeline) const
+{
+	if (CachedTimeline.NumAggregatedFrames == 0)
+	{
+		return;
+	}
+
+	NumFrames += CachedTimeline.NumAggregatedFrames;
+
+	FLinearColor TimelineColor = GetColorById(CachedTimeline.Id);
+
+	const float SampleW = Viewport.GetSampleWidth();
+	const int NumSamples = CachedTimeline.Samples.Num();
+
+	const float VY1 = FMath::RoundToFloat(Viewport.GetViewportYForValue(0.0));
+
+	for (int SampleIndex = 0; SampleIndex < NumSamples; SampleIndex++)
+	{
+		const FFrameTrackSample& Sample = CachedTimeline.Samples[SampleIndex];
+		if (Sample.NumFrames == 0)
+		{
+			continue;
+		}
+
+		NumDrawSamples++;
+
+		const float X = SampleIndex * SampleW;
+		float VY2;
+
+		FLinearColor ColorFill = TimelineColor;
+
+		if (Sample.LargestFrameDuration == std::numeric_limits<double>::infinity())
+		{
+			VY2 = Viewport.Height;
+			ColorFill.R = 0.0f;
+			ColorFill.G = 0.0f;
+			ColorFill.B = 0.0f;
+		}
+		else
+		{
+			VY2 = FMath::RoundToFloat(Viewport.GetViewportYForValue(Sample.LargestFrameDuration));
+			if (Sample.LargestFrameDuration > 1.0 / 30.0)
+			{
+				ColorFill.G *= 0.5f;
+				ColorFill.B *= 0.5f;
+			}
+			else if (Sample.LargestFrameDuration > 1.0 / 60.0)
+			{
+				ColorFill.B *= 0.5f;
+			}
+		}
+
+		const float Y = FMath::RoundToFloat(Viewport.Height) - VY2;
+		const float H = VY2 - VY1;
+
+		const FLinearColor ColorBorder(ColorFill.R * 0.75f, ColorFill.G * 0.75f, ColorFill.B * 0.75f, 1.0);
+
+		if (SampleW > 2.0f)
+		{
+			FPaintGeometry GeoFill = MAKE_PAINT_GEOMETRY_RC(AllottedGeometry, X + 1.0f, Y + 1.0f, SampleW - 2.0f, H - 2.0f);
+			FSlateDrawElement::MakeBox(ElementList, LayerId, GeoFill, WhiteBrush, DrawEffects, ColorFill);
+
+			//FPaintGeometry GeoBorder = MAKE_PAINT_GEOMETRY_RC(AllottedGeometry, X, Y, SampleW, H);
+			//FSlateDrawElement::MakeBox(ElementList, LayerId+1, GeoBorder, BorderBrush, DrawEffects, ColorBorder);
+
+			#define LOCAL_DRAW_RC(x, y, w, h) FSlateDrawElement::MakeBox(ElementList, LayerId, MAKE_PAINT_GEOMETRY_RC(AllottedGeometry, x, y, w, h), WhiteBrush, DrawEffects, ColorBorder)
+			LOCAL_DRAW_RC(X, Y, 1.0, H);
+			LOCAL_DRAW_RC(X + SampleW - 1.0f, Y, 1.0, H);
+			LOCAL_DRAW_RC(X + 1.0f, Y, SampleW - 2.0f, 1.0f);
+			LOCAL_DRAW_RC(X + 1.0f, Y + H - 1.0f, SampleW - 2.0f, 1.0f);
+			#undef DRAW_RC
+		}
+		else
+		{
+			FPaintGeometry GeoBorder = MAKE_PAINT_GEOMETRY_RC(AllottedGeometry, X, Y, SampleW, H);
+			FSlateDrawElement::MakeBox(ElementList, LayerId, GeoBorder, WhiteBrush, DrawEffects, ColorBorder);
+		}
+	}
+
+	LayerId++;
+}
+
+////////////////////////////////////////////////////////////////////////////////////////////////////
+
+void FFrameTrackDrawContext::DrawHoveredSample(const FFrameTrackSample& Sample) const
+{
+	const float SampleW = Viewport.GetSampleWidth();
+	const int32 FramesPerSample = Viewport.GetNumFramesPerSample();
+	const int32 FirstFrameIndex = Viewport.GetFirstFrameIndex();
+	const int SampleIndex = (Sample.LargestFrameIndex - FirstFrameIndex) / FramesPerSample;
+	const float X = SampleIndex * SampleW;
+
+	const float VY1 = FMath::RoundToFloat(Viewport.GetViewportYForValue(0.0));
+	float VY2;
+	if (Sample.LargestFrameDuration == std::numeric_limits<double>::infinity())
+	{
+		VY2 = Viewport.Height;
+	}
+	else
+	{
+		VY2 = FMath::RoundToFloat(Viewport.GetViewportYForValue(Sample.LargestFrameDuration));
+	}
+	const float Y = FMath::RoundToFloat(Viewport.Height) - VY2;
+	const float H = VY2 - VY1;
+
+	const FLinearColor ColorBorder(1.0f, 1.0f, 0.0f, 1.0);
+	FPaintGeometry GeoBorder = MAKE_PAINT_GEOMETRY_RC(AllottedGeometry, X - 1.0f, Y - 1.0f, SampleW + 2.0f, H + 2.0f);
+	FSlateDrawElement::MakeBox(ElementList, LayerId, GeoBorder, BorderBrush, DrawEffects, ColorBorder);
+
+	LayerId++;
+}
+
+////////////////////////////////////////////////////////////////////////////////////////////////////
+
+void FFrameTrackDrawContext::DrawHighlightedInterval(const FFrameTrackTimeline& CachedTimeline, const double StartTime, const double EndTime) const
+{
+	const int NumSamples = CachedTimeline.Samples.Num();
+
+	//TODO: binary search
+	int Index1 = 0;
+	int Index2 = NumSamples - 1;
+	while (Index1 < NumSamples && CachedTimeline.Samples[Index1].EndTime < StartTime)
+		Index1++;
+	while (Index2 >= Index1 && CachedTimeline.Samples[Index2].StartTime > EndTime)
+		Index2--;
+
+	if (Index1 <= Index2)
+	{
+		const float SampleW = Viewport.GetSampleWidth();
+		float X1 = Index1 * SampleW;
+		float X2 = Index2 * SampleW;
+
+		const float Y1 = 10.0f;
+		const float Y2 = Viewport.Height;
+		const float D = 2.0f;
+		const float H = 10.0f;
+
+		const FLinearColor Color(1.0f, 1.0f, 1.0f, 1.0f);
+
+		if (X1 >= 0.0f && X1 < Viewport.Width - 2.0f)
+		{
+			FPaintGeometry GeoL1 = MAKE_PAINT_GEOMETRY_RC(AllottedGeometry, X1 - D, Y1, D, H);
+			FSlateDrawElement::MakeBox(ElementList, LayerId, GeoL1, WhiteBrush, DrawEffects, Color);
+
+			FPaintGeometry GeoL2 = MAKE_PAINT_GEOMETRY_RC(AllottedGeometry, X1 - D, Y2 - H, D, H);
+			FSlateDrawElement::MakeBox(ElementList, LayerId, GeoL2, WhiteBrush, DrawEffects, Color);
+		}
+
+		if (X2 >= -2.0f && X2 < Viewport.Width)
+		{
+			FPaintGeometry GeoR1 = MAKE_PAINT_GEOMETRY_RC(AllottedGeometry, X2, Y1, D, H);
+			FSlateDrawElement::MakeBox(ElementList, LayerId, GeoR1, WhiteBrush, DrawEffects, Color);
+
+			FPaintGeometry GeoR2 = MAKE_PAINT_GEOMETRY_RC(AllottedGeometry, X2, Y2 - H, D, H);
+			FSlateDrawElement::MakeBox(ElementList, LayerId, GeoR2, WhiteBrush, DrawEffects, Color);
+		}
+
+		if (X1 < 0)
+			X1 = 0.0f;
+		if (X2 > Viewport.Width)
+			X2 = Viewport.Width;
+		if (X1 < X2)
+		{
+			FPaintGeometry GeoT = MAKE_PAINT_GEOMETRY_RC(AllottedGeometry, X1, Y1, X2 - X1, D);
+			FSlateDrawElement::MakeBox(ElementList, LayerId, GeoT, WhiteBrush, DrawEffects, Color);
+
+			FPaintGeometry GeoB = MAKE_PAINT_GEOMETRY_RC(AllottedGeometry, X1, Y2 - D, X2 - X1, D);
+			FSlateDrawElement::MakeBox(ElementList, LayerId, GeoB, WhiteBrush, DrawEffects, Color);
+		}
+	}
+}
+
+////////////////////////////////////////////////////////////////////////////////////////////////////
diff --git a/Engine/Source/Developer/TraceInsights/Private/Insights/ViewModels/FrameTrackHelper.h b/Engine/Source/Developer/TraceInsights/Private/Insights/ViewModels/FrameTrackHelper.h
new file mode 100644
--- /dev/null
+++ b/Engine/Source/Developer/TraceInsights/Private/Insights/ViewModels/FrameTrackHelper.h
@@ -0,0 +1,120 @@
+// Copyright 1998-2019 Epic Games, Inc. All Rights Reserved.
+
+#pragma once
+
+#include "CoreMinimal.h"
+#include "Styling/WidgetStyle.h"
+
+enum class ESlateDrawEffect : uint8;
+
+struct FGeometry;
+struct FSlateBrush;
+
+class FFrameTrackViewport;
+class FSlateWindowElementList;
+
+namespace Trace
+{
+	struct FFrame;
+}
+
+enum EFrameType
+{
+	Unknown = -1,
+
+	Game,
+	Render
+};
+
+struct FFrameTrackSample
+{
+	int32 NumFrames;
+	double TotalDuration; // sum of durations of all frames in this sample
+	double StartTime; // min start time of all frames in this sample
+	double EndTime; // max end time of all frames in this sample
+	double LargestFrameIndex; // index of the largest frame
+	double LargestFrameStartTime; // start time of the largest frame
+	double LargestFrameDuration; // duration of the largest frame
+
+	FFrameTrackSample()
+		: NumFrames(0)
+		, TotalDuration(0.0)
+		, StartTime(DBL_MAX)
+		, EndTime(-DBL_MAX)
+		, LargestFrameIndex(0)
+		, LargestFrameStartTime(0.0)
+		, LargestFrameDuration(0.0)
+	{}
+};
+
+struct FFrameTrackTimeline
+{
+	uint64 Id;
+	//EFrameType Type;
+
+	//double MinValue; // min value
+	//double MaxValue; // max value
+	//double AvgValue; // average value
+
+	//float SampleW; // width of a sample, in Slate units
+	//int FramesPerSample; // number of frames in a sample
+	//int FirstFrameIndex; // index of first frame in first sample; can be negative
+	//int NumSamples; // total number of samples
+	int NumAggregatedFrames; // total number of frames aggregated in samples; ie. sum of all Sample.NumFrames
+
+	TArray<FFrameTrackSample> Samples;
+};
+
+class FFrameTrackCacheContext : public FNoncopyable
+{
+public:
+	FFrameTrackCacheContext(const FFrameTrackViewport& InViewport);
+	//~FFrameTrackCacheContext();
+
+	void Begin();
+	bool BeginTimeline(FFrameTrackTimeline& CachedTimeline);
+	void AddEvent(const Trace::FFrame& Event);
+	void EndTimeline();
+	void End();
+
+public:
+	const FFrameTrackViewport& Viewport;
+
+	FFrameTrackTimeline* CurrentCachedTimeline; // pointer to current FFrameTrackTimeline; valid only between BeginTimeline - EndTimeline calls.
+
+	float SampleW; // width of a sample, in Slate units
+	int FramesPerSample; // number of frames in a sample
+	int FirstFrameIndex; // index of first frame in first sample; can be negative
+	int NumSamples; // total number of samples
+
+	// Debug stats.
+	int NumFrames; // number of added frames
+};
+
+class FFrameTrackDrawContext : public FNoncopyable
+{
+public:
+	FFrameTrackDrawContext(const FFrameTrackViewport& InViewport, const FGeometry& InAllottedGeometry, int32& InLayerId, FSlateWindowElementList& InElementList);
+	//~FFrameTrackDrawContext();
+
+	void DrawBackground(const FWidgetStyle& InWidgetStyle) const;
+	void DrawCached(const FFrameTrackTimeline& CachedTimeline) const;
+	void DrawHoveredSample(const FFrameTrackSample& Sample) const;
+	void DrawHighlightedInterval(const FFrameTrackTimeline& CachedTimeline, const double StartTime, const double EndTime) const;
+
+	FLinearColor GetColorById(int Id) const;
+
+public:
+	const FFrameTrackViewport& Viewport;
+	const FGeometry& AllottedGeometry;
+	int32& LayerId;
+	FSlateWindowElementList& ElementList;
+	const ESlateDrawEffect DrawEffects;
+
+	const FSlateBrush* WhiteBrush;
+	const FSlateBrush* BorderBrush;
+
+	// Debug stats.
+	mutable int NumFrames;
+	mutable int NumDrawSamples;
+};
diff --git a/Engine/Source/Developer/TraceInsights/Private/Insights/ViewModels/FrameTrackViewport.h b/Engine/Source/Developer/TraceInsights/Private/Insights/ViewModels/FrameTrackViewport.h
new file mode 100644
--- /dev/null
+++ b/Engine/Source/Developer/TraceInsights/Private/Insights/ViewModels/FrameTrackViewport.h
@@ -0,0 +1,228 @@
+// Copyright 1998-2019 Epic Games, Inc. All Rights Reserved.
+
+#pragma once
+
+#include "CoreMinimal.h"
+
+class FFrameTrackViewport
+{
+private:
+	const float SLATE_UNITS_TOLERANCE = 0.1f;
+
+public:
+	FFrameTrackViewport()
+	{
+		Reset();
+	}
+
+	void Reset()
+	{
+		Width = 0.0f;
+		Height = 0.0f;
+
+		MinIndex = 0;
+		MaxIndex = 0;
+
+		MinX = 0.0f;
+		MaxX = 0.0f;
+		PosX = 0.0f;
+		MinScaleX = 0.0001f;
+		MaxScaleX = 10.0f;
+		ScaleX = 1.0f;
+
+		MinValue = 0.0;
+		MaxValue = 1.0f;
+
+		MinY = 0.0f;
+		MaxY = 1.0f;
+		PosY = 0.0f;
+		MinScaleY = 0.001f;
+		MaxScaleY = 1000000.0f;
+		ScaleY = 1000.0f;
+
+		LastOnPaintMaxValue = 0.0;
+	}
+
+	void SetSize(const float InWidth, const float InHeight)
+	{
+		if (!FMath::IsNearlyEqual(Width, InWidth, SLATE_UNITS_TOLERANCE) ||
+			!FMath::IsNearlyEqual(Height, InHeight, SLATE_UNITS_TOLERANCE))
+		{
+			Width = InWidth;
+			Height = InHeight;
+			OnSizeChanged();
+		}
+	}
+
+	void SetMinMaxIndexInterval(const int32 InMinIndex, const int32 InMaxIndex)
+	{
+		MinIndex = InMinIndex;
+		MaxIndex = InMaxIndex;
+		UpdateMinMaxX();
+	}
+
+	int32 GetIndexAtPosX(const float X) const
+	{
+		return FMath::RoundToInt(X / ScaleX);
+	}
+
+	int32 GetIndexAtViewportX(const float VX) const
+	{
+		return FMath::RoundToInt((PosX + VX) / ScaleX);
+	}
+
+	float GetPosXForIndex(const int32 Index) const
+	{
+		return static_cast<float>(Index) * ScaleX;
+	}
+
+	float GetViewportXForIndex(const int32 Index) const
+	{
+		return static_cast<float>(Index) * ScaleX - PosX;
+	}
+
+	float GetRoundedViewportXForIndex(const int32 Index) const
+	{
+		return FMath::RoundToFloat(static_cast<float>(Index) * ScaleX - PosX);
+	}
+
+	void ScrollAtPosX(const float InPosX)
+	{
+		PosX = InPosX;
+		OnPosXChanged();
+	}
+
+	void ScrollAtIndex(const int32 Index)
+	{
+		PosX = static_cast<float>(Index) * ScaleX;
+		OnPosXChanged();
+	}
+
+	void CenterOnIndex(const int32 Index)
+	{
+		float W = FMath::Max(1.0f, ScaleX);
+		if (W > Width)
+		{
+			ScrollAtIndex(Index);
+		}
+		else
+		{
+			float VX = (Width - W) / 2.0f;
+			PosX = static_cast<float>(Index) * ScaleX - VX;
+			OnPosXChanged();
+		}
+	}
+
+	void CenterOnIndexInterval(const int32 IntervalStartIndex, const int32 IntervalEndIndex)
+	{
+		float W = FMath::Max(1.0f, ScaleX * static_cast<float>(IntervalEndIndex - IntervalStartIndex));
+		if (W > Width)
+		{
+			ScrollAtIndex(IntervalStartIndex);
+		}
+		else
+		{
+			float VX = (Width - W) / 2.0f;
+			PosX = static_cast<float>(IntervalStartIndex) * ScaleX - VX;
+			OnPosXChanged();
+		}
+	}
+
+	void ZoomWithFixedViewportX(const float InNewScaleX, const float VX)
+	{
+		float LocalNewScaleX = FMath::Clamp(InNewScaleX, MinScaleX, MaxScaleX);
+		if (LocalNewScaleX != ScaleX)
+		{
+			// Index at viewport X should remain the same. So we resolve equation:
+			//   (PosX + VX) / ScaleX = (NewPosX + VX) / NewScaleX
+			//   ==> NewPosX = (PosX + VX) / ScaleX * NewScaleX - VX
+			PosX = (PosX + VX) * (LocalNewScaleX / ScaleX) - VX;
+			ScaleX = LocalNewScaleX;
+
+			OnPosXChanged();
+			OnScaleXChanged();
+		}
+	}
+
+	float GetSampleWidth() const { return FMath::Max(1.0f, FMath::RoundToFloat(ScaleX)); }
+	int32 GetNumFramesPerSample() const { return FMath::Max(1, FMath::RoundToInt(1.0f / ScaleX)); }
+
+	int32 GetFirstFrameIndex() const { return GetIndexAtViewportX(0.0f); }
+
+	float GetPosYForValue(const double Value) const
+	{
+		return static_cast<float>(Value) * ScaleY;
+	}
+
+	float GetViewportYForValue(const double Value) const
+	{
+		return static_cast<float>(Value) * ScaleY - PosY;
+	}
+
+	void SetScaleY(const float InNewScaleY)
+	{
+		float LocalNewScaleY = FMath::Clamp(InNewScaleY, MinScaleY, MaxScaleY);
+		if (LocalNewScaleY != ScaleY)
+		{
+			ScaleY = LocalNewScaleY;
+			OnScaleYChanged();
+		}
+	}
+
+private:
+	void OnSizeChanged()
+	{
+	}
+
+	void OnPosXChanged()
+	{
+	}
+
+	void OnScaleXChanged()
+	{
+		UpdateMinMaxX();
+	}
+
+	void OnScaleYChanged()
+	{
+	}
+
+	void UpdateMinMaxX()
+	{
+		MinX = static_cast<float>(MinIndex) * ScaleX;
+		MaxX = static_cast<float>(MaxIndex) * ScaleX;
+
+		//float NewPosX = FMath::Clamp(PosX, MinX, MaxX);
+		//if (FMath::IsNearlyEqual(NewPosX, PosX, SLATE_UNITS_TOLERANCE))
+		//{
+		//	PosX = NewPosX;
+		//	OnPosXChanged();
+		//}
+	}
+
+public:
+	float Width; // width of viewport, in Slate units
+	float Height; // height of viewport, in Slate units
+
+	int32 MinIndex; // minimum index
+	int32 MaxIndex; // maximum index (exclusive)
+
+	float MinX; // minimum horizontal position, in Slate units
+	float MaxX; // maximum horizontal position, in Slate units
+	float PosX;
+	float MinScaleX; // minimum horizontal scale factor
+	float MaxScaleX; // maximum horizontal scale factor
+	float ScaleX; // scale factor between frame index and Slate units: if > 1, it represents number of Slate units for one frame index, otherwise it represents number of frames in one Slate unit
+
+	double MinValue;
+	double MaxValue;
+
+	float MinY; // minimum vertical scroll position, in Slate units
+	float MaxY; // maximum vertical scroll position, in Slate units
+	float PosY; // current vertical scroll position, in Slate units; origin at bottom
+	float MinScaleY; // minimum vertical scale factor
+	float MaxScaleY; // maximum vertical scale factor
+	float ScaleY; // scale factor between Value units and Slate units
+
+	mutable double LastOnPaintMaxValue;
+};
diff --git a/Engine/Source/Developer/TraceInsights/Private/Insights/ViewModels/GraphTrack.cpp b/Engine/Source/Developer/TraceInsights/Private/Insights/ViewModels/GraphTrack.cpp
new file mode 100644
--- /dev/null
+++ b/Engine/Source/Developer/TraceInsights/Private/Insights/ViewModels/GraphTrack.cpp
@@ -0,0 +1,679 @@
+// Copyright 1998-2019 Epic Games, Inc. All Rights Reserved.
+
+#include "GraphTrack.h"
+
+#include "Brushes/SlateColorBrush.h"
+#include "Brushes/SlateBorderBrush.h"
+#include "EditorStyleSet.h"
+#include "Rendering/DrawElements.h"
+#include "Styling/CoreStyle.h"
+#include "TraceServices/AnalysisService.h"
+#include "TraceServices/SessionService.h"
+
+// Insights
+#include "Insights/Common/PaintUtils.h"
+#include "Insights/InsightsManager.h"
+#include "Insights/TimingProfilerManager.h"
+#include "Insights/ViewModels/TimingTrackViewport.h"
+
+#define LOCTEXT_NAMESPACE "GraphTrack"
+
+////////////////////////////////////////////////////////////////////////////////////////////////////
+// FGraphTrackSeries
+////////////////////////////////////////////////////////////////////////////////////////////////////
+
+FGraphTrackSeries::FGraphTrackSeries()
+	: Color(0.0f, 0.5f, 1.0f, 1.0f)
+	, BorderColor(0.3f, 0.8f, 1.0f, 1.0f)
+	//, Points()
+	//, LinePoints()
+	//, Boxes()
+{
+}
+
+////////////////////////////////////////////////////////////////////////////////////////////////////
+
+FGraphTrackSeries::~FGraphTrackSeries()
+{
+}
+
+////////////////////////////////////////////////////////////////////////////////////////////////////
+// FGraphTrack
+////////////////////////////////////////////////////////////////////////////////////////////////////
+
+FGraphTrack::FGraphTrack(uint64 InTrackId)
+	: FBaseTimingTrack(InTrackId)
+	, WhiteBrush(FCoreStyle::Get().GetBrush("WhiteBrush"))
+	, PointBrush(FEditorStyle::GetBrush("Graph.ExecutionBubble"))
+	, BorderBrush(new FSlateBorderBrush(NAME_None, FMargin(1.0f)))
+	//, BorderBrush(FEditorStyle::GetBrush("PlainBorder"))
+	, Font(FCoreStyle::GetDefaultFontStyle("Regular", 8))
+	, bDrawPoints(true)
+	, bDrawPointsWithBorder(true)
+	, bDrawLines(true)
+	, bDrawLinesWithDuration(false)
+	, bDrawBoxes(true)
+	, BaselineY(0.0f)
+	, ScaleY(1.0)
+{
+}
+
+////////////////////////////////////////////////////////////////////////////////////////////////////
+
+FGraphTrack::~FGraphTrack()
+{
+	delete BorderBrush;
+}
+
+////////////////////////////////////////////////////////////////////////////////////////////////////
+
+void FGraphTrack::UpdateHoveredState(float MX, float MY, const FTimingTrackViewport& Viewport)
+{
+	if (MY >= Y && MY < Y + H)
+	{
+		bIsHovered = true;
+
+		if (MX < 100.0f && MY < Y + 14.0f)
+		{
+			bIsHeaderHovered = true;
+		}
+		else
+		{
+			bIsHeaderHovered = false;
+		}
+	}
+	else
+	{
+		bIsHovered = false;
+	}
+}
+
+////////////////////////////////////////////////////////////////////////////////////////////////////
+
+void FGraphTrack::UpdateStats()
+{
+	NumDrawPoints = 0;
+	NumDrawLines = 0;
+	NumDrawBoxes = 0;
+
+	for (const FGraphTrackSeries& Series : AllSeries)
+	{
+		NumDrawPoints += Series.Points.Num();
+		NumDrawLines += Series.LinePoints.Num() / 2;
+		NumDrawBoxes += Series.Boxes.Num();
+	}
+}
+
+////////////////////////////////////////////////////////////////////////////////////////////////////
+
+void FGraphTrack::Draw(FDrawContext& DC, const FTimingTrackViewport& Viewport) const
+{
+	DC.DrawBox(0.0f, Y, Viewport.Width, H, WhiteBrush, FLinearColor(0.1f, 0.1f, 0.1f, 1.0f));
+	DC.LayerId++;
+
+	//TODO: Set clipping to (0, Y, Viewport.Width, H)!
+
+	for (const FGraphTrackSeries& Series : AllSeries)
+	{
+		DrawSeries(Series, DC, Viewport);
+	}
+
+	// Draw bottom of the track (debug).
+	//DC.DrawBox(0.0f, Viewport.Height - 1.0f, Viewport.Width, 1.0f, WhiteBrush, FLinearColor(0.5f, 0.0f, 0.0f, 1.0f));
+	//DC.LayerId++;
+
+	// Draw baseline (Value == 0).
+	DC.DrawBox(0.0f, Y + BaselineY - 1.0f, Viewport.Width, 1.0f, WhiteBrush, FLinearColor(0.05f, 0.05f, 0.05f, 1.0f));
+	DC.LayerId++;
+}
+
+////////////////////////////////////////////////////////////////////////////////////////////////////
+
+void FGraphTrack::DrawSeries(const FGraphTrackSeries& Series, FDrawContext& DC, const FTimingTrackViewport& Viewport) const
+{
+	if (bDrawBoxes)
+	{
+		int NumBoxes = Series.Boxes.Num();
+		for (int Index = 0; Index < NumBoxes; ++Index)
+		{
+			const FGraphBox& Box = Series.Boxes[Index];
+			DC.DrawBox(Box.X, Y + Box.Y, Box.W, BaselineY - Box.Y, WhiteBrush, Series.Color);
+		}
+		DC.LayerId++;
+	}
+
+	if (bDrawLines)
+	{
+		FPaintGeometry Geo = DC.Geometry.ToPaintGeometry();
+		Geo.AppendTransform(FSlateLayoutTransform(FVector2D(0, Y)));
+		FSlateDrawElement::MakeLines(DC.ElementList, DC.LayerId, Geo, Series.LinePoints, DC.DrawEffects, Series.Color, false, 1.0f);
+		DC.LayerId++;
+	}
+
+	if (bDrawPoints)
+	{
+		int NumPoints = Series.Points.Num();
+
+		if (bDrawPointsWithBorder)
+		{
+			// Draw points (border).
+			for (int Index = 0; Index < NumPoints; ++Index)
+			{
+				const FVector2D& Pt = Series.Points[Index];
+				const float PtX = Pt.X - GraphTrackPointVisualSize / 2.0f - 1.0f;
+				const float PtY = Y + Pt.Y - GraphTrackPointVisualSize / 2.0f - 1.0f;
+				DC.DrawBox(PtX, PtY, GraphTrackPointVisualSize + 2.0f, GraphTrackPointVisualSize + 2.0f, PointBrush, Series.BorderColor);
+			}
+			DC.LayerId++;
+		}
+
+		// Draw points (interior).
+		for (int Index = 0; Index < NumPoints; ++Index)
+		{
+			const FVector2D& Pt = Series.Points[Index];
+			const float PtX = Pt.X - GraphTrackPointVisualSize / 2.0f;
+			const float PtY = Y + Pt.Y - GraphTrackPointVisualSize / 2.0f;
+			DC.DrawBox(PtX, PtY, GraphTrackPointVisualSize, GraphTrackPointVisualSize, PointBrush, Series.Color);
+		}
+		DC.LayerId++;
+
+		/*
+		//const float Angle = FMath::DegreesToRadians(45.0f);
+
+		// Draw points as rectangles.
+		const float PtSize = GraphTrackPointVisualSize - 2.0f;
+		//FVector2D RotationPoint(PtSize / 2.0f, PtSize / 2.0f);
+		for (int Index = 0; Index < NumPoints; ++Index)
+		{
+			const FVector2D& Pt = Series.Points[Index];
+			const float PtX = Pt.X - PtSize / 2.0f;
+			const float PtY = Y + Pt.Y - PtSize / 2.0f;
+			DC.DrawBox(PtX, PtY, PtSize, PtSize, PointBrush, Series.Color);
+			//DC.DrawRotatedBox(PtX, PtY, PtSize, PtSize, WhiteBrush, Series.Color, Angle, RotationPoint);
+		}
+		DC.LayerId++;
+
+		// Draw borders.
+		const float BorderPtSize = GraphTrackPointVisualSize;
+		//FVector2D BorderRotationPoint(BorderPtSize / 2.0f, BorderPtSize / 2.0f);
+		for (int Index = 0; Index < NumPoints; ++Index)
+		{
+			const FVector2D& Pt = Series.Points[Index];
+			const float PtX = Pt.X - BorderPtSize / 2.0f;
+			const float PtY = Y + Pt.Y - BorderPtSize / 2.0f;
+			DC.DrawBox(PtX, PtY, BorderPtSize, BorderPtSize, BorderBrush, Series.BorderColor);
+			//DC.DrawRotatedBox(PtX, PtY, BorderPtSize, BorderPtSize, BorderBrush, Series.BorderColor, Angle, BorderRotationPoint);
+		}
+		DC.LayerId++;
+		*/
+	}
+}
+
+////////////////////////////////////////////////////////////////////////////////////////////////////
+// FRandomGraphTrack
+////////////////////////////////////////////////////////////////////////////////////////////////////
+
+FRandomGraphTrack::FRandomGraphTrack(uint64 InTrackId)
+	: FGraphTrack(InTrackId)
+{
+	bDrawPoints = true;
+	bDrawPointsWithBorder = true;
+	bDrawLines = true;
+	bDrawLinesWithDuration = false;
+	bDrawBoxes = false;
+
+	AllSeries.AddDefaulted();
+	FGraphTrackSeries& Series = AllSeries[0];
+	Series.SetColor(FLinearColor(0.0f, 0.5f, 1.0f, 1.0f), FLinearColor(0.3f, 0.8f, 1.0f, 1.0f));
+}
+
+////////////////////////////////////////////////////////////////////////////////////////////////////
+
+FRandomGraphTrack::~FRandomGraphTrack()
+{
+}
+
+////////////////////////////////////////////////////////////////////////////////////////////////////
+
+void FRandomGraphTrack::Update(const FTimingTrackViewport& Viewport)
+{
+	NumAddedEvents = 0;
+
+	// TODO: vertical panning and zooming
+	BaselineY = H;
+	ScaleY = 1.0;
+
+	GenerateSeries(AllSeries[0], Viewport, 1000000);
+
+	UpdateStats();
+}
+
+////////////////////////////////////////////////////////////////////////////////////////////////////
+
+void FRandomGraphTrack::GenerateSeries(FGraphTrackSeries& Series, const FTimingTrackViewport& Viewport, const int EventCount)
+{
+	//////////////////////////////////////////////////
+	// Generate random events.
+
+	static const double DT1 = 0.0000001; // 100ns
+	static const double DT2 = 0.01; // 100ms
+	const float V1 = 0;
+	const float V2 = H;
+
+	struct FGraphEvent
+	{
+		double Time;
+		double Duration;
+		double Value;
+	};
+
+	TArray<FGraphEvent> Events;
+	Events.Reserve(EventCount);
+
+	FRandomStream RandomStream(0);
+	double NextT = 0.0;
+	for (int Index = 0; Index < EventCount; ++Index)
+	{
+		FGraphEvent Ev;
+		Ev.Time = NextT;
+		NextT += DT1 + RandomStream.GetFraction() * (DT2 - DT1);
+		Ev.Duration = DT1 + RandomStream.GetFraction() * (NextT - Ev.Time - DT1);
+		Ev.Value = V1 + RandomStream.GetFraction() * (V2 - V1);
+		Events.Add(Ev);
+	}
+
+	// Sort events by time.
+	//Events.Sort([](const FGraphEvent& A, const FGraphEvent& B) { return A.Time < B.Time; });
+
+	//////////////////////////////////////////////////
+	// Optimize and build draw lists.
+
+	FGraphTrackBuilder Builder(*this, Series, Viewport);
+
+	Builder.Begin();
+
+	int Index = 0;
+	while (Index < EventCount && Events[Index].Time < Viewport.StartTime)
+	{
+		++Index;
+	}
+	if (Index > 0)
+	{
+		Index--; // one point outside screen (left side)
+	}
+	while (Index < EventCount)
+	{
+		const FGraphEvent& Ev = Events[Index];
+		Builder.AddEvent(Ev.Time, Ev.Duration, Ev.Value);
+
+		if (Ev.Time > Viewport.EndTime)
+		{
+			// one point outside screen (right side)
+			break;
+		}
+
+		++Index;
+	}
+
+	Builder.End();
+}
+
+////////////////////////////////////////////////////////////////////////////////////////////////////
+// FFramesGraphTrack
+////////////////////////////////////////////////////////////////////////////////////////////////////
+
+FFramesGraphTrack::FFramesGraphTrack(uint64 InTrackId)
+	: FGraphTrack(InTrackId)
+{
+	bDrawPoints = true;
+	bDrawPointsWithBorder = true;
+	bDrawLines = true;
+	bDrawLinesWithDuration = true;
+	bDrawBoxes = false;
+
+	AllSeries.AddDefaulted(2);
+
+	FGraphTrackSeries& GameFramesSeries = AllSeries[0];
+	//GameFramesSeries.SetColor(FLinearColor(1.0f, 0.5f, 0.0f, 1.0f), FLinearColor(1.0f, 0.8f, 0.3f, 1.0f));
+	GameFramesSeries.SetColor(FLinearColor(0.3f, 0.3f, 1.0f, 1.0f), FLinearColor(1.0f, 1.0f, 1.0f, 1.0f));
+
+	FGraphTrackSeries& RenderFramesSeries = AllSeries[1];
+	//RenderFramesSeries.SetColor(FLinearColor(1.0f, 0.0f, 0.5f, 1.0f), FLinearColor(1.0f, 0.3f, 0.8f, 1.0f));
+	RenderFramesSeries.SetColor(FLinearColor(1.0f, 0.3f, 0.3f, 1.0f), FLinearColor(1.0f, 1.0f, 1.0f, 1.0f));
+}
+
+////////////////////////////////////////////////////////////////////////////////////////////////////
+
+FFramesGraphTrack::~FFramesGraphTrack()
+{
+}
+
+////////////////////////////////////////////////////////////////////////////////////////////////////
+
+void FFramesGraphTrack::Update(const FTimingTrackViewport& Viewport)
+{
+	NumAddedEvents = 0;
+
+	// TODO: Vertical panning and zooming needs to be moved out in a Viewport like controller.
+	BaselineY = H;
+	ScaleY = 200.0 / 0.1; // 200px = 100ms
+
+	UpdateSeries(AllSeries[0], Viewport, TraceFrameType_Game);
+	UpdateSeries(AllSeries[1], Viewport, TraceFrameType_Rendering);
+
+	UpdateStats();
+}
+
+////////////////////////////////////////////////////////////////////////////////////////////////////
+
+void FFramesGraphTrack::UpdateSeries(FGraphTrackSeries& Series, const FTimingTrackViewport& Viewport, ETraceFrameType FrameType)
+{
+	FGraphTrackBuilder Builder(*this, Series, Viewport);
+
+	Builder.Begin();
+
+	TSharedPtr<const Trace::IAnalysisSession> Session = FInsightsManager::Get()->GetSession();
+	if (Session.IsValid())
+	{
+		Trace::FAnalysisSessionReadScope _(*Session.Get());
+		Session->ReadFramesProvider([this, &Builder, FrameType](const Trace::IFrameProvider& FramesProvider)
+		{
+			uint64 FrameCount = FramesProvider.GetFrameCount(FrameType);
+			FramesProvider.EnumerateFrames(FrameType, 0, FrameCount - 1, [&Builder](const Trace::FFrame& Frame)
+			{
+				const double Duration = Frame.EndTime - Frame.StartTime;
+				Builder.AddEvent(Frame.StartTime, Duration, Duration);
+			});
+		});
+	}
+
+	Builder.End();
+}
+
+////////////////////////////////////////////////////////////////////////////////////////////////////
+// FGraphTrackBuilder
+////////////////////////////////////////////////////////////////////////////////////////////////////
+
+FGraphTrackBuilder::FGraphTrackBuilder(FGraphTrack& InTrack, FGraphTrackSeries& InSeries, const FTimingTrackViewport& InViewport)
+	: Track(InTrack)
+	, Series(InSeries)
+	, Viewport(InViewport)
+{
+}
+
+////////////////////////////////////////////////////////////////////////////////////////////////////
+
+void FGraphTrackBuilder::Begin()
+{
+	BeginPoints();
+	BeginConnectedLines();
+	BeginBoxes();
+}
+
+////////////////////////////////////////////////////////////////////////////////////////////////////
+
+void FGraphTrackBuilder::End()
+{
+	EndPoints();
+	EndConnectedLines();
+	EndBoxes();
+}
+
+////////////////////////////////////////////////////////////////////////////////////////////////////
+
+void FGraphTrackBuilder::AddEvent(double Time, double Duration, double Value)
+{
+	Track.NumAddedEvents++;
+
+	if (Track.bDrawPoints)
+	{
+		AddPoint(Time, Value);
+	}
+
+	if (Track.bDrawLines)
+	{
+		AddConnectedLine(Time, Value);
+
+		if (Track.bDrawLinesWithDuration)
+		{
+			AddConnectedLine(Time + Duration, Value);
+		}
+	}
+
+	if (Track.bDrawBoxes)
+	{
+		AddBox(Time, Duration, Value);
+	}
+}
+
+////////////////////////////////////////////////////////////////////////////////////////////////////
+// FGraphTrackBuilder - Points
+////////////////////////////////////////////////////////////////////////////////////////////////////
+
+void FGraphTrackBuilder::BeginPoints()
+{
+	Series.Points.Reset();
+
+	PointsCurrentX = -DBL_MAX;
+
+	int MaxPointsPerLineScan = FMath::CeilToInt(Track.H / GraphTrackPointDY);
+	PointsAtCurrentX.Reset();
+	PointsAtCurrentX.AddDefaulted(MaxPointsPerLineScan);
+	//PointsAtCurrentX.SetNum(MaxPointsPerLineScan);
+}
+
+////////////////////////////////////////////////////////////////////////////////////////////////////
+
+void FGraphTrackBuilder::AddPoint(double Time, double Value)
+{
+	const float X = Viewport.TimeToSlateUnitsRounded(Time);
+	if (X < -GraphTrackPointVisualSize / 2.0f || X >= Viewport.Width + GraphTrackPointVisualSize / 2.0f)
+	{
+		return;
+	}
+
+	// Align the X with a grid of GraphTrackPointDX pixels in size, in the global space (i.e. scroll independent).
+	const double AlignedX = FMath::RoundToDouble(Time * Viewport.ScaleX / GraphTrackPointDX) * GraphTrackPointDX;
+
+	if (AlignedX > PointsCurrentX + GraphTrackPointDX - 0.5)
+	{
+		FlushPoints();
+		PointsCurrentX = AlignedX;
+	}
+
+	const float Y = Track.GetYForValue(Value);
+
+	int Index = FMath::RoundToInt(Y / GraphTrackPointDY);
+	if (Index >= 0 && Index < PointsAtCurrentX.Num())
+	{
+		FPointInfo& Pt = PointsAtCurrentX[Index];
+		Pt.bValid = true;
+		Pt.X = X;
+		Pt.Y = Y;
+	}
+}
+
+////////////////////////////////////////////////////////////////////////////////////////////////////
+
+void FGraphTrackBuilder::FlushPoints()
+{
+	for (int Index = 0; Index < PointsAtCurrentX.Num(); ++Index)
+	{
+		FPointInfo& Pt = PointsAtCurrentX[Index];
+		if (Pt.bValid)
+		{
+			Pt.bValid = false;
+			//const float Y = Index * GraphTrackPointDY + GraphTrackPointDY / 2.0f;
+			Series.Points.Add(FVector2D(Pt.X, Pt.Y));
+		}
+	}
+}
+
+////////////////////////////////////////////////////////////////////////////////////////////////////
+
+void FGraphTrackBuilder::EndPoints()
+{
+	FlushPoints();
+}
+
+////////////////////////////////////////////////////////////////////////////////////////////////////
+// FGraphTrackBuilder - Connected Lines
+////////////////////////////////////////////////////////////////////////////////////////////////////
+
+void FGraphTrackBuilder::BeginConnectedLines()
+{
+	Series.LinePoints.Reset();
+
+	LinesCurrentX = -FLT_MAX;
+	LinesMinY = FLT_MAX;
+	LinesMaxY = -FLT_MAX;
+	LinesFirstY = FLT_MAX;
+	LinesLastY = FLT_MAX;
+	bIsLastLineAdded = false;
+}
+
+////////////////////////////////////////////////////////////////////////////////////////////////////
+
+void FGraphTrackBuilder::AddConnectedLine(double Time, double Value)
+{
+	if (bIsLastLineAdded)
+	{
+		return;
+	}
+
+	const float Y = Track.GetYForValue(Value);
+
+	const float X = Viewport.TimeToSlateUnitsRounded(Viewport.RestrictEndTime(Time));
+
+	ensure(X >= LinesCurrentX); // we are assuming events are already sorted by Time
+
+	if (X < 0)
+	{
+		LinesCurrentX = X;
+		LinesLastY = Y;
+		return;
+	}
+
+	if (X >= Viewport.Width)
+	{
+		if (!bIsLastLineAdded)
+		{
+			bIsLastLineAdded = true;
+
+			if (LinesLastY != FLT_MAX)
+			{
+				FlushConnectedLine();
+
+				Series.LinePoints.Add(FVector2D(LinesCurrentX, LinesLastY));
+				Series.LinePoints.Add(FVector2D(X, Y));
+			}
+
+			// Reset the "reduction line" so last FlushConnectedLine() call will do nothing.
+			LinesMinY = Y;
+			LinesMaxY = Y;
+		}
+		return;
+	}
+
+	if (X > LinesCurrentX)
+	{
+		if (LinesLastY != FLT_MAX)
+		{
+			FlushConnectedLine();
+
+			Series.LinePoints.Add(FVector2D(LinesCurrentX, LinesLastY));
+			Series.LinePoints.Add(FVector2D(X, Y));
+		}
+
+		// Advance the "reduction line".
+		LinesCurrentX = X;
+		LinesMinY = Y;
+		LinesMaxY = Y;
+		LinesFirstY = Y;
+		LinesLastY = Y;
+	}
+	else
+	{
+		// Merge current line with the "reduction line".
+		if (Y < LinesMinY)
+		{
+			LinesMinY = Y;
+		}
+		if (Y > LinesMaxY)
+		{
+			LinesMaxY = Y;
+		}
+		LinesLastY = Y;
+	}
+}
+
+////////////////////////////////////////////////////////////////////////////////////////////////////
+
+void FGraphTrackBuilder::FlushConnectedLine()
+{
+	if (LinesCurrentX >= 0.0f && LinesMinY != LinesMaxY)
+	{
+		Series.LinePoints.Add(FVector2D(LinesCurrentX, LinesMaxY));
+		Series.LinePoints.Add(FVector2D(LinesCurrentX, LinesMinY));
+	}
+}
+
+////////////////////////////////////////////////////////////////////////////////////////////////////
+
+void FGraphTrackBuilder::EndConnectedLines()
+{
+	FlushConnectedLine();
+}
+
+////////////////////////////////////////////////////////////////////////////////////////////////////
+// FGraphTrackBuilder - Boxes
+////////////////////////////////////////////////////////////////////////////////////////////////////
+
+void FGraphTrackBuilder::BeginBoxes()
+{
+	Series.Boxes.Reset();
+}
+
+////////////////////////////////////////////////////////////////////////////////////////////////////
+
+void FGraphTrackBuilder::AddBox(double Time, double Duration, double Value)
+{
+	float X1 = Viewport.TimeToSlateUnitsRounded(Time);
+	if (X1 > Viewport.Width)
+	{
+		return;
+	}
+
+	double EndTime = Viewport.RestrictEndTime(Time + Duration);
+	float X2 = Viewport.TimeToSlateUnitsRounded(EndTime);
+	if (X2 < 0)
+	{
+		return;
+	}
+
+	float W = X2 - X1;
+	ensure(W >= 0); // we expect events to be sorted
+
+	// Timing events are displayed with minimum 1px (including empty ones).
+	if (W == 0)
+	{
+		W = 1.0f;
+	}
+
+	// TODO: reduction algorithm
+	FGraphBox Box;
+	Box.X = X1;
+	Box.W = W;
+	Box.Y = Track.GetYForValue(Value);
+	Series.Boxes.Add(Box);
+}
+
+////////////////////////////////////////////////////////////////////////////////////////////////////
+
+void FGraphTrackBuilder::EndBoxes()
+{
+}
+
+////////////////////////////////////////////////////////////////////////////////////////////////////
+
+#undef LOCTEXT_NAMESPACE
diff --git a/Engine/Source/Developer/TraceInsights/Private/Insights/ViewModels/GraphTrack.h b/Engine/Source/Developer/TraceInsights/Private/Insights/ViewModels/GraphTrack.h
new file mode 100644
--- /dev/null
+++ b/Engine/Source/Developer/TraceInsights/Private/Insights/ViewModels/GraphTrack.h
@@ -0,0 +1,210 @@
+// Copyright 1998-2019 Epic Games, Inc. All Rights Reserved.
+
+#pragma once
+
+#include "CoreMinimal.h"
+#include "Fonts/SlateFontInfo.h"
+
+// Insights
+#include "Insights/ViewModels/BaseTimingTrack.h"
+
+struct FDrawContext;
+struct FSlateBrush;
+class FTimingTrackViewport;
+
+////////////////////////////////////////////////////////////////////////////////////////////////////
+
+static const float DefaultGraphTrackH = 100.0f; // in pixels (Slate units)
+
+// Visual size of points (in pixels).
+static const float GraphTrackPointVisualSize = 7.0f;
+
+// Size of points (in pixels) used in reduction algorithm.
+static const double GraphTrackPointDX = 3.0f;
+static const float GraphTrackPointDY = 3.0f;
+
+////////////////////////////////////////////////////////////////////////////////////////////////////
+
+struct FGraphBox
+{
+	float X;
+	float W;
+	float Y;
+};
+
+////////////////////////////////////////////////////////////////////////////////////////////////////
+
+class FGraphTrackSeries
+{
+	friend class FGraphTrack;
+	friend class FGraphTrackBuilder;
+
+public:
+	FGraphTrackSeries();
+	~FGraphTrackSeries();
+
+	void SetColor(FLinearColor InColor, FLinearColor InBorderColor)
+	{
+		Color = InColor;
+		BorderColor = InBorderColor;
+	}
+
+protected:
+	//FText Name;
+	//FText Description;
+
+	FLinearColor Color;
+	FLinearColor BorderColor;
+
+	TArray<FVector2D> Points;
+	TArray<FVector2D> LinePoints;
+	TArray<FGraphBox> Boxes;
+
+	//bool bIsVisible;
+};
+
+////////////////////////////////////////////////////////////////////////////////////////////////////
+
+class FGraphTrack : public FBaseTimingTrack
+{
+	friend class FGraphTrackBuilder;
+
+public:
+	FGraphTrack(uint64 InTrackId);
+	virtual ~FGraphTrack();
+
+	virtual void UpdateHoveredState(float MX, float MY, const FTimingTrackViewport& Viewport);
+
+	virtual void Update(const FTimingTrackViewport& Viewport) override = 0;
+
+	void Draw(FDrawContext& DC, const FTimingTrackViewport& Viewport) const;
+
+	int GetNumAddedEvents() const { return NumAddedEvents; }
+	int GetNumDrawPoints() const { return NumDrawPoints; }
+	int GetNumDrawLines() const { return NumDrawLines; }
+	int GetNumDrawBoxes() const { return NumDrawBoxes; }
+
+protected:
+	void UpdateStats();
+
+	void DrawSeries(const FGraphTrackSeries& Series, FDrawContext& DC, const FTimingTrackViewport& Viewport) const;
+
+	float GetYForValue(double Value) const
+	{
+		return BaselineY - static_cast<float>(Value * ScaleY); // TODO: vertical zooming and panning
+	}
+
+protected:
+	TArray<FGraphTrackSeries> AllSeries;
+
+	// Slate resources
+	const FSlateBrush* WhiteBrush;
+	const FSlateBrush* PointBrush;
+	const FSlateBrush* BorderBrush;
+	const FSlateFontInfo Font;
+
+	bool bDrawPoints;
+	bool bDrawPointsWithBorder;
+	bool bDrawLines;
+	bool bDrawLinesWithDuration;
+	bool bDrawBoxes;
+
+	float BaselineY; // Y position (in viewport local space) of the baseline (with Value == 0); in pixels (Slate units)
+	double ScaleY; // scale between Value units and viewport units; in pixels (Slate units) / Value unit
+
+	// Stats
+	int NumAddedEvents;
+	int NumDrawPoints;
+	int NumDrawLines;
+	int NumDrawBoxes;
+};
+
+////////////////////////////////////////////////////////////////////////////////////////////////////
+
+class FRandomGraphTrack : public FGraphTrack
+{
+public:
+	FRandomGraphTrack(uint64 InTrackId);
+	virtual ~FRandomGraphTrack();
+
+	virtual void Update(const FTimingTrackViewport& Viewport) override;
+
+protected:
+	void GenerateSeries(FGraphTrackSeries& Series, const FTimingTrackViewport& Viewport, const int EventCount);
+};
+
+////////////////////////////////////////////////////////////////////////////////////////////////////
+
+class FFramesGraphTrack : public FGraphTrack
+{
+public:
+	FFramesGraphTrack(uint64 InTrackId);
+	virtual ~FFramesGraphTrack();
+
+	virtual void Update(const FTimingTrackViewport& Viewport) override;
+
+protected:
+	void UpdateSeries(FGraphTrackSeries& Series, const FTimingTrackViewport& Viewport, ETraceFrameType FrameType);
+};
+
+////////////////////////////////////////////////////////////////////////////////////////////////////
+
+class FGraphTrackBuilder : public FNoncopyable
+{
+protected:
+	struct FPointInfo
+	{
+		bool bValid;
+		float X;
+		float Y;
+
+		FPointInfo() : bValid(false) {}
+	};
+
+public:
+	FGraphTrackBuilder(FGraphTrack& InTrack, FGraphTrackSeries& InSeries, const FTimingTrackViewport& InViewport);
+	//~FGraphTrackBuilder();
+
+	const FTimingTrackViewport& GetViewport() { return Viewport; }
+
+	void Begin();
+	void AddEvent(double Time, double Duration, double Value);
+	void End();
+
+protected:
+	void BeginPoints();
+	void AddPoint(double Time, double Value);
+	void FlushPoints();
+	void EndPoints();
+
+	void BeginConnectedLines();
+	void AddConnectedLine(double Time, double Value);
+	void FlushConnectedLine();
+	void EndConnectedLines();
+
+	void BeginBoxes();
+	void AddBox(double Time, double Duration, double Value);
+	void EndBoxes();
+
+protected:
+	FGraphTrack& Track;
+	FGraphTrackSeries& Series;
+	const FTimingTrackViewport& Viewport;
+
+	// Used by the point reduction algorithm.
+	double PointsCurrentX;
+	TArray<FPointInfo> PointsAtCurrentX;
+
+	// Used by the line reduction algorithm.
+	float LinesCurrentX;
+	float LinesMinY;
+	float LinesMaxY;
+	float LinesFirstY;
+	float LinesLastY;
+	bool bIsLastLineAdded;
+
+	// Used by the box reduction algorithm.
+	//...
+};
+
+////////////////////////////////////////////////////////////////////////////////////////////////////
