Copy up of Unreal Insights from //Tasks/UE4/Dev-Core-Trace.

diff --git a/Engine/Source/Developer/TraceInsights/Private/Insights/ViewModels/LogMessage.cpp b/Engine/Source/Developer/TraceInsights/Private/Insights/ViewModels/LogMessage.cpp
new file mode 100644
--- /dev/null
+++ b/Engine/Source/Developer/TraceInsights/Private/Insights/ViewModels/LogMessage.cpp
@@ -0,0 +1,198 @@
+// Copyright 1998-2019 Epic Games, Inc. All Rights Reserved.
+
+#include "LogMessage.h"
+
+#include "Misc/OutputDeviceHelper.h"
+#include "Misc/ScopeLock.h"
+#include "TraceServices/AnalysisService.h"
+#include "TraceServices/SessionService.h"
+
+// Insights
+#include "Insights/Common/TimeUtils.h"
+
+////////////////////////////////////////////////////////////////////////////////////////////////////
+// FLogMessageRecord
+////////////////////////////////////////////////////////////////////////////////////////////////////
+
+FLogMessageRecord::FLogMessageRecord()
+	: Index(0)
+	, Time(0)
+	, Verbosity(ELogVerbosity::Type::NoLogging)
+	, Category()
+	, Message()
+	, File()
+	, Line(0)
+{
+}
+
+////////////////////////////////////////////////////////////////////////////////////////////////////
+
+FLogMessageRecord::FLogMessageRecord(const Trace::FLogMessage& TraceLogMessage)
+	: Index(static_cast<int32>(TraceLogMessage.Index))
+	, Time(TraceLogMessage.Time)
+	, Verbosity(TraceLogMessage.Verbosity)
+	//, Category(FText::FromString(TraceLogMessage.Category))
+	//, Message(FText::FromString(TraceLogMessage.Message))
+	, File(FText::FromString(TraceLogMessage.File))
+	, Line(TraceLogMessage.Line)
+{
+	// Strip the "Log" prefix.
+	FString CategoryStr(TraceLogMessage.Category->Name);
+	if (CategoryStr.StartsWith(TEXT("Log")))
+		CategoryStr = CategoryStr.RightChop(3);
+	Category = FText::FromString(CategoryStr);
+
+	// Strip the trailing whitespaces (ex. some messages ends with "\n" and we do not want the LogView rows to have an unnecessary increased height).
+	FString MessageStr(TraceLogMessage.Message);
+	MessageStr.TrimEndInline();
+	Message = FText::FromString(MessageStr);
+}
+
+////////////////////////////////////////////////////////////////////////////////////////////////////
+
+FText FLogMessageRecord::GetIndexAsText() const
+{
+	return FText::AsNumber(Index);
+}
+
+////////////////////////////////////////////////////////////////////////////////////////////////////
+
+FText FLogMessageRecord::GetTimeAsText() const
+{
+	return FText::FromString(TimeUtils::FormatTimeHMS(Time, TimeUtils::Microsecond));
+}
+
+////////////////////////////////////////////////////////////////////////////////////////////////////
+
+FText FLogMessageRecord::GetVerbosityAsText() const
+{
+	return FText::FromString(FString(FOutputDeviceHelper::VerbosityToString(Verbosity)));
+}
+
+////////////////////////////////////////////////////////////////////////////////////////////////////
+
+FText FLogMessageRecord::GetCategoryAsText() const
+{
+	return Category;
+}
+
+////////////////////////////////////////////////////////////////////////////////////////////////////
+
+FText FLogMessageRecord::GetMessageAsText() const
+{
+	return Message;
+}
+
+////////////////////////////////////////////////////////////////////////////////////////////////////
+
+FText FLogMessageRecord::GetFileAsText() const
+{
+	return File;
+}
+
+////////////////////////////////////////////////////////////////////////////////////////////////////
+
+FText FLogMessageRecord::GetLineAsText() const
+{
+	return FText::AsNumber(Line);
+}
+
+////////////////////////////////////////////////////////////////////////////////////////////////////
+
+FText FLogMessageRecord::ToDisplayString() const
+{
+	FTextBuilder TextBuilder;
+	TextBuilder.AppendLine(GetTimeAsText());
+	TextBuilder.AppendLineFormat(NSLOCTEXT("SLogView", "CategoryLine", "Category: {0}"), Category);
+	TextBuilder.AppendLineFormat(NSLOCTEXT("SLogView", "MessageLine", "Message: {0}"), Message);
+	return TextBuilder.ToText();
+}
+
+////////////////////////////////////////////////////////////////////////////////////////////////////
+// FLogMessageCache
+////////////////////////////////////////////////////////////////////////////////////////////////////
+
+FLogMessageCache::FLogMessageCache()
+{
+}
+
+////////////////////////////////////////////////////////////////////////////////////////////////////
+
+void FLogMessageCache::Reset()
+{
+}
+
+////////////////////////////////////////////////////////////////////////////////////////////////////
+
+FLogMessageRecord& FLogMessageCache::Get(uint64 Index)
+{
+	{
+		FScopeLock Lock(&CriticalSection);
+		if (Map.Contains(Index))
+		{
+			return Map[Index];
+		}
+
+		// Just an arbitrary limit. Will purge the cache after this limit, to avoid using too much memory.
+		if (Map.Num() > 10000)
+		{
+			Map.Reset();
+		}
+	}
+
+	if (Session.IsValid())
+	{
+		Trace::FAnalysisSessionReadScope _(*Session.Get());
+		Session->ReadLogProvider(
+			[this, Index](const Trace::ILogProvider& LogProvider)
+		{
+			LogProvider.ReadMessage(
+				Index,
+				[this, Index](const Trace::FLogMessage& Message)
+			{
+				FScopeLock Lock(&CriticalSection);
+				FLogMessageRecord Entry(Message);
+				Map.Add(Index, Entry);
+			},
+				true); // resolve format string
+		});
+	}
+
+	{
+		FScopeLock Lock(&CriticalSection);
+		if (Map.Contains(Index))
+		{
+			return Map[Index];
+		}
+	}
+
+	return InvalidEntry;
+}
+
+////////////////////////////////////////////////////////////////////////////////////////////////////
+
+TSharedPtr<FLogMessageRecord> FLogMessageCache::GetUncached(uint64 Index, bool bResolveFormatString) const
+{
+	TSharedPtr<FLogMessageRecord> EntryPtr;
+
+	if (Session.IsValid())
+	{
+		Trace::FAnalysisSessionReadScope _(*Session.Get());
+		Session->ReadLogProvider(
+			[Index, bResolveFormatString, &EntryPtr](const Trace::ILogProvider& LogProvider)
+		{
+			LogProvider.ReadMessage(
+				Index,
+				[&EntryPtr](const Trace::FLogMessage& Message)
+			{
+				EntryPtr = MakeShareable(new FLogMessageRecord(Message));
+			},
+				bResolveFormatString);
+		});
+	}
+
+	return EntryPtr;
+}
+
+////////////////////////////////////////////////////////////////////////////////////////////////////
+////////////////////////////////////////////////////////////////////////////////////////////////////
diff --git a/Engine/Source/Developer/TraceInsights/Private/Insights/ViewModels/LogMessage.h b/Engine/Source/Developer/TraceInsights/Private/Insights/ViewModels/LogMessage.h
new file mode 100644
--- /dev/null
+++ b/Engine/Source/Developer/TraceInsights/Private/Insights/ViewModels/LogMessage.h
@@ -0,0 +1,85 @@
+// Copyright 1998-2019 Epic Games, Inc. All Rights Reserved.
+
+#pragma once
+
+#include "CoreMinimal.h"
+
+////////////////////////////////////////////////////////////////////////////////////////////////////
+
+namespace Trace
+{
+	struct FLogMessage;
+	class IAnalysisSession;
+}
+
+////////////////////////////////////////////////////////////////////////////////////////////////////
+
+class FLogMessageRecord
+{
+public:
+	FLogMessageRecord();
+	FLogMessageRecord(const Trace::FLogMessage& Message);
+
+	FText GetIndexAsText() const;
+	FText GetTimeAsText() const;
+	FText GetVerbosityAsText() const;
+	FText GetCategoryAsText() const;
+	FText GetMessageAsText() const;
+	FText GetFileAsText() const;
+	FText GetLineAsText() const;
+
+	FText ToDisplayString() const;
+
+public:
+	int32 Index;
+	double Time;
+	ELogVerbosity::Type Verbosity;
+	FText Category;
+	FText Message;
+	FText File;
+	uint32 Line;
+};
+
+////////////////////////////////////////////////////////////////////////////////////////////////////
+
+class FLogMessageCache
+{
+public:
+	FLogMessageCache();
+
+	void SetSession(TSharedPtr<const Trace::IAnalysisSession> InSession) { Session = InSession; }
+	void Reset();
+
+	FLogMessageRecord& Get(uint64 Index);
+	TSharedPtr<FLogMessageRecord> GetUncached(uint64 Index, bool bResolveFormatString) const;
+
+private:
+	FCriticalSection CriticalSection;
+	TSharedPtr<const Trace::IAnalysisSession> Session;
+	TMap<uint64, FLogMessageRecord> Map;
+	FLogMessageRecord InvalidEntry;
+};
+
+////////////////////////////////////////////////////////////////////////////////////////////////////
+
+class FLogMessageCategory
+{
+public:
+	FName Name;
+	bool bIsVisible;
+};
+
+////////////////////////////////////////////////////////////////////////////////////////////////////
+
+class FLogMessage
+{
+public:
+	FLogMessage(const int32 InIndex) : Index(InIndex) {}
+
+	int32 GetIndex() const { return Index; }
+
+private:
+	int32 Index;
+};
+
+////////////////////////////////////////////////////////////////////////////////////////////////////
diff --git a/Engine/Source/Developer/TraceInsights/Private/Insights/ViewModels/MarkersTimingTrack.cpp b/Engine/Source/Developer/TraceInsights/Private/Insights/ViewModels/MarkersTimingTrack.cpp
new file mode 100644
--- /dev/null
+++ b/Engine/Source/Developer/TraceInsights/Private/Insights/ViewModels/MarkersTimingTrack.cpp
@@ -0,0 +1,547 @@
+// Copyright 1998-2019 Epic Games, Inc. All Rights Reserved.
+
+#include "MarkersTimingTrack.h"
+
+#include "Brushes/SlateColorBrush.h"
+#include "Fonts/FontMeasure.h"
+#include "Framework/Application/SlateApplication.h"
+#include "Styling/CoreStyle.h"
+#include "TraceServices/AnalysisService.h"
+
+// Insights
+#include "Insights/Common/PaintUtils.h"
+#include "Insights/InsightsManager.h"
+#include "Insights/TimingProfilerManager.h"
+#include "Insights/ViewModels/TimingTrackViewport.h"
+
+////////////////////////////////////////////////////////////////////////////////////////////////////
+
+#define LOCTEXT_NAMESPACE "MarkersTimingTrack"
+
+////////////////////////////////////////////////////////////////////////////////////////////////////
+// FMarkersTimingTrack
+////////////////////////////////////////////////////////////////////////////////////////////////////
+
+FMarkersTimingTrack::FMarkersTimingTrack(uint64 InTrackId)
+	: FBaseTimingTrack(InTrackId)
+	//, TimeMarkerBoxes()
+	//, TimeMarkerTexts()
+	, bIsCollapsed(false)
+	, bUseOnlyBookmarks(true)
+	, bIsDirty(true)
+	, TargetHoveredAnimPercent(0.0f)
+	, CurrentHoveredAnimPercent(0.0f)
+	, NumLogMessages(0)
+	, NumDrawBoxes(0)
+	, NumDrawTexts(0)
+	, WhiteBrush(FCoreStyle::Get().GetBrush("WhiteBrush"))
+	, Font(FCoreStyle::GetDefaultFontStyle("Regular", 8))
+{
+}
+
+////////////////////////////////////////////////////////////////////////////////////////////////////
+
+FMarkersTimingTrack::~FMarkersTimingTrack()
+{
+}
+
+////////////////////////////////////////////////////////////////////////////////////////////////////
+
+void FMarkersTimingTrack::Reset()
+{
+	FBaseTimingTrack::Reset();
+
+	TimeMarkerBoxes.Reset();
+	TimeMarkerTexts.Reset();
+
+	bIsCollapsed = false;
+	bUseOnlyBookmarks = true;
+	bIsDirty = true;
+
+	TargetHoveredAnimPercent = 0.0f;
+	CurrentHoveredAnimPercent = 0.0f;
+
+	NumLogMessages = 0;
+	NumDrawBoxes = 0;
+	NumDrawTexts = 0;
+
+	UpdateHeight();
+}
+
+////////////////////////////////////////////////////////////////////////////////////////////////////
+
+void FMarkersTimingTrack::UpdateHoveredState(float MX, float MY, const FTimingTrackViewport& Viewport)
+{
+	if (MY >= Y && MY < Y + H)
+	{
+		bIsHovered = true;
+		TargetHoveredAnimPercent = 1.0f;
+
+		if (MX < 80.0f && MY < Y + 14.0f)
+		{
+			bIsHeaderHovered = true;
+		}
+		else
+		{
+			bIsHeaderHovered = false;
+		}
+	}
+	else
+	{
+		bIsHovered = false;
+		TargetHoveredAnimPercent = 0.0f;
+	}
+}
+
+////////////////////////////////////////////////////////////////////////////////////////////////////
+
+void FMarkersTimingTrack::Tick(const double InCurrentTime, const float InDeltaTime)
+{
+	if (CurrentHoveredAnimPercent != TargetHoveredAnimPercent)
+	{
+		if (CurrentHoveredAnimPercent < TargetHoveredAnimPercent)
+		{
+			const float ShowAnimSpeed = 2 * InDeltaTime;
+			CurrentHoveredAnimPercent += ShowAnimSpeed;
+			if (CurrentHoveredAnimPercent > TargetHoveredAnimPercent)
+			{
+				CurrentHoveredAnimPercent = TargetHoveredAnimPercent;
+			}
+		}
+		else
+		{
+			const float HideAnimSpeed = 3 * InDeltaTime;
+			CurrentHoveredAnimPercent -= HideAnimSpeed;
+			if (CurrentHoveredAnimPercent < TargetHoveredAnimPercent)
+			{
+				CurrentHoveredAnimPercent = TargetHoveredAnimPercent;
+			}
+		}
+	}
+}
+
+////////////////////////////////////////////////////////////////////////////////////////////////////
+
+void FMarkersTimingTrack::Update(const FTimingTrackViewport& InViewport)
+{
+	FTimeMarkerTrackBuilder CB(*this, InViewport);
+	CB.Begin();
+
+	TSharedPtr<const Trace::IAnalysisSession> Session = FInsightsManager::Get()->GetSession();
+	if (Session.IsValid())
+	{
+		Trace::FAnalysisSessionReadScope _(*Session.Get());
+
+		Session->ReadLogProvider([&CB](const Trace::ILogProvider& LogProvider)
+		{
+			CB.BeginLog(LogProvider);
+
+			LogProvider.EnumerateMessages(
+				CB.GetViewport().StartTime,
+				CB.GetViewport().EndTime,
+				[&CB](const Trace::FLogMessage& Message) { CB.AddLogMessage(Message); },
+				false);
+
+			CB.EndLog();
+		});
+	}
+
+	CB.End();
+}
+
+////////////////////////////////////////////////////////////////////////////////////////////////////
+
+void FMarkersTimingTrack::Draw(FDrawContext& DC, const FTimingTrackViewport& Viewport) const
+{
+	// Draw background.
+	DC.DrawBox(0.0f, Y, Viewport.Width, H, WhiteBrush, FLinearColor(0.04f, 0.04f, 0.04f, 1.0f));
+	DC.LayerId++;
+
+	// Draw the track's header, in background.
+	if (!IsHovered() || CurrentHoveredAnimPercent < 1.0)
+	{
+		DrawHeader(DC, true);
+	}
+
+	//////////////////////////////////////////////////
+	// Draw vertical lines.
+	// Multiple adjacent vertical lines with same color are merged into a single box.
+
+	float BoxY, BoxH;
+	if (IsCollapsed())
+	{
+		BoxY = Y;
+		BoxH = H;
+	}
+	else
+	{
+		BoxY = 0.0f;
+		BoxH = Viewport.Height;
+	}
+
+	int32 NumBoxes = TimeMarkerBoxes.Num();
+	for (int32 BoxIndex = 0; BoxIndex < NumBoxes; BoxIndex++)
+	{
+		const FTimeMarkerBoxInfo& Box = TimeMarkerBoxes[BoxIndex];
+		DC.DrawBox(Box.X, BoxY, Box.W, BoxH, WhiteBrush, Box.Color);
+		//NumDrawBoxes++;
+	}
+	DC.LayerId++;
+	NumDrawBoxes = NumBoxes;
+
+	//////////////////////////////////////////////////
+	// Draw texts (strings are already truncated).
+
+	const float CategoryY = Y + 2.0f;
+	const float MessageY = Y + (IsBookmarksTrack() ? 1.0f : 14.0f);
+
+	int32 NumTexts = TimeMarkerTexts.Num();
+	for (int32 TextIndex = 0; TextIndex < NumTexts; TextIndex++)
+	{
+		const FTimeMarkerTextInfo& TextInfo = TimeMarkerTexts[TextIndex];
+
+		if (!IsBookmarksTrack() && TextInfo.Category.Len() > 0)
+		{
+			DC.DrawText(TextInfo.X, CategoryY, TextInfo.Category, Font, TextInfo.Color);
+			NumDrawTexts++;
+		}
+
+		if (TextInfo.Message.Len() > 0)
+		{
+			DC.DrawText(TextInfo.X, MessageY, TextInfo.Message, Font, TextInfo.Color);
+			NumDrawTexts++;
+		}
+	}
+	DC.LayerId++;
+	//NumDrawTexts = NumTexts * 2;
+
+	//////////////////////////////////////////////////
+
+	// When hovered, the track's header is draw on top.
+	if (IsHovered() || CurrentHoveredAnimPercent > 0.0)
+	{
+		DrawHeader(DC, false);
+	}
+}
+
+
+////////////////////////////////////////////////////////////////////////////////////////////////////
+
+void FMarkersTimingTrack::DrawHeader(FDrawContext& DC, bool bFirstDraw) const
+{
+	const FString Name = IsBookmarksTrack() ? TEXT("Bookmarks") : TEXT("Logs");
+	//const float ArrowSize = 6.0f;
+	const float ArrowSizeX = 4.0f;
+	const float ArrowSizeY = 8.0f;
+	const float ArrowX = IsBookmarksTrack() ? 64.0f : 31.0f;
+	const float ArrowY = Y + 3.0f;
+	const float HeaderW = ArrowX + ArrowSizeX + 4.0f;
+	const float HeaderH = 14.0f;
+
+	if (!bFirstDraw)
+	{
+		DC.DrawBox(0.0f, Y, HeaderW, HeaderH, WhiteBrush, FLinearColor(0.04f, 0.04f, 0.04f, CurrentHoveredAnimPercent));
+		DC.LayerId++;
+	}
+
+	FLinearColor Color;
+	if (bFirstDraw || CurrentHoveredAnimPercent == 0.0)
+	{
+		Color = FLinearColor(0.07f, 0.07f, 0.07f, 1.0f);
+	}
+	else if (IsHeaderHovered())
+	{
+		Color = FLinearColor(1.0f, 1.0f, 0.0f, CurrentHoveredAnimPercent);
+	}
+	else
+	{
+		Color = FLinearColor(1.0f, 1.0f, 1.0f, CurrentHoveredAnimPercent);
+	}
+
+	// Draw "Bookmarks" or "Logs" text.
+	DC.DrawText(2.0f, Y + 1.0f, Name, Font, Color);
+
+	if (IsCollapsed())
+	{
+		// Draw "down arrow".
+		//for (float A = 0.0f; A < ArrowSize; A += 1.0f)
+		//{
+		//	DC.DrawBox(ArrowX - A, ArrowY + ArrowSize - A , 2.0f * A + 1.0f, 1.0f, WhiteBrush, Color);
+		//}
+
+		// Draw "right empty arrow".
+		TArray<FVector2D> Points;
+		Points.SetNum(4);
+		Points[0] = FVector2D(ArrowX, ArrowY);
+		Points[1] = FVector2D(ArrowX, ArrowY + ArrowSizeY);
+		Points[2] = FVector2D(ArrowX + ArrowSizeX, ArrowY + ArrowSizeY / 2.0f);
+		Points[3] = FVector2D(ArrowX, ArrowY);
+		FSlateDrawElement::MakeLines(DC.ElementList, DC.LayerId, DC.Geometry.ToPaintGeometry(), Points, DC.DrawEffects, Color, false, 1.0f);
+	}
+	else
+	{
+		// Draw "up arrow".
+		//for (float A = 0.0f; A < ArrowSize; A += 1.0f)
+		//{
+		//	DC.DrawBox(ArrowX - A, ArrowY + A, 2.0f * A + 1.0f, 1.0f, WhiteBrush, Color);
+		//}
+
+		// Draw "down-right filled arrow".
+		for (float A = 1.0f; A < ArrowSizeY; A += 1.0f)
+		{
+			DC.DrawBox(ArrowX - 3.0 + ArrowSizeY - A, ArrowY + A - 1.0f, A, 1.0f, WhiteBrush, Color);
+		}
+	}
+
+	DC.LayerId++;
+}
+
+////////////////////////////////////////////////////////////////////////////////////////////////////
+// FTimeMarkerTrackBuilder
+////////////////////////////////////////////////////////////////////////////////////////////////////
+
+FTimeMarkerTrackBuilder::FTimeMarkerTrackBuilder(FMarkersTimingTrack& InTrack, const FTimingTrackViewport& InViewport)
+	: Track(InTrack)
+	, Viewport(InViewport)
+	, FontMeasureService(FSlateApplication::Get().GetRenderer()->GetFontMeasureService())
+	, Font(FCoreStyle::GetDefaultFontStyle("Regular", 8))
+{
+}
+
+////////////////////////////////////////////////////////////////////////////////////////////////////
+
+void FTimeMarkerTrackBuilder::Begin()
+{
+	Track.ResetCache();
+	Track.NumLogMessages = 0;
+}
+
+////////////////////////////////////////////////////////////////////////////////////////////////////
+
+void FTimeMarkerTrackBuilder::BeginLog(const Trace::ILogProvider& LogProvider)
+{
+	LogProviderPtr = &LogProvider;
+
+	LastX1 = -1000.0f;
+	LastX2 = -1000.0f;
+	LastCategory = nullptr;
+	LastVerbosity = (ELogVerbosity::Type)0;
+	LastLogIndex = 0;
+}
+
+////////////////////////////////////////////////////////////////////////////////////////////////////
+
+void FTimeMarkerTrackBuilder::AddLogMessage(const Trace::FLogMessage& Message)
+{
+	Track.NumLogMessages++;
+
+	// Add also the log message imediately on the left of the screen (if any).
+	if (Track.NumLogMessages == 1 && Message.Index > 0)
+	{
+		// Note: Reading message at Index-1 will not work as expected when using filter!
+		//TODO: Search API like: LogProviderPtr->SearchMessage(StartIndex, ESearchDirection::Backward, LambdaPredicate, bResolveFormatString);
+		LogProviderPtr->ReadMessage(
+			Message.Index - 1,
+			[this](const Trace::FLogMessage& Message) { AddLogMessage(Message); },
+			false); // do not resolve format string
+	}
+
+	if (!Track.bUseOnlyBookmarks || FCString::Strcmp(*Message.Category->Name, TEXT("LogBookmark")) == 0)
+	{
+		float X = Viewport.TimeToSlateUnitsRounded(Message.Time);
+		if (X < 0)
+		{
+			X = -1.0f;
+		}
+		AddTimeMarker(X, *Message.Category->Name, Message.Verbosity, Message.Index);
+	}
+}
+
+////////////////////////////////////////////////////////////////////////////////////////////////////
+
+FLinearColor FTimeMarkerTrackBuilder::GetColorByCategory(const TCHAR* const Category)
+{
+	// Strip the "Log" prefix.
+	FString CategoryStr(Category);
+	if (CategoryStr.StartsWith(TEXT("Log")))
+		CategoryStr = CategoryStr.RightChop(3);
+
+	uint32 Hash = 0;
+	for (const TCHAR* c = *CategoryStr; *c; ++c)
+	{
+		Hash = (Hash + *c) * 0x2c2c57ed;
+	}
+
+	// Divided by 128.0 in order to force bright colors.
+	return FLinearColor(((Hash >> 16) & 0xFF) / 128.0f, ((Hash >> 8) & 0xFF) / 128.0f, (Hash & 0xFF) / 128.0f, 1.0f);
+}
+
+////////////////////////////////////////////////////////////////////////////////////////////////////
+
+FLinearColor FTimeMarkerTrackBuilder::GetColorByVerbosity(const ELogVerbosity::Type Verbosity)
+{
+	static FLinearColor Colors[] =
+	{
+		FLinearColor(0.0, 0.0, 0.0, 1.0), // NoLogging
+		FLinearColor(1.0, 0.0, 0.0, 1.0), // Fatal
+		FLinearColor(1.0, 0.3, 0.0, 1.0), // Error
+		FLinearColor(0.7, 0.5, 0.0, 1.0), // Warning
+		FLinearColor(0.0, 0.7, 0.0, 1.0), // Display
+		FLinearColor(0.0, 0.7, 1.0, 1.0), // Log
+		FLinearColor(0.7, 0.7, 0.7, 1.0), // Verbose
+		FLinearColor(1.0, 1.0, 1.0, 1.0), // VeryVerbose
+	};
+	static_assert(sizeof(Colors) / sizeof(FLinearColor) == (int)ELogVerbosity::Type::All + 1, "ELogVerbosity::Type has changed!?");
+	//return Colors[Verbosity & ELogVerbosity::VerbosityMask];
+	return Colors[Verbosity & 7]; // using 7 instead of ELogVerbosity::VerbosityMask (15) to make static analyzer happy
+}
+
+////////////////////////////////////////////////////////////////////////////////////////////////////
+
+void FTimeMarkerTrackBuilder::Flush(float AvailableTextW)
+{
+	// Is last marker valid?
+	if (LastCategory != nullptr)
+	{
+		const FLinearColor Color = GetColorByCategory(LastCategory);
+
+		bool bAddNewBox = true;
+		if (Track.TimeMarkerBoxes.Num() > 0)
+		{
+			FTimeMarkerBoxInfo& PrevBox = Track.TimeMarkerBoxes[Track.TimeMarkerBoxes.Num() - 1];
+			if (PrevBox.X + PrevBox.W == LastX1 &&
+				PrevBox.Color.R == Color.R &&
+				PrevBox.Color.G == Color.G &&
+				PrevBox.Color.B == Color.B)
+			{
+				// Extend previous box instead.
+				PrevBox.W += LastX2 - LastX1;
+				bAddNewBox = false;
+			}
+		}
+
+		if (bAddNewBox)
+		{
+			// Add new Box info to cache.
+			FTimeMarkerBoxInfo& Box = Track.TimeMarkerBoxes[Track.TimeMarkerBoxes.AddDefaulted()];
+			Box.X = LastX1;
+			Box.W = LastX2 - LastX1;
+			Box.Color = Color;
+			Box.Color.A = 0.5f;
+		}
+
+		if (AvailableTextW > 6.0f)
+		{
+			LogProviderPtr->ReadMessage(
+				LastLogIndex,
+				[this, AvailableTextW, &Color](const Trace::FLogMessage& Message)
+			{
+				// Strip the "Log" prefix.
+				FString CategoryStr(Message.Category->Name);
+				if (CategoryStr.StartsWith(TEXT("Log")))
+					CategoryStr = CategoryStr.RightChop(3);
+
+				const int32 LastWholeCharacterIndexCategory = FontMeasureService->FindLastWholeCharacterIndexBeforeOffset(CategoryStr, Font, FMath::RoundToInt(AvailableTextW - 2.0f));
+				const int32 LastWholeCharacterIndexMessage = FontMeasureService->FindLastWholeCharacterIndexBeforeOffset(Message.Message, Font, FMath::RoundToInt(AvailableTextW - 2.0f));
+
+				if (LastWholeCharacterIndexCategory >= 0 ||
+					LastWholeCharacterIndexMessage >= 0)
+				{
+					// Add new Text info to cache.
+					FTimeMarkerTextInfo& TextInfo = Track.TimeMarkerTexts[Track.TimeMarkerTexts.AddDefaulted()];
+					TextInfo.X = LastX2 + 2.0f;
+					TextInfo.Color = Color;
+					if (LastWholeCharacterIndexCategory >= 0)
+					{
+						TextInfo.Category = CategoryStr.Left(LastWholeCharacterIndexCategory + 1);
+					}
+					if (LastWholeCharacterIndexMessage >= 0)
+					{
+						TextInfo.Message.AppendChars(Message.Message, LastWholeCharacterIndexMessage + 1);
+					}
+				}
+			},
+				true); // resolve format string
+		}
+	}
+}
+
+////////////////////////////////////////////////////////////////////////////////////////////////////
+
+void FTimeMarkerTrackBuilder::AddTimeMarker(const float X, const TCHAR* const Category, const ELogVerbosity::Type Verbosity, const uint64 LogIndex)
+{
+	float W = X - LastX2;
+
+	if (W > 0.0f) // There is at least 1px from previous box?
+	{
+		// Flush previous marker (if any).
+		Flush(W);
+
+		// Begin new marker info.
+		LastX1 = X;
+		LastX2 = X + 1.0f;
+	}
+	else if (W == 0.0f) // Adjacent to previous box?
+	{
+		// Same color as previous marker?
+		if (Category == LastCategory)
+		{
+			// Extend previous box.
+			LastX2++;
+		}
+		else
+		{
+			// Flush previous marker (if any).
+			Flush(0.0f);
+
+			// Begin new box.
+			LastX1 = X;
+			LastX2 = X + 1.0f;
+		}
+	}
+	else // Overlaps previous box?
+	{
+		// Same color as previous marker?
+		if (Category == LastCategory)
+		{
+			// Keep previous box.
+		}
+		else
+		{
+			// Shrink previous box.
+			LastX2--;
+
+			if (LastX2 > LastX1)
+			{
+				// Flush previous marker (if any).
+				Flush(0.0f);
+			}
+
+			// Begin new box.
+			LastX1 = X;
+			LastX2 = X + 1.0f;
+		}
+	}
+
+	// Save marker.
+	LastCategory = Category;
+	LastVerbosity = Verbosity;
+	LastLogIndex = LogIndex;
+}
+
+////////////////////////////////////////////////////////////////////////////////////////////////////
+
+void FTimeMarkerTrackBuilder::EndLog()
+{
+	Flush(Viewport.Width - LastX2);
+}
+
+////////////////////////////////////////////////////////////////////////////////////////////////////
+
+void FTimeMarkerTrackBuilder::End()
+{
+}
+
+////////////////////////////////////////////////////////////////////////////////////////////////////
+
+#undef LOCTEXT_NAMESPACE
diff --git a/Engine/Source/Developer/TraceInsights/Private/Insights/ViewModels/MarkersTimingTrack.h b/Engine/Source/Developer/TraceInsights/Private/Insights/ViewModels/MarkersTimingTrack.h
new file mode 100644
--- /dev/null
+++ b/Engine/Source/Developer/TraceInsights/Private/Insights/ViewModels/MarkersTimingTrack.h
@@ -0,0 +1,156 @@
+// Copyright 1998-2019 Epic Games, Inc. All Rights Reserved.
+
+#pragma once
+
+#include "CoreMinimal.h"
+#include "Fonts/SlateFontInfo.h"
+
+// Insights
+#include "Insights/ViewModels/BaseTimingTrack.h"
+
+namespace Trace
+{
+	struct FLogMessage;
+	class ILogProvider;
+}
+
+struct FDrawContext;
+struct FSlateBrush;
+class FTimingTrackViewport;
+
+////////////////////////////////////////////////////////////////////////////////////////////////////
+
+static const float BookmarksTrackHeight = 14.0f;
+static const float TimeMarkersTrackHeight = 28.0f;
+
+////////////////////////////////////////////////////////////////////////////////////////////////////
+
+struct FTimeMarkerBoxInfo
+{
+	float X;
+	float W;
+	FLinearColor Color;
+};
+
+////////////////////////////////////////////////////////////////////////////////////////////////////
+
+struct FTimeMarkerTextInfo
+{
+	float X;
+	FLinearColor Color;
+	FString Category; // truncated Category string
+	FString Message; // truncated Message string
+};
+
+////////////////////////////////////////////////////////////////////////////////////////////////////
+
+class FMarkersTimingTrack : public FBaseTimingTrack
+{
+	friend class FTimeMarkerTrackBuilder;
+
+public:
+	FMarkersTimingTrack(uint64 InTrackId);
+	virtual ~FMarkersTimingTrack();
+
+	virtual void Reset() override;
+
+	bool IsCollapsed() const { return bIsCollapsed; }
+	void ToggleCollapsed() { bIsCollapsed = !bIsCollapsed; }
+
+	bool IsBookmarksTrack() const { return bUseOnlyBookmarks; }
+	void SetIsBookmarksTrackFlag(bool bInUseOnlyBookmarks)
+	{
+		bUseOnlyBookmarks = bInUseOnlyBookmarks;
+		UpdateHeight();
+	}
+
+	bool IsDirty() const { return bIsDirty; }
+	void SetDirtyFlag() { bIsDirty = true; }
+	void ClearDirtyFlag() { bIsDirty = false; }
+
+	// Stats
+	int GetNumLogMessages() const { return NumLogMessages; }
+	int GetNumBoxes() const { return TimeMarkerBoxes.Num(); }
+	int GetNumTexts() const { return TimeMarkerTexts.Num(); }
+
+	virtual void UpdateHoveredState(float MX, float MY, const FTimingTrackViewport& Viewport) override;
+	virtual void Tick(const double InCurrentTime, const float InDeltaTime) override;
+	virtual void Update(const FTimingTrackViewport& Viewport) override;
+
+	void Draw(FDrawContext& DC, const FTimingTrackViewport& Viewport) const;
+
+protected:
+	void ResetCache()
+	{
+		TimeMarkerBoxes.Reset();
+		TimeMarkerTexts.Reset();
+	}
+
+	void UpdateHeight()
+	{
+		H = bUseOnlyBookmarks ? BookmarksTrackHeight : TimeMarkersTrackHeight;
+	}
+
+	void DrawHeader(FDrawContext& DC, bool bFirstDraw) const;
+
+protected:
+	TArray<FTimeMarkerBoxInfo> TimeMarkerBoxes;
+	TArray<FTimeMarkerTextInfo> TimeMarkerTexts;
+
+	bool bIsCollapsed; // If false, the vertical lines will extend to entire viewport height; otherwise will be limited to this track's height.
+	bool bUseOnlyBookmarks; // If true, uses only bookmarks; otherwise it uses all log messages.
+	bool bIsDirty; // Source list of bookmarks/markers has changed. Cached Boxes and Texts are dirty.
+
+	float TargetHoveredAnimPercent; // [0.0 .. 1.0], 0.0 = hidden, 1.0 = visible
+	float CurrentHoveredAnimPercent;
+
+	// Stats
+	int NumLogMessages;
+	mutable int NumDrawBoxes;
+	mutable int NumDrawTexts;
+
+	// Slate resources
+	const FSlateBrush* WhiteBrush;
+	const FSlateFontInfo Font;
+};
+
+////////////////////////////////////////////////////////////////////////////////////////////////////
+
+class FTimeMarkerTrackBuilder : public FNoncopyable
+{
+public:
+	FTimeMarkerTrackBuilder(FMarkersTimingTrack& InTrack, const FTimingTrackViewport& InViewport);
+	//~FTimeMarkerTrackBuilder();
+
+	const FTimingTrackViewport& GetViewport() { return Viewport; }
+
+	void Begin();
+	void BeginLog(const Trace::ILogProvider& LogProvider);
+	void AddLogMessage(const Trace::FLogMessage& Message);
+	void EndLog();
+	void End();
+
+	static FLinearColor GetColorByCategory(const TCHAR* const Category);
+	static FLinearColor GetColorByVerbosity(const ELogVerbosity::Type Verbosity);
+
+private:
+	void Flush(float AvailableTextW);
+	void AddTimeMarker(const float X, const TCHAR* const Category, const ELogVerbosity::Type Verbosity, const uint64 LogIndex);
+
+public:
+	FMarkersTimingTrack& Track;
+	const FTimingTrackViewport& Viewport;
+
+	const TSharedRef<FSlateFontMeasure> FontMeasureService;
+	const FSlateFontInfo Font;
+
+	const Trace::ILogProvider* LogProviderPtr; // valid only between BeginLog() and EndLog()
+
+	float LastX1;
+	float LastX2;
+	const TCHAR* LastCategory;
+	ELogVerbosity::Type LastVerbosity;
+	uint64 LastLogIndex;
+};
+
+////////////////////////////////////////////////////////////////////////////////////////////////////
diff --git a/Engine/Source/Developer/TraceInsights/Private/Insights/ViewModels/TimeRulerTrack.cpp b/Engine/Source/Developer/TraceInsights/Private/Insights/ViewModels/TimeRulerTrack.cpp
new file mode 100644
--- /dev/null
+++ b/Engine/Source/Developer/TraceInsights/Private/Insights/ViewModels/TimeRulerTrack.cpp
@@ -0,0 +1,180 @@
+// Copyright 1998-2019 Epic Games, Inc. All Rights Reserved.
+
+#include "TimeRulerTrack.h"
+
+#include "Brushes/SlateColorBrush.h"
+#include "Fonts/FontMeasure.h"
+#include "Fonts/SlateFontInfo.h"
+#include "Framework/Application/SlateApplication.h"
+#include "Rendering/DrawElements.h"
+
+// Insights
+#include "Insights/Common/PaintUtils.h"
+#include "Insights/Common/TimeUtils.h"
+#include "Insights/ViewModels/TimingTrackViewport.h"
+
+#define LOCTEXT_NAMESPACE "TimeRulerTrack"
+
+////////////////////////////////////////////////////////////////////////////////////////////////////
+
+FTimeRulerTrack::FTimeRulerTrack(uint64 InTrackId)
+	: FBaseTimingTrack(InTrackId)
+	, WhiteBrush(FCoreStyle::Get().GetBrush("WhiteBrush"))
+	, Font(FCoreStyle::GetDefaultFontStyle("Regular", 8))
+{
+}
+
+////////////////////////////////////////////////////////////////////////////////////////////////////
+
+FTimeRulerTrack::~FTimeRulerTrack()
+{
+}
+
+////////////////////////////////////////////////////////////////////////////////////////////////////
+
+void FTimeRulerTrack::UpdateHoveredState(float MX, float MY, const FTimingTrackViewport& Viewport)
+{
+	if (MY >= Y && MY < Y + H)
+	{
+		bIsHovered = true;
+
+		//if (MX < 100.0f && MY < Y + 14.0f)
+		//{
+		//	bIsHeaderHovered = true;
+		//}
+		//else
+		//{
+		//	bIsHeaderHovered = false;
+		//}
+	}
+	else
+	{
+		bIsHovered = false;
+	}
+}
+
+////////////////////////////////////////////////////////////////////////////////////////////////////
+
+void FTimeRulerTrack::Draw(FDrawContext& DC, const FTimingTrackViewport& Viewport, const FVector2D& MousePosition, const bool bIsSelecting, const double SelectionStartTime, const double SelectionEndTime) const
+{
+	const float MinorTickMark = 5.0f;
+	const float MajorTickMark = 20 * MinorTickMark;
+
+	const float MinorTickMarkHeight = 5.0f;
+	const float MajorTickMarkHeight = 9.0f;
+
+	const float TextY = Y + MajorTickMarkHeight;
+
+	double MinorTickMarkTime = Viewport.GetDurationForViewportDX(MinorTickMark);
+	double MajorTickMarkTime = Viewport.GetDurationForViewportDX(MajorTickMark);
+
+	double VX = Viewport.StartTime * Viewport.ScaleX;
+	double MinorN = FMath::FloorToDouble(VX / static_cast<double>(MinorTickMark));
+	double MajorN = FMath::FloorToDouble(VX / static_cast<double>(MajorTickMark));
+	float MinorOX = static_cast<float>(FMath::RoundToDouble(MinorN * static_cast<double>(MinorTickMark) - VX));
+	float MajorOX = static_cast<float>(FMath::RoundToDouble(MajorN * static_cast<double>(MajorTickMark) - VX));
+
+	// Draw the time ruler's background.
+	DC.DrawBox(0.0f, Y, Viewport.Width, TimeRulerHeight, WhiteBrush, FLinearColor(0.09f, 0.09f, 0.09f, 1.0f));
+	DC.LayerId++;
+
+	// Draw the minor tick marks.
+	for (float X = MinorOX; X < Viewport.Width; X += MinorTickMark)
+	{
+		const bool bIsTenth = ((int)(((X - MajorOX) / MinorTickMark) + 0.4f) % 2 == 0);
+		const float MinorTickH = bIsTenth ? MinorTickMarkHeight : MinorTickMarkHeight - 1.0f;
+		DC.DrawBox(X, Y, 1.0f, MinorTickH, WhiteBrush,
+			bIsTenth ? FLinearColor(0.3f, 0.3f, 0.3f, 1.0f) : FLinearColor(0.25f, 0.25f, 0.25f, 1.0f));
+	}
+	// Draw the major tick marks.
+	for (float X = MajorOX; X < Viewport.Width; X += MajorTickMark)
+	{
+		DC.DrawBox(X, Y, 1.0f, MajorTickMarkHeight, WhiteBrush, FLinearColor(0.4f, 0.4f, 0.4f, 1.0f));
+	}
+	DC.LayerId++;
+
+	const double DT = static_cast<double>(MajorTickMark) / Viewport.ScaleX;
+	const double Precision = FMath::Max(DT / 10.0, TimeUtils::Nanosecond);
+
+	const TSharedRef<FSlateFontMeasure> FontMeasureService = FSlateApplication::Get().GetRenderer()->GetFontMeasureService();
+
+	// Draw the time at major tick marks.
+	for (float X = MajorOX; X < Viewport.Width + MajorTickMark; X += MajorTickMark)
+	{
+		const double T = Viewport.SlateUnitsToTime(X);
+		FString Text = TimeUtils::FormatTime(T, Precision);
+		const float TextWidth = FontMeasureService->Measure(Text, Font).X;
+		DC.DrawText(X - TextWidth / 2, TextY, Text, Font,
+			(T < 0 || T >= Viewport.MaxValidTime) ? FLinearColor(0.7f, 0.5f, 0.5f, 1.0f) : FLinearColor(0.8f, 0.8f, 0.8f, 1.0f));
+	}
+	DC.LayerId++;
+
+	bool bShowMousePos = !MousePosition.IsZero();
+	if (bShowMousePos)
+	{
+		const FLinearColor MousePosLineColor(0.9f, 0.9f, 0.9f, 0.1f);
+		const FLinearColor MousePosTextBackgroundColor(0.9f, 0.9f, 0.9f, 1.0f);
+		FLinearColor MousePosTextForegroundColor(0.1f, 0.1f, 0.1f, 1.0f);
+
+		// Time at current mouse position.
+		FString MousePosText;
+
+		const double MousePosTime = Viewport.SlateUnitsToTime(MousePosition.X);
+		const double MousePosPrecision = FMath::Max(DT / 100.0, TimeUtils::Nanosecond);
+		if (MousePosition.Y >= Y && MousePosition.Y < Y + H)
+		{
+			// If mouse is hovering the time ruller, format time with a better precision (split seconds in ms, us, ns and ps).
+			MousePosText = TimeUtils::FormatTimeSplit(MousePosTime, MousePosPrecision);
+		}
+		else
+		{
+			// Format current time with one more digit than the time at major tick marks.
+			MousePosText = TimeUtils::FormatTime(MousePosTime, MousePosPrecision);
+		}
+
+		const float MousePosTextWidth = FMath::RoundToFloat(FontMeasureService->Measure(MousePosText, Font).X);
+		static float CrtMousePosTextWidth = 0.0f;
+
+		if (!FMath::IsNearlyEqual(CrtMousePosTextWidth, MousePosTextWidth))
+		{
+			// Animate the box's width (to avoid flickering).
+			CrtMousePosTextWidth = CrtMousePosTextWidth * 0.6f + MousePosTextWidth * 0.4f;
+		}
+
+		float X = MousePosition.X;
+		float W = CrtMousePosTextWidth + 4.0f;
+		if (bIsSelecting && SelectionStartTime < SelectionEndTime)
+		{
+			float SelectionX1 = Viewport.TimeToSlateUnitsRounded(SelectionStartTime);
+			float SelectionX2 = Viewport.TimeToSlateUnitsRounded(SelectionEndTime);
+			if (X - SelectionX1 > 1.0f)
+				X = SelectionX2 + W / 2;
+			if (SelectionX2 - X > 1.0f)
+				X = SelectionX1 - W / 2;
+			MousePosTextForegroundColor = FLinearColor(FColor(32, 64, 128, 255));
+		}
+		else
+		{
+			// Draw horizontal line at mouse position.
+			//DC.DrawBox(0.0f, MousePosition.Y, Viewport.Width, 1.0f, WhiteBrush, MousePosLineColor);
+
+			// Draw vertical line at mouse position.
+			DC.DrawBox(MousePosition.X, 0.0f, 1.0f, Viewport.Height, WhiteBrush, MousePosLineColor);
+
+			// Stroke the vertical line above current time box.
+			DC.DrawBox(MousePosition.X, 0.0f, 1.0f, TextY, WhiteBrush, MousePosTextBackgroundColor);
+		}
+
+		// Fill the current time box.
+		DC.DrawBox(X - W / 2, TextY, W, 12.0f, WhiteBrush, MousePosTextBackgroundColor);
+		DC.LayerId++;
+
+		// Draw current time text.
+		DC.DrawText(X - MousePosTextWidth / 2, TextY, MousePosText, Font, MousePosTextForegroundColor);
+		DC.LayerId++;
+	}
+}
+
+////////////////////////////////////////////////////////////////////////////////////////////////////
+
+#undef LOCTEXT_NAMESPACE
diff --git a/Engine/Source/Developer/TraceInsights/Private/Insights/ViewModels/TimeRulerTrack.h b/Engine/Source/Developer/TraceInsights/Private/Insights/ViewModels/TimeRulerTrack.h
new file mode 100644
--- /dev/null
+++ b/Engine/Source/Developer/TraceInsights/Private/Insights/ViewModels/TimeRulerTrack.h
@@ -0,0 +1,48 @@
+// Copyright 1998-2019 Epic Games, Inc. All Rights Reserved.
+
+#pragma once
+
+#include "CoreMinimal.h"
+#include "Fonts/SlateFontInfo.h"
+
+// Insights
+#include "Insights/ViewModels/BaseTimingTrack.h"
+
+struct FDrawContext;
+struct FSlateBrush;
+class FTimingTrackViewport;
+
+////////////////////////////////////////////////////////////////////////////////////////////////////
+
+static const float TimeRulerHeight = 22.0f;
+
+////////////////////////////////////////////////////////////////////////////////////////////////////
+
+class FTimeRulerTrack : public FBaseTimingTrack
+{
+public:
+	FTimeRulerTrack(uint64 InTrackId);
+	virtual ~FTimeRulerTrack();
+
+	virtual void Reset() override
+	{
+		FBaseTimingTrack::Reset();
+
+		H = TimeRulerHeight;
+	}
+
+	virtual void UpdateHoveredState(float MX, float MY, const FTimingTrackViewport& Viewport);
+
+	void Draw(FDrawContext& DC, const FTimingTrackViewport& Viewport,
+				const FVector2D& MousePosition = FVector2D(0.0f, 0.0f),
+				const bool bIsSelecting = false,
+				const double SelectionStartTime = 0.0,
+				const double SelectionEndTime = 0.0) const;
+
+public:
+	// Slate resources
+	const FSlateBrush* WhiteBrush;
+	const FSlateFontInfo Font;
+};
+
+////////////////////////////////////////////////////////////////////////////////////////////////////
diff --git a/Engine/Source/Developer/TraceInsights/Private/Insights/ViewModels/TimerGroupingAndSorting.h b/Engine/Source/Developer/TraceInsights/Private/Insights/ViewModels/TimerGroupingAndSorting.h
new file mode 100644
--- /dev/null
+++ b/Engine/Source/Developer/TraceInsights/Private/Insights/ViewModels/TimerGroupingAndSorting.h
@@ -0,0 +1,39 @@
+// Copyright 1998-2019 Epic Games, Inc. All Rights Reserved.
+
+#pragma once
+
+#include "CoreMinimal.h"
+#include "Insights/ViewModels/TimerNode.h"
+
+////////////////////////////////////////////////////////////////////////////////////////////////////
+// Grouping
+
+/// Enumerates types of grouping or sorting for the timer nodes.
+enum class ETimerGroupingMode : uint32
+{
+	/// Creates a single group for all timers.
+	Flat,
+
+	/// Creates one group for one letter.
+	ByName,
+
+	/// Creates groups based on timer metadata group names.
+	ByMetaGroupName,
+
+	/// Creates one group for each timer type.
+	ByType,
+
+	ByTotalInclusiveTime,
+
+	ByTotalExclusiveTime,
+
+	ByInstanceCount,
+
+	/// Invalid enum type, may be used as a number of enumerations.
+	InvalidOrMax,
+};
+
+/// Type definition for shared pointers to instances of ETimerGroupingMode.
+typedef TSharedPtr<ETimerGroupingMode> ETimerGroupingModePtr;
+
+////////////////////////////////////////////////////////////////////////////////////////////////////
diff --git a/Engine/Source/Developer/TraceInsights/Private/Insights/ViewModels/TimerNode.cpp b/Engine/Source/Developer/TraceInsights/Private/Insights/ViewModels/TimerNode.cpp
new file mode 100644
--- /dev/null
+++ b/Engine/Source/Developer/TraceInsights/Private/Insights/ViewModels/TimerNode.cpp
@@ -0,0 +1,59 @@
+// Copyright 1998-2019 Epic Games, Inc. All Rights Reserved.
+
+#include "TimerNode.h"
+
+#define LOCTEXT_NAMESPACE "TimerNode"
+
+////////////////////////////////////////////////////////////////////////////////////////////////////
+
+void FTimerNode::ResetStats()
+{
+	Stats.Reset();
+}
+
+////////////////////////////////////////////////////////////////////////////////////////////////////
+
+void FTimerNode::SetStats(FTimerStats& InStats)
+{
+	Stats = InStats;
+}
+
+////////////////////////////////////////////////////////////////////////////////////////////////////
+
+const FText FTimerNode::GetNameEx() const
+{
+	FText Text = FText::GetEmpty();
+
+	if (IsGroup())
+	{
+		const int32 NumChildren = Children.Num();
+		const int32 NumFilteredChildren = FilteredChildren.Num();
+
+		if (NumFilteredChildren == NumChildren)
+		{
+			Text = FText::Format(LOCTEXT("TimerNodeGroupTextFmt1", "{0} ({1})"), FText::FromName(Name), FText::AsNumber(NumChildren));
+		}
+		else
+		{
+			Text = FText::Format(LOCTEXT("TimerNodeGroupTextFmt2", "{0} ({1} / {2})"), FText::FromName(Name), FText::AsNumber(NumFilteredChildren), FText::AsNumber(NumChildren));
+		}
+	}
+	else
+	{
+		//const bool bIsStatTracked = FProfilerManager::Get()->IsStatTracked(GroupOrStatNode->GetStatID());
+		//if (bIsStatTracked)
+		//{
+		//	Text = FText::Format(LOCTEXT("TimerNodeTrackedTextFmt", "{0}*"), FText::FromName(Name));
+		//}
+		//else
+		{
+			Text = FText::FromName(Name);
+		}
+	}
+
+	return Text;
+}
+
+////////////////////////////////////////////////////////////////////////////////////////////////////
+
+#undef LOCTEXT_NAMESPACE
diff --git a/Engine/Source/Developer/TraceInsights/Private/Insights/ViewModels/TimerNode.h b/Engine/Source/Developer/TraceInsights/Private/Insights/ViewModels/TimerNode.h
new file mode 100644
--- /dev/null
+++ b/Engine/Source/Developer/TraceInsights/Private/Insights/ViewModels/TimerNode.h
@@ -0,0 +1,277 @@
+// Copyright 1998-2019 Epic Games, Inc. All Rights Reserved.
+
+#pragma once
+
+#include "CoreMinimal.h"
+
+class FTimerNode;
+
+/** Type definition for shared pointers to instances of FTimerNode. */
+typedef TSharedPtr<class FTimerNode> FTimerNodePtr;
+
+/** Type definition for shared references to instances of FTimerNode. */
+typedef TSharedRef<class FTimerNode> FTimerNodeRef;
+
+/** Type definition for shared references to const instances of FTimerNode. */
+typedef TSharedRef<const class FTimerNode> FTimerNodeRefConst;
+
+/** Type definition for weak references to instances of FTimerNode. */
+typedef TWeakPtr<class FTimerNode> FTimerNodeWeak;
+
+enum ETimerNodeType
+{
+	/** The TimerNode is a CPU Scope timer. */
+	CpuScope,
+
+	/** The TimerNode is a GPU Scope timer. */
+	GpuScope,
+
+	/** The TimerNode is a Compute Scope timer. */
+	ComputeScope,
+
+	/** The TimerNode is a group node. */
+	Group,
+
+	/** Invalid enum type, may be used as a number of enumerations. */
+	InvalidOrMax,
+};
+
+struct FTimerStats
+{
+	double TotalInclusiveTime; /** Total inclusive time, in seconds. */
+	double MinInclusiveTime; /** Min inclusive time, in seconds. */
+	double MaxInclusiveTime; /** Max inclusive time, in seconds. */
+	double AverageInclusiveTime; /** Average inclusive time, in seconds. */
+	double MedianInclusiveTime; /** Median inclusive time, in seconds. */
+
+	double TotalExclusiveTime; /** Total exclusive time, in seconds. */
+	double MinExclusiveTime; /** Min exclusive time, in seconds. */
+	double MaxExclusiveTime; /** Max exclusive time, in seconds. */
+	double AverageExclusiveTime; /** Average exclusive time, in seconds. */
+	double MedianExclusiveTime; /** Median exclusive time, in seconds. */
+
+	uint64 InstanceCount; /** Number of instances. */
+	//uint64 MinInstanceCountPerFrame;
+	//uint64 MaxInstanceCountPerFrame;
+	//uint64 AverageInstanceCountPerFrame;
+	//uint64 MedianInstanceCountPerFrame;
+
+	FTimerStats()
+		: TotalInclusiveTime(0.0)
+		, MinInclusiveTime(DBL_MAX)
+		, MaxInclusiveTime(DBL_MIN)
+		, AverageInclusiveTime(0.0)
+		, MedianInclusiveTime(0.0)
+		, TotalExclusiveTime(0.0)
+		, MinExclusiveTime(DBL_MAX)
+		, MaxExclusiveTime(DBL_MIN)
+		, AverageExclusiveTime(0.0)
+		, MedianExclusiveTime(0.0)
+		, InstanceCount(0)
+	{
+	}
+
+	void Reset()
+	{
+		TotalInclusiveTime = 0.0;
+		MinInclusiveTime = DBL_MAX;
+		MaxInclusiveTime = DBL_MIN;
+		AverageInclusiveTime = 0.0;
+		MedianInclusiveTime = 0.0;
+		TotalExclusiveTime = 0.0;
+		MinExclusiveTime = DBL_MAX;
+		MaxExclusiveTime = DBL_MIN;
+		AverageExclusiveTime = 0.0;
+		MedianExclusiveTime = 0.0;
+		InstanceCount = 0;
+	}
+};
+
+/**
+ * Class used to store information about a timer node (used in timers' tree view).
+ */
+class FTimerNode : public TSharedFromThis<FTimerNode>
+{
+public:
+	static const uint64 InvalidId = -1;
+
+public:
+	/** Initialization constructor for the timer node. */
+	FTimerNode(uint64 InId, const FName InName, const FName InMetaGroupName, ETimerNodeType InType)
+		: Id(InId)
+		, Name(InName)
+		, MetaGroupName(InMetaGroupName)
+		, Type(InType)
+		, bForceExpandGroupNode(false)
+	{
+		ResetStats();
+	}
+
+	/** Initialization constructor for the group node. */
+	FTimerNode(const FName InGroupName)
+		: Id(0)
+		, Name(InGroupName)
+		, Type(ETimerNodeType::Group)
+		, bForceExpandGroupNode(false)
+	{
+		ResetStats();
+	}
+
+	/**
+	 * @return an Id of this timer, valid only for timer nodes.
+	 */
+	const uint64 GetId() const
+	{
+		return Id;
+	}
+
+	/**
+	 * @return a name of this node, group or timer.
+	 */
+	const FName& GetName() const
+	{
+		return Name;
+	}
+
+	/**
+	 * @return a name of this node, group or timer + addditional info, to display in Timers tree view.
+	*/
+	const FText GetNameEx() const;
+
+	/**
+	 * @return a name of the group that this timer node belongs to, taken from the metadata.
+	 */
+	const FName& GetMetaGroupName() const
+	{
+		return MetaGroupName;
+	}
+
+	/**
+	 * @return a type of this timer or ETimerNodeType::Group for group nodes.
+	 */
+	const ETimerNodeType& GetType() const
+	{
+		return Type;
+	}
+
+	/**
+	 * @return true, if this node is a group node.
+	 */
+	bool IsGroup() const
+	{
+		return Type == ETimerNodeType::Group;
+	}
+
+	/**
+	 * @return the stats of this timer.
+	 */
+	const FTimerStats& GetStats() const
+	{
+		return Stats;
+	}
+
+	void ResetStats();
+
+	void SetStats(FTimerStats& Stats);
+
+	/**
+	 * @return a const reference to the child nodes of this group.
+	 */
+	FORCEINLINE_DEBUGGABLE const TArray<FTimerNodePtr>& GetChildren() const
+	{
+		return Children;
+	}
+
+	/**
+	 * @return a const reference to the child nodes that should be visible to the UI based on filtering.
+	 */
+	FORCEINLINE_DEBUGGABLE const TArray<FTimerNodePtr>& GetFilteredChildren() const
+	{
+		return FilteredChildren;
+	}
+
+	/**
+	 * @return a weak reference to the group of this timer node, may be invalid.
+	 */
+	FTimerNodeWeak GetGroupPtr() const
+	{
+		return GroupPtr;
+	}
+
+	/**
+	 * @return a name of the fake group that this timer node belongs to.
+	 */
+	const FName& GetGroupName() const
+	{
+		return GroupPtr.Pin()->Name;
+	}
+
+	/**
+	 * @return a name of the fake group that this timer node belongs to.
+	 */
+	//const FName& GetGroupNameSafe() const
+	//{
+	//	return GroupPtr.IsValid() ? GroupPtr.Pin()->Name : NAME_None;
+	//}
+
+	bool IsFiltered() const
+	{
+		return false; // TODO
+	}
+
+public:
+	/** Sorts children using the specified class instance. */
+	template<typename TSortingClass>
+	void SortChildren(const TSortingClass& Instance)
+	{
+		Children.Sort(Instance);
+	}
+
+	/** Adds specified child to the children and sets group for it. */
+	FORCEINLINE_DEBUGGABLE void AddChildAndSetGroupPtr(const FTimerNodePtr& ChildPtr)
+	{
+		ChildPtr->GroupPtr = AsShared();
+		Children.Add(ChildPtr);
+	}
+
+	/** Adds specified child to the filtered children. */
+	FORCEINLINE_DEBUGGABLE void AddFilteredChild(const FTimerNodePtr& ChildPtr)
+	{
+		FilteredChildren.Add(ChildPtr);
+	}
+
+	/** Clears filtered children. */
+	void ClearFilteredChildren()
+	{
+		FilteredChildren.Reset();
+	}
+
+protected:
+	/** The Id of this timer or group. */
+	const uint64 Id;
+
+	/** The name of this timer or group. */
+	const FName Name;
+
+	/** The name of the group that this timer belongs to, based on the timer's metadata; only valid for timer nodes. */
+	const FName MetaGroupName;
+
+	/** Holds the type of this timer; for the group, this is Group. */
+	const ETimerNodeType Type;
+
+	/** Stats */
+	FTimerStats Stats;
+
+	/** Children of this node. */
+	TArray<FTimerNodePtr> Children;
+
+	/** Filtered children of this node. */
+	TArray<FTimerNodePtr> FilteredChildren;
+
+	/** A weak pointer to the group/parent of this node. */
+	FTimerNodeWeak GroupPtr;
+
+public:
+	/** Whether this group node should be expanded when the text filtering is enabled. */
+	bool bForceExpandGroupNode;
+};
diff --git a/Engine/Source/Developer/TraceInsights/Private/Insights/ViewModels/TimerNodeHelper.cpp b/Engine/Source/Developer/TraceInsights/Private/Insights/ViewModels/TimerNodeHelper.cpp
new file mode 100644
--- /dev/null
+++ b/Engine/Source/Developer/TraceInsights/Private/Insights/ViewModels/TimerNodeHelper.cpp
@@ -0,0 +1,136 @@
+// Copyright 1998-2019 Epic Games, Inc. All Rights Reserved.
+
+#include "TimerNodeHelper.h"
+#include "EditorStyleSet.h"
+
+#define LOCTEXT_NAMESPACE "TimerNode"
+
+////////////////////////////////////////////////////////////////////////////////////////////////////
+// TimerNode Type Helper
+////////////////////////////////////////////////////////////////////////////////////////////////////
+
+FText TimerNodeTypeHelper::ToName(const ETimerNodeType NodeType)
+{
+	static_assert(ETimerNodeType::InvalidOrMax == 4, "Not all cases are handled in switch below!?");
+	switch (NodeType)
+	{
+	case ETimerNodeType::GpuScope:			return LOCTEXT("Timer_Name_GpuScope", "GPU");
+	case ETimerNodeType::ComputeScope:		return LOCTEXT("Timer_Name_ComputeScope", "Compute");
+	case ETimerNodeType::CpuScope:			return LOCTEXT("Timer_Name_CpuScope", "CPU");
+	case ETimerNodeType::Group:				return LOCTEXT("Timer_Name_Group", "Group"); //TODO
+	default: return LOCTEXT("InvalidOrMax", "InvalidOrMax");
+	}
+}
+
+////////////////////////////////////////////////////////////////////////////////////////////////////
+
+FText TimerNodeTypeHelper::ToDescription(const ETimerNodeType NodeType)
+{
+	static_assert(ETimerNodeType::InvalidOrMax == 4, "Not all cases are handled in switch below!?");
+	switch (NodeType)
+	{
+	case ETimerNodeType::GpuScope:			return LOCTEXT("Timer_Desc_GpuScope", "GPU scope timer");
+	case ETimerNodeType::ComputeScope:		return LOCTEXT("Timer_Desc_ComputeScope", "Compute scope timer");
+	case ETimerNodeType::CpuScope:			return LOCTEXT("Timer_Desc_CpuScope", "CPU scope timer");
+	case ETimerNodeType::Group:				return LOCTEXT("Timer_Desc_Group", "Group timer node"); //TODO
+	default: return LOCTEXT("InvalidOrMax", "InvalidOrMax");
+	}
+}
+
+////////////////////////////////////////////////////////////////////////////////////////////////////
+
+FName TimerNodeTypeHelper::ToBrushName(const ETimerNodeType NodeType)
+{
+	static_assert(ETimerNodeType::InvalidOrMax == 4, "Not all cases are handled in switch below!?");
+	switch (NodeType)
+	{
+	case ETimerNodeType::GpuScope:			return TEXT("Profiler.FiltersAndPresets.StatTypeIcon"); //im:TODO
+	case ETimerNodeType::ComputeScope:		return TEXT("Profiler.FiltersAndPresets.StatTypeIcon"); //im:TODO
+	case ETimerNodeType::CpuScope:			return TEXT("Profiler.FiltersAndPresets.StatTypeIcon"); //im:TODO
+	case ETimerNodeType::Group:				return TEXT("Profiler.Misc.GenericGroup"); //TODO
+	default: return NAME_None;
+	}
+}
+
+////////////////////////////////////////////////////////////////////////////////////////////////////
+
+const FSlateBrush* TimerNodeTypeHelper::GetIconForGroup()
+{
+	return FEditorStyle::GetBrush(TEXT("Profiler.Misc.GenericGroup")); //TODO
+}
+
+////////////////////////////////////////////////////////////////////////////////////////////////////
+
+const FSlateBrush* TimerNodeTypeHelper::GetIconForTimerNodeType(const ETimerNodeType NodeType)
+{
+	static_assert(ETimerNodeType::InvalidOrMax == 4, "Not all cases are handled in switch below!?");
+	switch (NodeType)
+	{
+	case ETimerNodeType::GpuScope:			return FEditorStyle::GetBrush(TEXT("Profiler.Type.NumberFloat")); //TODO
+	case ETimerNodeType::ComputeScope:		return FEditorStyle::GetBrush(TEXT("Profiler.Type.NumberFloat")); //TODO
+	case ETimerNodeType::CpuScope:			return FEditorStyle::GetBrush(TEXT("Profiler.Type.NumberFloat")); //TODO
+	case ETimerNodeType::Group:				return FEditorStyle::GetBrush(TEXT("Profiler.Misc.GenericGroup")); //TODO
+	default: return nullptr;
+	}
+}
+
+////////////////////////////////////////////////////////////////////////////////////////////////////
+// TimerNode Grouping Helper
+////////////////////////////////////////////////////////////////////////////////////////////////////
+
+FText TimerNodeGroupingHelper::ToName(const ETimerGroupingMode GroupingMode)
+{
+	switch (GroupingMode)
+	{
+	case ETimerGroupingMode::Flat:					return LOCTEXT("Grouping_Name_Flat",				"Flat");
+	case ETimerGroupingMode::ByName:				return LOCTEXT("Grouping_Name_ByName",				"Timer Name");
+	case ETimerGroupingMode::ByMetaGroupName:		return LOCTEXT("Grouping_Name_MetaGroupName",		"Meta Group Name");
+	case ETimerGroupingMode::ByType:				return LOCTEXT("Grouping_Name_Type",				"Timer Type");
+	case ETimerGroupingMode::ByTotalInclusiveTime:	return LOCTEXT("Grouping_Name_TotalInclusiveTime",	"Total Inclusive Time");
+	case ETimerGroupingMode::ByTotalExclusiveTime:	return LOCTEXT("Grouping_Name_TotalExclusiveTime",	"Total Exclusive Time");
+	case ETimerGroupingMode::ByInstanceCount:		return LOCTEXT("Grouping_Name_InstanceCount",		"Instance Count");
+
+	default: return LOCTEXT("InvalidOrMax", "InvalidOrMax");
+	}
+}
+
+////////////////////////////////////////////////////////////////////////////////////////////////////
+
+FText TimerNodeGroupingHelper::ToDescription(const ETimerGroupingMode GroupingMode)
+{
+	switch (GroupingMode)
+	{
+	case ETimerGroupingMode::Flat:					return LOCTEXT("Grouping_Desc_Flat",				"Creates a single group. Includes all timers.");
+	case ETimerGroupingMode::ByName:				return LOCTEXT("Grouping_Desc_ByName",				"Creates one group for one letter.");
+	case ETimerGroupingMode::ByMetaGroupName:		return LOCTEXT("Grouping_Desc_MetaGroupName",		"Creates groups based on timer metadata group names.");
+	case ETimerGroupingMode::ByType:				return LOCTEXT("Grouping_Desc_Type",				"Creates one group for each timer type.");
+	case ETimerGroupingMode::ByTotalInclusiveTime:	return LOCTEXT("Grouping_Desc_TotalInclusiveTime",	"Creates one group for each logarithmic range ie. 0.001 - 0.01, 0.01 - 0.1, 0.1 - 1.0, 1.0 - 10.0 etc");
+	case ETimerGroupingMode::ByTotalExclusiveTime:	return LOCTEXT("Grouping_Desc_TotalExclusiveTime",	"Creates one group for each logarithmic range ie. 0.001 - 0.01, 0.01 - 0.1, 0.1 - 1.0, 1.0 - 10.0 etc");
+	case ETimerGroupingMode::ByInstanceCount:		return LOCTEXT("Grouping_Desc_InstanceCount",		"Creates one group for each logarithmic range ie. 0, 1 - 10, 10 - 100, 100 - 1000, etc");
+
+	default: return LOCTEXT("InvalidOrMax", "InvalidOrMax");
+	}
+}
+
+////////////////////////////////////////////////////////////////////////////////////////////////////
+
+FName TimerNodeGroupingHelper::ToBrushName(const ETimerGroupingMode GroupingMode)
+{
+	switch (GroupingMode)
+	{
+	//TODO: "UnrealProfiler.*Icon"
+	case ETimerGroupingMode::Flat:					return TEXT("Profiler.FiltersAndPresets.GroupNameIcon"); //TODO
+	case ETimerGroupingMode::ByName:				return TEXT("Profiler.FiltersAndPresets.GroupNameIcon"); //TODO
+	case ETimerGroupingMode::ByMetaGroupName	:	return TEXT("Profiler.FiltersAndPresets.StatNameIcon"); //TODO
+	case ETimerGroupingMode::ByType:				return TEXT("Profiler.FiltersAndPresets.StatTypeIcon"); //TODO
+	case ETimerGroupingMode::ByTotalInclusiveTime:	return TEXT("Profiler.FiltersAndPresets.StatValueIcon"); //TODO
+	case ETimerGroupingMode::ByTotalExclusiveTime:	return TEXT("Profiler.FiltersAndPresets.StatValueIcon"); //TODO
+	case ETimerGroupingMode::ByInstanceCount:		return TEXT("Profiler.FiltersAndPresets.StatValueIcon"); //TODO
+
+	default: return NAME_None;
+	}
+}
+
+////////////////////////////////////////////////////////////////////////////////////////////////////
+
+#undef LOCTEXT_NAMESPACE
diff --git a/Engine/Source/Developer/TraceInsights/Private/Insights/ViewModels/TimerNodeHelper.h b/Engine/Source/Developer/TraceInsights/Private/Insights/ViewModels/TimerNodeHelper.h
new file mode 100644
--- /dev/null
+++ b/Engine/Source/Developer/TraceInsights/Private/Insights/ViewModels/TimerNodeHelper.h
@@ -0,0 +1,240 @@
+// Copyright 1998-2019 Epic Games, Inc. All Rights Reserved.
+
+#pragma once
+
+#include "CoreMinimal.h"
+#include "Styling/SlateBrush.h"
+
+// Insights
+#include "Insights/ViewModels/TimerNode.h"
+#include "Insights/ViewModels/TimerGroupingAndSorting.h"
+
+////////////////////////////////////////////////////////////////////////////////////////////////////
+
+/** Helper struct that contains helper functions and classes for ETimerNodeType enum. */
+struct TimerNodeTypeHelper
+{
+	/**
+	 * @param Type - The value to get the text for.
+	 *
+	 * @return text representation of the specified ETimerNodeType value.
+	 */
+	static FText ToName(const ETimerNodeType Type);
+
+	/**
+	 * @param Type - The value to get the text for.
+	 *
+	 * @return text representation with more detailed explanation of the specified ETimerNodeType value.
+	 */
+	static FText ToDescription(const ETimerNodeType Type);
+
+	/**
+	 * @param Type - The value to get the brush name for.
+	 *
+	 * @return brush name of the specified ETimerNodeType value.
+	 */
+	static FName ToBrushName(const ETimerNodeType Type);
+
+	static const FSlateBrush* GetIconForGroup();
+	static const FSlateBrush* GetIconForTimerNodeType(const ETimerNodeType Type);
+};
+
+////////////////////////////////////////////////////////////////////////////////////////////////////
+
+/** Helper struct that contains grouping static functions and classes. */
+struct TimerNodeGroupingHelper
+{
+	/**
+	 * @param TimerGroupingMode - The value to get the text for.
+	 *
+	 * @return text representation of the specified ETimerGroupingMode value.
+	 */
+	static FText ToName(const ETimerGroupingMode TimerGroupingMode);
+
+	/**
+	 * @param TimerGroupingMode - The value to get the text for.
+	 *
+	 * @return text representation with more detailed explanation of the specified ETimerGroupingMode value.
+	 */
+	static FText ToDescription(const ETimerGroupingMode TimerGroupingMode);
+
+	/**
+	 * @param TimerGroupingMode - The value to get the brush name for.
+	 *
+	 * @return brush name of the specified ETimerGroupingMode value.
+	 */
+	static FName ToBrushName(const ETimerGroupingMode TimerGroupingMode);
+};
+
+////////////////////////////////////////////////////////////////////////////////////////////////////
+
+/** Helper struct that contains sorting static functions and classes. */
+struct TimerNodeSortingHelper
+{
+	//////////////////////////////////////////////////
+	// Sorting by timer's name.
+
+	struct ByNameAscending
+	{
+		FORCEINLINE_DEBUGGABLE bool operator()(const FTimerNodePtr& A, const FTimerNodePtr& B) const
+		{
+			return A->GetName().LexicalLess(B->GetName());
+		}
+	};
+
+	struct ByNameDescending
+	{
+		FORCEINLINE_DEBUGGABLE bool operator()(const FTimerNodePtr& A, const FTimerNodePtr& B) const
+		{
+			return B->GetName().LexicalLess(A->GetName());
+		}
+	};
+
+	//////////////////////////////////////////////////
+	// Sorting by timer's meta group name.
+	// If meta group name is the same then sort by name.
+
+	struct ByMetaGroupNameAscending
+	{
+		FORCEINLINE_DEBUGGABLE bool operator()(const FTimerNodePtr& A, const FTimerNodePtr& B) const
+		{
+			if (A->GetMetaGroupName() == B->GetMetaGroupName())
+			{
+				// Sort by timer name (ascending).
+				return A->GetName().LexicalLess(B->GetName());
+			}
+			else
+			{
+				return A->GetMetaGroupName().LexicalLess(B->GetMetaGroupName());
+			}
+		}
+	};
+
+	struct ByMetaGroupNameDescending
+	{
+		FORCEINLINE_DEBUGGABLE bool operator()(const FTimerNodePtr& A, const FTimerNodePtr& B) const
+		{
+			if (A->GetMetaGroupName() == B->GetMetaGroupName())
+			{
+				// Sort by timer name (ascending).
+				return A->GetName().LexicalLess(B->GetName());
+			}
+			else
+			{
+				return B->GetMetaGroupName().LexicalLess(A->GetMetaGroupName());
+			}
+		}
+	};
+
+	//////////////////////////////////////////////////
+	// Sorting by timer's type.
+	// If type is the same then sort by name.
+
+	struct ByTypeAscending
+	{
+		FORCEINLINE_DEBUGGABLE bool operator()(const FTimerNodePtr& A, const FTimerNodePtr& B) const
+		{
+			const ETimerNodeType& TypeA = A->GetType();
+			const ETimerNodeType& TypeB = B->GetType();
+
+			if (TypeA == TypeB)
+			{
+				// Sort by timer name (ascending).
+				return A->GetName().LexicalLess(B->GetName());
+			}
+			else
+			{
+				return TypeA < TypeB;
+			}
+		}
+	};
+
+	struct ByTypeDescending
+	{
+		FORCEINLINE_DEBUGGABLE bool operator()(const FTimerNodePtr& A, const FTimerNodePtr& B) const
+		{
+			const ETimerNodeType& TypeA = A->GetType();
+			const ETimerNodeType& TypeB = B->GetType();
+
+			if (TypeA == TypeB)
+			{
+				// Sort by timer name (ascending).
+				return A->GetName().LexicalLess(B->GetName());
+			}
+			else
+			{
+				return TypeA > TypeB;
+			}
+		}
+	};
+
+	//////////////////////////////////////////////////
+	// Sorting by a timer's stats value.
+	// If stats value is the same then sort by name.
+
+	#define SORT_BY_STATS_ASCENDING(Type, SortName, Stats) \
+		struct SortName \
+		{ \
+			FORCEINLINE_DEBUGGABLE bool operator()(const FTimerNodePtr& A, const FTimerNodePtr& B) const \
+			{ \
+				const Type ValueA = A->GetStats().Stats; \
+				const Type ValueB = B->GetStats().Stats; \
+				\
+				if (ValueA == ValueB) \
+				{ \
+					/* Sort by name (ascending). */ \
+					return A->GetName().LexicalLess(B->GetName()); \
+				} \
+				else \
+				{ \
+					return ValueA < ValueB; \
+				} \
+			} \
+		};
+
+	#define SORT_BY_STATS_DESCENDING(Type, SortName, Stats) \
+		struct SortName \
+		{ \
+			FORCEINLINE_DEBUGGABLE bool operator()(const FTimerNodePtr& A, const FTimerNodePtr& B) const \
+			{ \
+				const Type ValueA = A->GetStats().Stats; \
+				const Type ValueB = B->GetStats().Stats; \
+				\
+				if (ValueA == ValueB) \
+				{ \
+					/* Sort by name (ascending). */ \
+					return A->GetName().LexicalLess(B->GetName()); \
+				} \
+				else \
+				{ \
+					return ValueA > ValueB; \
+				} \
+			} \
+		};
+
+	#define SORT_BY_STATS(Type, SortName, Stats) \
+		SORT_BY_STATS_ASCENDING(Type, SortName##Ascending, Stats) \
+		SORT_BY_STATS_DESCENDING(Type, SortName##Descending, Stats)
+
+	SORT_BY_STATS(uint64, ByInstanceCount, InstanceCount);
+
+	SORT_BY_STATS(double, ByTotalInclusiveTime, TotalInclusiveTime);
+	SORT_BY_STATS(double, ByMinInclusiveTime, MinInclusiveTime);
+	SORT_BY_STATS(double, ByMaxInclusiveTime, MaxInclusiveTime);
+	SORT_BY_STATS(double, ByAverageInclusiveTime, AverageInclusiveTime);
+	SORT_BY_STATS(double, ByMedianInclusiveTime, MedianInclusiveTime);
+
+	SORT_BY_STATS(double, ByTotalExclusiveTime, TotalExclusiveTime);
+	SORT_BY_STATS(double, ByMinExclusiveTime, MinExclusiveTime);
+	SORT_BY_STATS(double, ByMaxExclusiveTime, MaxExclusiveTime);
+	SORT_BY_STATS(double, ByAverageExclusiveTime, AverageExclusiveTime);
+	SORT_BY_STATS(double, ByMedianExclusiveTime, MedianExclusiveTime);
+
+	#undef SORT_BY_STATS_ASCENDING
+	#undef SORT_BY_STATS_DESCENDING
+	#undef SORT_BY_STATS
+
+	//////////////////////////////////////////////////
+};
+
+////////////////////////////////////////////////////////////////////////////////////////////////////
diff --git a/Engine/Source/Developer/TraceInsights/Private/Insights/ViewModels/TimersViewColumn.h b/Engine/Source/Developer/TraceInsights/Private/Insights/ViewModels/TimersViewColumn.h
new file mode 100644
--- /dev/null
+++ b/Engine/Source/Developer/TraceInsights/Private/Insights/ViewModels/TimersViewColumn.h
@@ -0,0 +1,119 @@
+// Copyright 1998-2019 Epic Games, Inc. All Rights Reserved.
+
+#pragma once
+
+#include "CoreMinimal.h"
+#include "Types/SlateEnums.h"
+
+class FTimerNode;
+
+/// Holds information about a column in the timer view widget.
+class FTimersViewColumn
+{
+	friend struct FTimersViewColumnFactory;
+
+public:
+	enum EFlags : uint32
+	{
+		None = 0,
+
+		CanBeHidden = (1 << 0),
+		CanBeSorted = (1 << 1),
+		CanBeFiltered = (1 << 2),
+
+		NameColumnFlags = CanBeSorted | CanBeFiltered,
+		StatsColumnFlags = CanBeHidden | CanBeSorted | CanBeFiltered,
+	};
+
+	typedef TFunction<FText(const FTimersViewColumn& Column, const FTimerNode& TimerNode)> FGetFormattedValueFn;
+
+public:
+	/// Whether this column can be hidden.
+	bool bCanBeHidden() const { return ((uint32)Flags & (uint32)EFlags::CanBeHidden) != 0; }
+
+	/// Whether this column cab be used for sorting.
+	bool bCanBeSorted() const { return ((uint32)Flags & (uint32)EFlags::CanBeSorted) != 0; }
+
+	/// Where this column can be used to filtering displayed results.
+	bool bCanBeFiltered() const { return ((uint32)Flags & (uint32)EFlags::CanBeFiltered) != 0; }
+
+	/// If MinColumnWidth == MaxColumnWidth, this column has fixed width and cannot be resized.
+	bool bIsFixedColumnWidth() const { return MinColumnWidth == MaxColumnWidth; }
+
+	FText GetFormattedValue(const FTimerNode& TimerNode) const
+	{
+		return GetFormattedValueFn(*this, TimerNode);
+	}
+
+protected:
+	/// No default constructor.
+	FTimersViewColumn() = delete;
+
+	/// Initialization constructor, only used in FTimersViewColumnFactory.
+	FTimersViewColumn
+	(
+		int32 InOrder,
+		const FName InId,
+		const FName InSearchId,
+		FText InShortName,
+		FText InTitleName,
+		FText InDescription,
+		bool bInIsVisible,
+		const EFlags InFlags,
+		const EHorizontalAlignment InHorizontalAlignment,
+		const float InInitialColumnWidth,
+		const float InMinColumnWidth,
+		const float InMaxColumnWidth,
+		const FGetFormattedValueFn InGetFormattedValueFn
+	)
+		: Order(InOrder)
+		, Id(InId)
+		, SearchId(InSearchId)
+		, ShortName(MoveTemp(InShortName))
+		, TitleName(MoveTemp(InTitleName))
+		, Description(MoveTemp(InDescription))
+		, bIsVisible(bInIsVisible)
+		, Flags(InFlags)
+		, HorizontalAlignment(InHorizontalAlignment)
+		, InitialColumnWidth(InInitialColumnWidth)
+		, MinColumnWidth(InMinColumnWidth)
+		, MaxColumnWidth(InMaxColumnWidth)
+		, GetFormattedValueFn(InGetFormattedValueFn)
+	{
+	}
+
+public:
+	/// Order value, to sort columns in the tree view.
+	int32 Order;
+
+	/// Name of the column, name of the property.
+	FName Id;
+
+	/// Name of the column used by the searching system.
+	FName SearchId;
+
+	/// Short name of the column, displayed in the column header.
+	FText ShortName;
+
+	/// Title name of the column, displayed as title in the column tooltip.
+	FText TitleName;
+
+	/// Long name of the column, displayed in the column tooltip.
+	FText Description;
+
+	/// Is this column visible?
+	bool bIsVisible;
+
+	/// On/off switches.
+	EFlags Flags;
+
+	/// Horizontal alignment of the content in this column.
+	EHorizontalAlignment HorizontalAlignment;
+
+	float InitialColumnWidth; /// Initial column width.
+	float MinColumnWidth; /// Minimum column width.
+	float MaxColumnWidth; /// Maximum column width.
+
+	/// Custom function used to format (as an FText) the value of the timer node to be displayed by this column.
+	FGetFormattedValueFn GetFormattedValueFn;
+};
diff --git a/Engine/Source/Developer/TraceInsights/Private/Insights/ViewModels/TimersViewColumnFactory.h b/Engine/Source/Developer/TraceInsights/Private/Insights/ViewModels/TimersViewColumnFactory.h
new file mode 100644
--- /dev/null
+++ b/Engine/Source/Developer/TraceInsights/Private/Insights/ViewModels/TimersViewColumnFactory.h
@@ -0,0 +1,319 @@
+// Copyright 1998-2019 Epic Games, Inc. All Rights Reserved.
+
+#pragma once
+
+#include "CoreMinimal.h"
+
+// Insights
+#include "Insights/Common/TimeUtils.h"
+#include "Insights/ViewModels/TimerNodeHelper.h"
+#include "Insights/ViewModels/TimersViewColumn.h"
+
+#define LOCTEXT_NAMESPACE "STimerView"
+
+struct FTimersViewColumnFactory
+{
+	/** Default constructor. */
+	FTimersViewColumnFactory()
+	{
+		uint32 Index = 0;
+
+		// FTimerViewColumn
+		// (
+		//     Order,
+		//     Id,
+		//     SearchId,
+		//     ShortName,
+		//     TitleName,
+		//     Description,
+		//     bIsVisible,
+		//     Flags,
+		//     HorizontalAlignment,
+		//     InitialWidth, MinWidth, MaxWidth,
+		//     GetFormattedValueFn
+		// )
+
+		//////////////////////////////////////////////////
+		// Name
+
+		Collection.Add(new FTimersViewColumn
+		(
+			Index++,
+			TEXT("Name"),
+			TEXT("name"),
+			LOCTEXT("TimerNameColumnName", "Name"),
+			LOCTEXT("TimerNameColumnTitle", "Timer or Group Name"),
+			LOCTEXT("TimerNameColumnDesc", "Name of the timer or group."),
+			true,
+			FTimersViewColumn::EFlags::NameColumnFlags,
+			HAlign_Left,
+			246.0f, 10.0f, FLT_MAX,
+			[](const FTimersViewColumn& Column, const FTimerNode& TimerNode) -> FText
+			{ return FText::FromName(TimerNode.GetName()); }
+		));
+
+		//////////////////////////////////////////////////
+		// Meta Group Name
+
+		Collection.Add(new FTimersViewColumn
+		(
+			Index++,
+			TEXT("MetaGroupName"),
+			TEXT("metagroupname"),
+			LOCTEXT("TimerMetaGroupNameColumnName", "Meta Group"),
+			LOCTEXT("TimerMetaGroupNameColumnTitle", "Meta Group Name"),
+			LOCTEXT("TimerMetaGroupNameColumnDesc", "Name of the meta group."),
+			false,
+			(FTimersViewColumn::EFlags)(FTimersViewColumn::EFlags::CanBeHidden |
+										FTimersViewColumn::EFlags::CanBeSorted |
+										FTimersViewColumn::EFlags::CanBeFiltered),
+			HAlign_Left,
+			100.0f, 10.0f, FLT_MAX,
+			[](const FTimersViewColumn& Column, const FTimerNode& TimerNode) -> FText
+			{ return FText::FromName(TimerNode.GetMetaGroupName()); }
+		));
+
+		//////////////////////////////////////////////////
+		// Type
+
+		Collection.Add(new FTimersViewColumn
+		(
+			Index++,
+			TEXT("Type"),
+			TEXT("type"),
+			LOCTEXT("TimerTypeColumnName", "Type"),
+			LOCTEXT("TimerTypeColumnTitle", "Type"),
+			LOCTEXT("TimerTypeColumnDesc", "Type of timer or group."),
+			false,
+			(FTimersViewColumn::EFlags)(FTimersViewColumn::EFlags::CanBeHidden |
+										FTimersViewColumn::EFlags::CanBeSorted |
+										FTimersViewColumn::EFlags::CanBeFiltered),
+			HAlign_Left,
+			60.0f, 10.0f, FLT_MAX,
+			[](const FTimersViewColumn& Column, const FTimerNode& TimerNode) -> FText
+			{ return TimerNodeTypeHelper::ToName(TimerNode.GetType()); }
+		));
+
+		//////////////////////////////////////////////////
+		// Instance Count
+
+		Collection.Add(new FTimersViewColumn
+		(
+			Index++,
+			TEXT("InstanceCount"),
+			TEXT("count"),
+			LOCTEXT("InstanceCountName", "Count"),
+			LOCTEXT("InstanceCountTitle", "Instance Count"),
+			LOCTEXT("InstanceCountDesc", "Number of timer's instances"),
+			true,
+			FTimersViewColumn::EFlags::StatsColumnFlags,
+			HAlign_Right,
+			60.0f, 0.0f, FLT_MAX,
+			[](const FTimersViewColumn& Column, const FTimerNode& TimerNode) -> FText
+			{ return FText::AsNumber(TimerNode.GetStats().InstanceCount); }
+		));
+
+		//////////////////////////////////////////////////
+
+		constexpr float TotalTimeColumnInitialWidth = 60.0f;
+		constexpr float TotalTimeColumnMinWidth = 4.0f;
+		constexpr float TotalTimeColumnMaxWidth = FLT_MAX;
+
+		constexpr float TimeMsColumnInitialWidth = 50.0f;
+		constexpr float TimeMsColumnMinWidth = 4.0f;
+		constexpr float TimeMsColumnMaxWidth = FLT_MAX;
+
+		//////////////////////////////////////////////////
+		// Inclusive Time stats
+
+		Collection.Add(new FTimersViewColumn
+		(
+			Index++,
+			TEXT("TotalInclusiveTime"),
+			TEXT("inc"),
+			LOCTEXT("TotalInclusiveTimeName", "Incl."),
+			LOCTEXT("TotalInclusiveTimeTitle", "Total Inclusive Time"),
+			LOCTEXT("TotalInclusiveTimeDesc", "Total inclusive duration of selected timer's instances."),
+			true,
+			FTimersViewColumn::EFlags::StatsColumnFlags,
+			HAlign_Right,
+			TotalTimeColumnInitialWidth, TotalTimeColumnMinWidth, TotalTimeColumnMaxWidth,
+			[](const FTimersViewColumn& Column, const FTimerNode& TimerNode) -> FText
+			{ return FText::FromString(TimeUtils::FormatTimeAuto(TimerNode.GetStats().TotalInclusiveTime)); }
+		));
+
+		Collection.Add(new FTimersViewColumn
+		(
+			Index++,
+			TEXT("MinInclusiveTime"),
+			TEXT("mininc"),
+			LOCTEXT("MinInclusiveTimeName", "I.Min"),
+			LOCTEXT("MinInclusiveTimeTitle", "Min Inclusive Time (ms)"),
+			LOCTEXT("MinInclusiveTimeDesc", "Minimum inclusive duration of selected timer's instances, in milliseconds"),
+			false,
+			FTimersViewColumn::EFlags::StatsColumnFlags,
+			HAlign_Right,
+			TimeMsColumnInitialWidth, TimeMsColumnMinWidth, TimeMsColumnMaxWidth,
+			[](const FTimersViewColumn& Column, const FTimerNode& TimerNode) -> FText
+			{ return FText::FromString(TimeUtils::FormatTimeMs(TimerNode.GetStats().MinInclusiveTime)); }
+		));
+
+		Collection.Add(new FTimersViewColumn
+		(
+			Index++,
+			TEXT("MaxInclusiveTime"),
+			TEXT("maxinc"),
+			LOCTEXT("MaxInclusiveTimeName", "I.Max"),
+			LOCTEXT("MaxInclusiveTimeTitle", "Max Inclusive Time (ms)"),
+			LOCTEXT("MaxInclusiveTimeDesc", "Maximum inclusive duration of selected timer's instances, in milliseconds"),
+			false,
+			FTimersViewColumn::EFlags::StatsColumnFlags,
+			HAlign_Right,
+			TimeMsColumnInitialWidth, TimeMsColumnMinWidth, TimeMsColumnMaxWidth,
+			[](const FTimersViewColumn& Column, const FTimerNode& TimerNode) -> FText
+			{ return FText::FromString(TimeUtils::FormatTimeMs(TimerNode.GetStats().MaxInclusiveTime)); }
+		));
+
+		Collection.Add(new FTimersViewColumn
+		(
+			Index++,
+			TEXT("AvgInclusiveTime"),
+			TEXT("avginc"),
+			LOCTEXT("AvgInclusiveTimeName", "I.Avg."),
+			LOCTEXT("AvgInclusiveTimeTitle", "Average Inclusive Time (ms)"),
+			LOCTEXT("AvgInclusiveTimeDesc", "Average inclusive duration of selected timer's instances, in milliseconds"),
+			false,
+			FTimersViewColumn::EFlags::StatsColumnFlags,
+			HAlign_Right,
+			TimeMsColumnInitialWidth, TimeMsColumnMinWidth, TimeMsColumnMaxWidth,
+			[](const FTimersViewColumn& Column, const FTimerNode& TimerNode) -> FText
+			{ return FText::FromString(TimeUtils::FormatTimeMs(TimerNode.GetStats().AverageInclusiveTime)); }
+		));
+
+		Collection.Add(new FTimersViewColumn
+		(
+			Index++,
+			TEXT("MedInclusiveTime"),
+			TEXT("medinc"),
+			LOCTEXT("MedInclusiveTimeName", "I.Med."),
+			LOCTEXT("MedInclusiveTimeTitle", "Median Inclusive Time (ms)"),
+			LOCTEXT("MedInclusiveTimeDesc", "Median inclusive duration of selected timer's instances, in milliseconds"),
+			false,
+			FTimersViewColumn::EFlags::StatsColumnFlags,
+			HAlign_Right,
+			TimeMsColumnInitialWidth, TimeMsColumnMinWidth, TimeMsColumnMaxWidth,
+			[](const FTimersViewColumn& Column, const FTimerNode& TimerNode) -> FText
+			{ return FText::FromString(TimeUtils::FormatTimeMs(TimerNode.GetStats().MedianInclusiveTime)); }
+		));
+
+		//////////////////////////////////////////////////
+		// Exclusive Time stats
+
+		Collection.Add(new FTimersViewColumn
+		(
+			Index++,
+			TEXT("TotalExclusiveTime"),
+			TEXT("exc"),
+			LOCTEXT("TotalExclusiveTimeName", "Excl."),
+			LOCTEXT("TotalExclusiveTimeTitle", "Total Exclusive Time"),
+			LOCTEXT("TotalExclusiveTimeDesc", "Total exclusive duration of selected timer's instances."),
+			true,
+			FTimersViewColumn::EFlags::StatsColumnFlags,
+			HAlign_Right,
+			TotalTimeColumnInitialWidth, TotalTimeColumnMinWidth, TotalTimeColumnMaxWidth,
+			[](const FTimersViewColumn& Column, const FTimerNode& TimerNode) -> FText
+			{ return FText::FromString(TimeUtils::FormatTimeAuto(TimerNode.GetStats().TotalExclusiveTime)); }
+		));
+
+		Collection.Add(new FTimersViewColumn
+		(
+			Index++,
+			TEXT("MinExclusiveTime"),
+			TEXT("minexc"),
+			LOCTEXT("MinExclusiveTimeName", "E.Min"),
+			LOCTEXT("MinExclusiveTimeTitle", "Min Exclusive Time (ms)"),
+			LOCTEXT("MinExclusiveTimeDesc", "Minimum exclusive duration of selected timer's instances, in milliseconds"),
+			false,
+			FTimersViewColumn::EFlags::StatsColumnFlags,
+			HAlign_Right,
+			TimeMsColumnInitialWidth, TimeMsColumnMinWidth, TimeMsColumnMaxWidth,
+			[](const FTimersViewColumn& Column, const FTimerNode& TimerNode) -> FText
+			{ return FText::FromString(TimeUtils::FormatTimeMs(TimerNode.GetStats().MinExclusiveTime)); }
+		));
+
+		Collection.Add(new FTimersViewColumn
+		(
+			Index++,
+			TEXT("MaxExclusiveTime"),
+			TEXT("maxexc"),
+			LOCTEXT("MaxExclusiveTimeName", "E.Max"),
+			LOCTEXT("MaxExclusiveTimeTitle", "Max Exclusive Time (ms)"),
+			LOCTEXT("MaxExclusiveTimeDesc", "Maximum exclusive duration of selected timer's instances, in milliseconds"),
+			false,
+			FTimersViewColumn::EFlags::StatsColumnFlags,
+			HAlign_Right,
+			TimeMsColumnInitialWidth, TimeMsColumnMinWidth, TimeMsColumnMaxWidth,
+			[](const FTimersViewColumn& Column, const FTimerNode& TimerNode) -> FText
+			{ return FText::FromString(TimeUtils::FormatTimeMs(TimerNode.GetStats().MaxExclusiveTime)); }
+		));
+
+		Collection.Add(new FTimersViewColumn
+		(
+			Index++,
+			TEXT("AvgExclusiveTime"),
+			TEXT("avgexc"),
+			LOCTEXT("AvgExclusiveTimeName", "E.Avg."),
+			LOCTEXT("AvgExclusiveTimeTitle", "Average Exclusive Time (ms)"),
+			LOCTEXT("AvgExclusiveTimeDesc", "Average exclusive duration of selected timer's instances, in milliseconds"),
+			false,
+			FTimersViewColumn::EFlags::StatsColumnFlags,
+			HAlign_Right,
+			TimeMsColumnInitialWidth, TimeMsColumnMinWidth, TimeMsColumnMaxWidth,
+			[](const FTimersViewColumn& Column, const FTimerNode& TimerNode) -> FText
+			{ return FText::FromString(TimeUtils::FormatTimeMs(TimerNode.GetStats().AverageExclusiveTime)); }
+		));
+
+		Collection.Add(new FTimersViewColumn
+		(
+			Index++,
+			TEXT("MedExclusiveTime"),
+			TEXT("medexc"),
+			LOCTEXT("MedExclusiveTimeName", "E.Med."),
+			LOCTEXT("MedExclusiveTimeTitle", "Median Exclusive Time (ms)"),
+			LOCTEXT("MedExclusiveTimeDesc", "Median exclusive duration of timer's instances, in milliseconds"),
+			false,
+			FTimersViewColumn::EFlags::StatsColumnFlags,
+			HAlign_Right,
+			TimeMsColumnInitialWidth, TimeMsColumnMinWidth, TimeMsColumnMaxWidth,
+			[](const FTimersViewColumn& Column, const FTimerNode& TimerNode) -> FText
+			{ return FText::FromString(TimeUtils::FormatTimeMs(TimerNode.GetStats().MedianExclusiveTime)); }
+		));
+
+		//////////////////////////////////////////////////
+
+		for (int32 MapIndex = 0; MapIndex < Collection.Num(); MapIndex++)
+		{
+			ColumnIdToPtrMapping.Add(Collection[MapIndex]->Id, Collection[MapIndex]);
+		}
+	}
+
+	~FTimersViewColumnFactory()
+	{
+		for (const FTimersViewColumn* Column : Collection)
+			delete Column;
+		Collection.Reset();
+	}
+
+	/** Contains basic information about columns used in the timer view widget. Names should be localized. */
+	TArray<FTimersViewColumn*> Collection;
+
+	TMap<FName, const FTimersViewColumn*> ColumnIdToPtrMapping;
+
+	static const FTimersViewColumnFactory& Get()
+	{
+		static FTimersViewColumnFactory Instance;
+		return Instance;
+	}
+};
+
+#undef LOCTEXT_NAMESPACE
diff --git a/Engine/Source/Developer/TraceInsights/Private/Insights/ViewModels/TimingEvent.h b/Engine/Source/Developer/TraceInsights/Private/Insights/ViewModels/TimingEvent.h
new file mode 100644
--- /dev/null
+++ b/Engine/Source/Developer/TraceInsights/Private/Insights/ViewModels/TimingEvent.h
@@ -0,0 +1,50 @@
+// Copyright 1998-2019 Epic Games, Inc. All Rights Reserved.
+
+#pragma once
+
+#include "CoreMinimal.h"
+
+// Insights
+#include "Insights/ViewModels/TimerNode.h"
+
+class FTimingEventsTrack;
+
+struct FTimingEvent
+{
+	const FTimingEventsTrack* Track;
+	uint64 TypeId;
+	uint32 Depth;
+	double StartTime;
+	double EndTime;
+	double ExclusiveTime;
+
+	void Reset()
+	{
+		Track = nullptr;
+		TypeId = FTimerNode::InvalidId;
+		Depth = 0;
+		StartTime = 0.0;
+		EndTime = 0.0;
+		ExclusiveTime = 0.0;
+	}
+
+	bool IsValidTrack() const { return Track != nullptr; }
+	bool IsValid() const { return Track != nullptr && TypeId != FTimerNode::InvalidId; }
+
+	double Duration() const { return EndTime - StartTime; }
+
+	bool Equals(const FTimingEvent& Other) const
+	{
+		return Track == Other.Track
+			&& TypeId == Other.TypeId
+			&& Depth == Other.Depth
+			&& StartTime == Other.StartTime
+			&& EndTime == Other.EndTime;
+			//&& ExclusiveTime == Other.ExclusiveTime;
+	}
+
+	static bool AreEquals(const FTimingEvent& A, const FTimingEvent& B)
+	{
+		return A.Equals(B);
+	}
+};
diff --git a/Engine/Source/Developer/TraceInsights/Private/Insights/ViewModels/TimingEventsTrack.cpp b/Engine/Source/Developer/TraceInsights/Private/Insights/ViewModels/TimingEventsTrack.cpp
new file mode 100644
--- /dev/null
+++ b/Engine/Source/Developer/TraceInsights/Private/Insights/ViewModels/TimingEventsTrack.cpp
@@ -0,0 +1,116 @@
+// Copyright 1998-2019 Epic Games, Inc. All Rights Reserved.
+
+#include "TimingEventsTrack.h"
+
+#define LOCTEXT_NAMESPACE "TimingEventsTrack"
+
+////////////////////////////////////////////////////////////////////////////////////////////////////
+// FTimingEventsTrackLayout
+////////////////////////////////////////////////////////////////////////////////////////////////////
+
+void FTimingEventsTrackLayout::ForceNormalMode()
+{
+	bIsCompactMode = false;
+	EventH = NormalLayoutEventH;
+	EventDY = NormalLayoutEventDY;
+	TimelineDY = NormalLayoutTimelineDY;
+	MinTimelineH = NormalLayoutMinTimelineH;
+	TargetMinTimelineH = NormalLayoutMinTimelineH;
+}
+
+////////////////////////////////////////////////////////////////////////////////////////////////////
+
+void FTimingEventsTrackLayout::ForceCompactMode()
+{
+	bIsCompactMode = true;
+	EventH = CompactLayoutEventH;
+	EventDY = CompactLayoutEventDY;
+	TimelineDY = CompactLayoutTimelineDY;
+	MinTimelineH = CompactLayoutMinTimelineH;
+	TargetMinTimelineH = CompactLayoutMinTimelineH;
+}
+
+////////////////////////////////////////////////////////////////////////////////////////////////////
+
+void FTimingEventsTrackLayout::Update()
+{
+	const float LayoutTransitionSpeed = 0.25f;
+
+	if (bIsCompactMode)
+	{
+		if (EventH > CompactLayoutEventH)
+			EventH -= LayoutTransitionSpeed;
+		if (EventDY > CompactLayoutEventDY)
+			EventDY -= LayoutTransitionSpeed;
+		if (TimelineDY > CompactLayoutTimelineDY)
+			TimelineDY -= LayoutTransitionSpeed;
+	}
+	else
+	{
+		if (EventH < NormalLayoutEventH)
+			EventH += LayoutTransitionSpeed;
+		if (EventDY < NormalLayoutEventDY)
+			EventDY += LayoutTransitionSpeed;
+		if (TimelineDY < NormalLayoutTimelineDY)
+			TimelineDY += LayoutTransitionSpeed;
+	}
+
+	if (MinTimelineH > TargetMinTimelineH)
+		MinTimelineH -= LayoutTransitionSpeed;
+	else if (MinTimelineH < TargetMinTimelineH)
+		MinTimelineH += LayoutTransitionSpeed;
+}
+
+////////////////////////////////////////////////////////////////////////////////////////////////////
+// FTimingEventsTrack
+////////////////////////////////////////////////////////////////////////////////////////////////////
+
+FTimingEventsTrack::FTimingEventsTrack(uint64 InTrackId)
+	: FBaseTimingTrack(InTrackId)
+	, Depth(-1)
+	, bIsCollapsed(false)
+{
+}
+
+////////////////////////////////////////////////////////////////////////////////////////////////////
+
+FTimingEventsTrack::~FTimingEventsTrack()
+{
+}
+
+////////////////////////////////////////////////////////////////////////////////////////////////////
+
+void FTimingEventsTrack::Reset()
+{
+	FBaseTimingTrack::Reset();
+
+	Depth = -1;
+	bIsCollapsed = false;
+}
+
+////////////////////////////////////////////////////////////////////////////////////////////////////
+
+void FTimingEventsTrack::UpdateHoveredState(float MX, float MY, const FTimingTrackViewport& Viewport)
+{
+	if (MY >= Y && MY < Y + H)
+	{
+		bIsHovered = true;
+
+		if (MX < 100.0f && MY < Y + 14.0f)
+		{
+			bIsHeaderHovered = true;
+		}
+		else
+		{
+			bIsHeaderHovered = false;
+		}
+	}
+	else
+	{
+		bIsHovered = false;
+	}
+}
+
+////////////////////////////////////////////////////////////////////////////////////////////////////
+
+#undef LOCTEXT_NAMESPACE
diff --git a/Engine/Source/Developer/TraceInsights/Private/Insights/ViewModels/TimingEventsTrack.h b/Engine/Source/Developer/TraceInsights/Private/Insights/ViewModels/TimingEventsTrack.h
new file mode 100644
--- /dev/null
+++ b/Engine/Source/Developer/TraceInsights/Private/Insights/ViewModels/TimingEventsTrack.h
@@ -0,0 +1,64 @@
+// Copyright 1998-2019 Epic Games, Inc. All Rights Reserved.
+
+#pragma once
+
+#include "CoreMinimal.h"
+
+#include "Insights/ViewModels/BaseTimingTrack.h"
+
+////////////////////////////////////////////////////////////////////////////////////////////////////
+
+static const float RealMinTimelineH = 13.0f;
+
+static const float NormalLayoutEventH = 14.0f;
+static const float NormalLayoutEventDY = 2.0f;
+static const float NormalLayoutTimelineDY = 14.0f;
+static const float NormalLayoutMinTimelineH = 0.0f;
+
+static const float CompactLayoutEventH = 2.0f;
+static const float CompactLayoutEventDY = 1.0f;
+static const float CompactLayoutTimelineDY = 3.0f;
+static const float CompactLayoutMinTimelineH = 0.0f;
+
+////////////////////////////////////////////////////////////////////////////////////////////////////
+
+struct FTimingEventsTrackLayout
+{
+	bool bIsCompactMode;
+
+	float EventH; // height of a timing event, in Slate units
+	float EventDY; // vertical distance between two timing event sub-tracks, in Slate units
+	float TimelineDY; // space at top and bottom of each timeline, in Slate units
+	float MinTimelineH;
+	float TargetMinTimelineH;
+
+	float GetLaneY(uint32 Depth) const { return 1.0f + TimelineDY + Depth * (EventDY + EventH); }
+
+	void ForceNormalMode();
+	void ForceCompactMode();
+	void Update();
+};
+
+////////////////////////////////////////////////////////////////////////////////////////////////////
+
+class FTimingEventsTrack : public FBaseTimingTrack
+{
+public:
+	FTimingEventsTrack(uint64 InTrackId);
+	virtual ~FTimingEventsTrack();
+
+	virtual void Reset() override;
+	virtual void UpdateHoveredState(float MX, float MY, const FTimingTrackViewport& Viewport);
+
+public:
+	int32 Depth; // how many lanes has this track
+	bool bIsCollapsed;
+
+	// Cached OnPaint state.
+	//TArray<FEventBoxInfo> Boxes;
+	//TArray<FEventBoxInfo> MergedBorderBoxes;
+	//TArray<FEventBoxInfo> Borders;
+	//TArray<FTextInfo> Texts;
+};
+
+////////////////////////////////////////////////////////////////////////////////////////////////////
diff --git a/Engine/Source/Developer/TraceInsights/Private/Insights/ViewModels/TimingTrackViewport.h b/Engine/Source/Developer/TraceInsights/Private/Insights/ViewModels/TimingTrackViewport.h
new file mode 100644
--- /dev/null
+++ b/Engine/Source/Developer/TraceInsights/Private/Insights/ViewModels/TimingTrackViewport.h
@@ -0,0 +1,109 @@
+// Copyright 1998-2019 Epic Games, Inc. All Rights Reserved.
+
+#pragma once
+
+#include "CoreMinimal.h"
+
+class FTimingTrackViewport
+{
+public:
+	float Width; // width of viewport, in Slate units
+	float Height; // height of viewport, in Slate units
+
+	double MinValidTime; // min session time, in seconds
+	double MaxValidTime; // max session time, in seconds
+	double StartTime; // time of viewport's left side, in seconds
+	double EndTime; // time of viewport's right side, in seconds; computed when StartTime or Width changes
+	double MinScaleX;
+	double MaxScaleX;
+	double ScaleX; // scale factor between seconds and pixels (Slate units), in pixels (Slate units) per second
+
+	float TopOffset; // top offset (to allow the time ruller to be visible), in pixels (Slate units)
+	float ScrollHeight; // height of the vertical scrollable area, in pixels (Slate units)
+	float ScrollPosY; // current vertical scroll position, in pixels (Slate units)
+
+public:
+	FTimingTrackViewport()
+	{
+		Reset();
+	}
+
+	void Reset()
+	{
+		Width = 0.0f;
+		Height = 0.0f;
+
+		MinValidTime = 0.0;
+		MaxValidTime = 60.0; // 1m
+		StartTime = 0.0;
+		EndTime = 0.0;
+		MinScaleX = (5 * 20) / 3600.0; // 1h between major tick marks
+		MaxScaleX = 1.0E10; // 10ns between major tick marks
+		ScaleX = (5 * 20) / 5.0; // 5s between major tick marks
+
+		TopOffset = 0.0f;
+		ScrollHeight = 1.0f;
+		ScrollPosY = 0.0f;
+	}
+
+	float TimeToSlateUnits(const double Time) const
+	{
+		return static_cast<float>((Time - StartTime) * ScaleX);
+	}
+
+	float TimeToSlateUnitsRounded(const double Time) const
+	{
+		return static_cast<float>(FMath::RoundToDouble((Time - StartTime) * ScaleX));
+	}
+
+	float GetViewportDXForDuration(const double DT) const
+	{
+		return static_cast<float>(DT * ScaleX);
+	}
+
+	double GetDurationForViewportDX(const float DX) const
+	{
+		return static_cast<double>(DX) / ScaleX;
+	}
+
+	double SlateUnitsToTime(const float X) const
+	{
+		return StartTime + static_cast<double>(X) / ScaleX;
+	}
+
+	bool UpdateSize(const float InWidth, const float InHeight)
+	{
+		if (Width != InWidth || Height != InHeight)
+		{
+			Width = InWidth;
+			Height = InHeight;
+			EndTime = SlateUnitsToTime(Width);
+			return true;
+		}
+		return false;
+	}
+
+	double AlignTimeToPixel(const double InTime, const double InScaleX) const
+	{
+		return FMath::RoundToDouble(InTime * InScaleX) / InScaleX;
+	}
+
+	double AlignTimeToPixel(const double Time) const
+	{
+		return AlignTimeToPixel(Time, ScaleX);
+	}
+
+	bool ScrollAtTime(const double Time);
+	bool CenterOnTimeInterval(const double Time, const double Duration);
+
+	bool ZoomOnTimeInterval(const double Time, const double Duration);
+	bool ZoomWithFixedX(const double NewScaleX, const float X);
+
+	double RestrictEndTime(const double InEndTime) const;
+	double RestrictDuration(const double InStartTime, const double InEndTime) const;
+
+	void GetHorizontalScrollLimits(double& OutMinT, double& OutMaxT);
+	bool EnforceHorizontalScrollLimits(const double U);
+
+	float GetViewportY(const float Y) const { return TopOffset + Y - ScrollPosY; }
+};
