Copy up of Unreal Insights from //Tasks/UE4/Dev-Core-Trace.

diff --git a/Engine/Source/Developer/TraceInsights/Private/Insights/ViewModels/TimingViewDrawHelper.cpp b/Engine/Source/Developer/TraceInsights/Private/Insights/ViewModels/TimingViewDrawHelper.cpp
new file mode 100644
--- /dev/null
+++ b/Engine/Source/Developer/TraceInsights/Private/Insights/ViewModels/TimingViewDrawHelper.cpp
@@ -0,0 +1,600 @@
+// Copyright 1998-2019 Epic Games, Inc. All Rights Reserved.
+
+#include "TimingViewDrawHelper.h"
+
+#include "Brushes/SlateBorderBrush.h"
+#include "Brushes/SlateColorBrush.h"
+#include "EditorStyleSet.h"
+#include "Fonts/FontMeasure.h"
+#include "Fonts/SlateFontInfo.h"
+#include "Framework/Application/SlateApplication.h"
+#include "Rendering/DrawElements.h"
+#include "Styling/CoreStyle.h"
+#include "TraceServices/AnalysisService.h"
+
+// Insights
+#include "Insights/Common/PaintUtils.h"
+#include "Insights/Common/TimeUtils.h"
+#include "Insights/ViewModels/MarkersTimingTrack.h"
+#include "Insights/ViewModels/TimingEventsTrack.h"
+#include "Insights/ViewModels/TimingTrackViewport.h"
+
+////////////////////////////////////////////////////////////////////////////////////////////////////
+
+static const int32 EventBorderLayer = 0;
+static const int32 EventFillLayer = 1;
+static const int32 EventTextLayer = 2;
+
+static const int32 TimelineHeaderLayer = 3;
+static const int32 TimelineTextLayer = 4;
+
+static const int32 NumLayers = 5;
+
+////////////////////////////////////////////////////////////////////////////////////////////////////
+
+FTimingViewDrawHelper::FTimingViewDrawHelper(const FDrawContext& InDC, const FTimingTrackViewport& InViewport, const FTimingEventsTrackLayout& InLayout)
+	: DC(InDC)
+	, Viewport(InViewport)
+	, Layout(InLayout)
+	, WhiteBrush(FCoreStyle::Get().GetBrush("WhiteBrush"))
+	, BorderBrush(FEditorStyle::GetBrush("PlainBorder"))
+	, EventsBorderBrush(new FSlateBorderBrush(NAME_None, FMargin(1.0f)))
+	, BackgroundAreaBrush(WhiteBrush)
+	, ValidAreaColor(0.07f, 0.07f, 0.07f, 1.0f)
+	, InvalidAreaColor(0.1f, 0.07f, 0.07f, 1.0f)
+	, EventFont(FCoreStyle::GetDefaultFontStyle("Regular", 8))
+	, NumEvents(0)
+	, NumDrawBoxes(0)
+	, NumMergedBoxes(0)
+	, NumDrawBorders(0)
+	, NumDrawTexts(0)
+{
+}
+
+////////////////////////////////////////////////////////////////////////////////////////////////////
+
+FTimingViewDrawHelper::~FTimingViewDrawHelper()
+{
+	delete EventsBorderBrush; //im: is it safe to delete brush (i.e in OnPaint, imediately after use)?
+}
+
+////////////////////////////////////////////////////////////////////////////////////////////////////
+
+void FTimingViewDrawHelper::Begin()
+{
+	//NumEvents = 0;
+	//NumDrawBoxes = 0;
+	//NumDrawBorders = 0;
+	//NumDrawTexts = 0;
+
+	//NumDrawTimeMarkerBoxes = 0;
+	//NumDrawTimeMarkerTexts = 0;
+}
+
+////////////////////////////////////////////////////////////////////////////////////////////////////
+
+void FTimingViewDrawHelper::DrawBackground() const
+{
+	float X0 = Viewport.TimeToSlateUnitsRounded(0.0);
+	float X1 = Viewport.TimeToSlateUnitsRounded(Viewport.MaxValidTime);
+	const float W = FMath::CeilToFloat(Viewport.Width);
+	const float H = FMath::CeilToFloat(Viewport.Height);
+
+	if (X0 >= W || X1 <= 0.0f)
+	{
+		ValidX0 = 0.0f;
+		ValidX1 = W;
+
+		// Draw invalid area (entire view).
+		DC.DrawBox(0.0f, 0.0f, W, H, BackgroundAreaBrush, InvalidAreaColor);
+	}
+	else // X0 < W && X1 > 0
+	{
+		if (X0 > 0.0f)
+		{
+			// Draw invalid area (left).
+			DC.DrawBox(0.0f, 0.0f, X0, H, BackgroundAreaBrush, InvalidAreaColor);
+		}
+
+		if (X1 < W)
+		{
+			// Draw invalid area (right).
+			DC.DrawBox(X1, 0.0f, W - X1, H, BackgroundAreaBrush, InvalidAreaColor);
+		}
+
+		ValidX0 = FMath::Max(X0, 0.0f);
+		ValidX1 = FMath::Min(X1, W);
+
+		if (ValidX1 > ValidX0)
+		{
+			// Draw valid area.
+			DC.DrawBox(ValidX0, 0.0f, ValidX1 - ValidX0, H, BackgroundAreaBrush, ValidAreaColor);
+		}
+	}
+
+	DC.LayerId++;
+}
+
+////////////////////////////////////////////////////////////////////////////////////////////////////
+
+void FTimingViewDrawHelper::BeginTimelines()
+{
+	TimelineIndex = -1;
+	TimelineTopY = -Viewport.ScrollPosY;
+
+	float Y = TimelineTopY + Viewport.TopOffset;
+
+	if (Y > 0.0f && ValidX1 > ValidX0)
+	{
+		Y = FMath::Min(Y, Viewport.Height);
+
+		// Draw invalid area (top).
+		DC.DrawBox(DC.LayerId + TimelineHeaderLayer, ValidX0, 0.0f, ValidX1 - ValidX0, Y, BackgroundAreaBrush, InvalidAreaColor);
+	}
+}
+
+////////////////////////////////////////////////////////////////////////////////////////////////////
+
+bool FTimingViewDrawHelper::BeginTimeline(FTimingEventsTrack& Track)
+{
+	TimelineIndex++;
+
+	Track.Y = TimelineTopY + Viewport.ScrollPosY;
+
+	if (Track.H < Layout.MinTimelineH)
+		Track.H = Layout.MinTimelineH;
+
+	if (TimelineTopY + Track.H < -1.0f)
+	{
+		TimelineTopY += Track.H;
+		return false;
+	}
+
+	if (TimelineTopY > Viewport.Height - Viewport.TopOffset)
+	{
+		return false;
+	}
+
+	MaxDepth = -1;
+
+	// +1.0f is for horizontal line between timelines
+	TimelineY = Viewport.TopOffset + TimelineTopY + 1.0f + Layout.TimelineDY;
+
+	// Reset "last event X2" for all depths.
+	LastEventX2.Reset();
+
+	// Reset "last box" for all depths.
+	LastBox.Reset();
+
+	return true;
+}
+
+////////////////////////////////////////////////////////////////////////////////////////////////////
+
+void FTimingViewDrawHelper::EndTimeline(FTimingEventsTrack& Track, const TCHAR* TimelineName)
+{
+	// Flush merged boxes.
+	for (int32 Depth = 0; Depth <= MaxDepth; ++Depth)
+	{
+		const FBoxData& Box = LastBox[Depth];
+		if (Box.X1 < Box.X2)
+		{
+			const float EventY = TimelineY + (Layout.EventH + Layout.EventDY) * Depth;
+			DrawBox(Box, EventY, Layout.EventH);
+		}
+	}
+
+	Track.Depth = MaxDepth;
+
+	float NewH;
+	if (MaxDepth < 0)
+	{
+		NewH = Layout.MinTimelineH;
+	}
+	else //if (MaxDepth >= 0)
+	{
+		// 1.0f is for horizontal line between timelines
+		NewH = 1.0f + Layout.TimelineDY + Layout.EventH * (MaxDepth + 1) + Layout.EventDY * MaxDepth + Layout.TimelineDY;
+
+		if (NewH < RealMinTimelineH)
+			NewH = RealMinTimelineH;
+	}
+
+	if (Track.H < NewH)
+	{
+		if (Layout.bIsCompactMode)
+		{
+			Track.H = FMath::CeilToFloat(Track.H * 0.5f + NewH * 0.5f);
+		}
+		else
+		{
+			Track.H = FMath::CeilToFloat(Track.H * 0.9f + NewH * 0.1f);
+		}
+	}
+	else if (Track.H > NewH)
+	{
+		if (Layout.bIsCompactMode)
+		{
+			Track.H = FMath::FloorToFloat(Track.H * 0.5f + NewH * 0.5f);
+		}
+		else
+		{
+			Track.H = FMath::FloorToFloat(Track.H * 0.9f + NewH * 0.1f);
+		}
+	}
+
+	if (Track.H > 0)
+	{
+		const float Y = Viewport.TopOffset + TimelineTopY;
+
+		FLinearColor Color(0.05f, 0.05f, 0.05f, 1.0f);
+		FLinearColor TextColor(1.0f, 1.0f, 1.0f, 1.0f);
+
+		if (Track.IsSelected())
+		{
+			if (Track.IsHovered())
+			{
+				TextColor = FLinearColor(1.0f, 1.0f, 0.0f, 1.0f);
+			}
+			else
+			{
+				TextColor = FLinearColor(1.0f, 1.0f, 0.5f, 1.0f);
+			}
+		}
+		else if (Track.IsHovered())
+		{
+			TextColor = FLinearColor(1.0f, 1.0f, 0.0f, 1.0f);
+		}
+
+		// Draw a horizontal line between timelines.
+		DC.DrawBox(DC.LayerId + TimelineHeaderLayer, 0, Y, Viewport.Width, 1.0f, WhiteBrush, Color);
+
+		// Draw name of timeline.
+		//const FString Name = FString::Printf(TEXT("%d Y: %g, H: %g, VO: %g, VY: %g"), TimelineIndex, Track.Y, Track.H, Viewport.TopOffset, Viewport.ScrollPosY); //debug
+		const TSharedRef<FSlateFontMeasure> FontMeasureService = FSlateApplication::Get().GetRenderer()->GetFontMeasureService();
+		float NameWidth = FontMeasureService->Measure(TimelineName, EventFont).X;
+		DC.DrawBox(DC.LayerId + TimelineHeaderLayer, 0.0f, Y + 1.0f, NameWidth + 4.0f, 12.0f, WhiteBrush, Color);
+		DC.DrawText(DC.LayerId + TimelineTextLayer, 2.0f, Y, TimelineName, EventFont, TextColor);
+	}
+
+	TimelineTopY += Track.H;
+}
+
+////////////////////////////////////////////////////////////////////////////////////////////////////
+
+void FTimingViewDrawHelper::AddEvent(double EventStartTime, double EventEndTime, uint32 EventDepth, const TCHAR* EventName, uint32 Color)
+{
+	NumEvents++;
+
+	float EventX1 = Viewport.TimeToSlateUnitsRounded(EventStartTime);
+	if (EventX1 > Viewport.Width)
+	{
+		return;
+	}
+
+	double RestrictedEndTime = Viewport.RestrictEndTime(EventEndTime);
+	float EventX2 = Viewport.TimeToSlateUnitsRounded(RestrictedEndTime);
+	if (EventX2 < 0)
+	{
+		return;
+	}
+
+	// Timing events are displayed with minimum 1px (including empty ones).
+	if (EventX1 == EventX2)
+	{
+		EventX2 = EventX1 + 1.0f;
+	}
+
+	const int32 Depth = static_cast<int32>(EventDepth);
+	if (Depth > MaxDepth)
+	{
+		MaxDepth = Depth;
+	}
+
+	// Ensure we have enough slots in array. See LastBox[Depth] usage.
+	if (LastBox.Num() <= Depth)
+	{
+		LastBox.AddDefaulted(Depth + 1 - LastBox.Num());
+	}
+
+	float EventY = TimelineY + (Layout.EventH + Layout.EventDY) * Depth;
+	if (EventY < -Layout.EventH || EventY > Viewport.Height)
+	{
+		return;
+	}
+
+	// Ensure we have enough slots in array. See LastEventX2[Depth] usage.
+	while (LastEventX2.Num() <= Depth)
+	{
+		LastEventX2.Add(-1.0f);
+	}
+
+	// Limit event width on the viewport's left side.
+	// This also makes the text to be displayed in viewport, for very long events.
+	const float MinX = -2.0f; // -2 allows event border to remain outside screen
+	if (EventX1 < MinX)
+	{
+		EventX1 = MinX;
+	}
+
+	// Limit event width on the viewport's right side.
+	const float MaxX = Viewport.Width + 2.0f; // +2 allows event border to remain outside screen
+	if (EventX2 > MaxX)
+	{
+		EventX2 = MaxX;
+	}
+
+	float EventW = EventX2 - EventX1;
+
+	// Optimization...
+	if (EventW == 1.0f && EventX1 == LastEventX2[Depth] - 1.0f)
+	{
+		// Do no draw 1 pixel event if the last event was ended on that pixel.
+		return;
+	}
+
+	//////////////////////////////////////////////////
+	// Coloring
+
+	FLinearColor EventColorFill;
+
+	//im:TODO: EventColorFill = GetEventColorFn(Event);
+	if (Color == 0)
+	{
+		uint32 NameHash = 0;
+		for (const TCHAR* c = EventName; *c; ++c)
+		{
+			NameHash = (NameHash + *c) * 0x2c2c57ed;
+		}
+
+		Color = NameHash | 0xFF000000;
+
+		// Increase brightness.
+		//Color = ((NameHash | ((NameHash & 0x00808080) >> 1) | ((NameHash & 0x00808080) >> 2) | ((NameHash & 0x00808080) >> 3)) << 1) | 0xFF000000;
+
+		// Increase brightness.
+		//EventColorFill.R = ((NameHash >> 16) & 0xFF) / 128.0f;
+		//EventColorFill.G = ((NameHash >> 8) & 0xFF) / 128.0f;
+		//EventColorFill.B = (NameHash & 0xFF) / 128.0f;
+		//EventColorFill.A = 1.0f;
+	}
+	//else
+	{
+		EventColorFill.R = ((Color >> 16) & 0xFF) / 255.0f;
+		EventColorFill.G = ((Color >>  8) & 0xFF) / 255.0f;
+		EventColorFill.B = ((Color      ) & 0xFF) / 255.0f;
+		EventColorFill.A = ((Color >> 24) & 0xFF) / 255.0f;
+	}
+
+	// Maps hash to color directly.
+	//EventColorFill.R = ((Color >> 16) & 0xFF) / 255.0f;
+	//EventColorFill.G = ((Color >> 8) & 0xFF) / 255.0f;
+	//EventColorFill.B = (Color & 0xFF) / 255.0f;
+	//EventColorFill.A = 1.0f;
+
+	// Maps hash to color (sRGB to linear).
+	//const FColor ColorFill(Color);
+	//EventColorFill = ColorFill;
+
+	// Invereted color.
+	//EventColorFill.R = 1.0f - ((Color >> 16) & 0xFF) / 255.0f;
+	//EventColorFill.G = 1.0f - ((Color >> 8) & 0xFF) / 255.0f;
+	//EventColorFill.B = 1.0f - (Color & 0xFF) / 255.0f;
+	//EventColorFill.A = 1.0f;
+
+	// Bright color.
+	//EventColorFill.R = ((Color >> 16) & 0xFF) / 128.0f;
+	//EventColorFill.G = ((Color >> 8) & 0xFF) / 128.0f;
+	//EventColorFill.B = (Color & 0xFF) / 128.0f;
+	//EventColorFill.A = 1.0f;
+
+	// Bright color (adjusted R, G, B).
+	//const float BrightnessFactor = 0.59f / 255.0f; // [0 .. 0.59f / 255.0f]
+	//EventColorFill.R = ((Color >> 16) & 0xFF) * (BrightnessFactor / 0.30f);
+	//EventColorFill.G = ((Color >>  8) & 0xFF) * (BrightnessFactor / 0.59f);
+	//EventColorFill.B = ((Color      ) & 0xFF) * (BrightnessFactor / 0.11f);
+	//EventColorFill.A = 1.0f;
+
+	// Bright desaturated color.
+	//const float MinL = 0.3f;
+	//const float MulL = (1.0f - MinL) / 255.0f;
+	//EventColorFill.R = MinL + ((Color >> 16) & 0xFF) * MulL;
+	//EventColorFill.G = MinL + ((Color >> 8) & 0xFF) * MulL;
+	//EventColorFill.B = MinL + (Color & 0xFF) * MulL;
+	//EventColorFill.A = 1.0f;
+
+	// Dark desaturated color.
+	//const float MaxL = 0.5f;
+	//const float MulL = MaxL / 255.0f;
+	//EventColorFill.R = (((Color >> 16) & 0xFF) * MulL;
+	//EventColorFill.G = ((Color >> 8) & 0xFF) * MulL;
+	//EventColorFill.B = (Color & 0xFF) * MulL;
+	//EventColorFill.A = 1.0f;
+
+	//static const float BorderColorFactor = 0.75f; // darker border
+	static const float BorderColorFactor = 1.25f; // brighter border
+
+	//////////////////////////////////////////////////
+
+	if (EventW > 2.0f && Layout.EventH > 2.0f)
+	{
+		// Fill inside of the timing event box.
+		DC.DrawBox(DC.LayerId + EventFillLayer, EventX1 + 1.0f, EventY + 1.0f, EventW - 2.0f, Layout.EventH - 2.0f, WhiteBrush, EventColorFill);
+		NumDrawBoxes++;
+	}
+
+	// Save X2, for current depth.
+	LastEventX2[Depth] = EventX2;
+
+	// Draw border around the timing event box.
+	if (EventW > 2.0f)
+	{
+		FBoxData& Box = LastBox[Depth];
+		if (Box.X1 < Box.X2)
+		{
+			DrawBox(Box, EventY, Layout.EventH);
+			Box.Reset();
+		}
+
+		const FLinearColor EventColorBorder(EventColorFill.R * BorderColorFactor, EventColorFill.G * BorderColorFactor, EventColorFill.B * BorderColorFactor, 1.0f);
+		DC.DrawBox(DC.LayerId + EventBorderLayer, EventX1, EventY, EventW, Layout.EventH, EventsBorderBrush, EventColorBorder);
+		NumDrawBorders++;
+	}
+	else // 1px or 2px boxes
+	{
+		FBoxData& Box = LastBox[Depth];
+
+		// Check if we can merge this box with previous one, if any.
+		// Note: We are assuming events are processed in sorted order by X1.
+		if (Color == Box.Color && // same color
+			EventX1 <= Box.X2) // overlapping or adjacent
+		{
+			// Merge it with previous box.
+			Box.X2 = EventX2;
+			NumMergedBoxes++;
+		}
+		else
+		{
+			// Flush previous box, if any.
+			if (Box.X1 < Box.X2)
+			{
+				DrawBox(Box, EventY, Layout.EventH);
+			}
+
+			// Start new "merge box".
+			Box.X1 = EventX1;
+			Box.X2 = EventX2;
+			Box.Color = Color;
+			Box.LinearColor = FLinearColor(EventColorFill.R * BorderColorFactor, EventColorFill.G * BorderColorFactor, EventColorFill.B * BorderColorFactor, 1.0f);
+		}
+	}
+
+	// Draw the name of the timing event.
+	if (EventW > 8.0f && Layout.EventH > 10.0f)
+	{
+		FString Name = EventName;// +TEXT(" [") + FText::AsNumber(Event.Type->Id).ToString() + TEXT("]");
+		if (EventW > Name.Len() * 2.0f + 48.0f)
+		{
+			const double Duration = EventEndTime - EventStartTime;
+			Name += TEXT(" (");
+			Name += TimeUtils::FormatTimeAuto(Duration);
+			Name += TEXT(")");
+		}
+
+		const TSharedRef<FSlateFontMeasure> FontMeasureService = FSlateApplication::Get().GetRenderer()->GetFontMeasureService();
+		const int32 LastWholeCharacterIndex = FontMeasureService->FindLastWholeCharacterIndexBeforeOffset(Name, EventFont, FMath::RoundToInt(EventW - 2.0f));
+
+		if (LastWholeCharacterIndex >= 0)
+		{
+			// Grey threshold is shifted toward black (0.4 instead of 0.5 in test below) due to "area rule":
+			// a large gray surface (background of a timing event in this case) is perceived lighter than a smaller area (text pixels).
+			// Ref: https://books.google.ro/books?id=0pVr7dhmdWYC
+			const bool bIsDarkColor = (EventColorFill.ComputeLuminance() < 0.4f);
+
+			DC.DrawText(DC.LayerId + EventTextLayer, EventX1 + 2.0f, EventY + 1.0f, Name, 0, LastWholeCharacterIndex + 1, EventFont, bIsDarkColor ? FLinearColor::White : FLinearColor::Black);
+			NumDrawTexts++;
+		}
+	}
+}
+
+////////////////////////////////////////////////////////////////////////////////////////////////////
+
+void FTimingViewDrawHelper::DrawBox(const FBoxData& Box, const float EventY, const float EventH)
+{
+	DC.DrawBox(DC.LayerId + EventBorderLayer, Box.X1, EventY, Box.X2 - Box.X1, EventH, WhiteBrush, Box.LinearColor);
+	NumDrawBoxes++;
+}
+
+////////////////////////////////////////////////////////////////////////////////////////////////////
+
+void FTimingViewDrawHelper::EndTimelines()
+{
+	float Y = Viewport.TopOffset + TimelineTopY;
+
+	// Draw a last horizontal line.
+	DC.DrawBox(DC.LayerId + TimelineHeaderLayer, 0, Y, Viewport.Width, 1.0f, WhiteBrush, FLinearColor(0.05f, 0.05f, 0.05f, 1.0f));
+	Y += 1.0f;
+
+	if (Y < Viewport.Height && ValidX1 > ValidX0)
+	{
+		Y = FMath::Max(Y, 0.0f);
+
+		// Draw invalid area (bottom).
+		DC.DrawBox(DC.LayerId + TimelineHeaderLayer, ValidX0, Y, ValidX1 - ValidX0, Viewport.Height - Y, BackgroundAreaBrush, InvalidAreaColor);
+	}
+
+	DC.LayerId += NumLayers;
+}
+
+////////////////////////////////////////////////////////////////////////////////////////////////////
+
+void FTimingViewDrawHelper::DrawTimingEventHighlight(double StartTime, double EndTime, float Y, EHighlightMode Mode)
+{
+	float EventX1 = Viewport.TimeToSlateUnitsRounded(StartTime);
+	if (EventX1 > Viewport.Width)
+	{
+		return;
+	}
+
+	EndTime = Viewport.RestrictEndTime(EndTime);
+	float EventX2 = Viewport.TimeToSlateUnitsRounded(EndTime);
+	if (EventX2 < 0)
+	{
+		return;
+	}
+
+	if (EventX1 == EventX2)
+		EventX2 = EventX1 + 1.0f;
+
+	// Limit event width on the viewport's left side.
+	const float MinX = -2.0f; // -2 allows event border to remain outside screen
+	if (EventX1 < MinX)
+	{
+		EventX1 = MinX;
+	}
+
+	// Limit event width on the viewport's right side.
+	const float MaxX = Viewport.Width + 2.0f; // +2 allows event border to remain outside screen
+	if (EventX2 > MaxX)
+	{
+		EventX2 = MaxX;
+	}
+
+	float EventW = EventX2 - EventX1;
+
+	if (Mode == EHighlightMode::Hovered)
+	{
+		const FSlateBrush* HighlightBrush = BorderBrush;
+
+		const FLinearColor Color(1.0f, 1.0f, 0.0f, 1.0f);
+
+		// Draw border around the timing event box.
+		DC.DrawBox(EventX1 - 2.0f, Y - 2.0f, EventW + 4.0f, Layout.EventH + 4.0f, HighlightBrush, Color);
+	}
+	else // EHighlightMode::Selected or EHighlightMode::SelectedAndHovered
+	{
+		const FSlateBrush* SelectedBrush = BorderBrush;
+
+		const double Time = static_cast<double>(FPlatformTime::Cycles64()) * FPlatformTime::GetSecondsPerCycle64();
+		float S = FMath::Sin(2.0 * Time);
+		S = S * S;
+		const float B = (Mode == EHighlightMode::SelectedAndHovered) ? 0.0f : S;
+		const FLinearColor Color(S, S, B, 1.0f);
+
+		// Draw border around the timing event box.
+		DC.DrawBox(EventX1 - 2.0f, Y - 2.0f, EventW + 4.0f, Layout.EventH + 4.0f, SelectedBrush, Color);
+
+		//const FLinearColor Color2(1.0f - S, 1.0f - S, 1.0f - S, 1.0f);
+		//const FSlateBrush* SelectedBrush2 = FEditorStyle::GetBrush("LevelViewport.ActiveViewportBorder");
+		//DC.DrawBox(EventX1 - 2.0f, Y - 2.0f, EventW + 4.0f, Layout.EventH + 4.0f, SelectedBrush2, Color2);
+	}
+	DC.LayerId++;
+}
+
+////////////////////////////////////////////////////////////////////////////////////////////////////
+
+void FTimingViewDrawHelper::End()
+{
+}
+
+////////////////////////////////////////////////////////////////////////////////////////////////////
diff --git a/Engine/Source/Developer/TraceInsights/Private/Insights/ViewModels/TimingViewDrawHelper.h b/Engine/Source/Developer/TraceInsights/Private/Insights/ViewModels/TimingViewDrawHelper.h
new file mode 100644
--- /dev/null
+++ b/Engine/Source/Developer/TraceInsights/Private/Insights/ViewModels/TimingViewDrawHelper.h
@@ -0,0 +1,119 @@
+// Copyright 1998-2019 Epic Games, Inc. All Rights Reserved.
+
+#pragma once
+
+#include "CoreMinimal.h"
+#include "Fonts/SlateFontInfo.h"
+#include "Math/Color.h"
+
+////////////////////////////////////////////////////////////////////////////////////////////////////
+
+struct FSlateBrush;
+struct FDrawContext;
+
+class FTimingEventsTrack;
+struct FTimingEventsTrackLayout;
+class FTimingTrackViewport;
+
+////////////////////////////////////////////////////////////////////////////////////////////////////
+
+class FTimingViewDrawHelper : public FNoncopyable
+{
+public:
+	enum EHighlightMode
+	{
+		Hovered,
+		Selected,
+		SelectedAndHovered
+	};
+
+private:
+	struct FBoxData
+	{
+		float X1;
+		float X2;
+		uint32 Color;
+		FLinearColor LinearColor;
+
+		FBoxData() : X1(0.0f), X2(0.0f), Color(0) {}
+		void Reset() { X1 = 0.0f; X2 = 0.0f; Color = 0; }
+	};
+
+	//struct FEventBoxInfo
+	//{
+	//	float X1;
+	//	float X2;
+	//	int32 Depth;
+	//	uint32 Color;
+	//}
+
+	//struct FEventTextInfo
+	//{
+	//	float X;
+	//	float Y;
+	//	FString Text;
+	//	bool bUseDarkTextColor; // true if text needs to be displayed in Black, otherwise will be displayed in White
+	//};
+
+public:
+	FTimingViewDrawHelper(const FDrawContext& DC, const FTimingTrackViewport& InViewport, const FTimingEventsTrackLayout& Layout);
+	~FTimingViewDrawHelper();
+
+	const FTimingTrackViewport& GetViewport() { return Viewport; }
+
+	void Begin();
+
+	void DrawBackground() const;
+
+	//TODO: move the following in a Builder class
+	void BeginTimelines();
+	bool BeginTimeline(FTimingEventsTrack& Track);
+	void AddEvent(double StartTime, double EndTime, uint32 Depth, const TCHAR* EventName, uint32 Color = 0);
+	void EndTimeline(FTimingEventsTrack& Track, const TCHAR* TimelineName);
+	void EndTimelines();
+
+	void DrawTimingEventHighlight(double StartTime, double EndTime, float Y, EHighlightMode Mode);
+
+	void End();
+
+private:
+	void DrawBox(const FBoxData& Box, const float EventY, const float EventH);
+
+public:
+	const FDrawContext& DC;
+	const FTimingTrackViewport& Viewport;
+	const FTimingEventsTrackLayout& Layout;
+
+	const FSlateBrush* WhiteBrush;
+	const FSlateBrush* BorderBrush;
+	const FSlateBrush* EventsBorderBrush;
+	const FSlateBrush* BackgroundAreaBrush;
+	const FLinearColor ValidAreaColor;
+	const FLinearColor InvalidAreaColor;
+	const FSlateFontInfo EventFont;
+
+	//static const FSlateColorBrush SolidWhiteBrush = FSlateColorBrush(FColorList::White);
+	//static const FSlateBrush BorderBrush = FSlateBorderBrush(NAME_None, FMargin(1.0f));
+
+	mutable float ValidX0;
+	mutable float ValidX1;
+
+	float TimelineTopY;
+	float TimelineY;
+	int32 MaxDepth;
+	int TimelineIndex;
+
+	TArray<float> LastEventX2; // X2 value for last event on each depth, for current timeline
+	TArray<FBoxData> LastBox;
+
+	// Debug stats.
+	mutable int NumEvents;
+	mutable int NumDrawBoxes;
+	mutable int NumMergedBoxes;
+	mutable int NumDrawBorders;
+	mutable int NumDrawTexts;
+	mutable int NumDrawTimeMarkerBoxes;
+	mutable int NumDrawTimeMarkerTexts;
+};
+
+////////////////////////////////////////////////////////////////////////////////////////////////////
diff --git a/Engine/Source/Developer/TraceInsights/Private/Insights/Widgets/SFrameTrack.cpp b/Engine/Source/Developer/TraceInsights/Private/Insights/Widgets/SFrameTrack.cpp
new file mode 100644
--- /dev/null
+++ b/Engine/Source/Developer/TraceInsights/Private/Insights/Widgets/SFrameTrack.cpp
@@ -0,0 +1,887 @@
+// Copyright 1998-2019 Epic Games, Inc. All Rights Reserved.
+
+#include "SFrameTrack.h"
+
+#include "Brushes/SlateBorderBrush.h"
+#include "Brushes/SlateColorBrush.h"
+#include "EditorStyleSet.h"
+#include "Fonts/FontMeasure.h"
+#include "Fonts/SlateFontInfo.h"
+#include "Framework/Application/SlateApplication.h"
+#include "HAL/PlatformTime.h"
+#include "Rendering/DrawElements.h"
+#include "TraceServices/AnalysisService.h"
+#include "TraceServices/SessionService.h"
+#include "Widgets/Layout/SScrollBar.h"
+
+// Insights
+#include "Insights/Common/PaintUtils.h"
+#include "Insights/Common/Stopwatch.h"
+#include "Insights/Common/TimeUtils.h"
+#include "Insights/InsightsManager.h"
+#include "Insights/TimingProfilerCommon.h"
+#include "Insights/TimingProfilerManager.h"
+#include "Insights/ViewModels/FrameTrackHelper.h"
+#include "Insights/Widgets/STimingProfilerWindow.h"
+#include "Insights/Widgets/STimingView.h"
+
+#include <limits>
+
+////////////////////////////////////////////////////////////////////////////////////////////////////
+
+#define LOCTEXT_NAMESPACE "SFrameTrack"
+
+////////////////////////////////////////////////////////////////////////////////////////////////////
+
+SFrameTrack::SFrameTrack()
+	: HoveredSample(nullptr, nullptr)
+{
+	Reset();
+}
+
+////////////////////////////////////////////////////////////////////////////////////////////////////
+
+SFrameTrack::~SFrameTrack()
+{
+}
+
+////////////////////////////////////////////////////////////////////////////////////////////////////
+
+void SFrameTrack::Reset()
+{
+	Viewport.Reset();
+	bIsViewportDirty = true;
+
+	CachedTimelines.Reset();
+	TimelinesOrder.Reset();
+	TimelinesOrder.Add(0);
+	TimelinesOrder.Add(1);
+	TimelinesOrder.Add(2);
+	bIsStateDirty = true;
+
+	MousePosition = FVector2D::ZeroVector;
+
+	MousePositionOnButtonDown = FVector2D::ZeroVector;
+	ViewportPosXOnButtonDown = 0.0f;
+	ViewportPosYOnButtonDown = 0.0f;
+
+	MousePositionOnButtonUp = FVector2D::ZeroVector;
+
+	bIsLMB_Pressed = false;
+	bIsRMB_Pressed = false;
+
+	bIsRMB_Scrolling = false;
+
+	SelectionStartFrameIndex = 0;
+	SelectionEndFrameIndex = 0;
+
+	HoveredSample.Timeline = nullptr;
+	HoveredSample.Sample = nullptr;
+
+	//ThisGeometry
+
+	CursorType = EFrameTrackCursor::Default;
+
+	NumUpdatedFrames = 0;
+	UpdateDurationHistory.Reset();
+	DrawDurationHistory.Reset();
+	OnPaintDurationHistory.Reset();
+	LastOnPaintTime = FPlatformTime::Cycles64();
+}
+
+////////////////////////////////////////////////////////////////////////////////////////////////////
+
+void SFrameTrack::Construct(const FArguments& InArgs)
+{
+	ChildSlot
+	[
+		SNew(SOverlay)
+		.Visibility(EVisibility::SelfHitTestInvisible)
+
+		+ SOverlay::Slot()
+		.VAlign(VAlign_Top)
+		.Padding(FMargin(0, 0, 0, 0))
+		[
+			SAssignNew(HorizontalScrollBar, SScrollBar)
+			.Orientation(Orient_Horizontal)
+			.AlwaysShowScrollbar(false)
+			.Visibility(EVisibility::Visible)
+			.Thickness(FVector2D(4.0f, 4.0f))
+			.RenderOpacity(0.5)
+			.OnUserScrolled(this, &SFrameTrack::HorizontalScrollBar_OnUserScrolled)
+		]
+	];
+
+	UpdateHorizontalScrollBar();
+
+	BindCommands();
+}
+
+////////////////////////////////////////////////////////////////////////////////////////////////////
+
+void SFrameTrack::Tick(const FGeometry& AllottedGeometry, const double InCurrentTime, const float InDeltaTime)
+{
+	if (ThisGeometry != AllottedGeometry || bIsViewportDirty)
+	{
+		bIsViewportDirty = false;
+		const float ViewWidth = AllottedGeometry.GetLocalSize().X;
+		const float ViewHeight = AllottedGeometry.GetLocalSize().Y;
+		Viewport.SetSize(ViewWidth, ViewHeight);
+		bIsStateDirty = true;
+	}
+
+	ThisGeometry = AllottedGeometry;
+
+	if (!bIsRMB_Scrolling)
+	{
+		// Elastic snap to horizontal limits (+15%).
+
+		float MinX, MaxX;
+		if (Viewport.MaxX - Viewport.MinX < Viewport.Width)
+		{
+			MinX = Viewport.MaxX - Viewport.Width;
+			MaxX = Viewport.MinX;
+		}
+		else
+		{
+			MinX = Viewport.MinX - 0.15f * Viewport.Width;
+			MaxX = Viewport.MaxX - 0.85f * Viewport.Width;
+		}
+		const float U = 0.5f;
+
+		float PosX = Viewport.PosX;
+		if (PosX < MinX)
+		{
+			PosX = PosX * U + (1.0f - U) * MinX;
+			if (FMath::IsNearlyEqual(PosX, MinX, 1.0f))
+				PosX = MinX;
+		}
+		else if (PosX > MaxX)
+		{
+			PosX = PosX * U + (1.0f - U) * MaxX;
+			if (FMath::IsNearlyEqual(PosX, MaxX, 1.0f))
+				PosX = MaxX;
+			if (PosX < MinX)
+				PosX = MinX;
+		}
+
+		if (Viewport.PosX != PosX)
+		{
+			//Viewport.ScrollAtPosX(PosX);
+			Viewport.ScrollAtIndex(Viewport.GetIndexAtPosX(PosX));
+			bIsStateDirty = true;
+		}
+	}
+
+	TSharedPtr<const Trace::IAnalysisSession> Session = FInsightsManager::Get()->GetSession();
+	if (Session.IsValid())
+	{
+		Trace::FAnalysisSessionReadScope _(*Session.Get());
+		Session->ReadFramesProvider([this](const Trace::IFrameProvider& FramesProvider)
+		{
+			for (int32 FrameType = 0; FrameType < TraceFrameType_Count; ++FrameType)
+			{
+				FFrameTrackTimeline& CachedTimeline = CachedTimelines.FindOrAdd(FrameType);
+				CachedTimeline.Id = FrameType;
+
+				int32 NumFrames = FramesProvider.GetFrameCount(ETraceFrameType(FrameType));
+				if (NumFrames > Viewport.MaxIndex)
+				{
+					Viewport.SetMinMaxIndexInterval(0, NumFrames);
+					UpdateHorizontalScrollBar();
+					bIsStateDirty = true;
+
+					// Auto zoom out.
+					float SampleW = Viewport.GetSampleWidth();
+					int32 NumSamples = FMath::CeilToInt(Viewport.Width / SampleW);
+					if (NumFrames > NumSamples * Viewport.GetNumFramesPerSample())
+					{
+						ZoomHorizontally(-0.1f, 0.0f);
+						Viewport.ScrollAtPosX(0.0f);
+					}
+				}
+			}
+		});
+	}
+
+	if (bIsStateDirty)
+	{
+		bIsStateDirty = false;
+		UpdateState();
+	}
+}
+
+////////////////////////////////////////////////////////////////////////////////////////////////////
+
+void SFrameTrack::UpdateState()
+{
+	FStopwatch Stopwatch;
+	Stopwatch.Start();
+
+	for (auto It = CachedTimelines.CreateIterator(); It; ++It)
+	{
+		(*It).Value.NumAggregatedFrames = 0;
+	}
+
+	FFrameTrackCacheContext CC(Viewport);
+	CC.Begin();
+
+	TSharedPtr<const Trace::IAnalysisSession> Session = FInsightsManager::Get()->GetSession();
+	if (Session.IsValid())
+	{
+		Trace::FAnalysisSessionReadScope _(*Session.Get());
+		Session->ReadFramesProvider([this, &CC](const Trace::IFrameProvider& FramesProvider)
+		{
+			for (int32 FrameType = 0; FrameType < TraceFrameType_Count; ++FrameType)
+			{
+				FFrameTrackTimeline& CachedTimeline = CachedTimelines.FindOrAdd(FrameType);
+				CachedTimeline.Id = FrameType;
+
+				if (CC.BeginTimeline(CachedTimeline))
+				{
+					int32 StartIndex = FMath::Max(0, CC.Viewport.GetIndexAtViewportX(0.0f));
+					int32 EndIndex = CC.Viewport.GetIndexAtViewportX(Viewport.Width);
+					//int32 NumFrames = Timeline.GetEventCount();
+					//int32 EndIndex = FMath::Min(NumFrames, CC.Viewport.GetIndexAtViewportX(Viewport.Width));
+
+					FramesProvider.EnumerateFrames(ETraceFrameType(FrameType), static_cast<uint64>(StartIndex), static_cast<uint64>(EndIndex), [&CC](const Trace::FFrame& Frame)
+					{
+						CC.AddEvent(Frame);
+					});
+
+					CC.EndTimeline();
+				}
+			}
+		});
+	}
+
+	CC.End();
+
+	Stopwatch.Stop();
+	NumUpdatedFrames = CC.NumFrames;
+	UpdateDurationHistory.AddValue(Stopwatch.AccumulatedTime);
+}
+
+////////////////////////////////////////////////////////////////////////////////////////////////////
+
+FSampleRef SFrameTrack::GetSampleAtMousePosition(float X, float Y)
+{
+	float SampleW = Viewport.GetSampleWidth();
+	int SampleIndex = FMath::FloorToInt(X / SampleW);
+	if (SampleIndex >= 0)
+	{
+		const float ViewportHeight = FMath::RoundToFloat(Viewport.Height);
+		const float SampleY2 = ViewportHeight - FMath::RoundToFloat(Viewport.GetViewportYForValue(0.0));
+
+		// Search in reverse paint order.
+		for (int TimelineIndex = TimelinesOrder.Num() - 1; TimelineIndex >= 0; --TimelineIndex)
+		{
+			int32 TimelineId = TimelinesOrder[TimelineIndex];
+			if (CachedTimelines.Contains(TimelineId))
+			{
+				const FFrameTrackTimeline& CachedTimeline = CachedTimelines[TimelineId];
+				if (CachedTimeline.NumAggregatedFrames > 0 &&
+					SampleIndex < CachedTimeline.Samples.Num())
+				{
+					const FFrameTrackSample& Sample = CachedTimeline.Samples[SampleIndex];
+					if (Sample.NumFrames > 0)
+					{
+						float SampleY1;
+						if (Sample.LargestFrameDuration == std::numeric_limits<double>::infinity())
+						{
+							SampleY1 = 0.0;
+						}
+						else
+						{
+							SampleY1 = ViewportHeight - FMath::RoundToFloat(Viewport.GetViewportYForValue(Sample.LargestFrameDuration));
+						}
+						if (Y >= SampleY1 && Y <= SampleY2)
+						{
+							return FSampleRef(&CachedTimeline, &Sample);
+						}
+					}
+				}
+			}
+		}
+	}
+	return FSampleRef(nullptr, nullptr);
+}
+
+////////////////////////////////////////////////////////////////////////////////////////////////////
+
+void SFrameTrack::SelectFrameAtMousePosition(float X, float Y)
+{
+	FSampleRef SampleRef = GetSampleAtMousePosition(X, Y);
+	if (!SampleRef.IsValid())
+		SampleRef = GetSampleAtMousePosition(X - 1.0f, Y);
+	if (!SampleRef.IsValid())
+		SampleRef = GetSampleAtMousePosition(X + 1.0f, Y);
+
+	if (SampleRef.IsValid())
+	{
+		TSharedPtr<STimingProfilerWindow> Window = FTimingProfilerManager::Get()->GetProfilerWindow();
+		if (Window.IsValid())
+		{
+			const double StartTime = SampleRef.Sample->LargestFrameStartTime;
+			const double Duration = SampleRef.Sample->LargestFrameDuration;
+			Window->TimingView->CenterOnTimeInterval(StartTime, Duration);
+			Window->TimingView->SelectTimeInterval(StartTime, Duration);
+			FSlateApplication::Get().SetKeyboardFocus(Window->TimingView);
+		}
+	}
+}
+
+////////////////////////////////////////////////////////////////////////////////////////////////////
+
+int32 SFrameTrack::OnPaint(const FPaintArgs& Args, const FGeometry& AllottedGeometry, const FSlateRect& MyCullingRect, FSlateWindowElementList& OutDrawElements, int32 LayerId, const FWidgetStyle& InWidgetStyle, bool bParentEnabled) const
+{
+	const TSharedRef<FSlateFontMeasure> FontMeasureService = FSlateApplication::Get().GetRenderer()->GetFontMeasureService();
+
+	const bool bEnabled = ShouldBeEnabled(bParentEnabled);
+	const ESlateDrawEffect DrawEffects = bEnabled ? ESlateDrawEffect::None : ESlateDrawEffect::DisabledEffect;
+
+	const FSlateBrush* WhiteBrush = FCoreStyle::Get().GetBrush("WhiteBrush");
+
+	const float ViewWidth = AllottedGeometry.Size.X;
+	const float ViewHeight = AllottedGeometry.Size.Y;
+
+	FFrameTrackDrawContext DC(Viewport, AllottedGeometry, LayerId, OutDrawElements);
+
+	FStopwatch Stopwatch;
+	Stopwatch.Start();
+
+	DC.DrawBackground(InWidgetStyle);
+
+	// Draw frames.
+	{
+		DC.NumFrames = 0;
+		DC.NumDrawSamples = 0;
+
+		for (auto It = CachedTimelines.CreateConstIterator(); It; ++It)
+		{
+			DC.DrawCached((*It).Value);
+		}
+		//ensure(DC.NumFrames == NumUpdatedFrames);
+	}
+
+	// Highlight the mouse hovered sample.
+	if (HoveredSample.IsValid())
+	{
+		DC.DrawHoveredSample(*HoveredSample.Sample);
+
+		FString Text = FString::Format(TEXT("{0} frame {1} ({2})"),
+			{
+				(HoveredSample.Timeline->Id == 0) ? TEXT("Game") :
+				(HoveredSample.Timeline->Id == 1) ? TEXT("Render") : TEXT("Misc"),
+				FText::AsNumber(HoveredSample.Sample->LargestFrameIndex).ToString(),
+				TimeUtils::FormatTimeAuto(HoveredSample.Sample->LargestFrameDuration)
+			});
+
+		FSlateFontInfo SummaryFont = FCoreStyle::GetDefaultFontStyle("Regular", 8);
+		FVector2D TextSize = FontMeasureService->Measure(Text, SummaryFont);
+
+		const float DX = 2.0f;
+		const float W2 = TextSize.X/2 + DX;
+
+		float X1 = Viewport.GetViewportXForIndex(HoveredSample.Sample->LargestFrameIndex);
+		float CX = X1 + FMath::RoundToFloat(Viewport.GetSampleWidth() / 2);
+		if (CX + W2 > Viewport.Width)
+			CX = FMath::RoundToFloat(Viewport.Width - W2);
+		if (CX - W2 < 0)
+			CX = W2;
+
+		const float Y = 10.0f;
+		const float H = 14.0f;
+
+		FSlateDrawElement::MakeBox(
+			OutDrawElements,
+			LayerId,
+			MAKE_PAINT_GEOMETRY_RC(AllottedGeometry, CX - W2, Y, 2 * W2, H),
+			WhiteBrush,
+			DrawEffects,
+			FLinearColor(0.7, 0.7, 0.7, 0.9)
+		);
+
+		FSlateDrawElement::MakeText
+		(
+			OutDrawElements,
+			LayerId,
+			MAKE_PAINT_GEOMETRY_PT(AllottedGeometry, CX - W2 + DX, Y + 1.0f),
+			Text,
+			SummaryFont,
+			DrawEffects,
+			FLinearColor(0.0, 0.0, 0.0, 0.9)
+		);
+	}
+
+	TSharedPtr<STimingProfilerWindow> Window = FTimingProfilerManager::Get()->GetProfilerWindow();
+	if (Window.IsValid() && CachedTimelines.Num() > 0)
+	{
+		const double StartTime = Window->TimingView->GetViewport().StartTime;
+		const double EndTime = Window->TimingView->GetViewport().EndTime;
+		DC.DrawHighlightedInterval(CachedTimelines[0], StartTime, EndTime);
+	}
+
+	Stopwatch.Stop();
+	DrawDurationHistory.AddValue(Stopwatch.AccumulatedTime);
+
+	//////////////////////////////////////////////////
+
+	const bool bShouldDisplayDebugInfo = FInsightsManager::Get()->IsDebugInfoEnabled();
+	if (bShouldDisplayDebugInfo)
+	{
+		FSlateFontInfo SummaryFont = FCoreStyle::GetDefaultFontStyle("Regular", 8);
+		const float MaxFontCharHeight = FontMeasureService->Measure(TEXT("!"), SummaryFont).Y;
+		const float DbgDY = MaxFontCharHeight;
+
+		const float DbgW = 240.0f;
+		const float DbgH = DbgDY * 4 + 3.0f;
+		const float DbgX = ViewWidth - DbgW - 20.0f;
+		float DbgY = 7.0f;
+
+		++LayerId;
+		FSlateDrawElement::MakeBox(
+			OutDrawElements,
+			LayerId,
+			MAKE_PAINT_GEOMETRY_RC(AllottedGeometry, DbgX - 2.0f, DbgY - 2.0f, DbgW, DbgH),
+			WhiteBrush,
+			DrawEffects,
+			FLinearColor(1.0, 1.0, 1.0, 0.9)
+		);
+
+		++LayerId;
+		FLinearColor DbgTextColor(0.0, 0.0, 0.0, 0.9);
+
+		// Time interval since last OnPaint call.
+		const uint64 CurrentTime = FPlatformTime::Cycles64();
+		const uint64 OnPaintDuration = CurrentTime - LastOnPaintTime;
+		LastOnPaintTime = CurrentTime;
+		OnPaintDurationHistory.AddValue(OnPaintDuration); // saved for last 32 OnPaint calls
+		const uint64 AvgOnPaintDuration = OnPaintDurationHistory.ComputeAverage();
+		const uint64 AvgOnPaintDurationMs = FStopwatch::Cycles64ToMilliseconds(AvgOnPaintDuration);
+		const double AvgOnPaintFps = AvgOnPaintDurationMs != 0 ? 1.0 / FStopwatch::Cycles64ToSeconds(AvgOnPaintDuration) : 0.0;
+
+		const uint64 AvgUpdateDurationMs = FStopwatch::Cycles64ToMilliseconds(UpdateDurationHistory.ComputeAverage());
+		const uint64 AvgDrawDurationMs = FStopwatch::Cycles64ToMilliseconds(DrawDurationHistory.ComputeAverage());
+
+		// Draw performance info.
+		FSlateDrawElement::MakeText
+		(
+			OutDrawElements,
+			LayerId,
+			MAKE_PAINT_GEOMETRY_PT(AllottedGeometry, DbgX, DbgY),
+			FString::Printf(TEXT("U: %llu ms, D: %llu ms + %llu ms = %llu ms (%d fps)"),
+				AvgUpdateDurationMs, // caching time
+				AvgDrawDurationMs, // drawing time
+				AvgOnPaintDurationMs - AvgDrawDurationMs, // other overhead to OnPaint calls
+				AvgOnPaintDurationMs, // average time between two OnPaint calls
+				FMath::RoundToInt(AvgOnPaintFps)), // framerate of OnPaint calls
+			SummaryFont,
+			DrawEffects,
+			DbgTextColor
+		);
+		DbgY += DbgDY;
+
+		// Draw number of draw calls.
+		FSlateDrawElement::MakeText
+		(
+			OutDrawElements,
+			LayerId,
+			MAKE_PAINT_GEOMETRY_PT(AllottedGeometry, DbgX, DbgY),
+			FString::Format(TEXT("U: {0} frames, D: {1} samples"),
+				{
+					FText::AsNumber(NumUpdatedFrames).ToString(),
+					FText::AsNumber(DC.NumDrawSamples).ToString()
+				}),
+			SummaryFont,
+			DrawEffects,
+			DbgTextColor
+		);
+		DbgY += DbgDY;
+
+		// Draw viewport's horizontal info.
+		FSlateDrawElement::MakeText
+		(
+			OutDrawElements,
+			LayerId,
+			MAKE_PAINT_GEOMETRY_PT(AllottedGeometry, DbgX, DbgY),
+			FString::Printf(TEXT("SX: %.3f (%d %s), X: %.2f"),
+				Viewport.ScaleX,
+				(Viewport.ScaleX >= 1.0f) ? FMath::RoundToInt(Viewport.ScaleX) : FMath::RoundToInt(1.0f / Viewport.ScaleX),
+				(Viewport.ScaleX >= 1.0f) ? TEXT("px/frame") : TEXT("frames/px"),
+				Viewport.PosX),
+			SummaryFont,
+			DrawEffects,
+			DbgTextColor
+		);
+		DbgY += DbgDY;
+
+		// Draw viewport's vertical info.
+		FSlateDrawElement::MakeText
+		(
+			OutDrawElements,
+			LayerId,
+			MAKE_PAINT_GEOMETRY_PT(AllottedGeometry, DbgX, DbgY),
+			FString::Printf(TEXT("SY: %g, Y: %.2f"), Viewport.ScaleY, Viewport.PosY),
+			SummaryFont,
+			DrawEffects,
+			DbgTextColor
+		);
+		DbgY += DbgDY;
+	}
+
+	return SCompoundWidget::OnPaint(Args, AllottedGeometry, MyCullingRect, OutDrawElements, LayerId, InWidgetStyle, bParentEnabled && IsEnabled());
+}
+
+////////////////////////////////////////////////////////////////////////////////////////////////////
+
+FReply SFrameTrack::OnMouseButtonDown(const FGeometry& MyGeometry, const FPointerEvent& MouseEvent)
+{
+	FReply Reply = FReply::Unhandled();
+
+	if (!IsReady())
+	{
+		return Reply;
+	}
+
+	MousePositionOnButtonDown = MyGeometry.AbsoluteToLocal(MouseEvent.GetScreenSpacePosition());
+	ViewportPosXOnButtonDown = Viewport.PosX;
+	ViewportPosYOnButtonDown = Viewport.PosY;
+
+	if (MouseEvent.GetEffectingButton() == EKeys::LeftMouseButton)
+	{
+		bIsLMB_Pressed = true;
+
+		// Capture mouse.
+		Reply = FReply::Handled().CaptureMouse(SharedThis(this));
+	}
+	else if (MouseEvent.GetEffectingButton() == EKeys::RightMouseButton)
+	{
+		bIsRMB_Pressed = true;
+
+		// Capture mouse, so we can scroll outside this widget.
+		Reply = FReply::Handled().CaptureMouse(SharedThis(this));
+	}
+
+	return Reply;
+}
+
+////////////////////////////////////////////////////////////////////////////////////////////////////
+
+FReply SFrameTrack::OnMouseButtonUp(const FGeometry& MyGeometry, const FPointerEvent& MouseEvent)
+{
+	FReply Reply = FReply::Unhandled();
+
+	if (!IsReady())
+	{
+		return Reply;
+	}
+
+	MousePositionOnButtonUp = MyGeometry.AbsoluteToLocal(MouseEvent.GetScreenSpacePosition());
+
+	const bool bIsValidForMouseClick = MousePositionOnButtonUp.Equals(MousePositionOnButtonDown, MOUSE_SNAP_DISTANCE);
+
+	if (MouseEvent.GetEffectingButton() == EKeys::LeftMouseButton)
+	{
+		if (bIsLMB_Pressed)
+		{
+			// Release the mouse.
+			Reply = FReply::Handled().ReleaseMouseCapture();
+
+			if (bIsValidForMouseClick)
+			{
+				SelectFrameAtMousePosition(MousePositionOnButtonUp.X, MousePositionOnButtonUp.Y);
+			}
+
+			bIsLMB_Pressed = false;
+		}
+	}
+	else if (MouseEvent.GetEffectingButton() == EKeys::RightMouseButton)
+	{
+		if (bIsRMB_Pressed)
+		{
+			// Release mouse as we no longer scroll.
+			Reply = FReply::Handled().ReleaseMouseCapture();
+
+			if (bIsRMB_Scrolling)
+			{
+				CursorType = EFrameTrackCursor::Default;
+				bIsRMB_Scrolling = false;
+			}
+			else if (bIsValidForMouseClick)
+			{
+				ShowContextMenu(MouseEvent.GetScreenSpacePosition());
+				Reply = FReply::Handled();
+			}
+
+			bIsRMB_Pressed = false;
+		}
+	}
+
+	return Reply;
+}
+
+////////////////////////////////////////////////////////////////////////////////////////////////////
+
+FReply SFrameTrack::OnMouseMove(const FGeometry& MyGeometry, const FPointerEvent& MouseEvent)
+{
+	FReply Reply = FReply::Unhandled();
+
+	if (!IsReady())
+	{
+		return Reply;
+	}
+
+	MousePosition = MyGeometry.AbsoluteToLocal(MouseEvent.GetScreenSpacePosition());
+
+	if (MouseEvent.IsMouseButtonDown(EKeys::LeftMouseButton))
+	{
+		if (HasMouseCapture() && !MouseEvent.GetCursorDelta().IsZero())
+		{
+			// Inform other widgets that we have moved the selection box.
+			//SelectionBoxChangedEvent.Broadcast(SelectionBoxFrameStart, SelectionBoxFrameEnd);
+
+			Reply = FReply::Handled();
+		}
+	}
+	else if (MouseEvent.IsMouseButtonDown(EKeys::RightMouseButton))
+	{
+		if (HasMouseCapture() && !MouseEvent.GetCursorDelta().IsZero())
+		{
+			bIsRMB_Scrolling = true;
+			CursorType = EFrameTrackCursor::Hand;
+
+			HoveredSample.Reset();
+
+			float PosX = ViewportPosXOnButtonDown + (MousePositionOnButtonDown.X - MousePosition.X);
+			Viewport.ScrollAtIndex(Viewport.GetIndexAtPosX(PosX));
+			UpdateHorizontalScrollBar();
+			bIsStateDirty = true;
+
+			Reply = FReply::Handled();
+		}
+	}
+	else
+	{
+		HoveredSample = GetSampleAtMousePosition(MousePosition.X, MousePosition.Y);
+		if (!HoveredSample.IsValid())
+			HoveredSample = GetSampleAtMousePosition(MousePosition.X - 1.0f, MousePosition.Y);
+		if (!HoveredSample.IsValid())
+			HoveredSample = GetSampleAtMousePosition(MousePosition.X + 1.0f, MousePosition.Y);
+	}
+
+	return Reply;
+}
+
+////////////////////////////////////////////////////////////////////////////////////////////////////
+
+void SFrameTrack::OnMouseEnter(const FGeometry& MyGeometry, const FPointerEvent& MouseEvent)
+{
+
+}
+
+////////////////////////////////////////////////////////////////////////////////////////////////////
+
+void SFrameTrack::OnMouseLeave(const FPointerEvent& MouseEvent)
+{
+	if (!HasMouseCapture())
+	{
+		bIsLMB_Pressed = false;
+		bIsRMB_Pressed = false;
+
+		HoveredSample.Reset();
+
+		CursorType = EFrameTrackCursor::Default;
+	}
+}
+
+////////////////////////////////////////////////////////////////////////////////////////////////////
+
+FReply SFrameTrack::OnMouseWheel(const FGeometry& MyGeometry, const FPointerEvent& MouseEvent)
+{
+	MousePosition = MyGeometry.AbsoluteToLocal(MouseEvent.GetScreenSpacePosition());
+
+	if (MouseEvent.GetModifierKeys().IsShiftDown())
+	{
+		// Zoom in/out vertically.
+		const float Delta = MouseEvent.GetWheelDelta();
+		static const float ZoomStep = 0.25f; // as percent
+		float ScaleY;
+
+		if (Delta > 0)
+		{
+			ScaleY = Viewport.ScaleY * FMath::Pow(1.0f + ZoomStep, Delta);
+		}
+		else
+		{
+			ScaleY = Viewport.ScaleY * FMath::Pow(1.0f / (1.0f + ZoomStep), -Delta);
+		}
+
+		Viewport.SetScaleY(ScaleY);
+		//UpdateVerticalScrollBar();
+	}
+	else //if (MouseEvent.GetModifierKeys().IsControlDown())
+	{
+		// Zoom in/out horizontally.
+		const float Delta = MouseEvent.GetWheelDelta();
+		ZoomHorizontally(Delta, MousePosition.X);
+	}
+
+	return FReply::Handled();
+}
+
+////////////////////////////////////////////////////////////////////////////////////////////////////
+
+void SFrameTrack::ZoomHorizontally(const float Delta, const float X)
+{
+	static const float ZoomStep = 0.25f; // as percent
+	float ScaleX;
+
+	if (Delta > 0)
+	{
+		ScaleX = Viewport.ScaleX * FMath::Pow(1.0f + ZoomStep, Delta);
+	}
+	else
+	{
+		ScaleX = Viewport.ScaleX * FMath::Pow(1.0f / (1.0f + ZoomStep), -Delta);
+	}
+
+	// Snap to integer value of either: "number of frames per pixel (Slate unit)" or "number of pixels per frame".
+	if (ScaleX < 1.0f)
+	{
+		// frames per Slate unit
+		if (Delta > 0)
+			ScaleX = 1.0f / FMath::FloorToFloat(1.0f / ScaleX);
+		else
+			ScaleX = 1.0f / FMath::CeilToFloat(1.0f / ScaleX);
+	}
+	else
+	{
+		// Slate units per frame
+		if (Delta > 0)
+			ScaleX = FMath::CeilToFloat(ScaleX);
+		else
+			ScaleX = FMath::FloorToFloat(ScaleX);
+	}
+
+	//UE_LOG(TimingProfiler, Log, TEXT("%.2f, %.2f, %.2f"), Delta, Viewport.ScaleX, ScaleX);
+	Viewport.ZoomWithFixedViewportX(ScaleX, MousePosition.X);
+	UpdateHorizontalScrollBar();
+	bIsStateDirty = true;
+}
+
+////////////////////////////////////////////////////////////////////////////////////////////////////
+
+FReply SFrameTrack::OnMouseButtonDoubleClick(const FGeometry& MyGeometry, const FPointerEvent& MouseEvent)
+{
+	return FReply::Unhandled();
+}
+
+////////////////////////////////////////////////////////////////////////////////////////////////////
+
+FCursorReply SFrameTrack::OnCursorQuery(const FGeometry& MyGeometry, const FPointerEvent& CursorEvent) const
+{
+	FCursorReply CursorReply = FCursorReply::Unhandled();
+
+	if (CursorType == EFrameTrackCursor::Arrow)
+	{
+		CursorReply = FCursorReply::Cursor(EMouseCursor::ResizeLeftRight);
+	}
+	else if (CursorType == EFrameTrackCursor::Hand)
+	{
+		CursorReply = FCursorReply::Cursor(EMouseCursor::GrabHand);
+	}
+
+	return CursorReply;
+}
+
+////////////////////////////////////////////////////////////////////////////////////////////////////
+
+void SFrameTrack::ShowContextMenu(const FVector2D& ScreenSpacePosition)
+{
+
+}
+
+////////////////////////////////////////////////////////////////////////////////////////////////////
+
+void SFrameTrack::BindCommands()
+{
+
+}
+
+////////////////////////////////////////////////////////////////////////////////////////////////////
+
+void SFrameTrack::HorizontalScrollBar_OnUserScrolled(float ScrollOffset)
+{
+	const float SX = 1.0 / (Viewport.MaxX - Viewport.MinX);
+	const float ThumbSizeFraction = FMath::Clamp<float>(Viewport.Width * SX, 0.0f, 1.0f);
+	const float OffsetFraction = FMath::Clamp<float>(ScrollOffset, 0.0f, 1.0f - ThumbSizeFraction);
+
+	const float PosX = Viewport.MinX + OffsetFraction * (Viewport.MaxX - Viewport.MinX);
+	Viewport.ScrollAtPosX(PosX);
+	bIsStateDirty = true;
+
+	HorizontalScrollBar->SetState(OffsetFraction, ThumbSizeFraction);
+}
+
+////////////////////////////////////////////////////////////////////////////////////////////////////
+
+void SFrameTrack::UpdateHorizontalScrollBar()
+{
+	const float SX = 1.0 / (Viewport.MaxX - Viewport.MinX);
+	const float ThumbSizeFraction = FMath::Clamp<float>(Viewport.Width * SX, 0.0f, 1.0f);
+	const float ScrollOffset = (Viewport.PosX - Viewport.MinX) * SX;
+	const float OffsetFraction = FMath::Clamp<float>(ScrollOffset, 0.0f, 1.0f - ThumbSizeFraction);
+
+	HorizontalScrollBar->SetState(OffsetFraction, ThumbSizeFraction);
+}
+
+////////////////////////////////////////////////////////////////////////////////////////////////////
+/*
+void SFrameTrack::AddThreadTime(int32 InFrameIndex, const TMap<uint32, float>& InThreadMS, const TSharedRef<FProfilerStatMetaData>& InStatMetaData)
+{
+	FFrameThreadTimes FrameThreadTimes;
+	FrameThreadTimes.FrameNumber = InFrameIndex;
+	FrameThreadTimes.ThreadTimes = InThreadMS;
+	RecentlyAddedFrames.Add(FrameThreadTimes);
+
+	if (!bIsActiveTimerRegistered)
+	{
+		bIsActiveTimerRegistered = true;
+		RegisterActiveTimer(0.f, FWidgetActiveTimerDelegate::CreateSP(this, &SFrameTrack::EnsureDataUpdateDuringPreview));
+	}
+
+	StatMetadata = InStatMetaData;
+}
+*/
+////////////////////////////////////////////////////////////////////////////////////////////////////
+/*
+EActiveTimerReturnType SFrameTrack::EnsureDataUpdateDuringPreview(double InCurrentTime, float InDeltaTime)
+{
+	//if (RecentlyAddedFrames.Num() > 0)
+	//{
+	//	bUpdateData = true;
+	//	return EActiveTimerReturnType::Continue;
+	//}
+
+	bIsActiveTimerRegistered = false;
+	return EActiveTimerReturnType::Stop;
+}
+*/
+////////////////////////////////////////////////////////////////////////////////////////////////////
+/*
+void SFrameTrack::MoveSelectionBox(int32 FrameIndex)
+{
+	const int32 SelectionBoxSize = SelectionBoxFrameEnd - SelectionBoxFrameStart;
+	const int32 SelectionBoxHalfSize = SelectionBoxSize / 2;
+	const int32 CenterFrameIndex = FMath::Clamp(FrameIndex - SelectionBoxHalfSize, 0, Frames.Num() - 1 - SelectionBoxSize);
+
+	// Inform other widgets that we have moved the selection box.
+	SelectionBoxFrameStart = CenterFrameIndex;
+	SelectionBoxFrameEnd = CenterFrameIndex + SelectionBoxSize;
+	SelectionBoxChangedEvent.Broadcast(SelectionBoxFrameStart, SelectionBoxFrameEnd);
+}
+*/
+////////////////////////////////////////////////////////////////////////////////////////////////////
+
+#undef LOCTEXT_NAMESPACE
diff --git a/Engine/Source/Developer/TraceInsights/Private/Insights/Widgets/SFrameTrack.h b/Engine/Source/Developer/TraceInsights/Private/Insights/Widgets/SFrameTrack.h
new file mode 100644
--- /dev/null
+++ b/Engine/Source/Developer/TraceInsights/Private/Insights/Widgets/SFrameTrack.h
@@ -0,0 +1,189 @@
+// Copyright 1998-2019 Epic Games, Inc. All Rights Reserved.
+
+#pragma once
+
+#include "CoreMinimal.h"
+#include "Input/CursorReply.h"
+#include "Input/Reply.h"
+#include "Layout/Geometry.h"
+#include "Rendering/RenderingCommon.h"
+#include "Widgets/DeclarativeSyntaxSupport.h"
+#include "Widgets/SCompoundWidget.h"
+
+// Insights
+#include "Insights/Common/FixedCircularBuffer.h"
+#include "Insights/ViewModels/FrameTrackHelper.h"
+#include "Insights/ViewModels/FrameTrackViewport.h"
+
+class SScrollBar;
+
+////////////////////////////////////////////////////////////////////////////////////////////////////
+
+struct FSampleRef
+{
+	const FFrameTrackTimeline* Timeline;
+	const FFrameTrackSample* Sample;
+
+	FSampleRef(const FFrameTrackTimeline* InTimeline, const FFrameTrackSample* InSample)
+		: Timeline(InTimeline), Sample(InSample)
+	{
+	}
+
+	void Reset() { Timeline = nullptr; Sample = nullptr; }
+	bool IsValid() const { return Timeline != nullptr && Sample != nullptr; }
+};
+
+////////////////////////////////////////////////////////////////////////////////////////////////////
+/// Widget used to present frames data in a track.
+
+class SFrameTrack : public SCompoundWidget
+{
+	enum
+	{
+		/// Number of pixels.
+		MOUSE_SNAP_DISTANCE = 4,
+	};
+
+	struct EFrameTrackCursor
+	{
+		enum Type
+		{
+			Default,
+			Arrow,
+			Hand,
+		};
+	};
+
+public:
+	/// Default constructor.
+	SFrameTrack();
+
+	/// Virtual destructor.
+	virtual ~SFrameTrack();
+
+	/// Resets internal widget's data to the default one.
+	void Reset();
+
+	SLATE_BEGIN_ARGS(SFrameTrack)
+	{
+		_Clipping = EWidgetClipping::ClipToBounds;
+	}
+	SLATE_END_ARGS()
+
+	/// Construct this widget
+	/// @param InArgs The declaration data for this widget
+	void Construct(const FArguments& InArgs);
+
+	/// Ticks this widget. Override in derived classes, but always call the parent implementation.
+	/// @param  AllottedGeometry The space allotted for this widget
+	/// @param  InCurrentTime  Current absolute real time
+	/// @param  InDeltaTime  Real time passed since last tick
+	virtual void Tick(const FGeometry& AllottedGeometry, const double InCurrentTime, const float InDeltaTime) override;
+
+	virtual int32 OnPaint(const FPaintArgs& Args, const FGeometry& AllottedGeometry, const FSlateRect& MyCullingRect, FSlateWindowElementList& OutDrawElements, int32 LayerId, const FWidgetStyle& InWidgetStyle, bool bParentEnabled) const override;
+
+	virtual FReply OnMouseButtonDown(const FGeometry& MyGeometry, const FPointerEvent& MouseEvent) override;
+	virtual FReply OnMouseButtonUp(const FGeometry& MyGeometry, const FPointerEvent& MouseEvent) override;
+	virtual FReply OnMouseMove(const FGeometry& MyGeometry, const FPointerEvent& MouseEvent) override;
+	virtual void OnMouseEnter(const FGeometry& MyGeometry, const FPointerEvent& MouseEvent) override;
+	virtual void OnMouseLeave(const FPointerEvent& MouseEvent) override;
+	virtual FReply OnMouseWheel(const FGeometry& MyGeometry, const FPointerEvent& MouseEvent) override;
+	virtual FReply OnMouseButtonDoubleClick(const FGeometry& MyGeometry, const FPointerEvent& MouseEvent) override;
+
+	virtual FCursorReply OnCursorQuery(const FGeometry& MyGeometry, const FPointerEvent& CursorEvent) const override;
+
+protected:
+	//virtual FVector2D ComputeDesiredSize(float) const override
+	//{
+	//	return FVector2D(16.0f, 16.0f);
+	//}
+
+	bool IsReady() { return true; }
+
+	void UpdateState();
+
+	FSampleRef GetSampleAtMousePosition(float X, float Y);
+	void SelectFrameAtMousePosition(float X, float Y);
+
+	void ShowContextMenu(const FVector2D& ScreenSpacePosition);
+
+	/// Binds our UI commands to delegates.
+	void BindCommands();
+
+	/// Called when the user scrolls the horizontal scrollbar.
+	/// @param ScrollOffset Scroll offset as a fraction between 0 and 1.
+	void HorizontalScrollBar_OnUserScrolled(float ScrollOffset);
+	void UpdateHorizontalScrollBar();
+
+	void ZoomHorizontally(const float Delta, const float X);
+
+	//////////////////////////////////////////////////
+	// SelectionBoxChanged Event
+
+public:
+	/// The event to execute when the selection box has been changed.
+	DECLARE_EVENT_TwoParams(SFrameTrack, FSelectionBoxChangedEvent, int32 /*FrameStart*/, int32 /*FrameEnd*/);
+	FSelectionBoxChangedEvent& OnSelectionBoxChanged() { return SelectionBoxChangedEvent; }
+protected:
+	/// The event to execute when the selection box has been changed.
+	FSelectionBoxChangedEvent SelectionBoxChangedEvent;
+
+	//////////////////////////////////////////////////
+
+protected:
+	/// The track's viewport. Encapsulates info about position and scale.
+	FFrameTrackViewport Viewport;
+	bool bIsViewportDirty;
+
+	/// Cached info for timelines.
+	TMap<uint64, FFrameTrackTimeline> CachedTimelines;
+	TArray<int32> TimelinesOrder;
+
+	bool bIsStateDirty;
+
+	//////////////////////////////////////////////////
+
+	TSharedPtr<SScrollBar> HorizontalScrollBar;
+
+	//////////////////////////////////////////////////
+	// Panning, Zooming and Selection behaviors
+
+	/// The current mouse position.
+	FVector2D MousePosition;
+
+	/// Mouse position during the call on mouse button down.
+	FVector2D MousePositionOnButtonDown;
+	float ViewportPosXOnButtonDown;
+	float ViewportPosYOnButtonDown;
+
+	/// Mouse position during the call on mouse button up.
+	FVector2D MousePositionOnButtonUp;
+
+	bool bIsLMB_Pressed;
+	bool bIsRMB_Pressed;
+
+	/// True, if the user is currently interactively scrolling the view by holding the right mouse button and dragging.
+	bool bIsRMB_Scrolling;
+
+	//////////////////////////////////////////////////
+
+	int32 SelectionStartFrameIndex;
+	int32 SelectionEndFrameIndex;
+
+	FSampleRef HoveredSample;
+
+	//////////////////////////////////////////////////
+	// Misc
+
+	FGeometry ThisGeometry;
+
+	/// Cursor type.
+	EFrameTrackCursor::Type CursorType;
+
+	// Debug stats
+	int32 NumUpdatedFrames;
+	TFixedCircularBuffer<uint64, 32> UpdateDurationHistory;
+	mutable TFixedCircularBuffer<uint64, 32> DrawDurationHistory;
+	mutable TFixedCircularBuffer<uint64, 32> OnPaintDurationHistory;
+	mutable uint64 LastOnPaintTime;
+};
diff --git a/Engine/Source/Developer/TraceInsights/Private/Insights/Widgets/SGraphTrack.cpp b/Engine/Source/Developer/TraceInsights/Private/Insights/Widgets/SGraphTrack.cpp
new file mode 100644
--- /dev/null
+++ b/Engine/Source/Developer/TraceInsights/Private/Insights/Widgets/SGraphTrack.cpp
@@ -0,0 +1,331 @@
+// Copyright 1998-2019 Epic Games, Inc. All Rights Reserved.
+
+#include "SGraphTrack.h"
+
+#include "Brushes/SlateBorderBrush.h"
+#include "Brushes/SlateBoxBrush.h"
+#include "Brushes/SlateColorBrush.h"
+#include "Containers/ArrayBuilder.h"
+#include "Containers/MapBuilder.h"
+#include "EditorStyleSet.h"
+#include "Fonts/FontMeasure.h"
+#include "Fonts/SlateFontInfo.h"
+#include "Framework/Application/MenuStack.h"
+#include "Framework/Application/SlateApplication.h"
+#include "Framework/MultiBox/MultiBoxBuilder.h"
+#include "Layout/WidgetPath.h"
+#include "Misc/Paths.h"
+#include "Rendering/DrawElements.h"
+#include "SlateOptMacros.h"
+#include "Styling/CoreStyle.h"
+#include "Widgets/Input/SButton.h"
+#include "Widgets/Layout/SSpacer.h"
+#include "Widgets/SBoxPanel.h"
+#include "Widgets/SOverlay.h"
+#include "Widgets/Text/STextBlock.h"
+
+// Insights
+#include "Insights/Common/PaintUtils.h"
+#include "Insights/TimingProfilerManager.h"
+
+////////////////////////////////////////////////////////////////////////////////////////////////////
+
+#define LOCTEXT_NAMESPACE "SGraphTrack"
+
+////////////////////////////////////////////////////////////////////////////////////////////////////
+// SGraphTrack
+////////////////////////////////////////////////////////////////////////////////////////////////////
+
+SGraphTrack::SGraphTrack()
+	: TimeRulerTrack(0)
+	, GraphTrack(1)
+{
+	Reset();
+}
+
+////////////////////////////////////////////////////////////////////////////////////////////////////
+
+SGraphTrack::~SGraphTrack()
+{
+}
+
+////////////////////////////////////////////////////////////////////////////////////////////////////
+
+void SGraphTrack::Reset()
+{
+	TimeRulerTrack.Reset();
+	GraphTrack.Reset();
+	Viewport.ScaleX = (5 * 20) / 0.1; // 100ms between major tick marks
+}
+
+////////////////////////////////////////////////////////////////////////////////////////////////////
+
+void SGraphTrack::Construct(const FArguments& InArgs)
+{
+	ChildSlot
+	[
+		SNew(SOverlay)
+		.Visibility(EVisibility::SelfHitTestInvisible)
+
+		//+SOverlay::Slot()
+		//.HAlign(HAlign_Left)
+		//.VAlign(VAlign_Top)
+		//.Padding(FMargin(48.0f, 16.0f, 48.0f, 16.0f))	// Make some space for graph labels
+		//[
+		//	SAssignNew(GraphDescriptionsVBox, SVerticalBox)
+		//]
+	];
+
+	//BindCommands();
+}
+
+////////////////////////////////////////////////////////////////////////////////////////////////////
+
+void SGraphTrack::Tick(const FGeometry& AllottedGeometry, const double InCurrentTime, const float InDeltaTime)
+{
+	//ThisGeometry = AllottedGeometry;
+
+	float TrackWidth = AllottedGeometry.GetAbsoluteSize().X;
+	float TrackHeight = AllottedGeometry.GetAbsoluteSize().Y;
+
+	bool bIsGraphDirty = false;
+
+	if (Viewport.UpdateSize(TrackWidth, TrackHeight))
+	{
+		bIsGraphDirty = true;
+	}
+
+	if (GraphTrack.GetHeight() != TrackHeight)
+	{
+		bIsGraphDirty = true;
+	}
+
+	if (bIsGraphDirty)
+	{
+		bIsGraphDirty = false;
+
+		GraphTrack.SetPosY(TimeRulerTrack.GetHeight());
+		GraphTrack.SetHeight(Viewport.Height - TimeRulerTrack.GetHeight());
+		GraphTrack.Update(Viewport);
+	}
+}
+
+////////////////////////////////////////////////////////////////////////////////////////////////////
+
+//PRAGMA_DISABLE_OPTIMIZATION
+int32 SGraphTrack::OnPaint(const FPaintArgs& Args, const FGeometry& AllottedGeometry, const FSlateRect& MyCullingRect, FSlateWindowElementList& OutDrawElements, int32 LayerId, const FWidgetStyle& InWidgetStyle, bool bParentEnabled) const
+{
+	const bool bEnabled = ShouldBeEnabled(bParentEnabled);
+	const ESlateDrawEffect DrawEffects = bEnabled ? ESlateDrawEffect::None : ESlateDrawEffect::DisabledEffect;
+	FDrawContext DC(AllottedGeometry, MyCullingRect, InWidgetStyle, DrawEffects, OutDrawElements, LayerId);
+
+	GraphTrack.Draw(DC, Viewport);
+	TimeRulerTrack.Draw(DC, Viewport);
+
+	/*
+	static double TotalTime = 0.0f;
+	static uint32 NumCalls = 0;
+	const double StartTime = FPlatformTime::Seconds();
+
+	const TSharedRef<FSlateFontMeasure> FontMeasureService = FSlateApplication::Get().GetRenderer()->GetFontMeasureService();
+
+	// Rendering info.
+	const bool bEnabled = ShouldBeEnabled(bParentEnabled);
+	ESlateDrawEffect DrawEffects = bEnabled ? ESlateDrawEffect::None : ESlateDrawEffect::DisabledEffect;
+	const FSlateBrush* WhiteBrush = FCoreStyle::Get().GetBrush("WhiteBrush");
+
+	const float Width = AllottedGeometry.GetLocalSize().X;
+	const float Height = AllottedGeometry.GetLocalSize().Y;
+
+	// Draw background.
+	const FSlateBrush* AreaBrush = FEditorStyle::GetBrush("Profiler.LineGraphArea");
+	FSlateDrawElement::MakeBox
+	(
+		OutDrawElements,
+		LayerId,
+		AllottedGeometry.ToPaintGeometry(FVector2D(0,0), FVector2D(Width, Height)),
+		AreaBrush,
+		DrawEffects,
+		AreaBrush->GetTint(InWidgetStyle) * InWidgetStyle.GetColorAndOpacityTint()
+	);
+	LayerId++;
+
+	// Draw a test graph...
+	static const int EventCount = 1000;
+	const FLinearColor LineColor(0.0f, 0.5f, 1.0f, 1.0f);
+	TArray<FVector2D> LinePoints;
+	FVector2D V1(0.0f, 0.0f);
+	FVector2D V2;
+	FRandomStream RandomStream(0);
+	for (int Index = 0; Index < EventCount; ++Index)
+	{
+		V2.X = FMath::TruncToFloat((Width / (float)EventCount) * (float)Index);
+		V2.Y = Height * RandomStream.GetFraction();
+		if ((V2 - V1).Size() > 1.0)
+		{
+			LinePoints.Add(V1);
+			LinePoints.Add(V2);
+			FSlateDrawElement::MakeLines(OutDrawElements, LayerId, AllottedGeometry.ToPaintGeometry(), LinePoints, DrawEffects, LineColor);
+			LinePoints.Empty();
+			V1 = V2;
+		}
+	}
+	*/
+
+	//double sec = FPlatformTime::GetSecondsPerCycle();
+
+	//FSlateFontInfo SummaryFont = FCoreStyle::GetDefaultFontStyle("Regular", 8);
+	//const float MaxFontCharHeight = FontMeasureService->Measure(TEXT("!"), SummaryFont).Y;
+
+	/*
+	// Debug draw (slate fps)
+	++LayerId;
+	FSlateDrawElement::MakeBox(
+		OutDrawElements,
+		LayerId,
+		AllottedGeometry.ToPaintGeometry(FVector2D(8.0f, 8.0f), FVector2D(200.0, MaxFontCharHeight + 4.0)),
+		EventBrush,
+		DrawEffects,
+		FLinearColor(FColor(255, 255, 255))
+	);
+	static uint64 Time = 0;
+	const uint64 CurrentTime = FPlatformTime::Cycles64();
+	const uint64 Duration = CurrentTime - Time;
+	Time = CurrentTime;
+	static const int DurationCount = 64;
+	static uint64 Durations[DurationCount];
+	static uint64 DurationOffset = 0;
+	Durations[(DurationOffset++)% DurationCount] = Duration;
+	uint64 AvgDuration = 0;
+	for (int I = 0; I < DurationCount; ++I)
+		AvgDuration += Durations[I];
+	AvgDuration = AvgDuration / DurationCount;
+	const double AvgFps = AvgDuration > 0 ? 1.0 / (FPlatformTime::GetSecondsPerCycle64() * (double)AvgDuration) : 0.0;
+	const FString DebnugStatsStr = FString::Printf(TEXT("%d ms (%d ms) %.1f fps %d events"),
+		FMath::TruncToInt(AvgDuration * 1000 * FPlatformTime::GetSecondsPerCycle64()),
+		FMath::TruncToInt(Duration * 1000 * FPlatformTime::GetSecondsPerCycle64()),
+		AvgFps,
+		EventCount);
+	FSlateDrawElement::MakeText
+	(
+		OutDrawElements,
+		LayerId,
+		AllottedGeometry.ToOffsetPaintGeometry(FVector2D(10.0f, 10.0f)),
+		DebnugStatsStr,
+		SummaryFont,
+		DrawEffects,
+		FLinearColor::Red
+	);
+	*/
+
+	return SCompoundWidget::OnPaint(Args, AllottedGeometry, MyCullingRect, OutDrawElements, LayerId, InWidgetStyle, bParentEnabled && IsEnabled());
+}
+//PRAGMA_ENABLE_OPTIMIZATION
+
+////////////////////////////////////////////////////////////////////////////////////////////////////
+
+FReply SGraphTrack::OnMouseButtonDown(const FGeometry& MyGeometry, const FPointerEvent& MouseEvent)
+{
+	//...
+
+	return FReply::Unhandled();
+}
+
+////////////////////////////////////////////////////////////////////////////////////////////////////
+
+FReply SGraphTrack::OnMouseButtonUp(const FGeometry& MyGeometry, const FPointerEvent& MouseEvent)
+{
+	//...
+
+	return FReply::Unhandled();
+}
+
+////////////////////////////////////////////////////////////////////////////////////////////////////
+
+FReply SGraphTrack::OnMouseMove(const FGeometry& MyGeometry, const FPointerEvent& MouseEvent)
+{
+	//...
+
+	return FReply::Unhandled();
+}
+
+////////////////////////////////////////////////////////////////////////////////////////////////////
+
+void SGraphTrack::OnMouseEnter(const FGeometry& MyGeometry, const FPointerEvent& MouseEvent)
+{
+	//...
+}
+
+////////////////////////////////////////////////////////////////////////////////////////////////////
+
+void SGraphTrack::OnMouseLeave(const FPointerEvent& MouseEvent)
+{
+	//...
+}
+
+////////////////////////////////////////////////////////////////////////////////////////////////////
+
+FReply SGraphTrack::OnMouseWheel(const FGeometry& MyGeometry, const FPointerEvent& MouseEvent)
+{
+	//...
+
+	return FReply::Handled();
+}
+
+////////////////////////////////////////////////////////////////////////////////////////////////////
+
+FReply SGraphTrack::OnMouseButtonDoubleClick(const FGeometry& MyGeometry, const FPointerEvent& MouseEvent)
+{
+	//...
+
+	return FReply::Unhandled();
+}
+
+////////////////////////////////////////////////////////////////////////////////////////////////////
+
+void SGraphTrack::OnDragEnter(const FGeometry& MyGeometry, const FDragDropEvent& DragDropEvent)
+{
+	SCompoundWidget::OnDragEnter(MyGeometry, DragDropEvent);
+
+	//...
+}
+
+////////////////////////////////////////////////////////////////////////////////////////////////////
+
+void SGraphTrack::OnDragLeave(const FDragDropEvent& DragDropEvent)
+{
+	SCompoundWidget::OnDragLeave(DragDropEvent);
+
+	//...
+}
+
+////////////////////////////////////////////////////////////////////////////////////////////////////
+
+FReply SGraphTrack::OnDragOver(const FGeometry& MyGeometry, const FDragDropEvent& DragDropEvent)
+{
+	//...
+
+	return SCompoundWidget::OnDragOver(MyGeometry,DragDropEvent);
+}
+
+////////////////////////////////////////////////////////////////////////////////////////////////////
+
+FReply SGraphTrack::OnDrop(const FGeometry& MyGeometry, const FDragDropEvent& DragDropEvent)
+{
+	//...
+
+	return SCompoundWidget::OnDrop(MyGeometry,DragDropEvent);
+}
+
+////////////////////////////////////////////////////////////////////////////////////////////////////
+
+FCursorReply SGraphTrack::OnCursorQuery(const FGeometry& MyGeometry, const FPointerEvent& CursorEvent) const
+{
+	//...
+
+	return FCursorReply::Unhandled();
+}
+
+////////////////////////////////////////////////////////////////////////////////////////////////////
+
+#undef LOCTEXT_NAMESPACE
diff --git a/Engine/Source/Developer/TraceInsights/Private/Insights/Widgets/SIoProfilerWindow.cpp b/Engine/Source/Developer/TraceInsights/Private/Insights/Widgets/SIoProfilerWindow.cpp
new file mode 100644
--- /dev/null
+++ b/Engine/Source/Developer/TraceInsights/Private/Insights/Widgets/SIoProfilerWindow.cpp
@@ -0,0 +1,833 @@
+// Copyright 1998-2019 Epic Games, Inc. All Rights Reserved.
+
+#include "SIoProfilerWindow.h"
+
+#include "EditorStyleSet.h"
+#include "SlateOptMacros.h"
+#include "Widgets/Images/SImage.h"
+#include "Widgets/Layout/SBorder.h"
+#include "Widgets/Layout/SBox.h"
+#include "Widgets/Layout/SSpacer.h"
+#include "Widgets/Notifications/SNotificationList.h"
+#include "Widgets/SBoxPanel.h"
+#include "Widgets/Text/STextBlock.h"
+
+#if WITH_EDITOR
+	#include "EngineAnalytics.h"
+	#include "Runtime/Analytics/Analytics/Public/AnalyticsEventAttribute.h"
+	#include "Runtime/Analytics/Analytics/Public/Interfaces/IAnalyticsProvider.h"
+#endif // WITH_EDITOR
+
+// Insights
+#include "Insights/Version.h"
+#include "Insights/InsightsManager.h"
+#include "Insights/IoProfilerManager.h"
+#include "Insights/Widgets/SFrameTrack.h"
+#include "Insights/Widgets/SGraphTrack.h"
+#include "Insights/Widgets/SInsightsSettings.h"
+#include "Insights/Widgets/SLogView.h"
+//#include "Insights/Widgets/SIoProfilerToolbar.h"
+#include "Insights/Widgets/STimersView.h"
+#include "Insights/Widgets/STimingView.h"
+
+////////////////////////////////////////////////////////////////////////////////////////////////////
+
+#define LOCTEXT_NAMESPACE "SIoProfilerWindow"
+
+////////////////////////////////////////////////////////////////////////////////////////////////////
+
+// TODO: move this function to InsightsManager or in TraceSession.h
+static FText Io_GetTextForNotification(const EInsightsNotificationType NotificatonType, const ELoadingProgressState ProgressState, const FString& Filename, const float ProgressPercent = 0.0f)
+{
+	FText Result;
+
+	if (NotificatonType == EInsightsNotificationType::LoadingTraceFile)
+	{
+		if (ProgressState == ELoadingProgressState::Started)
+		{
+			FFormatNamedArguments Args;
+			Args.Add(TEXT("Filename"), FText::FromString(Filename));
+			Result = FText::Format(LOCTEXT("DescF_OfflineCapture_Started", "Started loading a file {Filename}"), Args);
+		}
+		else if (ProgressState == ELoadingProgressState::InProgress)
+		{
+			FFormatNamedArguments Args;
+			Args.Add(TEXT("Filename"), FText::FromString(Filename));
+			Args.Add(TEXT("LoadingProgressPercent"), FText::AsPercent(ProgressPercent));
+			Result = FText::Format(LOCTEXT("DescF_OfflineCapture_InProgress", "Loading a file {Filename} {LoadingProgressPercent}"), Args);
+		}
+		else if (ProgressState == ELoadingProgressState::Loaded)
+		{
+			FFormatNamedArguments Args;
+			Args.Add(TEXT("Filename"), FText::FromString(Filename));
+			Result = FText::Format(LOCTEXT("DescF_OfflineCapture_Loaded", "File {Filename} has been successfully loaded"), Args);
+		}
+		else if (ProgressState == ELoadingProgressState::Failed)
+		{
+			FFormatNamedArguments Args;
+			Args.Add(TEXT("Filename"), FText::FromString(Filename));
+			Result = FText::Format(LOCTEXT("DescF_OfflineCapture_Failed", "Failed to load file {Filename}"), Args);
+		}
+	}
+
+	return Result;
+}
+
+////////////////////////////////////////////////////////////////////////////////////////////////////
+
+SIoProfilerWindow::SIoProfilerWindow()
+	: DurationActive(0.0f)
+{
+}
+
+////////////////////////////////////////////////////////////////////////////////////////////////////
+
+SIoProfilerWindow::~SIoProfilerWindow()
+{
+	// Remove ourselves from the profiler manager.
+	//if (FIoProfilerManager::Get().IsValid())
+	//{
+	//	//...
+	//	FIoProfilerManager::Get()->OnViewModeChanged().RemoveAll(this);
+	//}
+
+#if WITH_EDITOR
+	if (DurationActive > 0.0f && FEngineAnalytics::IsAvailable())
+	{
+		FEngineAnalytics::GetProvider().RecordEvent(TEXT("Editor.Usage.Profiler"), FAnalyticsEventAttribute(TEXT("Duration"), DurationActive));
+	}
+#endif // WITH_EDITOR
+}
+
+////////////////////////////////////////////////////////////////////////////////////////////////////
+
+BEGIN_SLATE_FUNCTION_BUILD_OPTIMIZATION
+
+TSharedRef<SWidget> SIoProfilerWindow::ConstructMultiTrackView()
+{
+	return SNew(SSplitter)
+		.Orientation(Orient_Vertical)
+
+	+ SSplitter::Slot()
+		.Value(0.1f)
+		[
+			ConstructFramesTrack()
+		]
+
+	+ SSplitter::Slot()
+		.Value(0.2f)
+		[
+			ConstructGraphTrack()
+		]
+
+	+ SSplitter::Slot()
+		.Value(0.7f)
+		[
+			ConstructTimingTrack()
+		]
+
+	; // end of Widget
+}
+
+////////////////////////////////////////////////////////////////////////////////////////////////////
+
+TSharedRef<SWidget> SIoProfilerWindow::ConstructFramesTrack()
+{
+	return SNew(SVerticalBox)
+		.Visibility(this, &SIoProfilerWindow::IsFramesTrackVisible)
+		.IsEnabled(this, &SIoProfilerWindow::IsProfilerEnabled)
+
+	// Frame Track
+	+ SVerticalBox::Slot()
+		.FillHeight(1.0)
+		.Padding(0.0f, 0.0f, 0.0f, 0.0f)
+		[
+			SNew(SBox)
+				.HeightOverride(48.0f)
+				.IsEnabled(this, &SIoProfilerWindow::IsProfilerEnabled)
+				[
+					SAssignNew(FrameTrack, SFrameTrack)
+
+					//SNew(SHorizontalBox)
+					//
+					//+ SHorizontalBox::Slot()
+					//	.FillWidth(1.0f)
+					//	.Padding(0.0f)
+					//	.HAlign(HAlign_Fill)
+					//	.VAlign(VAlign_Fill)
+					//	[
+					//
+					//	]
+				]
+		]
+
+	; // end of Widget
+}
+
+////////////////////////////////////////////////////////////////////////////////////////////////////
+
+TSharedRef<SWidget> SIoProfilerWindow::ConstructGraphTrack()
+{
+	return SNew(SVerticalBox)
+		.Visibility(this, &SIoProfilerWindow::IsGraphTrackVisible)
+		.IsEnabled(this, &SIoProfilerWindow::IsProfilerEnabled)
+
+	/*
+	// Header
+	+ SVerticalBox::Slot()
+		.AutoHeight()
+		.Padding(0.0f, 0.0f, 0.0f, 0.0f)
+		[
+			SNew(SHorizontalBox)
+
+			+ SHorizontalBox::Slot()
+				.AutoWidth()
+				[
+					SNew(SImage)
+						.Image(FEditorStyle::GetBrush(TEXT("Profiler.Tab.GraphView")))
+				]
+
+			+ SHorizontalBox::Slot()
+				.AutoWidth()
+				[
+					SNew(STextBlock)
+						.Text(LOCTEXT("GraphTrackLabel", "Overview Timing Graph"))
+						.Margin(FMargin(0.0f, 2.0f, 0.0f, 0.0f))
+				]
+		]
+	*/
+
+	// Graph Track
+	+ SVerticalBox::Slot()
+		.FillHeight(1.0)
+		.Padding(0.0f, 0.0f, 0.0f, 0.0f)
+		[
+			SNew(SBox)
+				.HeightOverride(48.0f)
+				.IsEnabled(this, &SIoProfilerWindow::IsProfilerEnabled)
+				[
+					SNew(SHorizontalBox)
+
+					+ SHorizontalBox::Slot()
+						.FillWidth(1.0f)
+						.Padding(0.0f)
+						.HAlign(HAlign_Fill)
+						.VAlign(VAlign_Fill)
+						[
+							SAssignNew(GraphTrack, SGraphTrack)
+						]
+				]
+		]
+
+	; // end of Widget
+}
+
+////////////////////////////////////////////////////////////////////////////////////////////////////
+
+TSharedRef<SWidget> SIoProfilerWindow::ConstructTimingTrack()
+{
+	return SNew(SVerticalBox)
+		.Visibility(this, &SIoProfilerWindow::IsTimingTrackVisible)
+		.IsEnabled(this, &SIoProfilerWindow::IsProfilerEnabled)
+
+	/*
+	// Header
+	+ SVerticalBox::Slot()
+		.AutoHeight()
+		.Padding(0.0f, 0.0f, 0.0f, 0.0f)
+		[
+			SNew(SHorizontalBox)
+
+			+ SHorizontalBox::Slot()
+				.AutoWidth()
+				[
+					SNew(SImage)
+						.Image(FEditorStyle::GetBrush(TEXT("Profiler.Tab.GraphView")))
+				]
+
+			+ SHorizontalBox::Slot()
+				.AutoWidth()
+				[
+					SNew(STextBlock)
+						.Text(LOCTEXT("TimingLabel", "Main Timing View"))
+						.Margin(FMargin(0.0f, 2.0f, 0.0f, 0.0f))
+				]
+		]
+	*/
+
+	// Timing Track
+	+ SVerticalBox::Slot()
+		.FillHeight(1.0)
+		.Padding(0.0f, 0.0f, 0.0f, 0.0f)
+		[
+			SNew(SBox)
+				.HeightOverride(48.0f)
+				.IsEnabled(this, &SIoProfilerWindow::IsProfilerEnabled)
+				[
+					SNew(SHorizontalBox)
+
+					+ SHorizontalBox::Slot()
+						.FillWidth(1.0f)
+						.Padding(0.0f)
+						.HAlign(HAlign_Fill)
+						.VAlign(VAlign_Fill)
+						[
+							SAssignNew(TimingView, STimingView)
+						]
+				]
+		]
+
+	//// Threads Panel
+	//+ SVerticalBox::Slot()
+	//	.FillHeight(1.0)
+	//	.Padding(0.0f, 0.0f, 0.0f, 0.0f)
+	//	[
+	//		SAssignNew(ThreadsPanel, SVerticalBox)
+	//	]
+
+	; // end of Widget
+}
+
+////////////////////////////////////////////////////////////////////////////////////////////////////
+
+TSharedRef<SWidget> SIoProfilerWindow::ConstructTimersView()
+{
+	return SNew(SVerticalBox)
+		.Visibility(this, &SIoProfilerWindow::IsTimersViewVisible)
+		.IsEnabled(this, &SIoProfilerWindow::IsProfilerEnabled)
+
+	// Header
+	+ SVerticalBox::Slot()
+		.AutoHeight()
+		[
+			SNew(SHorizontalBox)
+
+			+ SHorizontalBox::Slot()
+				.AutoWidth()
+				[
+					SNew(SImage)
+						.Image(FEditorStyle::GetBrush(TEXT("Profiler.Tab.FiltersAndPresets")))
+				]
+
+			+ SHorizontalBox::Slot()
+				.AutoWidth()
+				[
+					SNew(STextBlock)
+						.Text(LOCTEXT("TimersLabel", "Timers"))
+				]
+		]
+
+	// Timers
+	+ SVerticalBox::Slot()
+		.FillHeight(1.0f)
+		.Padding(0.0f, 2.0f, 0.0f, 0.0f)
+		[
+			SAssignNew(TimersView, STimersView)
+		]
+
+	; // end of Widget
+}
+
+////////////////////////////////////////////////////////////////////////////////////////////////////
+
+TSharedRef<SWidget> SIoProfilerWindow::ConstructLogView()
+{
+	return SNew(SVerticalBox)
+		.Visibility(this, &SIoProfilerWindow::IsLogViewVisible)
+
+	// Header
+	+ SVerticalBox::Slot()
+		.AutoHeight()
+		[
+			SNew(SHorizontalBox)
+			+ SHorizontalBox::Slot()
+				.AutoWidth()
+				[
+					SNew(SImage)
+						.Image(FEditorStyle::GetBrush(TEXT("Profiler.Tab.FiltersAndPresets")))
+				]
+
+			+ SHorizontalBox::Slot()
+				.AutoWidth()
+				[
+					SNew(STextBlock)
+						.Text(LOCTEXT("LogViewLabel", "Log View"))
+				]
+		]
+
+	// Log View
+	+ SVerticalBox::Slot()
+		.FillHeight(1.0f)
+		[
+			SAssignNew(LogView, SLogView)
+		]
+
+	; // end of Widget
+}
+
+////////////////////////////////////////////////////////////////////////////////////////////////////
+
+void SIoProfilerWindow::Construct(const FArguments& InArgs)
+{
+	ChildSlot
+		[
+			SNew(SOverlay)
+
+			// Version
+			+ SOverlay::Slot()
+				.HAlign(HAlign_Right)
+				.VAlign(VAlign_Top)
+				.Padding(0.0f, -16.0f, 0.0f, 0.0f)
+				[
+					SNew(STextBlock)
+						.Clipping(EWidgetClipping::ClipToBoundsWithoutIntersecting)
+						.Text(LOCTEXT("UnrealInsightsVersion", UNREAL_INSIGHTS_VERSION_STRING_EX))
+						.ColorAndOpacity(FLinearColor(0.15f, 0.15f, 0.15f, 1.0f))
+				]
+
+			// Overlay slot for the main profiler window area
+			+ SOverlay::Slot()
+				.HAlign(HAlign_Fill)
+				.VAlign(VAlign_Fill)
+				[
+					SAssignNew(MainContentPanel, SVerticalBox)
+
+					// Toolbar
+					//+ SVerticalBox::Slot()
+					//	.AutoHeight()
+					//	[
+					//TODO:		SNew(SIoProfilerToolbar)
+					//	]
+
+					+ SVerticalBox::Slot()
+						.HAlign(HAlign_Fill)
+						.VAlign(VAlign_Fill)
+						[
+							SNew(SSplitter)
+								.Orientation(Orient_Horizontal)
+
+							+ SSplitter::Slot()
+								.Value(0.65f)
+								[
+									SNew(SSplitter)
+										.Orientation(Orient_Vertical)
+
+									// Multi Track View
+									+ SSplitter::Slot()
+										.Value(0.75f)
+										[
+											//ConstructMultiTrackView()
+											ConstructTimingTrack()
+										]
+
+									// Log View
+									//+ SSplitter::Slot()
+									//	.Value(0.25f)
+									//	[
+									//		ConstructLogView()
+									//	]
+								]
+
+							// Timers View
+							//+ SSplitter::Slot()
+							//	.Value(0.35f)
+							//	//.SizeRule(SSplitter::SizeToContent)
+							//	.Expose(TimersViewSlot)
+							//	[
+							//		ConstructTimersView()
+							//	]
+						]
+				]
+
+			// session hint overlay
+			+ SOverlay::Slot()
+				.HAlign(HAlign_Center)
+				.VAlign(VAlign_Center)
+				[
+					SNew(SBorder)
+						.Visibility(this, &SIoProfilerWindow::IsSessionOverlayVisible)
+						.BorderImage(FEditorStyle::GetBrush("NotificationList.ItemBackground"))
+						.Padding(8.0f)
+						[
+							SNew(STextBlock)
+								.Text(LOCTEXT("SelectTraceOverlayText", "Please select a trace."))
+						]
+				]
+
+			// notification area overlay
+			+ SOverlay::Slot()
+				.HAlign(HAlign_Right)
+				.VAlign(VAlign_Bottom)
+				.Padding(16.0f)
+				[
+					SAssignNew(NotificationList, SNotificationList)
+				]
+
+			// profiler settings overlay
+			+ SOverlay::Slot()
+				.HAlign(HAlign_Center)
+				.VAlign(VAlign_Center)
+				.Expose(OverlaySettingsSlot)
+		];
+
+	TimingView->SetGraphTrackVisible(true);
+
+	//ProfilerMiniView->OnSelectionBoxChanged().AddSP(GraphPanel.ToSharedRef(), &SIoProfilerGraphPanel::MiniView_OnSelectionBoxChanged);
+	//FIoProfilerManager::Get()->OnViewModeChanged().AddSP(this, &SIoProfilerWindow::ProfilerManager_OnViewModeChanged);
+}
+
+END_SLATE_FUNCTION_BUILD_OPTIMIZATION
+
+////////////////////////////////////////////////////////////////////////////////////////////////////
+
+EVisibility SIoProfilerWindow::IsFramesTrackVisible() const
+{
+	if (FIoProfilerManager::Get()->IsFramesTrackVisible() &&
+		FInsightsManager::Get()->GetSession().IsValid())
+	{
+		return EVisibility::Visible;
+	}
+
+	return EVisibility::Collapsed;
+}
+
+////////////////////////////////////////////////////////////////////////////////////////////////////
+
+EVisibility SIoProfilerWindow::IsGraphTrackVisible() const
+{
+	if (FIoProfilerManager::Get()->IsGraphTrackVisible() &&
+		FInsightsManager::Get()->GetSession().IsValid())
+	{
+		return EVisibility::Visible;
+	}
+
+	return EVisibility::Collapsed;
+}
+
+////////////////////////////////////////////////////////////////////////////////////////////////////
+
+EVisibility SIoProfilerWindow::IsTimingTrackVisible() const
+{
+	if (FIoProfilerManager::Get()->IsTimingTrackVisible() &&
+		FInsightsManager::Get()->GetSession().IsValid())
+	{
+		return EVisibility::Visible;
+	}
+
+	return EVisibility::Collapsed;
+}
+
+////////////////////////////////////////////////////////////////////////////////////////////////////
+
+EVisibility SIoProfilerWindow::IsTimersViewVisible() const
+{
+	if (FIoProfilerManager::Get()->IsTimersViewVisible() &&
+		FInsightsManager::Get()->GetSession().IsValid())
+	{
+		return EVisibility::Visible;
+	}
+
+	return EVisibility::Collapsed;
+}
+
+////////////////////////////////////////////////////////////////////////////////////////////////////
+
+EVisibility SIoProfilerWindow::IsLogViewVisible() const
+{
+	if (FIoProfilerManager::Get()->IsLogViewVisible() &&
+		FInsightsManager::Get()->GetSession().IsValid())
+	{
+		return EVisibility::Visible;
+	}
+
+	return EVisibility::Collapsed;
+}
+
+////////////////////////////////////////////////////////////////////////////////////////////////////
+
+EVisibility SIoProfilerWindow::IsSessionOverlayVisible() const
+{
+	if (FInsightsManager::Get()->GetSession().IsValid())
+	{
+		return EVisibility::Hidden;
+	}
+
+	return EVisibility::Visible;
+}
+
+////////////////////////////////////////////////////////////////////////////////////////////////////
+
+bool SIoProfilerWindow::IsProfilerEnabled() const
+{
+	const bool bIsActive = FInsightsManager::Get()->GetSession().IsValid();
+	return bIsActive;
+}
+
+////////////////////////////////////////////////////////////////////////////////////////////////////
+
+void SIoProfilerWindow::ManageLoadingProgressNotificationState(const FString& Filename, const EInsightsNotificationType NotificatonType, const ELoadingProgressState ProgressState, const float LoadingProgress)
+{
+	const FString BaseFilename = FPaths::GetBaseFilename(Filename);
+
+	if (ProgressState == ELoadingProgressState::Started)
+	{
+		const bool bContains = ActiveNotifications.Contains(Filename);
+		if (!bContains)
+		{
+			FNotificationInfo NotificationInfo(Io_GetTextForNotification(NotificatonType, ProgressState, BaseFilename));
+			NotificationInfo.bFireAndForget = false;
+			NotificationInfo.bUseLargeFont = false;
+
+			// Add two buttons, one for cancel, one for loading the received file.
+			if (NotificatonType == EInsightsNotificationType::LoadingTraceFile)
+			{
+				const FText CancelButtonText = LOCTEXT("CancelButton_Text", "Cancel");
+				const FText CancelButtonTT = LOCTEXT("CancelButton_TTText", "Hides this notification");
+				const FText LoadButtonText = LOCTEXT("LoadButton_Text", "Load file");
+				const FText LoadButtonTT = LOCTEXT("LoadButton_TTText", "Loads the received file and hides this notification");
+
+				NotificationInfo.ButtonDetails.Add(FNotificationButtonInfo(CancelButtonText, CancelButtonTT,
+					FSimpleDelegate::CreateSP(this, &SIoProfilerWindow::SendingServiceSideCapture_Cancel, Filename), SNotificationItem::CS_Success));
+				NotificationInfo.ButtonDetails.Add(FNotificationButtonInfo(LoadButtonText, LoadButtonTT,
+					FSimpleDelegate::CreateSP(this, &SIoProfilerWindow::SendingServiceSideCapture_Load, Filename), SNotificationItem::CS_Success));
+			}
+
+			SNotificationItemWeak NotificationItem = NotificationList->AddNotification(NotificationInfo);
+			NotificationItem.Pin()->SetCompletionState(SNotificationItem::CS_Pending);
+			ActiveNotifications.Add(Filename, NotificationItem);
+		}
+	}
+	else if (ProgressState == ELoadingProgressState::InProgress)
+	{
+		const SNotificationItemWeak* LoadingProgressPtr = ActiveNotifications.Find(Filename);
+		if (LoadingProgressPtr)
+		{
+			TSharedPtr<SNotificationItem> LoadingProcessPinned = LoadingProgressPtr->Pin();
+			LoadingProcessPinned->SetText(Io_GetTextForNotification(NotificatonType, ProgressState, BaseFilename, LoadingProgress));
+			LoadingProcessPinned->SetCompletionState(SNotificationItem::CS_Pending);
+		}
+	}
+	else if (ProgressState == ELoadingProgressState::Loaded)
+	{
+		const SNotificationItemWeak* LoadingProgressPtr = ActiveNotifications.Find(Filename);
+		if (LoadingProgressPtr)
+		{
+			TSharedPtr<SNotificationItem> LoadingProcessPinned = LoadingProgressPtr->Pin();
+			LoadingProcessPinned->SetText(Io_GetTextForNotification(NotificatonType, ProgressState, BaseFilename));
+			LoadingProcessPinned->SetCompletionState(SNotificationItem::CS_Success);
+
+			// Notifications for received files are handled by the user.
+			if (NotificatonType == EInsightsNotificationType::LoadingTraceFile)
+			{
+				LoadingProcessPinned->ExpireAndFadeout();
+				ActiveNotifications.Remove(Filename);
+			}
+		}
+	}
+	else if (ProgressState == ELoadingProgressState::Failed)
+	{
+		const SNotificationItemWeak* LoadingProgressPtr = ActiveNotifications.Find(Filename);
+		if (LoadingProgressPtr)
+		{
+			TSharedPtr<SNotificationItem> LoadingProcessPinned = LoadingProgressPtr->Pin();
+			LoadingProcessPinned->SetText(Io_GetTextForNotification(NotificatonType, ProgressState, BaseFilename));
+			LoadingProcessPinned->SetCompletionState(SNotificationItem::CS_Fail);
+
+			LoadingProcessPinned->ExpireAndFadeout();
+			ActiveNotifications.Remove(Filename);
+		}
+	}
+	else if (ProgressState == ELoadingProgressState::Cancelled)
+	{
+		const SNotificationItemWeak* LoadingProgressPtr = ActiveNotifications.Find(Filename);
+		if (LoadingProgressPtr)
+		{
+			TSharedPtr<SNotificationItem> LoadingProcessPinned = LoadingProgressPtr->Pin();
+			LoadingProcessPinned->ExpireAndFadeout();
+			ActiveNotifications.Remove(Filename);
+		}
+	}
+}
+
+////////////////////////////////////////////////////////////////////////////////////////////////////
+
+void SIoProfilerWindow::SendingServiceSideCapture_Cancel(const FString Filename)
+{
+	const SNotificationItemWeak* LoadingProgressPtr = ActiveNotifications.Find(Filename);
+	if (LoadingProgressPtr)
+	{
+		TSharedPtr<SNotificationItem> LoadingProcessPinned = LoadingProgressPtr->Pin();
+		LoadingProcessPinned->ExpireAndFadeout();
+		ActiveNotifications.Remove(Filename);
+	}
+}
+
+////////////////////////////////////////////////////////////////////////////////////////////////////
+
+void SIoProfilerWindow::SendingServiceSideCapture_Load(const FString Filename)
+{
+	const SNotificationItemWeak* LoadingProgressPtr = ActiveNotifications.Find(Filename);
+	if (LoadingProgressPtr)
+	{
+		TSharedPtr<SNotificationItem> LoadingProcessPinned = LoadingProgressPtr->Pin();
+		LoadingProcessPinned->ExpireAndFadeout();
+		ActiveNotifications.Remove(Filename);
+
+		const FString PathName = FPaths::ProfilingDir() + TEXT("UnrealStats/Received/");
+		const FString TraceFilepath = PathName + Filename;
+		FInsightsManager::Get()->LoadTraceFile(TraceFilepath);
+	}
+}
+
+////////////////////////////////////////////////////////////////////////////////////////////////////
+
+EActiveTimerReturnType SIoProfilerWindow::UpdateActiveDuration(double InCurrentTime, float InDeltaTime)
+{
+	DurationActive += InDeltaTime;
+
+	// The profiler window will explicitly unregister this active timer when the mouse leaves.
+	return EActiveTimerReturnType::Continue;
+}
+
+////////////////////////////////////////////////////////////////////////////////////////////////////
+
+void SIoProfilerWindow::OnMouseEnter(const FGeometry& MyGeometry, const FPointerEvent& MouseEvent)
+{
+	SCompoundWidget::OnMouseEnter(MyGeometry, MouseEvent);
+
+	if (!ActiveTimerHandle.IsValid())
+	{
+		ActiveTimerHandle = RegisterActiveTimer(0.f, FWidgetActiveTimerDelegate::CreateSP(this, &SIoProfilerWindow::UpdateActiveDuration));
+	}
+}
+
+////////////////////////////////////////////////////////////////////////////////////////////////////
+
+void SIoProfilerWindow::OnMouseLeave(const FPointerEvent& MouseEvent)
+{
+	SCompoundWidget::OnMouseLeave(MouseEvent);
+
+	auto PinnedActiveTimerHandle = ActiveTimerHandle.Pin();
+	if (PinnedActiveTimerHandle.IsValid())
+	{
+		UnRegisterActiveTimer(PinnedActiveTimerHandle.ToSharedRef());
+	}
+}
+
+////////////////////////////////////////////////////////////////////////////////////////////////////
+
+FReply SIoProfilerWindow::OnKeyDown(const FGeometry& MyGeometry, const FKeyEvent& InKeyEvent)
+{
+	return FIoProfilerManager::Get()->GetCommandList()->ProcessCommandBindings(InKeyEvent) ? FReply::Handled() : FReply::Unhandled();
+}
+
+////////////////////////////////////////////////////////////////////////////////////////////////////
+
+FReply SIoProfilerWindow::OnDragOver(const FGeometry& MyGeometry, const FDragDropEvent& DragDropEvent)
+{
+	TSharedPtr<FExternalDragOperation> DragDropOp = DragDropEvent.GetOperationAs<FExternalDragOperation>();
+	if (DragDropOp.IsValid())
+	{
+		if (DragDropOp->HasFiles())
+		{
+			const TArray<FString>& Files = DragDropOp->GetFiles();
+			if (Files.Num() == 1)
+			{
+				const FString DraggedFileExtension = FPaths::GetExtension(Files[0], true);
+				if (DraggedFileExtension == TEXT(".utrace"))
+				{
+					return FReply::Handled();
+				}
+			}
+		}
+	}
+
+	return SCompoundWidget::OnDragOver(MyGeometry,DragDropEvent);
+}
+
+////////////////////////////////////////////////////////////////////////////////////////////////////
+
+FReply SIoProfilerWindow::OnDrop(const FGeometry& MyGeometry, const FDragDropEvent& DragDropEvent)
+{
+	TSharedPtr<FExternalDragOperation> DragDropOp = DragDropEvent.GetOperationAs<FExternalDragOperation>();
+	if (DragDropOp.IsValid())
+	{
+		if (DragDropOp->HasFiles())
+		{
+			// For now, only allow a single file.
+			const TArray<FString>& Files = DragDropOp->GetFiles();
+			if (Files.Num() == 1)
+			{
+				const FString DraggedFileExtension = FPaths::GetExtension(Files[0], true);
+				if (DraggedFileExtension == TEXT(".utrace"))
+				{
+					// Enqueue load operation.
+					FInsightsManager::Get()->LoadTraceFile(Files[0]);
+					return FReply::Handled();
+				}
+			}
+		}
+	}
+
+	return SCompoundWidget::OnDrop(MyGeometry,DragDropEvent);
+}
+
+////////////////////////////////////////////////////////////////////////////////////////////////////
+
+void SIoProfilerWindow::OpenProfilerSettings()
+{
+	MainContentPanel->SetEnabled(false);
+	(*OverlaySettingsSlot)
+	[
+		SNew(SBorder)
+		.BorderImage(FEditorStyle::GetBrush("NotificationList.ItemBackground"))
+		.Padding(8.0f)
+		[
+			SNew(SInsightsSettings)
+			.OnClose(this, &SIoProfilerWindow::CloseProfilerSettings)
+			.SettingPtr(&FInsightsManager::GetSettings())
+		]
+	];
+}
+
+////////////////////////////////////////////////////////////////////////////////////////////////////
+
+void SIoProfilerWindow::CloseProfilerSettings()
+{
+	// Close the profiler settings by simply replacing widget with a null one.
+	(*OverlaySettingsSlot)
+	[
+		SNullWidget::NullWidget
+	];
+	MainContentPanel->SetEnabled(true);
+}
+
+////////////////////////////////////////////////////////////////////////////////////////////////////
+
+//void SIoProfilerWindow::ProfilerManager_OnViewModeChanged(EIoProfilerViewMode NewViewMode)
+//{
+//	if (NewViewMode == EIoProfilerViewMode::LineIndexBased)
+//	{
+//		EventGraphPanel->SetVisibility(EVisibility::Visible);
+//		EventGraphPanel->SetEnabled(true);
+//
+//		(*FiltersAndPresetsSlot)
+//		[
+//			FiltersAndPresets.ToSharedRef()
+//		];
+//	}
+//	else if (NewViewMode == EIoProfilerViewMode::ThreadViewTimeBased)
+//	{
+//		EventGraphPanel->SetVisibility(EVisibility::Collapsed);
+//		EventGraphPanel->SetEnabled(false);
+//
+//		(*FiltersAndPresetsSlot)
+//		[
+//			SNullWidget::NullWidget
+//		];
+//	}
+//}
+
+////////////////////////////////////////////////////////////////////////////////////////////////////
+
+#undef LOCTEXT_NAMESPACE
diff --git a/Engine/Source/Developer/TraceInsights/Private/Insights/Widgets/SIoProfilerWindow.h b/Engine/Source/Developer/TraceInsights/Private/Insights/Widgets/SIoProfilerWindow.h
new file mode 100644
--- /dev/null
+++ b/Engine/Source/Developer/TraceInsights/Private/Insights/Widgets/SIoProfilerWindow.h
@@ -0,0 +1,187 @@
+// Copyright 1998-2019 Epic Games, Inc. All Rights Reserved.
+
+#pragma once
+
+#include "CoreMinimal.h"
+#include "Input/Reply.h"
+#include "Layout/Visibility.h"
+#include "Misc/Guid.h"
+#include "SlateFwd.h"
+#include "Widgets/DeclarativeSyntaxSupport.h"
+#include "Widgets/Layout/SSplitter.h"
+#include "Widgets/SCompoundWidget.h"
+#include "Widgets/SOverlay.h"
+
+// Insights
+#include "Insights/InsightsManager.h"
+
+class FActiveTimerHandle;
+class SVerticalBox;
+class SFrameTrack;
+class SGraphTrack;
+class STimersView;
+class STimingView;
+class SLogView;
+
+/// Type definition for shared pointers to instances of SNotificationItem.
+typedef TSharedPtr<class SNotificationItem> SNotificationItemPtr;
+
+/// Type definition for shared references to instances of SNotificationItem.
+typedef TSharedRef<class SNotificationItem> SNotificationItemRef;
+
+/// Type definition for weak references to instances of SNotificationItem.
+typedef TWeakPtr<class SNotificationItem> SNotificationItemWeak;
+
+/// Implements the timing profiler window.
+class SIoProfilerWindow : public SCompoundWidget
+{
+public:
+	/// Default constructor.
+	SIoProfilerWindow();
+
+	/// Virtual destructor.
+	virtual ~SIoProfilerWindow();
+
+	SLATE_BEGIN_ARGS(SIoProfilerWindow){}
+	SLATE_END_ARGS()
+
+	/// Constructs this widget.
+	void Construct(const FArguments& InArgs);
+
+	void ManageLoadingProgressNotificationState(const FString& Filename, const EInsightsNotificationType NotificatonType, const ELoadingProgressState ProgressState, const float LoadingProgress);
+
+	void OpenProfilerSettings();
+	void CloseProfilerSettings();
+
+protected:
+	TSharedRef<SWidget> ConstructMultiTrackView();
+	TSharedRef<SWidget> ConstructFramesTrack();
+	TSharedRef<SWidget> ConstructGraphTrack();
+	TSharedRef<SWidget> ConstructTimingTrack();
+	TSharedRef<SWidget> ConstructTimersView();
+	TSharedRef<SWidget> ConstructLogView();
+
+	/// Callback for determining the visibility of the Frames track.
+	EVisibility IsFramesTrackVisible() const;
+
+	/// Callback for determining the visibility of the Graph track.
+	EVisibility IsGraphTrackVisible() const;
+
+	/// Callback for determining the visibility of the Timing track.
+	EVisibility IsTimingTrackVisible() const;
+
+	/// Callback for determining the visibility of the Timers View.
+	EVisibility IsTimersViewVisible() const;
+
+	/// Callback for determining the visibility of the Log View.
+	EVisibility IsLogViewVisible() const;
+
+	/// Callback for determining the visibility of the 'Select a session' overlay.
+	EVisibility IsSessionOverlayVisible() const;
+
+	/// Callback for getting the enabled state of the profiler window.
+	bool IsProfilerEnabled() const;
+
+	void SendingServiceSideCapture_Cancel(const FString Filename);
+
+	void SendingServiceSideCapture_Load(const FString Filename);
+
+private:
+
+	/**
+	 * Ticks this widget.  Override in derived classes, but always call the parent implementation.
+	 *
+	 * @param  AllottedGeometry The space allotted for this widget
+	 * @param  InCurrentTime  Current absolute real time
+	 * @param  InDeltaTime  Real time passed since last tick
+	 */
+	//virtual void Tick(const FGeometry& AllottedGeometry, const double InCurrentTime, const float InDeltaTime) override;
+
+	/**
+	 * The system will use this event to notify a widget that the cursor has entered it. This event is NOT bubbled.
+	 *
+	 * @param MyGeometry The Geometry of the widget receiving the event
+	 * @param MouseEvent Information about the input event
+	 */
+	virtual void OnMouseEnter(const FGeometry& MyGeometry, const FPointerEvent& MouseEvent) override;
+
+	/**
+	 * The system will use this event to notify a widget that the cursor has left it. This event is NOT bubbled.
+	 *
+	 * @param MouseEvent Information about the input event
+	 */
+	virtual void OnMouseLeave(const FPointerEvent& MouseEvent) override;
+
+	/**
+	 * Called after a key is pressed when this widget has focus
+	 *
+	 * @param MyGeometry The Geometry of the widget receiving the event
+	 * @param  InKeyEvent  Key event
+	 *
+	 * @return  Returns whether the event was handled, along with other possible actions
+	 */
+	virtual FReply OnKeyDown(const FGeometry& MyGeometry, const FKeyEvent& InKeyEvent) override;
+
+	/**
+	 * Called when the user is dropping something onto a widget; terminates drag and drop.
+	 *
+	 * @param MyGeometry      The geometry of the widget receiving the event.
+	 * @param DragDropEvent   The drag and drop event.
+	 *
+	 * @return A reply that indicated whether this event was handled.
+	 */
+	virtual FReply OnDrop(const FGeometry& MyGeometry, const FDragDropEvent& DragDropEvent) override;
+
+	/**
+	 * Called during drag and drop when the the mouse is being dragged over a widget.
+	 *
+	 * @param MyGeometry      The geometry of the widget receiving the event.
+	 * @param DragDropEvent   The drag and drop event.
+	 *
+	 * @return A reply that indicated whether this event was handled.
+	 */
+	virtual FReply OnDragOver(const FGeometry& MyGeometry, const FDragDropEvent& DragDropEvent)  override;
+
+private:
+	/// Updates the amount of time the profiler has been active
+	EActiveTimerReturnType UpdateActiveDuration(double InCurrentTime, float InDeltaTime);
+
+public:
+	/// Widget for the frame track
+	TSharedPtr<SFrameTrack> FrameTrack;
+
+	/// Widget for the graph track
+	TSharedPtr<SGraphTrack> GraphTrack;
+
+	/// Holds all timing tracks (one for each thread).
+	TSharedPtr<SVerticalBox> ThreadsPanel;
+
+	/// Widget for the timing track
+	TSharedPtr<STimingView> TimingView;
+
+	/// Holds the Timers view widget/slot.
+	SSplitter::FSlot* TimersViewSlot;
+	TSharedPtr<STimersView> TimersView;
+
+	/// Widget for the log view
+	TSharedPtr<SLogView> LogView;
+
+	/// Widget for the non-intrusive notifications.
+	TSharedPtr<SNotificationList> NotificationList;
+
+	/// Holds all active and visible notifications, stored as FGuid -> SNotificationItemWeak.
+	TMap<FString, SNotificationItemWeak> ActiveNotifications;
+
+	/// Overlay slot which contains the profiler settings widget.
+	SOverlay::FOverlaySlot* OverlaySettingsSlot;
+
+	/// The number of seconds the profiler has been active
+	float DurationActive;
+
+private:
+	/// The handle to the active update duration tick
+	TWeakPtr<FActiveTimerHandle> ActiveTimerHandle;
+
+	/// Holds all widgets for the profiler window like menu bar, toolbar and tabs.
+	TSharedPtr<SVerticalBox> MainContentPanel;
+};
