Copy up of Unreal Insights from //Tasks/UE4/Dev-Core-Trace.

diff --git a/Engine/Source/Developer/TraceInsights/Private/Insights/Widgets/SLogView.cpp b/Engine/Source/Developer/TraceInsights/Private/Insights/Widgets/SLogView.cpp
new file mode 100644
--- /dev/null
+++ b/Engine/Source/Developer/TraceInsights/Private/Insights/Widgets/SLogView.cpp
@@ -0,0 +1,1250 @@
+// Copyright 1998-2019 Epic Games, Inc. All Rights Reserved.
+
+#include "SLogView.h"
+
+#include "Async/AsyncWork.h"
+#include "EditorStyleSet.h"
+#include "Framework/Application/SlateApplication.h"
+#include "Framework/MultiBox/MultiBoxBuilder.h"
+#include "SlateOptMacros.h"
+#include "Widgets/Images/SImage.h"
+#include "Widgets/Input/SComboButton.h"
+#include "Widgets/Layout/SScrollBox.h"
+#include "Widgets/Text/STextBlock.h"
+
+// Insights
+#include "Insights/Common/TimeUtils.h"
+#include "Insights/InsightsManager.h"
+#include "Insights/TimingProfilerCommon.h"
+#include "Insights/TimingProfilerManager.h"
+#include "Insights/ViewModels/TimingViewDrawHelper.h"
+#include "Insights/Widgets/STimingProfilerWindow.h"
+#include "Insights/Widgets/STimingView.h"
+
+////////////////////////////////////////////////////////////////////////////////////////////////////
+
+#define LOCTEXT_NAMESPACE "SLogView"
+
+
+////////////////////////////////////////////////////////////////////////////////////////////////////
+// SLogMessageRow
+////////////////////////////////////////////////////////////////////////////////////////////////////
+
+class SLogMessageRow : public SMultiColumnTableRow<TSharedPtr<FLogMessage>>
+{
+	SLATE_BEGIN_ARGS(SLogMessageRow) {}
+	SLATE_END_ARGS()
+
+public:
+	/**
+	 * Constructs the widget.
+	 *
+	 * @param InArgs The construction arguments.
+	 * @param InLogMessage The log message displayed by this row.
+	 * @param InOwnerTableView The table to which the row must be added.
+	 */
+	void Construct(const FArguments& InArgs, TSharedPtr<FLogMessage> InLogMessage, TSharedRef<SLogView> InParentWidget, const TSharedRef<STableViewBase>& InOwnerTableView)
+	{
+		WeakLogMessage = MoveTemp(InLogMessage);
+		WeakParentWidget = InParentWidget;
+
+		SMultiColumnTableRow<TSharedPtr<FLogMessage>>::Construct(FSuperRowType::FArguments(), InOwnerTableView);
+
+		TSharedRef<SWidget> Row = ChildSlot.GetChildAt(0);
+
+		ChildSlot
+		[
+			SNew(SBorder)
+				.BorderImage(FCoreStyle::Get().GetBrush("WhiteBrush"))
+				.BorderBackgroundColor(this, &SLogMessageRow::GetBackgroundColor)
+				[
+					Row
+				]
+		];
+	}
+
+public:
+	virtual TSharedRef<SWidget> GenerateWidgetForColumn(const FName& ColumnName) override
+	{
+		if (ColumnName == LogViewColumns::IdColumnName)
+		{
+			return SNew(SBox)
+				.Padding(FMargin(4.0, 0.0))
+				[
+					SNew(STextBlock)
+					.Text(this, &SLogMessageRow::GetIndex)
+				];
+		}
+		else if (ColumnName == LogViewColumns::TimeColumnName)
+		{
+			return SNew(SBox)
+				.Padding(FMargin(4.0, 0.0))
+				[
+					SNew(STextBlock)
+					.Text(this, &SLogMessageRow::GetTime)
+					.ColorAndOpacity(this, &SLogMessageRow::GetColorAndOpacity)
+				];
+		}
+		else if (ColumnName == LogViewColumns::VerbosityColumnName)
+		{
+			return SNew(SBox)
+				.Padding(FMargin(4.0, 0.0))
+				[
+					SNew(STextBlock)
+					.Text(this, &SLogMessageRow::GetVerbosity)
+					.ColorAndOpacity(this, &SLogMessageRow::GetColorByVerbosity)
+				];
+		}
+		else if (ColumnName == LogViewColumns::CategoryColumnName)
+		{
+			return SNew(SBox)
+				.Padding(FMargin(4.0, 0.0))
+				[
+					SNew(STextBlock)
+					.Text(this, &SLogMessageRow::GetCategory)
+					.ColorAndOpacity(this, &SLogMessageRow::GetColorByCategory)
+				];
+		}
+		else if (ColumnName == LogViewColumns::MessageColumnName)
+		{
+			return SNew(SBox)
+				.Padding(FMargin(4.0, 0.0))
+				[
+					SNew(STextBlock)
+					.Text(this, &SLogMessageRow::GetMessage)
+					.HighlightText(this, &SLogMessageRow::GetMessageHighlightText)
+				];
+		}
+		else if (ColumnName == LogViewColumns::FileColumnName)
+		{
+			return SNew(SBox)
+				.Padding(FMargin(4.0, 0.0))
+				[
+					SNew(STextBlock)
+					.Text(this, &SLogMessageRow::GetFile)
+				];
+		}
+		else if (ColumnName == LogViewColumns::LineColumnName)
+		{
+			return SNew(SBox)
+				.Padding(FMargin(4.0, 0.0))
+				[
+					SNew(STextBlock)
+					.Text(this, &SLogMessageRow::GetLine)
+				];
+		}
+		else
+		{
+			return SNew(STextBlock).Text(LOCTEXT("UnknownColumn", "Unknown Column"));
+		}
+	}
+
+	FSlateColor GetBackgroundColor() const
+	{
+		TSharedPtr<SLogView> ParentWidgetPin = WeakParentWidget.Pin();
+		TSharedPtr<FLogMessage> LogMessagePin = WeakLogMessage.Pin();
+		if (ParentWidgetPin.IsValid() && LogMessagePin.IsValid())
+		{
+			TSharedPtr<FLogMessage> SelectedLogMessage = ParentWidgetPin->GetSelectedLogMessage();
+			if (!SelectedLogMessage || SelectedLogMessage->GetIndex() != LogMessagePin->GetIndex()) // if row is not selected
+			{
+				FLogMessageRecord& CacheEntry = ParentWidgetPin->GetCache().Get(LogMessagePin->GetIndex());
+				const double Time = CacheEntry.Time;
+
+				TSharedPtr<STimingProfilerWindow> Window = FTimingProfilerManager::Get()->GetProfilerWindow();
+				if (Window && Window->TimingView)
+				{
+					if (Window->TimingView->IsTimeSelectedInclusive(Time))
+					{
+						return FSlateColor(FLinearColor(0.25f, 0.5f, 1.0f, 0.25f));
+					}
+				}
+			}
+		}
+
+		return FSlateColor(FLinearColor(0.0f, 0.0f, 0.0f, 0.0f));
+	}
+
+	FSlateColor GetColorAndOpacity() const
+	{
+		bool IsSelected = false;
+
+		TSharedPtr<SLogView> ParentWidgetPin = WeakParentWidget.Pin();
+		TSharedPtr<FLogMessage> LogMessagePin = WeakLogMessage.Pin();
+		if (ParentWidgetPin.IsValid() && LogMessagePin.IsValid())
+		{
+			TSharedPtr<FLogMessage> SelectedLogMessage = ParentWidgetPin->GetSelectedLogMessage();
+			if (SelectedLogMessage && SelectedLogMessage->GetIndex() == LogMessagePin->GetIndex())
+			{
+				IsSelected = true;
+			}
+
+			FLogMessageRecord& CacheEntry = ParentWidgetPin->GetCache().Get(LogMessagePin->GetIndex());
+			const double Time = CacheEntry.Time;
+
+			TSharedPtr<STimingProfilerWindow> Window = FTimingProfilerManager::Get()->GetProfilerWindow();
+			if (Window && Window->TimingView)
+			{
+				if (Window->TimingView->IsTimeSelectedInclusive(Time))
+				{
+					if (IsSelected)
+					{
+						//return FSlateColor(FLinearColor(0.0f, 0.1f, 0.5f, 1.0f));
+						return FSlateColor(FLinearColor(0.0f, 0.05f, 0.2f, 1.0f));
+					}
+					else
+					{
+						//return FSlateColor(FLinearColor(0.2f, 0.4f, 0.8f, 1.0f));
+						return FSlateColor(FLinearColor(0.4f, 0.8f, 1.6f, 1.0f));
+					}
+				}
+			}
+		}
+
+		if (IsSelected)
+		{
+			return FSlateColor(FLinearColor::Black);
+		}
+		else
+		{
+			return FSlateColor(FLinearColor::White);
+		}
+	}
+
+	FText GetIndex() const
+	{
+		TSharedPtr<SLogView> ParentWidgetPin = WeakParentWidget.Pin();
+		TSharedPtr<FLogMessage> LogMessagePin = WeakLogMessage.Pin();
+		if (ParentWidgetPin.IsValid() && LogMessagePin.IsValid())
+		{
+			FLogMessageRecord& CacheEntry = ParentWidgetPin->GetCache().Get(LogMessagePin->GetIndex());
+			return CacheEntry.GetIndexAsText();
+		}
+		else
+		{
+			return FText();
+		}
+	}
+
+	FText GetTime() const
+	{
+		TSharedPtr<SLogView> ParentWidgetPin = WeakParentWidget.Pin();
+		TSharedPtr<FLogMessage> LogMessagePin = WeakLogMessage.Pin();
+		if (ParentWidgetPin.IsValid() && LogMessagePin.IsValid())
+		{
+			FLogMessageRecord& CacheEntry = ParentWidgetPin->GetCache().Get(LogMessagePin->GetIndex());
+			return CacheEntry.GetTimeAsText();
+		}
+		else
+		{
+			return FText();
+		}
+	}
+
+	FText GetVerbosity() const
+	{
+		TSharedPtr<SLogView> ParentWidgetPin = WeakParentWidget.Pin();
+		TSharedPtr<FLogMessage> LogMessagePin = WeakLogMessage.Pin();
+		if (ParentWidgetPin.IsValid() && LogMessagePin.IsValid())
+		{
+			FLogMessageRecord& CacheEntry = ParentWidgetPin->GetCache().Get(LogMessagePin->GetIndex());
+			return CacheEntry.GetVerbosityAsText();
+		}
+		else
+		{
+			return FText();
+		}
+	}
+
+	FSlateColor GetColorByVerbosity() const
+	{
+		if (IsSelected())
+			return FSlateColor(FLinearColor::Black);
+
+		TSharedPtr<SLogView> ParentWidgetPin = WeakParentWidget.Pin();
+		TSharedPtr<FLogMessage> LogMessagePin = WeakLogMessage.Pin();
+		if (ParentWidgetPin.IsValid() && LogMessagePin.IsValid())
+		{
+			FLogMessageRecord& CacheEntry = ParentWidgetPin->GetCache().Get(LogMessagePin->GetIndex());
+			return FSlateColor(FTimeMarkerTrackBuilder::GetColorByVerbosity(CacheEntry.Verbosity));
+		}
+		else
+		{
+			return FSlateColor(FLinearColor::White);
+		}
+	}
+
+	FText GetCategory() const
+	{
+		TSharedPtr<SLogView> ParentWidgetPin = WeakParentWidget.Pin();
+		TSharedPtr<FLogMessage> LogMessagePin = WeakLogMessage.Pin();
+		if (ParentWidgetPin.IsValid() && LogMessagePin.IsValid())
+		{
+			FLogMessageRecord& CacheEntry = ParentWidgetPin->GetCache().Get(LogMessagePin->GetIndex());
+			return CacheEntry.GetCategoryAsText();
+		}
+		else
+		{
+			return FText();
+		}
+	}
+
+	FSlateColor GetColorByCategory() const
+	{
+		if (IsSelected())
+			return FSlateColor(FLinearColor::Black);
+
+		TSharedPtr<SLogView> ParentWidgetPin = WeakParentWidget.Pin();
+		TSharedPtr<FLogMessage> LogMessagePin = WeakLogMessage.Pin();
+		if (ParentWidgetPin.IsValid() && LogMessagePin.IsValid())
+		{
+			FLogMessageRecord& CacheEntry = ParentWidgetPin->GetCache().Get(LogMessagePin->GetIndex());
+			return FSlateColor(FTimeMarkerTrackBuilder::GetColorByCategory(*CacheEntry.Category.ToString()));
+		}
+		else
+		{
+			return FSlateColor(FLinearColor::White);
+		}
+	}
+
+	FText GetMessage() const
+	{
+		TSharedPtr<SLogView> ParentWidgetPin = WeakParentWidget.Pin();
+		TSharedPtr<FLogMessage> LogMessagePin = WeakLogMessage.Pin();
+		if (ParentWidgetPin.IsValid() && LogMessagePin.IsValid())
+		{
+			FLogMessageRecord& CacheEntry = ParentWidgetPin->GetCache().Get(LogMessagePin->GetIndex());
+			return CacheEntry.GetMessageAsText();
+		}
+		else
+		{
+			return FText();
+		}
+	}
+
+	FText GetMessageHighlightText() const
+	{
+		TSharedPtr<SLogView> ParentWidgetPin = WeakParentWidget.Pin();
+		if (ParentWidgetPin.IsValid())
+		{
+			return ParentWidgetPin->GetFilterText();
+		}
+		else
+		{
+			return FText();
+		}
+	}
+
+	FText GetFile() const
+	{
+		TSharedPtr<SLogView> ParentWidgetPin = WeakParentWidget.Pin();
+		TSharedPtr<FLogMessage> LogMessagePin = WeakLogMessage.Pin();
+		if (ParentWidgetPin.IsValid() && LogMessagePin.IsValid())
+		{
+			FLogMessageRecord& CacheEntry = ParentWidgetPin->GetCache().Get(LogMessagePin->GetIndex());
+			return CacheEntry.GetFileAsText();
+		}
+		else
+		{
+			return FText();
+		}
+	}
+
+	FText GetLine() const
+	{
+		TSharedPtr<SLogView> ParentWidgetPin = WeakParentWidget.Pin();
+		TSharedPtr<FLogMessage> LogMessagePin = WeakLogMessage.Pin();
+		if (ParentWidgetPin.IsValid() && LogMessagePin.IsValid())
+		{
+			FLogMessageRecord& CacheEntry = ParentWidgetPin->GetCache().Get(LogMessagePin->GetIndex());
+			return CacheEntry.GetLineAsText();
+		}
+		else
+		{
+			return FText();
+		}
+	}
+
+	FText GetRowToolTip() const
+	{
+		TSharedPtr<SLogView> ParentWidgetPin = WeakParentWidget.Pin();
+		TSharedPtr<FLogMessage> LogMessagePin = WeakLogMessage.Pin();
+		if (ParentWidgetPin.IsValid() && LogMessagePin.IsValid())
+		{
+			FLogMessageRecord& CacheEntry = ParentWidgetPin->GetCache().Get(LogMessagePin->GetIndex());
+			return CacheEntry.ToDisplayString();
+		}
+		else
+		{
+			return FText();
+		}
+	}
+
+private:
+	TWeakPtr<FLogMessage> WeakLogMessage;
+	TWeakPtr<SLogView> WeakParentWidget;
+};
+
+////////////////////////////////////////////////////////////////////////////////////////////////////
+// SLogView
+////////////////////////////////////////////////////////////////////////////////////////////////////
+
+SLogView::SLogView()
+	: FilteringStartIndex(0)
+	, FilteringEndIndex(0)
+	, FilteringChangeNumber(0)
+	, bIsFilteringAsyncTaskCancelRequested(false)
+	, TotalNumCategories(0)
+	, TotalNumMessages(0)
+	, bIsDirty(false)
+{
+}
+
+////////////////////////////////////////////////////////////////////////////////////////////////////
+
+SLogView::~SLogView()
+{
+	// Remove ourselves from the profiler manager.
+	//if (FTimingProfilerManager::Get().IsValid())
+	//{
+	//	//TODO: FTimimgProfilerManager::Get()->OnRequestLogViewUpdate().RemoveAll(this);
+	//}
+
+	Reset();
+}
+
+////////////////////////////////////////////////////////////////////////////////////////////////////
+
+void SLogView::Reset()
+{
+	//ListView
+	//ExternalScrollbar
+	FilterTextBox->SetText(FText::GetEmpty());
+
+	Filter.Reset();
+	FilterChangeNumber = 0;
+
+	FilteringStartIndex = 0;
+	FilteringEndIndex = 0;
+	FilteringChangeNumber = 0;
+
+	// Clean up our async task if we're deleted before it is completed.
+	if (FilteringAsyncTask.IsValid())
+	{
+		if (!FilteringAsyncTask->Cancel())
+		{
+			bIsFilteringAsyncTaskCancelRequested = true;
+			FilteringAsyncTask->EnsureCompletion();
+		}
+		FilteringAsyncTask.Reset();
+	}
+
+	//bIsFilteringAsyncTaskCancelRequested = false;
+	//FilteringStopwatch.Stop();
+
+	TotalNumCategories = 0;
+	TotalNumMessages = 0;
+
+	bIsDirty = false;
+	DirtyStopwatch.Stop();
+
+	StatsText = FText::GetEmpty();
+
+	Cache.Reset();
+
+	Messages.Reset();
+
+	ListView->RebuildList();
+}
+
+////////////////////////////////////////////////////////////////////////////////////////////////////
+
+BEGIN_SLATE_FUNCTION_BUILD_OPTIMIZATION
+void SLogView::Construct(const FArguments& InArgs)
+{
+	SAssignNew(ExternalScrollbar, SScrollBar)
+	.AlwaysShowScrollbar(true);
+
+	ChildSlot
+	[
+		SNew(SVerticalBox)
+
+		// Toolbar
+		+ SVerticalBox::Slot()
+		.AutoHeight()
+		[
+			SNew(SHorizontalBox)
+
+			// Icon
+			+ SHorizontalBox::Slot()
+			.VAlign(VAlign_Center)
+			.AutoWidth()
+			[
+				SNew(SImage)
+				.Image(FEditorStyle::GetBrush(TEXT("Log.TabIcon")))
+			]
+
+			// Label
+			+ SHorizontalBox::Slot()
+			.VAlign(VAlign_Center)
+			.Padding(2.0f, 0.0f)
+			.AutoWidth()
+			[
+				SNew(STextBlock)
+				.Text(LOCTEXT("LogViewLabel", "Log View"))
+			]
+
+			// Verbosity Threshold Filter
+			+SHorizontalBox::Slot()
+			.Padding(2.0f, 0.0f, 0.0f, 0.0f)
+			.VAlign(VAlign_Center)
+			.AutoWidth()
+			[
+				SNew(SComboButton)
+				.ComboButtonStyle(FEditorStyle::Get(), "GenericFilters.ComboButtonStyle")
+				.ForegroundColor(FLinearColor::White)
+				.ContentPadding(0)
+				.ToolTipText(LOCTEXT("VerbosityThresholdFilterToolTip", "Filter log messages by verbosity threshold."))
+				.OnGetMenuContent(this, &SLogView::MakeVerbosityThresholdMenu)
+				.HasDownArrow(true)
+				.ContentPadding(FMargin(1.0f, 1.0f, 1.0f, 0.0f))
+				.ButtonContent()
+				[
+					SNew(SHorizontalBox)
+
+					+SHorizontalBox::Slot()
+					.Padding(0.0f)
+					.AutoWidth()
+					[
+						SNew(STextBlock)
+						.TextStyle(FEditorStyle::Get(), "GenericFilters.TextStyle")
+						.Font(FEditorStyle::Get().GetFontStyle("FontAwesome.9"))
+						.Text(FText::FromString(FString(TEXT("\xf0b0"))) /*fa-filter*/)
+					]
+
+					+SHorizontalBox::Slot()
+					.Padding(2.0f, 0.0f, 0.0f, 0.0f)
+					.AutoWidth()
+					[
+						SNew(STextBlock)
+						.TextStyle(FEditorStyle::Get(), "GenericFilters.TextStyle")
+						.Text(LOCTEXT("VerbosityThresholdFilter", "Verbosity Threshold"))
+					]
+				]
+			]
+
+			// Category Filter
+			+SHorizontalBox::Slot()
+			.Padding(2.0f, 0.0f, 0.0f, 0.0f)
+			.VAlign(VAlign_Center)
+			.AutoWidth()
+			[
+				SNew(SComboButton)
+				.ComboButtonStyle(FEditorStyle::Get(), "GenericFilters.ComboButtonStyle")
+				.ForegroundColor(FLinearColor::White)
+				.ContentPadding(0)
+				.ToolTipText(LOCTEXT("CategoryFilterToolTip", "Filter log messages by category."))
+				.OnGetMenuContent(this, &SLogView::MakeCategoryFilterMenu)
+				.HasDownArrow(true)
+				.ContentPadding(FMargin(1.0f, 1.0f, 1.0f, 0.0f))
+				.ButtonContent()
+				[
+					SNew(SHorizontalBox)
+
+					+SHorizontalBox::Slot()
+					.Padding(0.0f)
+					.AutoWidth()
+					[
+						SNew(STextBlock)
+						.TextStyle(FEditorStyle::Get(), "GenericFilters.TextStyle")
+						.Font(FEditorStyle::Get().GetFontStyle("FontAwesome.9"))
+						.Text(FText::FromString(FString(TEXT("\xf0b0"))) /*fa-filter*/)
+					]
+
+					+SHorizontalBox::Slot()
+					.Padding(2.0f, 0.0f, 0.0f, 0.0f)
+					.AutoWidth()
+					[
+						SNew(STextBlock)
+						.TextStyle(FEditorStyle::Get(), "GenericFilters.TextStyle")
+						.Text(LOCTEXT("CategoryFilter", "Category Filter"))
+					]
+				]
+			]
+
+			// Text Filter (Search Box)
+			+ SHorizontalBox::Slot()
+			.VAlign(VAlign_Center)
+			.Padding(2.0f, 0.0f)
+			.AutoWidth()
+			[
+				SAssignNew(FilterTextBox, SSearchBox)
+				.HintText(LOCTEXT("FilterTextBoxHint", "Search log messages"))
+				.ToolTipText(LOCTEXT("FilterTextBoxToolTip", "Type here to filter the list of log messages."))
+				.OnTextChanged(this, &SLogView::FilterTextBox_OnTextChanged)
+			]
+
+			// Stats Text
+			+ SHorizontalBox::Slot()
+			.VAlign(VAlign_Center)
+			.Padding(2.0f, 0.0f)
+			.AutoWidth()
+			[
+				SNew(STextBlock)
+				.Text(this, &SLogView::GetStatsText)
+				.ColorAndOpacity(this, &SLogView::GetStatsTextColor)
+			]
+		]
+
+		+ SVerticalBox::Slot()
+		.FillHeight(1.0f)
+		[
+			SNew(SBox)
+			.VAlign(VAlign_Fill)
+			[
+				SNew(SHorizontalBox)
+
+				+ SHorizontalBox::Slot()
+				.FillWidth(1.0f)
+				.Padding(0.0f)
+				.VAlign(VAlign_Fill)
+				[
+					SNew(SScrollBox)
+					.Orientation(Orient_Horizontal)
+
+					+ SScrollBox::Slot()
+					.VAlign(VAlign_Fill)
+					[
+						SNew(SBorder)
+						.BorderImage(FEditorStyle::GetBrush("ToolPanel.GroupBorder"))
+						.Padding(0.0f)
+						[
+							SAssignNew(ListView, SListView<TSharedPtr<FLogMessage>>)
+							.ExternalScrollbar(ExternalScrollbar)
+							.ItemHeight(20.0f)
+							.SelectionMode(ESelectionMode::Single)
+							.OnSelectionChanged(this, &SLogView::OnSelectionChanged)
+							.ListItemsSource(&Messages)
+							.OnGenerateRow(this, &SLogView::OnGenerateRow)
+							.ConsumeMouseWheel(EConsumeMouseWheel::Always)
+							.OnContextMenuOpening(FOnContextMenuOpening::CreateSP(this, &SLogView::ListView_GetContextMenu))
+							.HeaderRow
+							(
+								SNew(SHeaderRow)
+
+								+ SHeaderRow::Column(LogViewColumns::IdColumnName)
+								.ManualWidth(60.0f)
+								.DefaultLabel(LOCTEXT("IdColumn", "Index"))
+
+								+ SHeaderRow::Column(LogViewColumns::TimeColumnName)
+								.ManualWidth(94.0f)
+								.DefaultLabel(LOCTEXT("TimeColumn", "Time"))
+
+								+ SHeaderRow::Column(LogViewColumns::VerbosityColumnName)
+								.ManualWidth(80.0f)
+								.DefaultLabel(LOCTEXT("VerbosityColumn", "Verbosity"))
+
+								+ SHeaderRow::Column(LogViewColumns::CategoryColumnName)
+								.ManualWidth(120.0f)
+								.DefaultLabel(LOCTEXT("CategoryColumn", "Category"))
+
+								+ SHeaderRow::Column(LogViewColumns::MessageColumnName)
+								.ManualWidth(880.0f)
+								.DefaultLabel(LOCTEXT("MessageColumn", "Message"))
+
+								+ SHeaderRow::Column(LogViewColumns::FileColumnName)
+								.ManualWidth(600.0f)
+								.DefaultLabel(LOCTEXT("FileColumn", "File"))
+
+								+ SHeaderRow::Column(LogViewColumns::LineColumnName)
+								.ManualWidth(60.0f)
+								.DefaultLabel(LOCTEXT("LineColumn", "Line"))
+							)
+						]
+					]
+				]
+
+				+ SHorizontalBox::Slot()
+				.AutoWidth()
+				.Padding(0.0f)
+				[
+					SNew(SBox)
+					.WidthOverride(FOptionalSize(13.0f))
+					[
+						ExternalScrollbar.ToSharedRef()
+					]
+				]
+			]
+		]
+	];
+
+	// Register ourselves with the profiler manager.
+	//TODO: FTimingProfilerManager::Get()->OnRequestLogViewUpdate().AddSP(this, &SLogView::ProfilerManager_OnRequestLogViewUpdate);
+}
+END_SLATE_FUNCTION_BUILD_OPTIMIZATION
+
+////////////////////////////////////////////////////////////////////////////////////////////////////
+
+TSharedRef<ITableRow> SLogView::OnGenerateRow(TSharedPtr<FLogMessage> InLogMessage, const TSharedRef<STableViewBase>& OwnerTable)
+{
+	// Generate a row for the log message corresponding to InLogMessage.
+	return SNew(SLogMessageRow, InLogMessage, SharedThis(this), OwnerTable);
+}
+
+////////////////////////////////////////////////////////////////////////////////////////////////////
+
+void SLogView::Tick(const FGeometry& AllottedGeometry, const double InCurrentTime, const float InDeltaTime)
+{
+	int NewMessageCount = 0;
+
+	TSharedPtr<const Trace::IAnalysisSession> Session = FInsightsManager::Get()->GetSession();
+
+	Cache.SetSession(Session);
+
+	if (Session.IsValid())
+	{
+		Trace::FAnalysisSessionReadScope _(*Session.Get());
+		Session->ReadLogProvider([this, &NewMessageCount](const Trace::ILogProvider& LogProvider)
+		{
+			NewMessageCount = static_cast<int>(LogProvider.GetMessageCount());
+
+			//TODO: show only categories that are used in current trace
+			//TODO: cause of duplicates: a) runtime, b) case insensitive, c) stripped "Log" prefix
+			//TODO: FString vs. FText vs. FName ?
+
+			//TODO: int NumCategories = static_cast<int>(LogProvider.GetCategoriesCount());
+			int NumCategories = 0;
+			LogProvider.EnumerateCategories([&NumCategories](const Trace::FLogCategory& Category)
+			{
+				NumCategories++;
+			});
+
+			if (NumCategories != TotalNumCategories)
+			{
+				TotalNumCategories = NumCategories;
+				UE_LOG(TimingProfiler, Log, TEXT("[LogView] Total Log Categories: %d"), TotalNumCategories);
+
+				TSet<FName> Categories;
+				LogProvider.EnumerateCategories([&Categories](const Trace::FLogCategory& Category)
+				{
+					FString CategoryStr(Category.Name);
+					if (CategoryStr.StartsWith(TEXT("Log")))
+						CategoryStr = CategoryStr.RightChop(3);
+					if (Categories.Contains(FName(*CategoryStr)))
+					{
+						UE_LOG(TimingProfiler, Log, TEXT("[LogView] Duplicated Log Category: \"%s\""), *Category.Name);
+					}
+					Categories.Add(FName(*CategoryStr));
+				});
+				Filter.SyncAvailableCategories(Categories);
+				UE_LOG(TimingProfiler, Log, TEXT("[LogView] Unique Log Categories: %d"), Filter.GetAvailableLogCategories().Num());
+
+				//Cache.Reset();
+				Messages.Reset();
+				TotalNumMessages = 0;
+				//ListView->RebuildList();
+				bIsDirty = true;
+				DirtyStopwatch.Start();
+			}
+		});
+	}
+
+	if (NewMessageCount != TotalNumMessages)
+	{
+		bIsDirty = true;
+		DirtyStopwatch.Start();
+
+		if (NewMessageCount > TotalNumMessages)
+		{
+			if (Filter.IsFilterSet())
+			{
+				// Filter messages async.
+				if (!FilteringAsyncTask.IsValid())
+				{
+					FilteringStopwatch.Restart();
+
+					bIsFilteringAsyncTaskCancelRequested = false;
+					FilteringStartIndex = TotalNumMessages;
+					FilteringEndIndex = NewMessageCount;
+					FilteringChangeNumber = Filter.GetChangeNumber();
+					FilteringAsyncTask = MakeUnique<FAsyncTask<FLogFilteringAsyncTask>>(FilteringStartIndex, FilteringEndIndex, Filter, SharedThis(this));
+					UE_LOG(TimingProfiler, Log, TEXT("[LogView] Start async task for filtering by%s%s%s (\"%s\") (%d to %d)"),
+						Filter.IsFilterSetByVerbosity() ? TEXT(" Verbosity,") : TEXT(""),
+						Filter.IsFilterSetByCategory() ? TEXT(" Category,") : TEXT(""),
+						Filter.IsFilterSetByText() ? TEXT(" Text") : TEXT(""),
+						*Filter.GetFilterText().ToString(),
+						FilteringStartIndex, FilteringEndIndex);
+					FilteringAsyncTask->StartBackgroundTask();
+				}
+				else
+				{
+					// A task is already in progress.
+					if (FilteringStartIndex == TotalNumMessages &&
+						FilteringEndIndex <= NewMessageCount &&
+						FilteringChangeNumber == Filter.GetChangeNumber())
+					{
+						// The filter is still valid. Just wait.
+					}
+					else
+					{
+						// The filter used by running task is obsolete. Cancel the task.
+						bIsFilteringAsyncTaskCancelRequested = true;
+					}
+				}
+			}
+			else // no filtering
+			{
+				FilteringStopwatch.Restart();
+
+				for (int Index = TotalNumMessages; Index < NewMessageCount; Index++)
+				{
+					FLogMessage LogMessage(Index);
+					Messages.Add(MakeShared<FLogMessage>(MoveTemp(LogMessage)));
+				}
+
+				const int NumAddedMessages = NewMessageCount - TotalNumMessages;
+				TotalNumMessages = NewMessageCount;
+				ListView->RebuildList();
+				bIsDirty = false;
+				DirtyStopwatch.Reset();
+				UpdateStatsText();
+
+				FilteringStopwatch.Stop();
+				uint64 DurationMs = FilteringStopwatch.GetAccumulatedTimeMs();
+				if (DurationMs > 10) // avoids spams
+				{
+					UE_LOG(TimingProfiler, Log, TEXT("[LogView] Updated (no filter; %d added / %d total messages) in %llu ms."),
+						NumAddedMessages, NewMessageCount, DurationMs);
+				}
+			}
+		}
+		else // if (NewMessageCount < TotalNumMessages)
+		{
+			// Just reset. On next Tick() the list will grow if needed.
+			UE_LOG(TimingProfiler, Log, TEXT("[LogView] RESET"));
+			Cache.Reset();
+			Messages.Reset();
+			TotalNumMessages = 0;
+			ListView->RebuildList();
+			bIsDirty = (NewMessageCount != 0);
+			if (bIsDirty)
+			{
+				DirtyStopwatch.Start();
+			}
+			else
+			{
+				DirtyStopwatch.Reset();
+			}
+			UpdateStatsText();
+		}
+	}
+
+	if (FilteringAsyncTask.IsValid() &&
+		FilteringAsyncTask->IsDone())
+	{
+		// A filtering async task has completed. Check if filter used is still valid.
+		if (!bIsFilteringAsyncTaskCancelRequested &&
+			FilteringStartIndex == TotalNumMessages &&
+			FilteringEndIndex <= NewMessageCount &&
+			FilteringChangeNumber == Filter.GetChangeNumber())
+		{
+			FLogFilteringAsyncTask& Task = FilteringAsyncTask->GetTask();
+			const TArray<uint32>& FilteredMessages = Task.GetFilteredMessages();
+
+			// Add filtered messages to current Messages array.
+			const int NumFilteredMessages = FilteredMessages.Num();
+			for (int Index = 0; Index < NumFilteredMessages; Index++)
+			{
+				FLogMessage LogMessage(FilteredMessages[Index]);
+				Messages.Add(MakeShared<FLogMessage>(MoveTemp(LogMessage)));
+			}
+
+			TotalNumMessages = Task.GetEndIndex();
+			ListView->RebuildList();
+			bIsDirty = false;
+			DirtyStopwatch.Reset();
+			UpdateStatsText();
+
+			FilteringStopwatch.Stop();
+			uint64 DurationMs = FilteringStopwatch.GetAccumulatedTimeMs();
+			if (DurationMs > 10) // avoids spams
+			{
+				int NumAsyncFilteredMessages = Task.GetEndIndex() - Task.GetStartIndex();
+				double Speed = static_cast<double>(NumAsyncFilteredMessages) / FilteringStopwatch.GetAccumulatedTime();
+				UE_LOG(TimingProfiler, Log, TEXT("[LogView] Updated (%d added / %d async filtered / %d total messages) in %llu ms (%.2f messages/second)."),
+					NumFilteredMessages, NumAsyncFilteredMessages, TotalNumMessages, DurationMs, Speed);
+			}
+		}
+
+		FilteringAsyncTask.Reset();
+	}
+
+	//SCompoundWidget::Tick(AllottedGeometry, InCurrentTime, InDeltaTime);
+}
+
+////////////////////////////////////////////////////////////////////////////////////////////////////
+
+TSharedPtr<FLogMessage> SLogView::GetSelectedLogMessage() const
+{
+	TArray<TSharedPtr<FLogMessage>> SelectedItems = ListView->GetSelectedItems();
+	return (SelectedItems.Num() == 1) ? SelectedItems[0] : nullptr;
+}
+
+////////////////////////////////////////////////////////////////////////////////////////////////////
+
+void SLogView::OnSelectionChanged(TSharedPtr<FLogMessage> LogMessage, ESelectInfo::Type SelectInfo)
+{
+	TSharedPtr<STimingProfilerWindow> Window = FTimingProfilerManager::Get()->GetProfilerWindow();
+	if (Window && Window->TimingView)
+	{
+		// Single item selection.
+		if (LogMessage.IsValid())
+		{
+			const double Time = Cache.Get(LogMessage->GetIndex()).Time;
+
+			if (FSlateApplication::Get().GetModifierKeys().IsShiftDown())
+			{
+				Window->TimingView->SelectToTimeMarker(Time);
+			}
+			else
+			{
+				Window->TimingView->SetAndCenterOnTimeMarker(Time);
+			}
+		}
+	}
+}
+
+////////////////////////////////////////////////////////////////////////////////////////////////////
+
+void SLogView::FilterTextBox_OnTextChanged(const FText& InFilterText)
+{
+	if (Filter.GetFilterText().ToString().Equals(InFilterText.ToString(), ESearchCase::CaseSensitive))
+	{
+		// nothing to do
+		return;
+	}
+
+	// Set filter phrases.
+	Filter.SetFilterText(InFilterText);
+
+	// Report possible syntax errors back to the user.
+	FilterTextBox->SetError(Filter.GetSyntaxErrors());
+
+	OnFilterChanged();
+}
+
+////////////////////////////////////////////////////////////////////////////////////////////////////
+
+void SLogView::OnFilterChanged()
+{
+	const FString FilterText = FilterTextBox->GetText().ToString();
+	UE_LOG(TimingProfiler, Log, TEXT("[LogView] OnFilterChanged: \"%s\""), *FilterText);
+	Cache.Reset();
+	Messages.Reset();
+	TotalNumMessages = 0;
+	bIsDirty = true;
+	DirtyStopwatch.Start();
+	// The next Tick() will update the filtered list of messages.
+}
+
+////////////////////////////////////////////////////////////////////////////////////////////////////
+
+void SLogView::UpdateStatsText()
+{
+	if (Messages.Num() == TotalNumMessages)
+	{
+		StatsText = FText::Format(LOCTEXT("StatsText1", "{0} logs"), FText::AsNumber(TotalNumMessages));
+	}
+	else
+	{
+		StatsText = FText::Format(LOCTEXT("StatsText2", "{0} / {1} logs"), FText::AsNumber(Messages.Num()), FText::AsNumber(TotalNumMessages));
+	}
+}
+
+////////////////////////////////////////////////////////////////////////////////////////////////////
+
+FText SLogView::GetStatsText() const
+{
+	if (bIsDirty)
+	{
+		DirtyStopwatch.Update();
+		double DT = DirtyStopwatch.GetAccumulatedTime();
+		if (DT > 1.0)
+		{
+			return FText::Format(LOCTEXT("StatsTextEx", "{0} (filtering... please wait... {1}s)"), StatsText, FText::AsNumber(FMath::RoundToInt(DT)));
+		}
+	}
+
+	return StatsText;
+}
+
+////////////////////////////////////////////////////////////////////////////////////////////////////
+
+FSlateColor SLogView::GetStatsTextColor() const
+{
+	if (bIsDirty)
+	{
+		return FSlateColor(FLinearColor(1.0f, 0.5f, 0.5f, 1.0f));
+	}
+	else
+	{
+		return FSlateColor(FLinearColor(1.0f, 1.0f, 1.0f, 1.0f));
+	}
+}
+
+////////////////////////////////////////////////////////////////////////////////////////////////////
+
+TSharedPtr<SWidget> SLogView::ListView_GetContextMenu()
+{
+	TSharedPtr<FLogMessage> SelectedLogMessage = GetSelectedLogMessage();
+
+	FMenuBuilder MenuBuilder(/*bInShouldCloseWindowAfterMenuSelection=*/true, nullptr);
+
+	MenuBuilder.BeginSection("LogViewContextMenu");
+	{
+		if (SelectedLogMessage.IsValid())
+		{
+			FLogMessageRecord& Record = Cache.Get(SelectedLogMessage->GetIndex());
+			FName CategoryName(*Record.Category.ToString());
+
+			MenuBuilder.AddMenuEntry(
+				FText::Format(LOCTEXT("HideCategory", "Hide \"{0}\" Category"), Record.Category),
+				FText::Format(LOCTEXT("HideCategory_Tooltip", "Hide the \"{0}\" log category."), Record.Category),
+				FSlateIcon(),
+				FUIAction(FExecuteAction::CreateSP(this, &SLogView::ToggleCategory_Execute, CategoryName)),
+				NAME_None,
+				EUserInterfaceActionType::Button
+			);
+
+			MenuBuilder.AddMenuEntry(
+				FText::Format(LOCTEXT("ShowOnlyCategory", "Show Only \"{0}\" Category"), Record.Category),
+				FText::Format(LOCTEXT("ShowOnlyCategory_Tooltip", "Show only the \"{0}\" log category (hide all other log categories)."), Record.Category),
+				FSlateIcon(),
+				FUIAction(FExecuteAction::CreateSP(this, &SLogView::ShowOnlyCategory_Execute, CategoryName)),
+				NAME_None,
+				EUserInterfaceActionType::Button
+			);
+		}
+
+		MenuBuilder.AddMenuEntry(
+			LOCTEXT("ShowAllCategoriesCtxMenu", "Show All Categories"),
+			LOCTEXT("ShowAllCategoriesCtxMenu_Tooltip", "Reset category filter (show all log categories)."),
+			FSlateIcon(),
+			FUIAction(FExecuteAction::CreateSP(this, &SLogView::ShowAllCategories_Execute)),
+			NAME_None,
+			EUserInterfaceActionType::Button
+		);
+	}
+	MenuBuilder.EndSection();
+
+	return MenuBuilder.MakeWidget();
+}
+
+////////////////////////////////////////////////////////////////////////////////////////////////////
+
+TSharedRef<SWidget> SLogView::MakeVerbosityThresholdMenu()
+{
+	FMenuBuilder MenuBuilder(/*bInShouldCloseWindowAfterMenuSelection=*/true, nullptr);
+
+	MenuBuilder.BeginSection("LogViewVerbosityThreshold"/*, LOCTEXT("VerbosityThresholdHeading", "Verbosity Threshold")*/);
+	CreateVerbosityThresholdMenuSection(MenuBuilder);
+	MenuBuilder.EndSection();
+
+	return MenuBuilder.MakeWidget();
+}
+
+////////////////////////////////////////////////////////////////////////////////////////////////////
+
+void  SLogView::CreateVerbosityThresholdMenuSection(FMenuBuilder& MenuBuilder)
+{
+	struct FVerbosityThresholdInfo
+	{
+		ELogVerbosity::Type Verbosity;
+		FText Label;
+		FText ToolTip;
+	};
+
+	FVerbosityThresholdInfo VerbosityThresholds[] =
+	{
+		{ ELogVerbosity::VeryVerbose, LOCTEXT("VerbosityThresholdVeryVerbose", "VeryVerbose"), LOCTEXT("VerbosityThresholdVeryVerbose_Tooltip", "Show all log messages (any verbosity level)."), },
+		{ ELogVerbosity::Verbose,     LOCTEXT("VerbosityThresholdVerbose",     "Verbose"),     LOCTEXT("VerbosityThresholdVerbose_Tooltip",     "Show Verbose, Log, Display, Warning, Error and Fatal log messages (i.e. hide VeryVerbose log messages)."), },
+		{ ELogVerbosity::Log,         LOCTEXT("VerbosityThresholdLog",         "Log"),         LOCTEXT("VerbosityThresholdLog_Tooltip",         "Show Log, Display, Warning, Error and Fatal log messages (i.e. hide Verbose and VeryVerbose log messages)."), },
+		{ ELogVerbosity::Display,     LOCTEXT("VerbosityThresholdDisplay",     "Display"),     LOCTEXT("VerbosityThresholdDisplay_Tooltip",     "Show Display, Warning, Error and Fatal log messages (i.e. hide Log, Verbose and VeryVerbose log messages)."), },
+		{ ELogVerbosity::Warning,     LOCTEXT("VerbosityThresholdWarning",     "Warning"),     LOCTEXT("VerbosityThresholdWarning_Tooltip",     "Show only Warning, Error and Fatal log messages."), },
+		{ ELogVerbosity::Error,       LOCTEXT("VerbosityThresholdError",       "Error"),       LOCTEXT("VerbosityThresholdError_Tooltip",       "Show only Error and Fatal log messages."), },
+		{ ELogVerbosity::Fatal,       LOCTEXT("VerbosityThresholdFatal",       "Fatal"),       LOCTEXT("VerbosityThresholdFatal_Tooltip",       "Show only Fatal log messages."), },
+	};
+
+	for (int Index = 0; Index < sizeof(VerbosityThresholds) / sizeof(FVerbosityThresholdInfo); ++Index)
+	{
+		const FVerbosityThresholdInfo& Threshold = VerbosityThresholds[Index];
+
+		/*
+		MenuBuilder.AddMenuEntry(
+			Threshold.Label,
+			Threshold.ToolTip,
+			FSlateIcon(),
+			FUIAction(FExecuteAction::CreateSP(this, &SLogView::VerbosityThreshold_Execute, Threshold.Verbosity),
+				FCanExecuteAction::CreateLambda([] { return true; }),
+				FIsActionChecked::CreateSP(this, &SLogView::VerbosityThreshold_IsChecked, Threshold.Verbosity)),
+			NAME_None,
+			EUserInterfaceActionType::ToggleButton
+		);
+		*/
+
+		const TSharedRef<SWidget> TextBlock = SNew(STextBlock)
+			.Text(Threshold.Label)
+			.ShadowColorAndOpacity(FLinearColor(0.05f, 0.05f, 0.05f, 1.0f))
+			.ShadowOffset(FVector2D(1.0f, 1.0f))
+			.ColorAndOpacity(FSlateColor(FTimeMarkerTrackBuilder::GetColorByVerbosity(Threshold.Verbosity)));
+
+		MenuBuilder.AddMenuEntry(
+			FUIAction(FExecuteAction::CreateSP(this, &SLogView::VerbosityThreshold_Execute, Threshold.Verbosity),
+				FCanExecuteAction::CreateLambda([] { return true; }),
+				FIsActionChecked::CreateSP(this, &SLogView::VerbosityThreshold_IsChecked, Threshold.Verbosity)),
+			TextBlock,
+			NAME_None,
+			Threshold.ToolTip,
+			EUserInterfaceActionType::ToggleButton
+		);
+	}
+}
+
+////////////////////////////////////////////////////////////////////////////////////////////////////
+
+TSharedRef<SWidget> SLogView::MakeCategoryFilterMenu()
+{
+	FMenuBuilder MenuBuilder(/*bInShouldCloseWindowAfterMenuSelection=*/true, nullptr);
+
+	MenuBuilder.BeginSection("LogViewShowAllCategories", LOCTEXT("QuickFilterHeading", "Quick Filter"));
+	{
+		MenuBuilder.AddMenuEntry(
+			LOCTEXT("ShowAllCategories", "Show/Hide All"),
+			LOCTEXT("ShowAllCategories_Tooltip", "Change filtering to show/hide all categories"),
+			FSlateIcon(),
+			FUIAction(FExecuteAction::CreateSP(this, &SLogView::ShowHideAllCategories_Execute),
+				FCanExecuteAction::CreateLambda([] { return true; }),
+				FIsActionChecked::CreateSP(this, &SLogView::ShowHideAllCategories_IsChecked)),
+			NAME_None,
+			EUserInterfaceActionType::ToggleButton
+		);
+	}
+	MenuBuilder.EndSection();
+
+	MenuBuilder.BeginSection("LogViewCategoriesEntries", LOCTEXT("CategoriesHeading", "Categories"));
+	CreateCategoriesFilterMenuSection(MenuBuilder);
+	MenuBuilder.EndSection();
+
+	return MenuBuilder.MakeWidget();
+}
+
+////////////////////////////////////////////////////////////////////////////////////////////////////
+
+void SLogView::CreateCategoriesFilterMenuSection(FMenuBuilder& MenuBuilder)
+{
+	for (const FName CategoryName : Filter.GetAvailableLogCategories())
+	{
+		const FString CategoryString = CategoryName.ToString();
+		const FText CategoryText(FText::AsCultureInvariant(CategoryString));
+
+		//MenuBuilder.AddMenuEntry(
+		//	FText::AsCultureInvariant(CategoryString),
+		//	FText::Format(LOCTEXT("Category_Tooltip", "Filter the Log View to show/hide category: {0}"), CategoryText),
+		//	FSlateIcon(),
+		//	FUIAction(FExecuteAction::CreateSP(this, &SLogView::ToggleCategory_Execute, CategoryName),
+		//		FCanExecuteAction::CreateLambda([] { return true; }),
+		//		FIsActionChecked::CreateSP(this, &SLogView::ToggleCategory_IsChecked, CategoryName)),
+		//	NAME_None,
+		//	EUserInterfaceActionType::ToggleButton
+		//);
+
+		const TSharedRef<SWidget> TextBlock = SNew(STextBlock)
+			.Text(FText::AsCultureInvariant(CategoryString))
+			.ShadowColorAndOpacity(FLinearColor(0.05f, 0.05f, 0.05f, 1.0f))
+			.ShadowOffset(FVector2D(1.0f, 1.0f))
+			.ColorAndOpacity(FSlateColor(FTimeMarkerTrackBuilder::GetColorByCategory(*CategoryString)));
+
+		MenuBuilder.AddMenuEntry(
+			FUIAction(FExecuteAction::CreateSP(this, &SLogView::ToggleCategory_Execute, CategoryName),
+				FCanExecuteAction::CreateLambda([] { return true; }),
+				FIsActionChecked::CreateSP(this, &SLogView::ToggleCategory_IsChecked, CategoryName)),
+			TextBlock,
+			NAME_None,
+			FText::Format(LOCTEXT("Category_Tooltip", "Filter the Log View to show/hide category: {0}"), CategoryText),
+			EUserInterfaceActionType::ToggleButton
+		);
+	}
+}
+
+////////////////////////////////////////////////////////////////////////////////////////////////////
+
+bool SLogView::VerbosityThreshold_IsChecked(ELogVerbosity::Type Verbosity) const
+{
+	return Verbosity <= Filter.GetVerbosityThreshold();
+}
+
+////////////////////////////////////////////////////////////////////////////////////////////////////
+
+void SLogView::VerbosityThreshold_Execute(ELogVerbosity::Type Verbosity)
+{
+	Filter.SetVerbosityThreshold(Verbosity);
+	OnFilterChanged();
+}
+
+////////////////////////////////////////////////////////////////////////////////////////////////////
+
+bool SLogView::ShowHideAllCategories_IsChecked() const
+{
+	return Filter.IsShowAllCategoriesEnabled();
+}
+
+////////////////////////////////////////////////////////////////////////////////////////////////////
+
+void SLogView::ShowHideAllCategories_Execute()
+{
+	Filter.ToggleShowAllCategories();
+	OnFilterChanged();
+}
+
+////////////////////////////////////////////////////////////////////////////////////////////////////
+
+void SLogView::ShowAllCategories_Execute()
+{
+	if (Filter.IsShowAllCategoriesEnabled())
+	{
+		Filter.ToggleShowAllCategories(); // hide
+		Filter.ToggleShowAllCategories(); // show
+	}
+	else
+	{
+		Filter.ToggleShowAllCategories(); // show
+	}
+	OnFilterChanged();
+}
+
+////////////////////////////////////////////////////////////////////////////////////////////////////
+
+void SLogView::ShowOnlyCategory_Execute(FName InName)
+{
+	Filter.EnableOnlyCategory(InName);
+	OnFilterChanged();
+}
+
+////////////////////////////////////////////////////////////////////////////////////////////////////
+
+bool SLogView::ToggleCategory_IsChecked(FName InName) const
+{
+	return Filter.IsLogCategoryEnabled(InName);
+}
+
+////////////////////////////////////////////////////////////////////////////////////////////////////
+
+void SLogView::ToggleCategory_Execute(FName InName)
+{
+	Filter.ToggleLogCategory(InName);
+	OnFilterChanged();
+}
+
+////////////////////////////////////////////////////////////////////////////////////////////////////
+
+#undef LOCTEXT_NAMESPACE
diff --git a/Engine/Source/Developer/TraceInsights/Private/Insights/Widgets/STimerTableCell.h b/Engine/Source/Developer/TraceInsights/Private/Insights/Widgets/STimerTableCell.h
new file mode 100644
--- /dev/null
+++ b/Engine/Source/Developer/TraceInsights/Private/Insights/Widgets/STimerTableCell.h
@@ -0,0 +1,140 @@
+// Copyright 1998-2019 Epic Games, Inc. All Rights Reserved.
+
+#pragma once
+
+#include "CoreMinimal.h"
+#include "Widgets/DeclarativeSyntaxSupport.h"
+#include "Widgets/SCompoundWidget.h"
+
+// Insights
+#include "Insights/ViewModels/TimerNode.h"
+
+DECLARE_DELEGATE_TwoParams(FSetHoveredTableCell, const FName /*ColumnId*/, const FTimerNodePtr /*SamplePtr*/);
+DECLARE_DELEGATE_RetVal_OneParam(bool, FIsColumnVisibleDelegate, const FName /*ColumnId*/);
+DECLARE_DELEGATE_RetVal_OneParam(EHorizontalAlignment, FGetColumnOutlineHAlignmentDelegate, const FName /*ColumnId*/);
+
+class STimerTableCell : public SCompoundWidget
+{
+public:
+	SLATE_BEGIN_ARGS(STimerTableCell) {}
+		SLATE_EVENT(FSetHoveredTableCell, OnSetHoveredTableCell)
+		SLATE_ATTRIBUTE(FText, HighlightText)
+		SLATE_ARGUMENT(FTimerNodePtr, TimerNodePtr)
+		SLATE_ARGUMENT(FName, ColumnId)
+		SLATE_ARGUMENT(bool, IsTimerNameColumn)
+	SLATE_END_ARGS()
+
+	void Construct(const FArguments& InArgs, const TSharedRef<class ITableRow>& TableRow);
+
+protected:
+	TSharedRef<SWidget> GenerateWidgetForColumn(const FArguments& InArgs, const TSharedRef<class ITableRow>& TableRow);
+	TSharedRef<SWidget> GenerateWidgetForNameColumn(const FArguments& InArgs, const TSharedRef<class ITableRow>& TableRow);
+	TSharedRef<SWidget> GenerateWidgetForStatsColumn(const FArguments& InArgs, const TSharedRef<class ITableRow>& TableRow);
+
+	/**
+	 * The system will use this event to notify a widget that the cursor has entered it. This event is NOT bubbled.
+	 *
+	 * @param MyGeometry The Geometry of the widget receiving the event
+	 * @param MouseEvent Information about the input event
+	 */
+	virtual void OnMouseEnter(const FGeometry& MyGeometry, const FPointerEvent& MouseEvent) override
+	{
+		SCompoundWidget::OnMouseEnter(MyGeometry, MouseEvent);
+		SetHoveredTableCellDelegate.ExecuteIfBound(ColumnId, TimerNodePtr);
+	}
+
+	/**
+	 * The system will use this event to notify a widget that the cursor has left it. This event is NOT bubbled.
+	 *
+	 * @param MouseEvent Information about the input event
+	 */
+	virtual void OnMouseLeave(const FPointerEvent& MouseEvent) override
+	{
+		SCompoundWidget::OnMouseLeave(MouseEvent);
+		SetHoveredTableCellDelegate.ExecuteIfBound(NAME_None, nullptr);
+	}
+
+	/**
+	 * Called during drag and drop when the drag enters a widget.
+	 *
+	 * Enter/Leave events in slate are meant as lightweight notifications.
+	 * So we do not want to capture mouse or set focus in response to these.
+	 * However, OnDragEnter must also support external APIs (e.g. OLE Drag/Drop)
+	 * Those require that we let them know whether we can handle the content
+	 * being dragged OnDragEnter.
+	 *
+	 * The concession is to return a can_handled/cannot_handle
+	 * boolean rather than a full FReply.
+	 *
+	 * @param MyGeometry      The geometry of the widget receiving the event.
+	 * @param DragDropEvent   The drag and drop event.
+	 *
+	 * @return A reply that indicated whether the contents of the DragDropEvent can potentially be processed by this widget.
+	 */
+	virtual void OnDragEnter(const FGeometry& MyGeometry, const FDragDropEvent& DragDropEvent) override
+	{
+		SCompoundWidget::OnDragEnter(MyGeometry, DragDropEvent);
+		SetHoveredTableCellDelegate.ExecuteIfBound(ColumnId, TimerNodePtr);
+	}
+
+	/**
+	 * Called during drag and drop when the drag leaves a widget.
+	 *
+	 * @param DragDropEvent   The drag and drop event.
+	 */
+	virtual void OnDragLeave(const FDragDropEvent& DragDropEvent)  override
+	{
+		SCompoundWidget::OnDragLeave(DragDropEvent);
+		SetHoveredTableCellDelegate.ExecuteIfBound(NAME_None, nullptr);
+	}
+
+	EVisibility GetHintIconVisibility() const
+	{
+		return IsHovered() ? EVisibility::Visible : EVisibility::Hidden;
+	}
+
+	EVisibility GetCulledEventsIconVisibility() const
+	{
+		//return TimerNodePtr->HasCulledChildren() ? EVisibility::Visible : EVisibility::Collapsed;
+		return EVisibility::Collapsed;
+	}
+
+	FText GetNameEx() const
+	{
+		return TimerNodePtr->GetNameEx();
+	}
+
+	FSlateColor GetColorAndOpacity() const
+	{
+		const FLinearColor TextColor =
+			TimerNodePtr->IsFiltered() ? FLinearColor(1.0f, 1.0f, 1.0f, 0.5f) :
+			FLinearColor::White;
+		return TextColor;
+	}
+
+	FSlateColor GetStatsColorAndOpacity() const
+	{
+		const FLinearColor TextColor =
+			TimerNodePtr->IsFiltered() ? FLinearColor(1.0f, 1.0f, 1.0f, 0.5f) :
+			TimerNodePtr->GetStats().InstanceCount == 0 ? FLinearColor(1.0f, 1.0f, 1.0f, 0.6f) :
+			FLinearColor::White;
+		return TextColor;
+	}
+
+	FLinearColor GetShadowColorAndOpacity() const
+	{
+		const FLinearColor ShadowColor =
+			TimerNodePtr->IsFiltered() ? FLinearColor(0.f, 0.f, 0.f, 0.25f) :
+			FLinearColor(0.0f, 0.0f, 0.0f, 0.5f);
+		return ShadowColor;
+	}
+
+protected:
+	/** A shared pointer to the timer node. */
+	FTimerNodePtr TimerNodePtr;
+
+	/** The Id of the column where this timer belongs. */
+	FName ColumnId;
+
+	FSetHoveredTableCell SetHoveredTableCellDelegate;
+};
diff --git a/Engine/Source/Developer/TraceInsights/Private/Insights/Widgets/STimerTableRow.cpp b/Engine/Source/Developer/TraceInsights/Private/Insights/Widgets/STimerTableRow.cpp
new file mode 100644
--- /dev/null
+++ b/Engine/Source/Developer/TraceInsights/Private/Insights/Widgets/STimerTableRow.cpp
@@ -0,0 +1,272 @@
+// Copyright 1998-2019 Epic Games, Inc. All Rights Reserved.
+
+#include "STimerTableRow.h"
+
+#include "Widgets/SOverlay.h"
+#include "Widgets/Images/SImage.h"
+
+// Insights
+#include "Insights/ViewModels/TimersViewColumnFactory.h"
+#include "Insights/Widgets/STimersViewTooltip.h"
+#include "Insights/Widgets/STimerTableCell.h"
+
+#define LOCTEXT_NAMESPACE "STimersView"
+
+BEGIN_SLATE_FUNCTION_BUILD_OPTIMIZATION
+
+void STimerTableRow::Construct(const FArguments& InArgs, const TSharedRef<STableViewBase>& InOwnerTableView)
+{
+	OnShouldBeEnabled = InArgs._OnShouldBeEnabled;
+	IsColumnVisibleDelegate = InArgs._OnIsColumnVisible;
+	SetHoveredTableCellDelegate = InArgs._OnSetHoveredTableCell;
+	GetColumnOutlineHAlignmentDelegate = InArgs._OnGetColumnOutlineHAlignmentDelegate;
+
+	HighlightText = InArgs._HighlightText;
+	HighlightedTimerName = InArgs._HighlightedTimerName;
+
+	TimerNodePtr = InArgs._TimerNodePtr;
+
+	SetEnabled(TAttribute<bool>(this, &STimerTableRow::HandleShouldBeEnabled));
+
+	SMultiColumnTableRow<FTimerNodePtr>::Construct(SMultiColumnTableRow<FTimerNodePtr>::FArguments(), InOwnerTableView);
+
+	/*
+	const FSlateBrush* const NodeIcon = TimerNodeTypeHelper::GetIconForTimerNodeType(InTimerNode->GetType());
+	const TSharedRef<SToolTip> Tooltip = InTimerNode->IsGroup() ? SNew(SToolTip) : STimersViewTooltip(InTimerNode->GetId()).GetTooltip();
+
+	ChildSlot
+		[
+			SNew(SHorizontalBox)
+
+			// Expander arrow.
+			+ SHorizontalBox::Slot()
+			.AutoWidth()
+			.HAlign(HAlign_Right)
+			.VAlign(VAlign_Center)
+			[
+				SNew(SExpanderArrow, SharedThis(this))
+			]
+
+			// Icon to visualize group or timer type.
+			+ SHorizontalBox::Slot()
+			.AutoWidth()
+			.Padding(0.0f, 0.0f, 8.0f, 0.0f)
+			[
+				SNew(SImage)
+				.Image(NodeIcon)
+				.ToolTip(Tooltip)
+			]
+
+			// Description text.
+			+ SHorizontalBox::Slot()
+			.AutoWidth()
+			.VAlign(VAlign_Center)
+			.HAlign(HAlign_Left)
+			.Padding(FMargin(2.0f, 0.0f))
+			[
+				SNew(STextBlock)
+				.Text(this, &STimerTableRow::GetText)
+				.HighlightText(InArgs._HighlightText)
+				.TextStyle(FEditorStyle::Get(), TEXT("Profiler.Tooltip"))
+				.ColorAndOpacity(this, &STimerTableRow::GetColorAndOpacity)
+			]
+
+			+ SHorizontalBox::Slot()
+			.AutoWidth()
+			.HAlign(HAlign_Center)
+			.VAlign(VAlign_Top)
+			.Padding(0.0f, 1.0f, 0.0f, 0.0f)
+			[
+				SNew(SImage)
+				.Visibility(!InTimerNode->IsGroup() ? EVisibility::Visible : EVisibility::Collapsed)
+				.Image(FEditorStyle::GetBrush("Profiler.Tooltip.HintIcon10"))
+				.ToolTip(Tooltip)
+			]
+		];
+
+	STableRow<FTimerNodePtr>::ConstructInternal(STableRow::FArguments().ShowSelection(true), InOwnerTableView);
+	*/
+}
+
+TSharedRef<SWidget> STimerTableRow::GenerateWidgetForColumn(const FName& ColumnId)
+{
+	return
+
+	SNew(SOverlay)
+	.Visibility(EVisibility::SelfHitTestInvisible)
+
+	+SOverlay::Slot()
+	.Padding(0.0f)
+	[
+		SNew(SImage)
+		.Image(FEditorStyle::GetBrush("Profiler.LineGraphArea"))
+		.ColorAndOpacity(this, &STimerTableRow::GetBackgroundColorAndOpacity)
+	]
+
+	+SOverlay::Slot()
+	.Padding(0.0f)
+	[
+		SNew(SImage)
+		.Image(this, &STimerTableRow::GetOutlineBrush, ColumnId)
+		.ColorAndOpacity(this, &STimerTableRow::GetOutlineColorAndOpacity)
+	]
+
+	+SOverlay::Slot()
+	[
+		SNew(STimerTableCell, SharedThis(this))
+		.Visibility(this, &STimerTableRow::IsColumnVisible, ColumnId)
+		.TimerNodePtr(TimerNodePtr)
+		.ColumnId(ColumnId)
+		.HighlightText(HighlightText)
+		.IsTimerNameColumn(ColumnId == FTimersViewColumnFactory::Get().Collection[0]->Id) // name column
+		.OnSetHoveredTableCell(this, &STimerTableRow::OnSetHoveredTableCell)
+	];
+}
+
+END_SLATE_FUNCTION_BUILD_OPTIMIZATION
+
+FReply STimerTableRow::OnDragDetected(const FGeometry& MyGeometry, const FPointerEvent& MouseEvent)
+{
+	//im:TODO
+	//if (MouseEvent.IsMouseButtonDown(EKeys::LeftMouseButton))
+	//{
+	//	if (TimerNode->IsGroup())
+	//	{
+	//		// Add all timer Ids for the group.
+	//		TArray<int32> TimerIds;
+	//		const TArray<FTimerNodePtr>& FilteredChildren = TimerNode->GetFilteredChildren();
+	//		const int32 NumFilteredChildren = FilteredChildren.Num();
+	//
+	//		TimerIds.Reserve(NumFilteredChildren);
+	//		for (int32 Nx = 0; Nx < NumFilteredChildren; ++Nx)
+	//		{
+	//			TimerIds.Add(FilteredChildren[Nx]->GetId());
+	//		}
+	//
+	//		return FReply::Handled().BeginDragDrop(FStatIDDragDropOp::NewGroup(TimerIds, TimerNode->GetName().GetPlainNameString()));
+	//	}
+	//	else
+	//	{
+	//		return FReply::Handled().BeginDragDrop(FStatIDDragDropOp::NewSingle(TimerNode->GetId(), TimerNode->GetName().GetPlainNameString()));
+	//	}
+	//}
+
+	return SMultiColumnTableRow<FTimerNodePtr>::OnDragDetected(MyGeometry, MouseEvent);
+}
+
+FText STimerTableRow::GetText() const
+{
+	FText Text = FText::GetEmpty();
+
+	if (TimerNodePtr->IsGroup())
+	{
+		Text = FText::Format(LOCTEXT("TimerNode_GroupNodeTextFmt", "{0} ({1})"), FText::FromName(TimerNodePtr->GetName()), FText::AsNumber(TimerNodePtr->GetChildren().Num()));
+	}
+	else
+	{
+		Text = FText::FromName(TimerNodePtr->GetName());
+	}
+
+	return Text;
+}
+
+FSlateFontInfo STimerTableRow::GetFont() const
+{
+	const bool bIsStatTracked = false;//im:TODO: FTimingProfilerManager::Get()->IsStatTracked(TimerNodePtr->GetTimerId());
+	const FSlateFontInfo FontInfo = bIsStatTracked ? FEditorStyle::GetFontStyle("BoldFont") : FEditorStyle::GetFontStyle("NormalFont");
+	return FontInfo;
+}
+
+FSlateColor STimerTableRow::GetColorAndOpacity() const
+{
+	//const bool bIsStatTracked = FTimingProfilerManager::Get()->IsStatTracked(TimerNodePtr->GetId());
+	//const FSlateColor Color = bIsStatTracked ? FTimingProfilerManager::Get()->GetColorForTimerId(TimerNodePtr->GetId()) : FLinearColor::White;
+	//return Color;
+	//FTimingProfilerManager::GetSettings().GetColorForStat(StatName)
+	return FLinearColor::White;
+}
+
+FSlateColor STimerTableRow::GetBackgroundColorAndOpacity() const
+{
+	return GetBackgroundColorAndOpacity(TimerNodePtr->GetStats().TotalInclusiveTime);
+	//return FLinearColor(0.0f, 0.0f, 0.0f, 1.0f)
+}
+
+FSlateColor STimerTableRow::GetBackgroundColorAndOpacity(double Time) const
+{
+	const FLinearColor Color =	Time > TimeUtils::Second      ? FLinearColor(0.3f, 0.0f, 0.0f, 1.0f) :
+								Time > TimeUtils::Milisecond  ? FLinearColor(0.3f, 0.1f, 0.0f, 1.0f) :
+								Time > TimeUtils::Microsecond ? FLinearColor(0.0f, 0.1f, 0.0f, 1.0f) :
+																FLinearColor(0.0f, 0.0f, 0.0f, 1.0f);
+	return Color;
+}
+
+FSlateColor STimerTableRow::GetOutlineColorAndOpacity() const
+{
+	const FLinearColor NoColor(0.0f, 0.0f, 0.0f, 0.0f);
+	const bool bShouldBeHighlighted = TimerNodePtr->GetName() == HighlightedTimerName.Get();
+	const FLinearColor OutlineColorAndOpacity = bShouldBeHighlighted ? FLinearColor(FColorList::SlateBlue) : NoColor;
+	return OutlineColorAndOpacity;
+}
+
+const FSlateBrush* STimerTableRow::GetOutlineBrush(const FName ColumnId) const
+{
+	EHorizontalAlignment Result = HAlign_Center;
+	if (IsColumnVisibleDelegate.IsBound())
+	{
+		Result = GetColumnOutlineHAlignmentDelegate.Execute(ColumnId);
+	}
+
+	const FSlateBrush* Brush = nullptr;
+	if (Result == HAlign_Left)
+	{
+		Brush = FEditorStyle::GetBrush("Profiler.EventGraph.Border.L");
+	}
+	else if(Result == HAlign_Right)
+	{
+		Brush = FEditorStyle::GetBrush("Profiler.EventGraph.Border.R");
+	}
+	else
+	{
+		Brush = FEditorStyle::GetBrush("Profiler.EventGraph.Border.TB");
+	}
+	return Brush;
+}
+
+bool STimerTableRow::HandleShouldBeEnabled() const
+{
+	bool bResult = false;
+
+	if (TimerNodePtr->IsGroup())
+	{
+		bResult = true;
+	}
+	else
+	{
+		if (OnShouldBeEnabled.IsBound())
+		{
+			bResult = OnShouldBeEnabled.Execute(TimerNodePtr->GetId());
+		}
+	}
+
+	return bResult;
+}
+
+EVisibility STimerTableRow::IsColumnVisible(const FName ColumnId) const
+{
+	EVisibility Result = EVisibility::Collapsed;
+
+	if (IsColumnVisibleDelegate.IsBound())
+	{
+		Result = IsColumnVisibleDelegate.Execute(ColumnId) ? EVisibility::Visible : EVisibility::Collapsed;
+	}
+
+	return Result;
+}
+
+void STimerTableRow::OnSetHoveredTableCell(const FName InColumnId, const FTimerNodePtr InSamplePtr)
+{
+	SetHoveredTableCellDelegate.ExecuteIfBound(InColumnId, InSamplePtr);
+}
+
+#undef LOCTEXT_NAMESPACE
diff --git a/Engine/Source/Developer/TraceInsights/Private/Insights/Widgets/STimerTableRow.h b/Engine/Source/Developer/TraceInsights/Private/Insights/Widgets/STimerTableRow.h
new file mode 100644
--- /dev/null
+++ b/Engine/Source/Developer/TraceInsights/Private/Insights/Widgets/STimerTableRow.h
@@ -0,0 +1,93 @@
+// Copyright 1998-2019 Epic Games, Inc. All Rights Reserved.
+
+#pragma once
+
+#include "CoreMinimal.h"
+#include "Widgets/DeclarativeSyntaxSupport.h"
+#include "Widgets/Views/SHeaderRow.h"
+#include "Widgets/Views/STableViewBase.h"
+#include "Widgets/Views/STableRow.h"
+#include "Widgets/Views/STreeView.h"
+
+// Insights
+#include "Insights/ViewModels/TimerNodeHelper.h"
+
+DECLARE_DELEGATE_RetVal_OneParam(bool, FShouldBeEnabledDelegate, const uint32 /*TimerId*/);
+DECLARE_DELEGATE_RetVal_OneParam(bool, FIsColumnVisibleDelegate, const FName /*ColumnId*/);
+DECLARE_DELEGATE_TwoParams(FSetHoveredTableCell, const FName /*ColumnId*/, const FTimerNodePtr /*TimerNodePtr*/);
+DECLARE_DELEGATE_RetVal_OneParam(EHorizontalAlignment, FGetColumnOutlineHAlignmentDelegate, const FName /*ColumnId*/);
+
+/** Widget that represents a table row in the timers' tree control. Generates widgets for each column on demand. */
+class STimerTableRow : public SMultiColumnTableRow<FTimerNodePtr>
+{
+public:
+	SLATE_BEGIN_ARGS(STimerTableRow) {}
+		SLATE_EVENT(FShouldBeEnabledDelegate, OnShouldBeEnabled)
+		SLATE_EVENT(FIsColumnVisibleDelegate, OnIsColumnVisible)
+		SLATE_EVENT(FSetHoveredTableCell, OnSetHoveredTableCell)
+		SLATE_EVENT(FGetColumnOutlineHAlignmentDelegate, OnGetColumnOutlineHAlignmentDelegate)
+		SLATE_ATTRIBUTE(FText, HighlightText)
+		SLATE_ATTRIBUTE(FName, HighlightedTimerName)
+		SLATE_ARGUMENT(FTimerNodePtr, TimerNodePtr)
+	SLATE_END_ARGS()
+
+public:
+	void Construct(const FArguments& InArgs, const TSharedRef<STableViewBase>& InOwnerTableView);
+
+	virtual TSharedRef<SWidget> GenerateWidgetForColumn(const FName& ColumnId) override;
+
+	/**
+	 * Called when Slate detects that a widget started to be dragged.
+	 * Usage:
+	 * A widget can ask Slate to detect a drag.
+	 * OnMouseDown() reply with FReply::Handled().DetectDrag(SharedThis(this)).
+	 * Slate will either send an OnDragDetected() event or do nothing.
+	 * If the user releases a mouse button or leaves the widget before
+	 * a drag is triggered (maybe user started at the very edge) then no event will be
+	 * sent.
+	 *
+	 * @param  InMyGeometry  Widget geometry
+	 * @param  InMouseEvent  MouseMove that triggered the drag
+	 *
+	 */
+	virtual FReply OnDragDetected(const FGeometry& MyGeometry, const FPointerEvent& MouseEvent) override;
+
+protected:
+	/**
+	 * @return a text which describes this table row, refers to both groups and timers
+	 */
+	FText GetText() const;
+
+	/**
+	 * @return a font style which is used to draw this table row, refers to both groups and timers
+	 */
+	FSlateFontInfo GetFont() const;
+
+	/**
+	 * @return a color and opacity value used to draw this table row, refers to both groups and timers
+	 */
+	FSlateColor GetColorAndOpacity() const;
+
+	FSlateColor GetBackgroundColorAndOpacity() const;
+	FSlateColor GetBackgroundColorAndOpacity(double Time) const;
+	FSlateColor GetOutlineColorAndOpacity() const;
+	const FSlateBrush* GetOutlineBrush(const FName ColumnId) const;
+	bool HandleShouldBeEnabled() const;
+	EVisibility IsColumnVisible(const FName ColumnId) const;
+	void OnSetHoveredTableCell(const FName ColumnId, const FTimerNodePtr SamplePtr);
+
+protected:
+	/** Data context for this table row. */
+	FTimerNodePtr TimerNodePtr;
+
+	FShouldBeEnabledDelegate OnShouldBeEnabled;
+	FIsColumnVisibleDelegate IsColumnVisibleDelegate;
+	FSetHoveredTableCell SetHoveredTableCellDelegate;
+	FGetColumnOutlineHAlignmentDelegate GetColumnOutlineHAlignmentDelegate;
+
+	/** Text to be highlighted on timer name. */
+	TAttribute<FText> HighlightText;
+
+	/** Name of the timer that should be drawn as highlighted. */
+	TAttribute<FName> HighlightedTimerName;
+};
