Copy up of Unreal Insights from //Tasks/UE4/Dev-Core-Trace.

diff --git a/Engine/Source/Developer/TraceInsights/Private/Insights/Widgets/STimersView.cpp b/Engine/Source/Developer/TraceInsights/Private/Insights/Widgets/STimersView.cpp
new file mode 100644
--- /dev/null
+++ b/Engine/Source/Developer/TraceInsights/Private/Insights/Widgets/STimersView.cpp
@@ -0,0 +1,1796 @@
+// Copyright 1998-2019 Epic Games, Inc. All Rights Reserved.
+
+#include "STimersView.h"
+
+#include "EditorStyleSet.h"
+#include "Framework/MultiBox/MultiBoxBuilder.h"
+#include "SlateOptMacros.h"
+#include "TraceServices/AnalysisService.h"
+#include "Widgets/Input/SCheckBox.h"
+#include "Widgets/Layout/SScrollBox.h"
+#include "Widgets/Input/SSearchBox.h"
+#include "Widgets/Images/SImage.h"
+#include "Widgets/Layout/SGridPanel.h"
+#include "Widgets/Layout/SSeparator.h"
+#include "Widgets/SToolTip.h"
+#include "Widgets/Views/STableViewBase.h"
+
+// Insights
+#include "Insights/TimingProfilerCommon.h"
+#include "Insights/TimingProfilerManager.h"
+//#include "Insights/ViewModels/TraceSession.h"
+#include "Insights/ViewModels/TimerNodeHelper.h"
+#include "Insights/ViewModels/TimersViewColumnFactory.h"
+#include "Insights/Widgets/STimersViewTooltip.h"
+#include "Insights/Widgets/STimerTableRow.h"
+
+#define LOCTEXT_NAMESPACE "STimersView"
+
+////////////////////////////////////////////////////////////////////////////////////////////////////
+
+STimersView::STimersView()
+	: bExpansionSaved(false)
+	, GroupingMode(ETimerGroupingMode::ByType)
+	, ColumnSortMode(EColumnSortMode::Type::Descending)
+	, ColumnBeingSorted(FName(TEXT("TotalInclusiveTime")))
+{
+	FMemory::Memset(bTimerNodeIsVisible, 1);
+}
+
+////////////////////////////////////////////////////////////////////////////////////////////////////
+
+STimersView::~STimersView()
+{
+	// Remove ourselves from the profiler manager.
+	if (FInsightsManager::Get().IsValid())
+	{
+		FInsightsManager::Get()->OnRequestTimersUpdate().RemoveAll(this);
+	}
+}
+
+////////////////////////////////////////////////////////////////////////////////////////////////////
+
+BEGIN_SLATE_FUNCTION_BUILD_OPTIMIZATION
+void STimersView::Construct(const FArguments& InArgs)
+{
+	SAssignNew(ExternalScrollbar, SScrollBar)
+	.AlwaysShowScrollbar(true);
+
+	ChildSlot
+	[
+		SNew(SVerticalBox)
+
+		+SVerticalBox::Slot()
+		.VAlign(VAlign_Center)
+		.AutoHeight()
+		[
+			SNew(SBorder)
+			.BorderImage(FEditorStyle::GetBrush("ToolPanel.GroupBorder"))
+			.Padding(2.0f)
+			[
+				SNew(SVerticalBox)
+
+				// Search box
+				+SVerticalBox::Slot()
+				.VAlign(VAlign_Center)
+				.Padding(2.0f)
+				.AutoHeight()
+				[
+					SAssignNew(SearchBox, SSearchBox)
+					.HintText(LOCTEXT("SearchBoxHint", "Search timers or groups"))
+					.OnTextChanged(this, &STimersView::SearchBox_OnTextChanged)
+					.IsEnabled(this, &STimersView::SearchBox_IsEnabled)
+					.ToolTipText(LOCTEXT("FilterSearchHint", "Type here to search timer or group"))
+				]
+
+				// Group by
+				+SVerticalBox::Slot()
+				.VAlign(VAlign_Center)
+				.Padding(2.0f)
+				.AutoHeight()
+				[
+					SNew(SHorizontalBox)
+
+					+SHorizontalBox::Slot()
+					.FillWidth(1.0f)
+					.VAlign(VAlign_Center)
+					[
+						SNew(STextBlock)
+						.Text(LOCTEXT("GroupByText", "Group by"))
+					]
+
+					+SHorizontalBox::Slot()
+					.FillWidth(2.0f)
+					.VAlign(VAlign_Center)
+					[
+						SAssignNew(GroupByComboBox, SComboBox<TSharedPtr<ETimerGroupingMode>>)
+						.ToolTipText(this, &STimersView::GroupBy_GetSelectedTooltipText)
+						.OptionsSource(&GroupByOptionsSource)
+						.OnSelectionChanged(this, &STimersView::GroupBy_OnSelectionChanged)
+						.OnGenerateWidget(this, &STimersView::GroupBy_OnGenerateWidget)
+						[
+							SNew(STextBlock)
+							.Text(this, &STimersView::GroupBy_GetSelectedText)
+						]
+					]
+				]
+
+				// Check boxes for: GpuScope, ComputeScope, CpuScope
+				+ SVerticalBox::Slot()
+				.VAlign(VAlign_Center)
+				.Padding(2.0f)
+				.AutoHeight()
+				[
+					SNew(SHorizontalBox)
+
+					+SHorizontalBox::Slot()
+					.Padding(FMargin(0.0f,0.0f,1.0f,0.0f))
+					.FillWidth(1.0f)
+					[
+						GetToggleButtonForTimerType(ETimerNodeType::GpuScope)
+					]
+
+					+SHorizontalBox::Slot()
+					.Padding(FMargin(1.0f,0.0f,1.0f,0.0f))
+					.FillWidth(1.0f)
+					[
+						GetToggleButtonForTimerType(ETimerNodeType::ComputeScope)
+					]
+
+					+SHorizontalBox::Slot()
+					.Padding(FMargin(1.0f,0.0f,1.0f,0.0f))
+					.FillWidth(1.0f)
+					[
+						GetToggleButtonForTimerType(ETimerNodeType::CpuScope)
+					]
+				]
+			]
+		]
+
+		// Tree view
+		+SVerticalBox::Slot()
+		.FillHeight(1.0f)
+		.Padding(0.0f, 6.0f, 0.0f, 0.0f)
+		[
+			SNew(SHorizontalBox)
+
+			+ SHorizontalBox::Slot()
+			.FillWidth(1.0f)
+			.Padding(0.0f)
+			[
+				SNew(SScrollBox)
+				.Orientation(Orient_Horizontal)
+
+				+ SScrollBox::Slot()
+				[
+					SNew(SBorder)
+					.BorderImage(FEditorStyle::GetBrush("ToolPanel.GroupBorder"))
+					.Padding(0.0f)
+					[
+						SAssignNew(TreeView, STreeView<FTimerNodePtr>)
+						.ExternalScrollbar(ExternalScrollbar)
+						.SelectionMode(ESelectionMode::Multi)
+						.TreeItemsSource(&FilteredGroupNodes)
+						.OnGetChildren(this, &STimersView::TreeView_OnGetChildren)
+						.OnGenerateRow(this, &STimersView::TreeView_OnGenerateRow)
+						.OnSelectionChanged(this, &STimersView::TreeView_OnSelectionChanged)
+						.OnMouseButtonDoubleClick(this, &STimersView::TreeView_OnMouseButtonDoubleClick)
+						.OnContextMenuOpening(FOnContextMenuOpening::CreateSP(this, &STimersView::TreeView_GetMenuContent))
+						.ItemHeight(12.0f)
+						.HeaderRow
+						(
+							SAssignNew(TreeViewHeaderRow, SHeaderRow)
+							.Visibility(EVisibility::Visible)
+						)
+					]
+				]
+			]
+
+			+ SHorizontalBox::Slot()
+			.AutoWidth()
+			.Padding(0.0f)
+			[
+				SNew(SBox)
+				.WidthOverride(FOptionalSize(13.0f))
+				[
+					ExternalScrollbar.ToSharedRef()
+				]
+			]
+		]
+	];
+
+	InitializeAndShowHeaderColumns();
+	//BindCommands();
+
+	// Create the search filters: text based, type based etc.
+	TextFilter = MakeShareable(new FTimerNodeTextFilter(FTimerNodeTextFilter::FItemToStringArray::CreateSP(this, &STimersView::HandleItemToStringArray)));
+	Filters = MakeShareable(new FTimerNodeFilterCollection());
+	Filters->Add(TextFilter);
+
+	CreateGroupByOptionsSources();
+
+	// Register ourselves with the profiler manager.
+	FInsightsManager::Get()->OnRequestTimersUpdate().AddSP(this, &STimersView::ProfilerManager_OnRequestTimersUpdate);
+}
+END_SLATE_FUNCTION_BUILD_OPTIMIZATION
+
+////////////////////////////////////////////////////////////////////////////////////////////////////
+
+TSharedPtr<SWidget> STimersView::TreeView_GetMenuContent() const
+{
+	const TArray<FTimerNodePtr> SelectedTimerNodes = TreeView->GetSelectedItems();
+	const int NumSelectedTimerNodes = SelectedTimerNodes.Num();
+	FTimerNodePtr SelectedTimerNode = NumSelectedTimerNodes ? SelectedTimerNodes[0] : nullptr;
+
+	const FTimersViewColumn* const * ColumnPtrPtr = FTimersViewColumnFactory::Get().ColumnIdToPtrMapping.Find(HoveredColumnId);
+	const FTimersViewColumn* const ColumnPtr = (ColumnPtrPtr != nullptr) ? *ColumnPtrPtr : nullptr;
+
+	FText SelectionStr;
+	FText PropertyName;
+	FText PropertyValue;
+
+	if (NumSelectedTimerNodes == 0)
+	{
+		SelectionStr = LOCTEXT("NothingSelected", "Nothing selected");
+	}
+	else if (NumSelectedTimerNodes == 1)
+	{
+		if (ColumnPtr != nullptr)
+		{
+			PropertyName = ColumnPtr->ShortName;
+			PropertyValue = ColumnPtr->GetFormattedValue(*SelectedTimerNode);
+		}
+		SelectionStr = FText::FromName(SelectedTimerNode->GetName());
+	}
+	else
+	{
+		SelectionStr = LOCTEXT("MultipleSelection", "Multiple selection");
+	}
+
+	const bool bShouldCloseWindowAfterMenuSelection = true;
+	FMenuBuilder MenuBuilder(bShouldCloseWindowAfterMenuSelection, NULL);
+
+	// Selection menu
+	MenuBuilder.BeginSection("Selection", LOCTEXT("ContextMenu_Header_Selection", "Selection"));
+	{
+		struct FLocal
+		{
+			static bool ReturnFalse()
+			{
+				return false;
+			}
+		};
+
+		FUIAction DummyUIAction;
+		DummyUIAction.CanExecuteAction = FCanExecuteAction::CreateStatic(&FLocal::ReturnFalse);
+		MenuBuilder.AddMenuEntry
+		(
+			SelectionStr,
+			LOCTEXT("ContextMenu_Selection", "Currently selected items"),
+			FSlateIcon(FEditorStyle::GetStyleSetName(), "@missing.icon"), DummyUIAction, NAME_None, EUserInterfaceActionType::Button
+		);
+	}
+	MenuBuilder.EndSection();
+
+	/*
+	// Culling/Filtering menu
+	MenuBuilder.BeginSection("Culling/Filtering", LOCTEXT("ContextMenu_Header_Root", "Root"));
+	{
+		// Culling menu
+		FText CullingDesc;
+
+		if (!Column.bCanBeCulled)
+		{
+			CullingDesc = LOCTEXT("ContextMenu_Culling_DescErrCol", "Culling not available, please select a different column");
+		}
+		else if (NumSelectedTimerNodes == 1)
+		{
+			CullingDesc = FText::Format(LOCTEXT("ContextMenu_Culling_DescFmt", "Cull events to '{0}' based on '{1}'"), PropertyValue, PropertyName);
+		}
+		else
+		{
+			CullingDesc = LOCTEXT("ContextMenu_Culling_DescErrEve", "Culling not available, please select one event");
+		}
+
+		MenuBuilder.AddMenuEntry
+		(
+			CullingDesc,
+			LOCTEXT("ContextMenu_Culling_Desc_TT", "Culls the event graph based on the property value of the selected event"),
+			FSlateIcon(FEditorStyle::GetStyleSetName(), "Profiler.EventGraph.CullEvents"), CullByProperty_Custom(SelectedTimerNode, HoveredColumnID, false), NAME_None, EUserInterfaceActionType::Button
+		);
+
+		// Filtering menu
+		FText FilteringDesc;
+
+		if (!Column.bCanBeFiltered)
+		{
+			FilteringDesc = LOCTEXT("ContextMenu_Filtering_DescErrCol", "Filtering not available, please select a different column");
+		}
+		else if (NumSelectedTimerNodes == 1)
+		{
+			FilteringDesc = FText::Format(LOCTEXT("ContextMenu_Filtering_DescFmt", "Filter events to '{0}' based on '{1}'"), PropertyValue, PropertyName);
+		}
+		else
+		{
+			FilteringDesc = LOCTEXT("ContextMenu_Filtering_DescErrEve", "Filtering not available, please select one event");
+		}
+
+		MenuBuilder.AddMenuEntry
+		(
+			FilteringDesc,
+			LOCTEXT("ContextMenu_Filtering_Desc_TT", "Filters the event graph based on the property value of the selected event"),
+			FSlateIcon(FEditorStyle::GetStyleSetName(), "Profiler.EventGraph.FilterEvents"), FilterOutByProperty_Custom(SelectedTimerNode, HoveredColumnID, false), NAME_None, EUserInterfaceActionType::Button
+		);
+
+		MenuBuilder.AddMenuSeparator();
+		MenuBuilder.AddMenuEntry
+		(
+			LOCTEXT("ContextMenu_ClearHistory", "Reset to default"),
+			LOCTEXT("ContextMenu_ClearHistory_Reset_Desc", "For the selected event graph resets root/culling/filter to the default state and clears the history"),
+			FSlateIcon(FEditorStyle::GetStyleSetName(), "Profiler.Misc.ResetToDefault"), ClearHistory_Custom(), NAME_None, EUserInterfaceActionType::Button
+		);
+	}
+	MenuBuilder.EndSection();
+	*/
+
+	/*
+		Event graph coloring based on inclusive time as a gradient from black to red?
+	*/
+
+	MenuBuilder.BeginSection("Misc", LOCTEXT("ContextMenu_Header_Misc", "Miscellaneous"));
+	{
+		/*
+		FUIAction Action_CopySelectedToClipboard
+		(
+			FExecuteAction::CreateSP(this, &STimersView::ContextMenu_CopySelectedToClipboard_Execute),
+			FCanExecuteAction::CreateSP(this, &STimersView::ContextMenu_CopySelectedToClipboard_CanExecute)
+		);
+		MenuBuilder.AddMenuEntry
+		(
+			LOCTEXT("ContextMenu_Header_Misc_CopySelectedToClipboard", "Copy To Clipboard"),
+			LOCTEXT("ContextMenu_Header_Misc_CopySelectedToClipboard_Desc", "Copies selection to clipboard"),
+			FSlateIcon(FEditorStyle::GetStyleSetName(), "Profiler.Misc.CopyToClipboard"), Action_CopySelectedToClipboard, NAME_None, EUserInterfaceActionType::Button
+		);
+		*/
+
+		MenuBuilder.AddSubMenu
+		(
+			LOCTEXT("ContextMenu_Header_Misc_Sort", "Sort By"),
+			LOCTEXT("ContextMenu_Header_Misc_Sort_Desc", "Sort by column"),
+			FNewMenuDelegate::CreateSP(this, &STimersView::TreeView_BuildSortByMenu),
+			false,
+			FSlateIcon(FEditorStyle::GetStyleSetName(), "Profiler.Misc.SortBy")
+		);
+	}
+	MenuBuilder.EndSection();
+
+	MenuBuilder.BeginSection("Columns", LOCTEXT("ContextMenu_Header_Columns", "Columns"));
+	{
+		MenuBuilder.AddSubMenu
+		(
+			LOCTEXT("ContextMenu_Header_Columns_View", "View Column"),
+			LOCTEXT("ContextMenu_Header_Columns_View_Desc", "Hides or shows columns"),
+			FNewMenuDelegate::CreateSP(this, &STimersView::TreeView_BuildViewColumnMenu),
+			false,
+			FSlateIcon(FEditorStyle::GetStyleSetName(), "Profiler.EventGraph.ViewColumn")
+		);
+
+		FUIAction Action_ShowAllColumns
+		(
+			FExecuteAction::CreateSP(this, &STimersView::ContextMenu_ShowAllColumns_Execute),
+			FCanExecuteAction::CreateSP(this, &STimersView::ContextMenu_ShowAllColumns_CanExecute)
+		);
+		MenuBuilder.AddMenuEntry
+		(
+			LOCTEXT("ContextMenu_Header_Columns_ShowAllColumns", "Show All Columns"),
+			LOCTEXT("ContextMenu_Header_Columns_ShowAllColumns_Desc", "Resets tree view to show all columns"),
+			FSlateIcon(FEditorStyle::GetStyleSetName(), "Profiler.EventGraph.ResetColumn"), Action_ShowAllColumns, NAME_None, EUserInterfaceActionType::Button
+		);
+
+		FUIAction Action_ShowMinMaxMedColumns
+		(
+			FExecuteAction::CreateSP(this, &STimersView::ContextMenu_ShowMinMaxMedColumns_Execute),
+			FCanExecuteAction::CreateSP(this, &STimersView::ContextMenu_ShowMinMaxMedColumns_CanExecute)
+		);
+		MenuBuilder.AddMenuEntry
+		(
+			LOCTEXT("ContextMenu_Header_Columns_ShowMinMaxMedColumns", "Reset Columns to Min/Max/Median Preset"),
+			LOCTEXT("ContextMenu_Header_Columns_ShowMinMaxMedColumns_Desc", "Resets columns to Min/Max/Median preset"),
+			FSlateIcon(FEditorStyle::GetStyleSetName(), "Profiler.EventGraph.ResetColumn"), Action_ShowMinMaxMedColumns, NAME_None, EUserInterfaceActionType::Button
+		);
+
+		FUIAction Action_ResetColumns
+		(
+			FExecuteAction::CreateSP(this, &STimersView::ContextMenu_ResetColumns_Execute),
+			FCanExecuteAction::CreateSP(this, &STimersView::ContextMenu_ResetColumns_CanExecute)
+		);
+		MenuBuilder.AddMenuEntry
+		(
+			LOCTEXT("ContextMenu_Header_Columns_ResetColumns", "Reset Columns to Default"),
+			LOCTEXT("ContextMenu_Header_Columns_ResetColumns_Desc", "Resets columns to default"),
+			FSlateIcon(FEditorStyle::GetStyleSetName(), "Profiler.EventGraph.ResetColumn"), Action_ResetColumns, NAME_None, EUserInterfaceActionType::Button
+		);
+	}
+	MenuBuilder.EndSection();
+
+	return MenuBuilder.MakeWidget();
+}
+
+////////////////////////////////////////////////////////////////////////////////////////////////////
+
+void STimersView::TreeView_BuildSortByMenu(FMenuBuilder& MenuBuilder)
+{
+	// TODO: Refactor later @see TSharedPtr<SWidget> SCascadePreviewViewportToolBar::GenerateViewMenu() const
+	MenuBuilder.BeginSection("ColumnName", LOCTEXT("ContextMenu_Header_Misc_ColumnName", "Column Name"));
+	for (auto It = TreeViewHeaderColumns.CreateConstIterator(); It; ++It)
+	{
+		const FTimersViewColumn& Column = It.Value();
+
+		if (Column.bIsVisible && Column.bCanBeSorted())
+		{
+			FUIAction Action_SortByColumn
+			(
+				FExecuteAction::CreateSP(this, &STimersView::ContextMenu_SortByColumn_Execute, Column.Id),
+				FCanExecuteAction::CreateSP(this, &STimersView::ContextMenu_SortByColumn_CanExecute, Column.Id),
+				FIsActionChecked::CreateSP(this, &STimersView::ContextMenu_SortByColumn_IsChecked, Column.Id)
+			);
+			MenuBuilder.AddMenuEntry
+			(
+				Column.TitleName,
+				Column.Description,
+				FSlateIcon(), Action_SortByColumn, NAME_None, EUserInterfaceActionType::RadioButton
+			);
+		}
+	}
+	MenuBuilder.EndSection();
+
+	//-----------------------------------------------------------------------------
+
+	MenuBuilder.BeginSection("SortMode", LOCTEXT("ContextMenu_Header_Misc_Sort_SortMode", "Sort Mode"));
+	{
+		FUIAction Action_SortAscending
+		(
+			FExecuteAction::CreateSP(this, &STimersView::ContextMenu_SortMode_Execute, EColumnSortMode::Ascending),
+			FCanExecuteAction::CreateSP(this, &STimersView::ContextMenu_SortMode_CanExecute, EColumnSortMode::Ascending),
+			FIsActionChecked::CreateSP(this, &STimersView::ContextMenu_SortMode_IsChecked, EColumnSortMode::Ascending)
+		);
+		MenuBuilder.AddMenuEntry
+		(
+			LOCTEXT("ContextMenu_Header_Misc_Sort_SortAscending", "Sort Ascending"),
+			LOCTEXT("ContextMenu_Header_Misc_Sort_SortAscending_Desc", "Sorts ascending"),
+			FSlateIcon(FEditorStyle::GetStyleSetName(), "Profiler.Misc.SortAscending"), Action_SortAscending, NAME_None, EUserInterfaceActionType::RadioButton
+		);
+
+		FUIAction Action_SortDescending
+		(
+			FExecuteAction::CreateSP(this, &STimersView::ContextMenu_SortMode_Execute, EColumnSortMode::Descending),
+			FCanExecuteAction::CreateSP(this, &STimersView::ContextMenu_SortMode_CanExecute, EColumnSortMode::Descending),
+			FIsActionChecked::CreateSP(this, &STimersView::ContextMenu_SortMode_IsChecked, EColumnSortMode::Descending)
+		);
+		MenuBuilder.AddMenuEntry
+		(
+			LOCTEXT("ContextMenu_Header_Misc_Sort_SortDescending", "Sort Descending"),
+			LOCTEXT("ContextMenu_Header_Misc_Sort_SortDescending_Desc", "Sorts descending"),
+			FSlateIcon(FEditorStyle::GetStyleSetName(), "Profiler.Misc.SortDescending"), Action_SortDescending, NAME_None, EUserInterfaceActionType::RadioButton
+		);
+	}
+	MenuBuilder.EndSection();
+}
+
+////////////////////////////////////////////////////////////////////////////////////////////////////
+
+void STimersView::TreeView_BuildViewColumnMenu(FMenuBuilder& MenuBuilder)
+{
+	MenuBuilder.BeginSection("ViewColumn", LOCTEXT("ContextMenu_Header_Columns_View", "View Column"));
+
+	for (auto It = TreeViewHeaderColumns.CreateConstIterator(); It; ++It)
+	{
+		const FTimersViewColumn& Column = It.Value();
+
+		FUIAction Action_ToggleColumn
+		(
+			FExecuteAction::CreateSP(this, &STimersView::ContextMenu_ToggleColumn_Execute, Column.Id),
+			FCanExecuteAction::CreateSP(this, &STimersView::ContextMenu_ToggleColumn_CanExecute, Column.Id),
+			FIsActionChecked::CreateSP(this, &STimersView::ContextMenu_ToggleColumn_IsChecked, Column.Id)
+		);
+		MenuBuilder.AddMenuEntry
+		(
+			Column.TitleName,
+			Column.Description,
+			FSlateIcon(), Action_ToggleColumn, NAME_None, EUserInterfaceActionType::ToggleButton
+		);
+	}
+
+	MenuBuilder.EndSection();
+}
+
+////////////////////////////////////////////////////////////////////////////////////////////////////
+
+void STimersView::InitializeAndShowHeaderColumns()
+{
+	//ColumnSortMode = EColumnSortMode::Ascending;
+	//ColumnBeingSorted = FTimersViewColumnFactory::Get().Collection[0]->Id;
+
+	const int32 NumColumns = FTimersViewColumnFactory::Get().Collection.Num();
+	for (int32 ColumnIndex = 0; ColumnIndex < NumColumns; ColumnIndex++)
+	{
+		TreeViewHeaderRow_CreateColumnArgs(ColumnIndex);
+	}
+
+	for (auto It = TreeViewHeaderColumns.CreateConstIterator(); It; ++It)
+	{
+		const FTimersViewColumn& Column = It.Value();
+
+		if (Column.bIsVisible)
+		{
+			TreeViewHeaderRow_ShowColumn(Column.Id);
+		}
+	}
+}
+
+////////////////////////////////////////////////////////////////////////////////////////////////////
+
+void STimersView::TreeViewHeaderRow_CreateColumnArgs(const int32 ColumnIndex)
+{
+	const FTimersViewColumn& Column = *FTimersViewColumnFactory::Get().Collection[ColumnIndex];
+	SHeaderRow::FColumn::FArguments ColumnArgs;
+
+	ColumnArgs
+		.ColumnId(Column.Id)
+		.DefaultLabel(Column.ShortName)
+		.SortMode(EColumnSortMode::None)
+		.HAlignHeader(HAlign_Fill)
+		.VAlignHeader(VAlign_Fill)
+		.HeaderContentPadding(TOptional<FMargin>(2.0f))
+		.HAlignCell(HAlign_Fill)
+		.VAlignCell(VAlign_Fill)
+		.SortMode(this, &STimersView::GetSortModeForColumn, Column.Id)
+		.OnSort(this, &STimersView::OnSortModeChanged)
+		.ManualWidth(Column.InitialColumnWidth)
+		.FixedWidth(Column.bIsFixedColumnWidth() ? Column.InitialColumnWidth : TOptional<float>())
+		.HeaderContent()
+		[
+			SNew(SHorizontalBox)
+			.ToolTip(STimersViewTooltip::GetColumnTooltip(Column))
+
+			+ SHorizontalBox::Slot()
+			.HAlign(Column.HorizontalAlignment)
+			.VAlign(VAlign_Center)
+			[
+				SNew(STextBlock)
+				.Text(Column.ShortName)
+			]
+		]
+		.MenuContent()
+		[
+			TreeViewHeaderRow_GenerateColumnMenu(Column)
+		];
+
+	TreeViewHeaderColumnArgs.Add(Column.Id, ColumnArgs);
+	TreeViewHeaderColumns.Add(Column.Id, Column);
+}
+
+////////////////////////////////////////////////////////////////////////////////////////////////////
+
+void STimersView::TreeViewHeaderRow_ShowColumn(const FName ColumnId)
+{
+	FTimersViewColumn& Column = TreeViewHeaderColumns.FindChecked(ColumnId);
+	Column.bIsVisible = true;
+	SHeaderRow::FColumn::FArguments& ColumnArgs = TreeViewHeaderColumnArgs.FindChecked(ColumnId);
+
+	const int32 NumColumns = TreeViewHeaderRow->GetColumns().Num();
+	int32 ColumnIndex = 0;
+	for (; ColumnIndex < NumColumns; ColumnIndex++)
+	{
+		const SHeaderRow::FColumn& CurrentColumn = TreeViewHeaderRow->GetColumns()[ColumnIndex];
+		const FTimersViewColumn& CurrentTimersViewColumn = TreeViewHeaderColumns.FindChecked(CurrentColumn.ColumnId);
+		if (Column.Order < CurrentTimersViewColumn.Order)
+			break;
+	}
+
+	TreeViewHeaderRow->InsertColumn(ColumnArgs, ColumnIndex);
+}
+
+////////////////////////////////////////////////////////////////////////////////////////////////////
+
+TSharedRef<SWidget> STimersView::TreeViewHeaderRow_GenerateColumnMenu(const FTimersViewColumn& Column)
+{
+	bool bIsMenuVisible = false;
+
+	const bool bShouldCloseWindowAfterMenuSelection = true;
+	FMenuBuilder MenuBuilder(bShouldCloseWindowAfterMenuSelection, NULL);
+	{
+		if (Column.bCanBeHidden())
+		{
+			MenuBuilder.BeginSection("Column", LOCTEXT("TreeViewHeaderRow_Header_Column", "Column"));
+
+			FUIAction Action_HideColumn
+			(
+				FExecuteAction::CreateSP(this, &STimersView::HeaderMenu_HideColumn_Execute, Column.Id),
+				FCanExecuteAction::CreateSP(this, &STimersView::HeaderMenu_HideColumn_CanExecute, Column.Id)
+			);
+
+			MenuBuilder.AddMenuEntry
+			(
+				LOCTEXT("TreeViewHeaderRow_HideColumn", "Hide"),
+				LOCTEXT("TreeViewHeaderRow_HideColumn_Desc", "Hides the selected column"),
+				FSlateIcon(), Action_HideColumn, NAME_None, EUserInterfaceActionType::Button
+			);
+			bIsMenuVisible = true;
+
+			MenuBuilder.EndSection();
+		}
+
+		if (Column.bCanBeSorted())
+		{
+			MenuBuilder.BeginSection("SortMode", LOCTEXT("ContextMenu_Header_Misc_Sort_SortMode", "Sort Mode"));
+
+			FUIAction Action_SortAscending
+			(
+				FExecuteAction::CreateSP(this, &STimersView::HeaderMenu_SortMode_Execute, Column.Id, EColumnSortMode::Ascending),
+				FCanExecuteAction::CreateSP(this, &STimersView::HeaderMenu_SortMode_CanExecute, Column.Id, EColumnSortMode::Ascending),
+				FIsActionChecked::CreateSP(this, &STimersView::HeaderMenu_SortMode_IsChecked, Column.Id, EColumnSortMode::Ascending)
+			);
+			MenuBuilder.AddMenuEntry
+			(
+				LOCTEXT("ContextMenu_Header_Misc_Sort_SortAscending", "Sort Ascending"),
+				LOCTEXT("ContextMenu_Header_Misc_Sort_SortAscending_Desc", "Sorts ascending"),
+				FSlateIcon(FEditorStyle::GetStyleSetName(), "Profiler.Misc.SortAscending"), Action_SortAscending, NAME_None, EUserInterfaceActionType::RadioButton
+			);
+
+			FUIAction Action_SortDescending
+			(
+				FExecuteAction::CreateSP(this, &STimersView::HeaderMenu_SortMode_Execute, Column.Id, EColumnSortMode::Descending),
+				FCanExecuteAction::CreateSP(this, &STimersView::HeaderMenu_SortMode_CanExecute, Column.Id, EColumnSortMode::Descending),
+				FIsActionChecked::CreateSP(this, &STimersView::HeaderMenu_SortMode_IsChecked, Column.Id, EColumnSortMode::Descending)
+			);
+			MenuBuilder.AddMenuEntry
+			(
+				LOCTEXT("ContextMenu_Header_Misc_Sort_SortDescending", "Sort Descending"),
+				LOCTEXT("ContextMenu_Header_Misc_Sort_SortDescending_Desc", "Sorts descending"),
+				FSlateIcon(FEditorStyle::GetStyleSetName(), "Profiler.Misc.SortDescending"), Action_SortDescending, NAME_None, EUserInterfaceActionType::RadioButton
+			);
+			bIsMenuVisible = true;
+
+			MenuBuilder.EndSection();
+		}
+
+		if (Column.bCanBeFiltered())
+		{
+			MenuBuilder.BeginSection("FilterMode", LOCTEXT("ContextMenu_Header_Misc_Filter_FilterMode", "Filter Mode"));
+			bIsMenuVisible = true;
+			MenuBuilder.EndSection();
+		}
+	}
+
+	/*
+	TODO:
+	- Show top ten
+	- Show top bottom
+	- Filter by list (avg, median, 10%, 90%, etc.)
+	- Text box for filtering for each column instead of one text box used for filtering
+	- Grouping button for flat view modes (show at most X groups, show all groups for names)
+	*/
+
+	return bIsMenuVisible ? MenuBuilder.MakeWidget() : (TSharedRef<SWidget>)SNullWidget::NullWidget;
+}
+
+////////////////////////////////////////////////////////////////////////////////////////////////////
+
+void STimersView::ProfilerManager_OnRequestTimersUpdate()
+{
+	TSharedPtr<const Trace::IAnalysisSession> SessionLocal = FInsightsManager::Get()->GetSession();
+
+	if (SessionLocal != Session)
+	{
+		Session = SessionLocal;
+		RebuildTree();
+	}
+	else
+	{
+		UpdateTree();
+	}
+}
+
+////////////////////////////////////////////////////////////////////////////////////////////////////
+
+void STimersView::UpdateTree()
+{
+	// Create groups, sort timers within the group and apply filtering.
+	CreateGroups();
+	SortTimers();
+	ApplyFiltering();
+}
+
+////////////////////////////////////////////////////////////////////////////////////////////////////
+
+void STimersView::ApplyFiltering()
+{
+	FilteredGroupNodes.Reset();
+
+	// Apply filter to all groups and its children.
+	const int32 NumGroups = GroupNodes.Num();
+	for (int32 GroupIndex = 0; GroupIndex < NumGroups; ++GroupIndex)
+	{
+		FTimerNodePtr& GroupPtr = GroupNodes[GroupIndex];
+		GroupPtr->ClearFilteredChildren();
+		const bool bIsGroupVisible = Filters->PassesAllFilters(GroupPtr);
+
+		const TArray<FTimerNodePtr>& GroupChildren = GroupPtr->GetChildren();
+		const int32 NumChildren = GroupChildren.Num();
+		int32 NumVisibleChildren = 0;
+		for (int32 Cx = 0; Cx < NumChildren; ++Cx)
+		{
+			// Add a child.
+			const FTimerNodePtr& NodePtr = GroupChildren[Cx];
+			const bool bIsChildVisible = Filters->PassesAllFilters(NodePtr) && bTimerNodeIsVisible[NodePtr->GetType()];
+			if (bIsChildVisible)
+			{
+				GroupPtr->AddFilteredChild(NodePtr);
+				NumVisibleChildren++;
+			}
+		}
+
+		if (bIsGroupVisible || NumVisibleChildren > 0)
+		{
+			// Add a group.
+			FilteredGroupNodes.Add(GroupPtr);
+			GroupPtr->bForceExpandGroupNode = true;
+		}
+		else
+		{
+			GroupPtr->bForceExpandGroupNode = false;
+		}
+	}
+
+	// Only expand timer nodes if we have a text filter.
+	const bool bNonEmptyTextFilter = !TextFilter->GetRawFilterText().IsEmpty();
+	if (bNonEmptyTextFilter)
+	{
+		if (!bExpansionSaved)
+		{
+			ExpandedNodes.Empty();
+			TreeView->GetExpandedItems(ExpandedNodes);
+			bExpansionSaved = true;
+		}
+
+		for (int32 Fx = 0; Fx < FilteredGroupNodes.Num(); Fx++)
+		{
+			const FTimerNodePtr& GroupPtr = FilteredGroupNodes[Fx];
+			TreeView->SetItemExpansion(GroupPtr, GroupPtr->bForceExpandGroupNode);
+		}
+	}
+	else
+	{
+		if (bExpansionSaved)
+		{
+			// Restore previously expanded nodes when the text filter is disabled.
+			TreeView->ClearExpandedItems();
+			for (auto It = ExpandedNodes.CreateConstIterator(); It; ++It)
+			{
+				TreeView->SetItemExpansion(*It, true);
+			}
+			bExpansionSaved = false;
+		}
+	}
+
+	// Request tree refresh
+	TreeView->RequestTreeRefresh();
+}
+
+////////////////////////////////////////////////////////////////////////////////////////////////////
+
+void STimersView::HandleItemToStringArray(const FTimerNodePtr& FTimerNodePtr, TArray<FString>& OutSearchStrings) const
+{
+	OutSearchStrings.Add(FTimerNodePtr->GetName().GetPlainNameString());
+}
+
+////////////////////////////////////////////////////////////////////////////////////////////////////
+
+TSharedRef<SWidget> STimersView::GetToggleButtonForTimerType(const ETimerNodeType NodeType)
+{
+	return SNew(SCheckBox)
+		.Style(FEditorStyle::Get(), "ToggleButtonCheckbox")
+		.HAlign(HAlign_Center)
+		.Padding(2.0f)
+		.OnCheckStateChanged(this, &STimersView::FilterByTimerType_OnCheckStateChanged, NodeType)
+		.IsChecked(this, &STimersView::FilterByTimerType_IsChecked, NodeType)
+		.ToolTipText(TimerNodeTypeHelper::ToDescription(NodeType))
+		[
+			SNew(SHorizontalBox)
+
+			+ SHorizontalBox::Slot()
+				.AutoWidth()
+				.VAlign(VAlign_Center)
+				[
+					SNew(SImage)
+						.Image(TimerNodeTypeHelper::GetIconForTimerNodeType(NodeType))
+				]
+
+			+SHorizontalBox::Slot()
+				.Padding(2.0f, 0.0f, 0.0f, 0.0f)
+				.VAlign(VAlign_Center)
+				[
+					SNew(STextBlock)
+						.Text(TimerNodeTypeHelper::ToName(NodeType))
+						.TextStyle(FEditorStyle::Get(), TEXT("Profiler.Caption"))
+				]
+		];
+}
+
+////////////////////////////////////////////////////////////////////////////////////////////////////
+
+void STimersView::FilterByTimerType_OnCheckStateChanged(ECheckBoxState NewRadioState, const ETimerNodeType InStatType)
+{
+	bTimerNodeIsVisible[InStatType] = NewRadioState == ECheckBoxState::Checked;
+	ApplyFiltering();
+}
+
+////////////////////////////////////////////////////////////////////////////////////////////////////
+
+ECheckBoxState STimersView::FilterByTimerType_IsChecked(const ETimerNodeType InStatType) const
+{
+	return bTimerNodeIsVisible[InStatType] ? ECheckBoxState::Checked : ECheckBoxState::Unchecked;
+}
+
+////////////////////////////////////////////////////////////////////////////////////////////////////
+// TreeView
+////////////////////////////////////////////////////////////////////////////////////////////////////
+
+void STimersView::TreeView_Refresh()
+{
+	if (TreeView.IsValid())
+	{
+		TreeView->RequestTreeRefresh();
+	}
+}
+
+////////////////////////////////////////////////////////////////////////////////////////////////////
+
+void STimersView::TreeView_OnSelectionChanged(FTimerNodePtr SelectedItem, ESelectInfo::Type SelectInfo)
+{
+	if (SelectInfo != ESelectInfo::Direct)
+	{
+		TArray<FTimerNodePtr> SelectedItems = TreeView->GetSelectedItems();
+		if (SelectedItems.Num() == 1)
+		{
+			//HighlightedTimerName = SelectedItems[0]->GetName();
+		}
+		else
+		{
+			//HighlightedTimerName = NAME_None;
+		}
+	}
+}
+
+////////////////////////////////////////////////////////////////////////////////////////////////////
+
+void STimersView::TreeView_OnGetChildren(FTimerNodePtr InParent, TArray<FTimerNodePtr>& OutChildren)
+{
+	OutChildren = InParent->GetFilteredChildren();
+}
+
+////////////////////////////////////////////////////////////////////////////////////////////////////
+
+void STimersView::TreeView_OnMouseButtonDoubleClick(FTimerNodePtr TimerNode)
+{
+	if (!TimerNode->IsGroup())
+	{
+		//im:TODO: const bool bIsStatTracked = FTimingProfilerManager::Get()->IsStatTracked(TimerNode->GetId());
+	//	if (!bIsStatTracked)
+	//	{
+	//		// Add a new graph.
+	//		FTimingProfilerManager::Get()->TrackStat(TimerNode->GetId());
+	//	}
+	//	else
+	//	{
+	//		// Remove a graph
+	//		FTimingProfilerManager::Get()->UntrackStat(TimerNode->GetId());
+	//	}
+	}
+	else
+	{
+		const bool bIsGroupExpanded = TreeView->IsItemExpanded(TimerNode);
+		TreeView->SetItemExpansion(TimerNode, !bIsGroupExpanded);
+	}
+}
+
+////////////////////////////////////////////////////////////////////////////////////////////////////
+// Tree View's Table Row
+////////////////////////////////////////////////////////////////////////////////////////////////////
+
+TSharedRef<ITableRow> STimersView::TreeView_OnGenerateRow(FTimerNodePtr TimerNodePtr, const TSharedRef<STableViewBase>& OwnerTable)
+{
+	TSharedRef<ITableRow> TableRow =
+		SNew(STimerTableRow, OwnerTable)
+		.OnShouldBeEnabled(this, &STimersView::TableRow_ShouldBeEnabled)
+		.OnIsColumnVisible(this, &STimersView::TableRow_IsColumnVisible)
+		.OnSetHoveredTableCell(this, &STimersView::TableRow_SetHoveredTableCell)
+		.OnGetColumnOutlineHAlignmentDelegate(this, &STimersView::TableRow_GetColumnOutlineHAlignment)
+		.HighlightText(this, &STimersView::TableRow_GetHighlightText)
+		.HighlightedTimerName(this, &STimersView::TableRow_GetHighlightedTimerName)
+		.TimerNodePtr(TimerNodePtr);
+
+	return TableRow;
+}
+
+////////////////////////////////////////////////////////////////////////////////////////////////////
+
+bool STimersView::TableRow_IsColumnVisible(const FName ColumnId) const
+{
+	bool bResult = false;
+	const FTimersViewColumn& ColumnPtr = TreeViewHeaderColumns.FindChecked(ColumnId);
+	return ColumnPtr.bIsVisible;
+}
+
+////////////////////////////////////////////////////////////////////////////////////////////////////
+
+void STimersView::TableRow_SetHoveredTableCell(const FName ColumnId, const FTimerNodePtr TimerNodePtr)
+{
+	HoveredColumnId = ColumnId;
+
+	const bool bIsAnyMenusVisible = FSlateApplication::Get().AnyMenusVisible();
+	if (!HasMouseCapture() && !bIsAnyMenusVisible)
+	{
+		HoveredTimerNodePtr = TimerNodePtr;
+	}
+
+	//UE_LOG(TimingProfiler, Log, TEXT("%s -> %s"), *HoveredColumnId.GetPlainNameString(), TimerNodePtr.IsValid() ? *TimerNodePtr->GetName().GetPlainNameString() : TEXT("nullptr"));
+}
+
+////////////////////////////////////////////////////////////////////////////////////////////////////
+
+EHorizontalAlignment STimersView::TableRow_GetColumnOutlineHAlignment(const FName ColumnId) const
+{
+	const TIndirectArray<SHeaderRow::FColumn>& Columns = TreeViewHeaderRow->GetColumns();
+	const int32 LastColumnIdx = Columns.Num() - 1;
+
+	// First column
+	if (Columns[0].ColumnId == ColumnId)
+	{
+		return HAlign_Left;
+	}
+	// Last column
+	else if (Columns[LastColumnIdx].ColumnId == ColumnId)
+	{
+		return HAlign_Right;
+	}
+	// Middle columns
+	{
+		return HAlign_Center;
+	}
+}
+
+////////////////////////////////////////////////////////////////////////////////////////////////////
+
+FText STimersView::TableRow_GetHighlightText() const
+{
+	return SearchBox->GetText();
+}
+
+////////////////////////////////////////////////////////////////////////////////////////////////////
+
+FName STimersView::TableRow_GetHighlightedTimerName() const
+{
+	return HighlightedTimerName;
+}
+
+////////////////////////////////////////////////////////////////////////////////////////////////////
+
+bool STimersView::TableRow_ShouldBeEnabled(const uint32 TimerId) const
+{
+	return true;//im:TODO: Session->GetAggregatedStat(TimerId) != nullptr;
+}
+
+////////////////////////////////////////////////////////////////////////////////////////////////////
+// SearchBox
+////////////////////////////////////////////////////////////////////////////////////////////////////
+
+void STimersView::SearchBox_OnTextChanged(const FText& InFilterText)
+{
+	TextFilter->SetRawFilterText(InFilterText);
+	SearchBox->SetError(TextFilter->GetFilterErrorText());
+	ApplyFiltering();
+}
+
+////////////////////////////////////////////////////////////////////////////////////////////////////
+
+bool STimersView::SearchBox_IsEnabled() const
+{
+	return TimerNodes.Num() > 0;
+}
+
+////////////////////////////////////////////////////////////////////////////////////////////////////
+// GroupBy
+////////////////////////////////////////////////////////////////////////////////////////////////////
+
+void STimersView::CreateGroups()
+{
+	TMap<FName, FTimerNodePtr> GroupNodeSet;
+
+	if (GroupingMode == ETimerGroupingMode::Flat)
+	{
+		const FName GroupName(TEXT("All"));
+		FTimerNodePtr* GroupPtr = GroupNodeSet.Find(GroupName);
+		if (!GroupPtr)
+		{
+			GroupPtr = &GroupNodeSet.Add(GroupName, MakeShareable(new FTimerNode(GroupName)));
+		}
+
+		for (const FTimerNodePtr& TimerNodePtr : TimerNodes)
+		{
+			(*GroupPtr)->AddChildAndSetGroupPtr(TimerNodePtr);
+		}
+
+		TreeView->SetItemExpansion(*GroupPtr, true);
+	}
+	// Creates groups based on stat metadata groups.
+	else if (GroupingMode == ETimerGroupingMode::ByMetaGroupName)
+	{
+		for (const FTimerNodePtr& TimerNodePtr : TimerNodes)
+		{
+			const FName GroupName = TimerNodePtr->GetMetaGroupName();
+
+			FTimerNodePtr* GroupPtr = GroupNodeSet.Find(GroupName);
+			if (!GroupPtr)
+			{
+				GroupPtr = &GroupNodeSet.Add(GroupName, MakeShareable(new FTimerNode(GroupName)));
+			}
+
+			(*GroupPtr)->AddChildAndSetGroupPtr(TimerNodePtr);
+			TreeView->SetItemExpansion(*GroupPtr, true);
+		}
+	}
+	// Creates one group for each stat type.
+	else if (GroupingMode == ETimerGroupingMode::ByType)
+	{
+		for (const FTimerNodePtr& TimerNodePtr : TimerNodes)
+		{
+			const FName GroupName = *TimerNodeTypeHelper::ToName(TimerNodePtr->GetType()).ToString();
+
+			FTimerNodePtr* GroupPtr = GroupNodeSet.Find(GroupName);
+			if (!GroupPtr)
+			{
+				GroupPtr = &GroupNodeSet.Add(GroupName, MakeShareable(new FTimerNode(GroupName)));
+			}
+
+			(*GroupPtr)->AddChildAndSetGroupPtr(TimerNodePtr);
+			TreeView->SetItemExpansion(*GroupPtr, true);
+		}
+	}
+	// Creates one group for one letter.
+	else if (GroupingMode == ETimerGroupingMode::ByName)
+	{
+		for (const FTimerNodePtr& TimerNodePtr : TimerNodes)
+		{
+			const FName GroupName = *TimerNodePtr->GetName().GetPlainNameString().Left(1);
+
+			FTimerNodePtr* GroupPtr = GroupNodeSet.Find(GroupName);
+			if (!GroupPtr)
+			{
+				GroupPtr = &GroupNodeSet.Add(GroupName, MakeShareable(new FTimerNode(GroupName)));
+			}
+
+			(*GroupPtr)->AddChildAndSetGroupPtr(TimerNodePtr);
+		}
+	}
+	// Creates one group for each logarithmic range ie. 0.001 - 0.01, 0.01 - 0.1, 0.1 - 1.0, 1.0 - 10.0, etc.
+	else if (GroupingMode == ETimerGroupingMode::ByTotalInclusiveTime)
+	{
+		//im:TODO:
+	}
+	// Creates one group for each logarithmic range ie. 0.001 - 0.01, 0.01 - 0.1, 0.1 - 1.0, 1.0 - 10.0, etc.
+	else if (GroupingMode == ETimerGroupingMode::ByTotalExclusiveTime)
+	{
+		//im:TODO:
+	}
+	// Creates one group for each logarithmic range ie. 0, 1 - 10, 10 - 100, 100 - 1000, etc.
+	else if (GroupingMode == ETimerGroupingMode::ByInstanceCount)
+	{
+		//im:TODO:
+	}
+
+	GroupNodeSet.GenerateValueArray(GroupNodes);
+
+	// Sort by a fake group name.
+	GroupNodes.Sort(TimerNodeSortingHelper::ByNameAscending());
+}
+
+////////////////////////////////////////////////////////////////////////////////////////////////////
+
+void STimersView::CreateGroupByOptionsSources()
+{
+	GroupByOptionsSource.Reset(3);
+
+	// Must be added in order of elements in the ETimerGroupingMode.
+	GroupByOptionsSource.Add(MakeShareable(new ETimerGroupingMode(ETimerGroupingMode::Flat)));
+	GroupByOptionsSource.Add(MakeShareable(new ETimerGroupingMode(ETimerGroupingMode::ByName)));
+	//GroupByOptionsSource.Add(MakeShareable(new ETimerGroupingMode(ETimerGroupingMode::ByMetaGroupName)));
+	GroupByOptionsSource.Add(MakeShareable(new ETimerGroupingMode(ETimerGroupingMode::ByType)));
+
+	ETimerGroupingModePtr* GroupingModePtrPtr = GroupByOptionsSource.FindByPredicate([&](const ETimerGroupingModePtr InGroupingModePtr) { return *InGroupingModePtr == GroupingMode; });
+	if (GroupingModePtrPtr != nullptr)
+	{
+		GroupByComboBox->SetSelectedItem(*GroupingModePtrPtr);
+	}
+
+	GroupByComboBox->RefreshOptions();
+}
+
+////////////////////////////////////////////////////////////////////////////////////////////////////
+
+void STimersView::GroupBy_OnSelectionChanged(TSharedPtr<ETimerGroupingMode> NewGroupingMode, ESelectInfo::Type SelectInfo)
+{
+	if (SelectInfo != ESelectInfo::Direct)
+	{
+		GroupingMode = *NewGroupingMode;
+
+		// Create groups, sort timers within the group and apply filtering.
+		CreateGroups();
+		SortTimers();
+		ApplyFiltering();
+	}
+}
+
+////////////////////////////////////////////////////////////////////////////////////////////////////
+
+TSharedRef<SWidget> STimersView::GroupBy_OnGenerateWidget(TSharedPtr<ETimerGroupingMode> InGroupingMode) const
+{
+	return SNew(STextBlock)
+		.Text(TimerNodeGroupingHelper::ToName(*InGroupingMode))
+		.ToolTipText(TimerNodeGroupingHelper::ToDescription(*InGroupingMode));
+}
+
+////////////////////////////////////////////////////////////////////////////////////////////////////
+
+FText STimersView::GroupBy_GetSelectedText() const
+{
+	return TimerNodeGroupingHelper::ToName(GroupingMode);
+}
+
+////////////////////////////////////////////////////////////////////////////////////////////////////
+
+FText STimersView::GroupBy_GetSelectedTooltipText() const
+{
+	return TimerNodeGroupingHelper::ToDescription(GroupingMode);
+}
+
+////////////////////////////////////////////////////////////////////////////////////////////////////
+// SortBy
+////////////////////////////////////////////////////////////////////////////////////////////////////
+
+void STimersView::SortTimers()
+{
+	const int32 NumGroups = GroupNodes.Num();
+
+	#define CHECK_AND_SORT_COLUMN(ColumnId, SortTypeName) \
+		if (ColumnBeingSorted == ColumnId) \
+		{ \
+			if (ColumnSortMode == EColumnSortMode::Type::Descending) \
+			{ \
+				for (int32 ID = 0; ID < NumGroups; ++ID) \
+				{ \
+					GroupNodes[ID]->SortChildren(TimerNodeSortingHelper::SortTypeName##Descending()); \
+				} \
+			} \
+			else /*if (ColumnSortMode == EColumnSortMode::Type::Ascending)*/ \
+			{ \
+				for (int32 ID = 0; ID < NumGroups; ++ID) \
+				{ \
+					GroupNodes[ID]->SortChildren(TimerNodeSortingHelper::SortTypeName##Ascending()); \
+				} \
+			} \
+		}
+
+		 CHECK_AND_SORT_COLUMN(FName(TEXT("Name")), ByName)
+	else CHECK_AND_SORT_COLUMN(FName(TEXT("MetaGroupName")), ByMetaGroupName)
+	else CHECK_AND_SORT_COLUMN(FName(TEXT("Type")), ByType)
+	else CHECK_AND_SORT_COLUMN(FName(TEXT("InstanceCount")), ByInstanceCount)
+	else CHECK_AND_SORT_COLUMN(FName(TEXT("TotalInclusiveTime")), ByTotalInclusiveTime)
+	else CHECK_AND_SORT_COLUMN(FName(TEXT("MinInclusiveTime")), ByMinInclusiveTime)
+	else CHECK_AND_SORT_COLUMN(FName(TEXT("MaxInclusiveTime")), ByMaxInclusiveTime)
+	else CHECK_AND_SORT_COLUMN(FName(TEXT("AvgInclusiveTime")), ByAverageInclusiveTime)
+	else CHECK_AND_SORT_COLUMN(FName(TEXT("MedInclusiveTime")), ByMedianInclusiveTime)
+	else CHECK_AND_SORT_COLUMN(FName(TEXT("TotalExclusiveTime")), ByTotalExclusiveTime)
+	else CHECK_AND_SORT_COLUMN(FName(TEXT("MinExclusiveTime")), ByMinExclusiveTime)
+	else CHECK_AND_SORT_COLUMN(FName(TEXT("MaxExclusiveTime")), ByMaxExclusiveTime)
+	else CHECK_AND_SORT_COLUMN(FName(TEXT("AvgExclusiveTime")), ByAverageExclusiveTime)
+	else CHECK_AND_SORT_COLUMN(FName(TEXT("MedExclusiveTime")), ByMedianExclusiveTime)
+}
+
+////////////////////////////////////////////////////////////////////////////////////////////////////
+
+EColumnSortMode::Type STimersView::GetSortModeForColumn(const FName ColumnId) const
+{
+	if (ColumnBeingSorted != ColumnId)
+	{
+		return EColumnSortMode::None;
+	}
+
+	return ColumnSortMode;
+}
+
+void STimersView::SetSortModeForColumn(const FName& ColumnId, const EColumnSortMode::Type SortMode)
+{
+	ColumnBeingSorted = ColumnId;
+	ColumnSortMode = SortMode;
+
+	// Sort timers and apply filtering.
+	SortTimers();
+	ApplyFiltering();
+}
+
+void STimersView::OnSortModeChanged(const EColumnSortPriority::Type SortPriority, const FName& ColumnId, const EColumnSortMode::Type SortMode)
+{
+	SetSortModeForColumn(ColumnId, SortMode);
+	TreeView_Refresh();
+}
+
+////////////////////////////////////////////////////////////////////////////////////////////////////
+// SortMode action (HeaderMenu)
+////////////////////////////////////////////////////////////////////////////////////////////////////
+
+bool STimersView::HeaderMenu_SortMode_IsChecked(const FName ColumnId, const EColumnSortMode::Type InSortMode)
+{
+	return ColumnBeingSorted == ColumnId && ColumnSortMode == InSortMode;
+}
+
+bool STimersView::HeaderMenu_SortMode_CanExecute(const FName ColumnId, const EColumnSortMode::Type InSortMode) const
+{
+	const FTimersViewColumn& Column = TreeViewHeaderColumns.FindChecked(ColumnId);
+	const bool bIsValid = Column.bCanBeSorted();
+
+	bool bCanExecute = ColumnBeingSorted != ColumnId ? true : ColumnSortMode != InSortMode;
+	bCanExecute = bCanExecute && bIsValid;
+
+	return bCanExecute;
+}
+
+void STimersView::HeaderMenu_SortMode_Execute(const FName ColumnId, const EColumnSortMode::Type InSortMode)
+{
+	SetSortModeForColumn(ColumnId, InSortMode);
+	TreeView_Refresh();
+}
+
+////////////////////////////////////////////////////////////////////////////////////////////////////
+// SortMode action (ContextMenu)
+////////////////////////////////////////////////////////////////////////////////////////////////////
+
+bool STimersView::ContextMenu_SortMode_IsChecked(const EColumnSortMode::Type InSortMode)
+{
+	return ColumnSortMode == InSortMode;
+}
+
+bool STimersView::ContextMenu_SortMode_CanExecute(const EColumnSortMode::Type InSortMode) const
+{
+	return ColumnSortMode != InSortMode;
+}
+
+void STimersView::ContextMenu_SortMode_Execute(const EColumnSortMode::Type InSortMode)
+{
+	SetSortModeForColumn(ColumnBeingSorted, InSortMode);
+	TreeView_Refresh();
+}
+
+////////////////////////////////////////////////////////////////////////////////////////////////////
+// SortByColumn action (ContextMenu)
+////////////////////////////////////////////////////////////////////////////////////////////////////
+
+bool STimersView::ContextMenu_SortByColumn_IsChecked(const FName ColumnId)
+{
+	return ColumnId == ColumnBeingSorted;
+}
+
+bool STimersView::ContextMenu_SortByColumn_CanExecute(const FName ColumnId) const
+{
+	return ColumnId != ColumnBeingSorted;
+}
+
+void STimersView::ContextMenu_SortByColumn_Execute(const FName ColumnId)
+{
+	SetSortModeForColumn(ColumnId, EColumnSortMode::Descending);
+	TreeView_Refresh();
+}
+
+////////////////////////////////////////////////////////////////////////////////////////////////////
+// HideColumn action (HeaderMenu)
+////////////////////////////////////////////////////////////////////////////////////////////////////
+
+bool STimersView::HeaderMenu_HideColumn_CanExecute(const FName ColumnId) const
+{
+	const FTimersViewColumn& Column = TreeViewHeaderColumns.FindChecked(ColumnId);
+	return Column.bCanBeHidden();
+}
+
+void STimersView::HeaderMenu_HideColumn_Execute(const FName ColumnId)
+{
+	FTimersViewColumn& Column = TreeViewHeaderColumns.FindChecked(ColumnId);
+	Column.bIsVisible = false;
+	TreeViewHeaderRow->RemoveColumn(ColumnId);
+}
+
+////////////////////////////////////////////////////////////////////////////////////////////////////
+// ToggleColumn action (ContextMenu)
+////////////////////////////////////////////////////////////////////////////////////////////////////
+
+bool STimersView::ContextMenu_ToggleColumn_IsChecked(const FName ColumnId)
+{
+	const FTimersViewColumn& Column = TreeViewHeaderColumns.FindChecked(ColumnId);
+	return Column.bIsVisible;
+}
+
+bool STimersView::ContextMenu_ToggleColumn_CanExecute(const FName ColumnId) const
+{
+	const FTimersViewColumn& Column = TreeViewHeaderColumns.FindChecked(ColumnId);
+	return Column.bCanBeHidden();
+}
+
+void STimersView::ContextMenu_ToggleColumn_Execute(const FName ColumnId)
+{
+	FTimersViewColumn& Column = TreeViewHeaderColumns.FindChecked(ColumnId);
+	if (Column.bIsVisible)
+	{
+		HeaderMenu_HideColumn_Execute(ColumnId);
+	}
+	else
+	{
+		TreeViewHeaderRow_ShowColumn(ColumnId);
+	}
+}
+
+////////////////////////////////////////////////////////////////////////////////////////////////////
+// "Show All Columns" action (ContextMenu)
+////////////////////////////////////////////////////////////////////////////////////////////////////
+
+bool STimersView::ContextMenu_ShowAllColumns_CanExecute() const
+{
+	return true;
+}
+
+////////////////////////////////////////////////////////////////////////////////////////////////////
+
+void STimersView::ContextMenu_ShowAllColumns_Execute()
+{
+	ColumnSortMode = EColumnSortMode::Descending;
+	ColumnBeingSorted = FName(TEXT("TotalInclusiveTime"));
+
+	const int32 NumColumns = FTimersViewColumnFactory::Get().Collection.Num();
+	for (int32 ColumnIndex = 0; ColumnIndex < NumColumns; ColumnIndex++)
+	{
+		const FTimersViewColumn& DefaultColumn = *FTimersViewColumnFactory::Get().Collection[ColumnIndex];
+		const FTimersViewColumn& CurrentColumn = TreeViewHeaderColumns.FindChecked(DefaultColumn.Id);
+
+		if (!CurrentColumn.bIsVisible)
+		{
+			TreeViewHeaderRow_ShowColumn(DefaultColumn.Id);
+		}
+	}
+}
+
+////////////////////////////////////////////////////////////////////////////////////////////////////
+// "Show Min/Max/Median Columns" action (ContextMenu)
+////////////////////////////////////////////////////////////////////////////////////////////////////
+
+bool STimersView::ContextMenu_ShowMinMaxMedColumns_CanExecute() const
+{
+	return true;
+}
+
+////////////////////////////////////////////////////////////////////////////////////////////////////
+
+void STimersView::ContextMenu_ShowMinMaxMedColumns_Execute()
+{
+	TSet<FName> Preset =
+	{
+		TEXT("Name"),
+		TEXT("InstanceCount"),
+		TEXT("TotalInclusiveTime"),
+		TEXT("MinInclusiveTime"),
+		TEXT("MaxInclusiveTime"),
+		TEXT("MedInclusiveTime"),
+		TEXT("TotalExclusiveTime"),
+		TEXT("MinExclusiveTime"),
+		TEXT("MaxExclusiveTime"),
+		TEXT("MedExclusiveTime"),
+	};
+
+	ColumnSortMode = EColumnSortMode::Descending;
+	ColumnBeingSorted = FName(TEXT("TotalInclusiveTime"));
+
+	const int32 NumColumns = FTimersViewColumnFactory::Get().Collection.Num();
+	for (int32 ColumnIndex = 0; ColumnIndex < NumColumns; ColumnIndex++)
+	{
+		const FTimersViewColumn& DefaultColumn = *FTimersViewColumnFactory::Get().Collection[ColumnIndex];
+		const FTimersViewColumn& CurrentColumn = TreeViewHeaderColumns.FindChecked(DefaultColumn.Id);
+
+		bool bIsVisible = Preset.Contains(DefaultColumn.Id);
+		if (bIsVisible && !CurrentColumn.bIsVisible)
+		{
+			TreeViewHeaderRow_ShowColumn(DefaultColumn.Id);
+		}
+		else if (!bIsVisible && CurrentColumn.bIsVisible)
+		{
+			HeaderMenu_HideColumn_Execute(DefaultColumn.Id);
+		}
+	}
+}
+
+////////////////////////////////////////////////////////////////////////////////////////////////////
+// ResetColumns action (ContextMenu)
+////////////////////////////////////////////////////////////////////////////////////////////////////
+
+bool STimersView::ContextMenu_ResetColumns_CanExecute() const
+{
+	return true;
+}
+
+////////////////////////////////////////////////////////////////////////////////////////////////////
+
+void STimersView::ContextMenu_ResetColumns_Execute()
+{
+	ColumnSortMode = EColumnSortMode::Descending;
+	ColumnBeingSorted = FName(TEXT("TotalInclusiveTime"));
+
+	const int32 NumColumns = FTimersViewColumnFactory::Get().Collection.Num();
+	for (int32 ColumnIndex = 0; ColumnIndex < NumColumns; ColumnIndex++)
+	{
+		const FTimersViewColumn& DefaultColumn = *FTimersViewColumnFactory::Get().Collection[ColumnIndex];
+		const FTimersViewColumn& CurrentColumn = TreeViewHeaderColumns.FindChecked(DefaultColumn.Id);
+
+		if (DefaultColumn.bIsVisible && !CurrentColumn.bIsVisible)
+		{
+			TreeViewHeaderRow_ShowColumn(DefaultColumn.Id);
+		}
+		else if (!DefaultColumn.bIsVisible && CurrentColumn.bIsVisible)
+		{
+			HeaderMenu_HideColumn_Execute(DefaultColumn.Id);
+		}
+	}
+}
+
+////////////////////////////////////////////////////////////////////////////////////////////////////
+
+void STimersView::RebuildTree(bool bRebuildEvenIfNotChanged)
+{
+	bool bTimerNodesListChanged = false;
+
+	if (bRebuildEvenIfNotChanged)
+	{
+		TimerNodes.Empty(TimerNodes.Num());
+		//TimerNodesMap.Empty(TimerNodesMap.Num());
+		TimerNodesIdMap.Empty(TimerNodesIdMap.Num());
+
+		bTimerNodesListChanged = true;
+	}
+
+	if (Session.IsValid())
+	{
+		Trace::FAnalysisSessionReadScope _(*Session.Get());
+
+		Session->ReadTimingProfilerProvider([this, bRebuildEvenIfNotChanged, &bTimerNodesListChanged](const Trace::ITimingProfilerProvider& TimingProfilerProvider)
+		{
+			TimingProfilerProvider.ReadTimers([this, bRebuildEvenIfNotChanged, &bTimerNodesListChanged](const Trace::FTimingProfilerTimer* Timers, uint64 TimersCount)
+			{
+				if (TimersCount != TimerNodes.Num() || bRebuildEvenIfNotChanged)
+				{
+					TimerNodes.Empty(TimerNodes.Num());
+					//TimerNodesMap.Empty(TimerNodesMap.Num());
+					TimerNodesIdMap.Empty(TimerNodesIdMap.Num());
+
+					for (uint64 TimerIndex = 0; TimerIndex < TimersCount; ++TimerIndex)
+					{
+						const Trace::FTimingProfilerTimer& Timer = Timers[TimerIndex];
+						FName Name(*Timer.Name);// +TEXT(" [GPU]")));
+						FName Group(Timer.IsGpuTimer ? TEXT("GPU") : TEXT("CPU"));
+						ETimerNodeType Type = Timer.IsGpuTimer ? ETimerNodeType::GpuScope : ETimerNodeType::CpuScope;
+						FTimerNode* TimerPtr = new FTimerNode(Timer.Id, Name, Group, Type);
+						FTimerNodePtr TimerNodePtr = MakeShareable(TimerPtr);
+						TimerNodes.Add(TimerNodePtr);
+						//TimerNodesMap.Add(Name, TimerNodePtr);
+						TimerNodesIdMap.Add(Timer.Id, TimerNodePtr);
+					}
+
+					bTimerNodesListChanged = true;
+				}
+			});
+		});
+	}
+
+	if (bTimerNodesListChanged)
+	{
+		UpdateTree();
+		UpdateStats(StatsStartTime, StatsEndTime);
+	}
+}
+
+////////////////////////////////////////////////////////////////////////////////////////////////////
+
+struct FTimeCalculationHelper
+{
+	struct FStackEntry
+	{
+		double StartTime;
+		double AccumulatedTimeInChildren;
+		uint32 EventType;
+	};
+
+	FTimeCalculationHelper(double InIntervalStartTime, double InIntervalEndTime, TMap<uint64, FTimerStatsEx>& InStatsMap)
+		: IntervalStartTime(InIntervalStartTime)
+		, IntervalEndTime(InIntervalEndTime)
+		, StatsMap(InStatsMap)
+	{
+
+	}
+
+	template<typename CallbackType>
+	void UpdateStats(const Trace::ITimingProfilerProvider::Timeline& Timeline, CallbackType Callback)
+	{
+		EventStack.Empty(EventStack.Num());
+		Timeline.EnumerateEvents(IntervalStartTime, IntervalEndTime, [this, Callback](bool IsEnter, double Time, const Trace::FTimingProfilerEvent& Event)
+		{
+			if (IsEnter)
+			{
+				FStackEntry& StackEntry = EventStack.AddDefaulted_GetRef();
+				StackEntry.StartTime = FMath::Max(IntervalStartTime, Time);
+				StackEntry.AccumulatedTimeInChildren = 0.0;
+				StackEntry.EventType = Event.TimerIndex;
+			}
+			else
+			{
+				FStackEntry& StackEntry = EventStack.Last();
+				double EventInclusiveTime = FMath::Min(IntervalEndTime, Time) - StackEntry.StartTime;
+				check(EventInclusiveTime >= 0.0);
+				double EventExclusiveTime = EventInclusiveTime - StackEntry.AccumulatedTimeInChildren;
+				EventStack.Pop(false);
+				if (EventStack.Num())
+				{
+					EventStack.Last().AccumulatedTimeInChildren += EventInclusiveTime;
+				}
+
+				FTimerStatsEx& Stats = StatsMap.FindOrAdd(Event.TimerIndex);
+				double EventNonRecursiveInclusiveTime = EventInclusiveTime;
+				for (const FStackEntry& AncestorStackEntry : EventStack)
+				{
+					if (AncestorStackEntry.EventType == Event.TimerIndex)
+					{
+						EventNonRecursiveInclusiveTime = 0.0;
+					}
+				}
+				Callback(Stats, EventNonRecursiveInclusiveTime, EventExclusiveTime);
+			}
+		});
+	}
+
+	TArray<FStackEntry> EventStack;
+	double IntervalStartTime;
+	double IntervalEndTime;
+	TMap<uint64, FTimerStatsEx>& StatsMap;
+};
+
+void STimersView::UpdateStats(double StartTime, double EndTime)
+{
+	for (const FTimerNodePtr& TimerNodePtr : TimerNodes)
+	{
+		TimerNodePtr->ResetStats();
+	}
+
+	if (Session.IsValid() &&
+		StartTime < EndTime)
+	{
+		struct ComputeStatsContext
+		{
+			double StartTime;
+			double EndTime;
+			TMap<uint64, FTimerStatsEx> Map;
+		} Ctx;
+
+		Ctx.StartTime = StartTime;
+		Ctx.EndTime = EndTime;
+
+		const bool bComputeMedian = true;
+
+		FTimeCalculationHelper CalculationHelper(Ctx.StartTime, Ctx.EndTime, Ctx.Map);
+
+		{
+			// Compute instance count and total/min/max inclusive/exclusive times for each timer.
+			// Iterate through all timing events (all thread timelines).
+			Trace::FAnalysisSessionReadScope _(*Session.Get());
+			Session->ReadTimingProfilerProvider([&CalculationHelper](const Trace::ITimingProfilerProvider& TimingProfilerProvider)
+			{
+				TimingProfilerProvider.EnumerateTimelines([&CalculationHelper](const Trace::ITimingProfilerProvider::Timeline& Timeline)
+				{
+					CalculationHelper.UpdateStats(Timeline, UpdateTotalMinMaxTimerStats);
+				});
+			});
+
+			// Now, as we know min/max inclusive/exclusive times for each timer, we can compute histogram and median values.
+			if (bComputeMedian)
+			{
+				// Update bucket size (DT) for computing histogram.
+				for (auto& KV : Ctx.Map)
+				{
+					PreComputeHistogram(KV.Value);
+				}
+
+				// Compute histogram.
+				// Iterate again through all timing events (all timelines).
+				Session->ReadTimingProfilerProvider([&CalculationHelper](const Trace::ITimingProfilerProvider& TimingProfilerProvider)
+				{
+					TimingProfilerProvider.EnumerateTimelines([&CalculationHelper](const Trace::ITimingProfilerProvider::Timeline& Timeline)
+					{
+						CalculationHelper.UpdateStats(Timeline, UpdateHistogramForTimerStats);
+					});
+				});
+			}
+		}
+
+		// Compute average and median inclusive/exclusive times.
+		for (auto& KV : Ctx.Map)
+		{
+			PostProcessTimerStats(KV.Value, bComputeMedian);
+
+			// Update the CPU timer node.
+			FTimerNodePtr* TimerNodePtrPtr = TimerNodesIdMap.Find(KV.Key);
+			if (TimerNodePtrPtr != nullptr)
+			{
+				(*TimerNodePtrPtr)->SetStats(KV.Value.BaseStats);
+			}
+		}
+
+		StatsStartTime = Ctx.StartTime;
+		StatsEndTime = Ctx.EndTime;
+
+		UpdateTree();
+		TreeView->RebuildList();
+	}
+}
+
+////////////////////////////////////////////////////////////////////////////////////////////////////
+
+void STimersView::UpdateTotalMinMaxTimerStats(FTimerStatsEx& StatsEx, double InclTime, double ExclTime)
+{
+	FTimerStats& Stats = StatsEx.BaseStats;
+
+	Stats.TotalInclusiveTime += InclTime;
+	if (InclTime < Stats.MinInclusiveTime)
+		Stats.MinInclusiveTime = InclTime;
+	if (InclTime > Stats.MaxInclusiveTime)
+	Stats.MaxInclusiveTime = InclTime;
+
+	Stats.TotalExclusiveTime += ExclTime;
+	if (ExclTime < Stats.MinExclusiveTime)
+		Stats.MinExclusiveTime = ExclTime;
+	if (ExclTime > Stats.MaxExclusiveTime)
+	Stats.MaxExclusiveTime = ExclTime;
+
+	Stats.InstanceCount++;
+}
+
+////////////////////////////////////////////////////////////////////////////////////////////////////
+
+void STimersView::PreComputeHistogram(FTimerStatsEx& StatsEx)
+{
+	const FTimerStats& Stats = StatsEx.BaseStats;
+
+	// Each bucket (Histogram[i]) will be centered on a value.
+	// I.e. First bucket (bucket 0) is centered on Min value: [Min-DT/2, Min+DT/2)
+	// and last bucket (bucket N-1) is centered on Max value: [Max-DT/2, Max+DT/2).
+
+	constexpr double InvHistogramLen = 1.0 / static_cast<double>(HistogramLen - 1);
+
+	if (Stats.MaxInclusiveTime == Stats.MinInclusiveTime)
+		StatsEx.InclDT = 1.0; // single large bucket
+	else
+		StatsEx.InclDT = (Stats.MaxInclusiveTime - Stats.MinInclusiveTime) * InvHistogramLen;
+
+	if (Stats.MaxExclusiveTime == Stats.MinExclusiveTime)
+		StatsEx.ExclDT = 1.0; // single large bucket
+	else
+		StatsEx.ExclDT = (Stats.MaxExclusiveTime - Stats.MinExclusiveTime) * InvHistogramLen;
+}
+
+////////////////////////////////////////////////////////////////////////////////////////////////////
+
+void STimersView::UpdateHistogramForTimerStats(FTimerStatsEx& StatsEx, double InclTime, double ExclTime)
+{
+	const FTimerStats& Stats = StatsEx.BaseStats;
+
+	int InclIndex = static_cast<int>((InclTime - Stats.MinInclusiveTime + StatsEx.InclDT / 2) / StatsEx.InclDT);
+	ensure(InclIndex >= 0);
+	if (InclIndex < 0)
+	{
+		InclIndex = 0;
+	}
+	ensure(InclIndex < HistogramLen);
+	if (InclIndex >= HistogramLen)
+	{
+		InclIndex = HistogramLen - 1;
+	}
+	StatsEx.InclHistogram[InclIndex]++;
+
+	int ExclIndex = static_cast<int>((ExclTime - Stats.MinExclusiveTime + StatsEx.ExclDT / 2) / StatsEx.ExclDT);
+	ensure(ExclIndex >= 0);
+	if (ExclIndex < 0)
+	{
+		ExclIndex = 0;
+	}
+	ensure(ExclIndex < HistogramLen);
+	if (ExclIndex >= HistogramLen)
+	{
+		ExclIndex = HistogramLen - 1;
+	}
+	StatsEx.ExclHistogram[ExclIndex]++;
+}
+
+////////////////////////////////////////////////////////////////////////////////////////////////////
+
+void STimersView::PostProcessTimerStats(FTimerStatsEx& StatsEx, bool bComputeMedian)
+{
+	FTimerStats& Stats = StatsEx.BaseStats;
+
+	// Compute average inclusive/exclusive times.
+	ensure(Stats.InstanceCount > 0);
+	double InvCount = 1.0f / static_cast<double>(Stats.InstanceCount);
+	Stats.AverageInclusiveTime = Stats.TotalInclusiveTime * InvCount;
+	Stats.AverageExclusiveTime = Stats.TotalExclusiveTime * InvCount;
+
+	if (bComputeMedian)
+	{
+		const int HalfCount = Stats.InstanceCount / 2;
+
+		// Compute median inclusive time.
+		int InclCount = 0;
+		for (int HistogramIndex = 0; HistogramIndex < HistogramLen; HistogramIndex++)
+		{
+			InclCount += StatsEx.InclHistogram[HistogramIndex];
+			if (InclCount > HalfCount)
+			{
+				Stats.MedianInclusiveTime = Stats.MinInclusiveTime + HistogramIndex * StatsEx.InclDT;
+				if (HistogramIndex > 0 &&
+					Stats.InstanceCount % 2 == 0 &&
+					InclCount - StatsEx.InclHistogram[HistogramIndex] == HalfCount)
+				{
+					const double PrevMedian = Stats.MinInclusiveTime + (HistogramIndex - 1) * StatsEx.InclDT;
+					Stats.MedianInclusiveTime = (Stats.MedianInclusiveTime + PrevMedian) / 2;
+				}
+				break;
+			}
+		}
+
+		// Compute median exclusive time.
+		int ExclCount = 0;
+		for (int HistogramIndex = 0; HistogramIndex < HistogramLen; HistogramIndex++)
+		{
+			ExclCount += StatsEx.InclHistogram[HistogramIndex];
+			if (ExclCount > HalfCount)
+			{
+				Stats.MedianExclusiveTime = Stats.MinExclusiveTime + HistogramIndex * StatsEx.ExclDT;
+				if (HistogramIndex > 0 &&
+					Stats.InstanceCount % 2 == 0 &&
+					ExclCount - StatsEx.ExclHistogram[HistogramIndex] == HalfCount)
+				{
+					const double PrevMedian = Stats.MinExclusiveTime + (HistogramIndex - 1) * StatsEx.ExclDT;
+					Stats.MedianExclusiveTime = (Stats.MedianExclusiveTime + PrevMedian) / 2;
+				}
+				break;
+			}
+		}
+	}
+}
+
+////////////////////////////////////////////////////////////////////////////////////////////////////
+
+void STimersView::SelectTimerNode(uint64 Id)
+{
+	FTimerNodePtr* TimerNodePtrPtr = TimerNodesIdMap.Find(Id);
+	if (TimerNodePtrPtr != nullptr)
+	{
+		//UE_LOG(TimingProfiler, Log, TEXT("Select and RequestScrollIntoView %s"), *(*TimerNodePtrPtr)->GetName().ToString());
+		TreeView->SetSelection(*TimerNodePtrPtr);
+		TreeView->RequestScrollIntoView(*TimerNodePtrPtr);
+	}
+}
+
+////////////////////////////////////////////////////////////////////////////////////////////////////
+
+#undef LOCTEXT_NAMESPACE
diff --git a/Engine/Source/Developer/TraceInsights/Private/Insights/Widgets/STimersView.h b/Engine/Source/Developer/TraceInsights/Private/Insights/Widgets/STimersView.h
new file mode 100644
--- /dev/null
+++ b/Engine/Source/Developer/TraceInsights/Private/Insights/Widgets/STimersView.h
@@ -0,0 +1,334 @@
+// Copyright 1998-2019 Epic Games, Inc. All Rights Reserved.
+
+#pragma once
+
+#include "CoreMinimal.h"
+#include "Misc/FilterCollection.h"
+#include "Misc/TextFilter.h"
+#include "SlateFwd.h"
+#include "Widgets/DeclarativeSyntaxSupport.h"
+#include "Widgets/Input/SComboBox.h"
+#include "Widgets/SCompoundWidget.h"
+#include "Widgets/SWidget.h"
+#include "Widgets/Views/SHeaderRow.h"
+#include "Widgets/Views/STableRow.h"
+#include "Widgets/Views/STableViewBase.h"
+#include "Widgets/Views/STreeView.h"
+
+// Insights
+#include "Insights/ViewModels/TimerGroupingAndSorting.h"
+#include "Insights/ViewModels/TimerNode.h"
+#include "Insights/ViewModels/TimersViewColumn.h"
+
+class FMenuBuilder;
+
+namespace Trace
+{
+	class IAnalysisSession;
+	struct FTimelineEvent;
+}
+
+static constexpr int HistogramLen = 100; // number of buckets per histogram
+
+struct FTimerStatsEx
+{
+	FTimerStats BaseStats;
+
+	// Histogram for computing median inclusive time.
+	int InclHistogram[HistogramLen];
+	double InclDT; // bucket size
+
+	// Histogram for computing median exclusive time.
+	int ExclHistogram[HistogramLen];
+	double ExclDT; // bucket size
+
+	FTimerStatsEx()
+	{
+		FMemory::Memzero(InclHistogram, sizeof(int) * HistogramLen);
+		InclDT = 1.0;
+
+		FMemory::Memzero(ExclHistogram, sizeof(int) * HistogramLen);
+		ExclDT = 1.0;
+	}
+};
+
+/** The filter collection - used for updating the list of timer nodes. */
+typedef TFilterCollection<const FTimerNodePtr&> FTimerNodeFilterCollection;
+
+/** The text based filter - used for updating the list of timer nodes. */
+typedef TTextFilter<const FTimerNodePtr&> FTimerNodeTextFilter;
+
+/** A custom widget used to display the list of timers. */
+class STimersView : public SCompoundWidget
+{
+public:
+	/** Default constructor. */
+	STimersView();
+
+	/** Virtual destructor. */
+	virtual ~STimersView();
+
+	SLATE_BEGIN_ARGS(STimersView){}
+	SLATE_END_ARGS()
+
+	/**
+	 * Construct this widget
+	 * @param	InArgs	- The declaration data for this widget
+	 */
+	void Construct(const FArguments& InArgs);
+
+	void RebuildTree(bool bRebuildEvenIfNotChanged = true);
+	void UpdateStats(double StartTime, double EndTime);
+
+	void SelectTimerNode(uint64 Id);
+
+	//const TSet<FTimerNodePtr>& GetTimerNodes() const { return TimerNodes; }
+	//const TMap<uint64, FTimerNodePtr> GetTimerNodesIdMap() const { return TimerNodesIdMap; }
+	const FTimerNodePtr* GetTimerNode(uint64 Id) const { return TimerNodesIdMap.Find(Id); }
+
+protected:
+
+	void UpdateTree();
+
+	static void UpdateTotalMinMaxTimerStats(FTimerStatsEx& Stats, double InclTime, double ExclTime);
+	static void PreComputeHistogram(FTimerStatsEx& StatsEx);
+	static void UpdateHistogramForTimerStats(FTimerStatsEx& StatsEx, double InclTime, double ExclTime);
+	static void PostProcessTimerStats(FTimerStatsEx& StatsEx, bool bComputeMedian);
+
+	/** Called when the filter and presets widget should be updated with the latest data. */
+	void ProfilerManager_OnRequestTimersUpdate();
+
+	/**
+	 * Populates OutSearchStrings with the strings that should be used in searching.
+	 *
+	 * @param GroupOrStatNodePtr	- the group and stat node to get a text description from.
+	 * @param OutSearchStrings		- an array of strings to use in searching.
+	 *
+	 */
+	void HandleItemToStringArray(const FTimerNodePtr& GroupOrStatNodePtr, TArray<FString>& OutSearchStrings) const;
+
+	////////////////////////////////////////////////////////////////////////////////////////////////////
+	// Tree View - Context Menu
+
+	TSharedPtr<SWidget> TreeView_GetMenuContent() const;
+	void TreeView_BuildSortByMenu(FMenuBuilder& MenuBuilder);
+	void TreeView_BuildViewColumnMenu(FMenuBuilder& MenuBuilder);
+
+	////////////////////////////////////////////////////////////////////////////////////////////////////
+	// Tree View - Columns' Header
+
+	void InitializeAndShowHeaderColumns();
+	void TreeViewHeaderRow_CreateColumnArgs(const int32 ColumnIndex);
+	void TreeViewHeaderRow_ShowColumn(const FName ColumnId);
+	TSharedRef<SWidget> TreeViewHeaderRow_GenerateColumnMenu(const FTimersViewColumn& Column);
+
+	////////////////////////////////////////////////////////////////////////////////////////////////////
+	// Tree View - Misc
+
+	void TreeView_Refresh();
+
+	/**
+	 * Called by STreeView to retrieves the children for the specified parent item.
+	 * @param InParent		- The parent node to retrieve the children from.
+	 * @param OutChildren	- List of children for the parent node.
+	 */
+	void TreeView_OnGetChildren(FTimerNodePtr InParent, TArray<FTimerNodePtr>& OutChildren);
+
+	/** Called by STreeView when selection has changed. */
+	void TreeView_OnSelectionChanged(FTimerNodePtr SelectedItem, ESelectInfo::Type SelectInfo);
+
+	/** Called by STreeView when a tree item is double clicked. */
+	void TreeView_OnMouseButtonDoubleClick(FTimerNodePtr TreeNode);
+
+	////////////////////////////////////////////////////////////////////////////////////////////////////
+	// Tree View - Table Row
+
+	/** Called by STreeView to generate a table row for the specified item. */
+	TSharedRef<ITableRow> TreeView_OnGenerateRow(FTimerNodePtr TreeNode, const TSharedRef<STableViewBase>& OwnerTable);
+
+	bool TableRow_IsColumnVisible(const FName ColumnId) const;
+	void TableRow_SetHoveredTableCell(const FName ColumnId, const FTimerNodePtr TimerNodePtr);
+	EHorizontalAlignment TableRow_GetColumnOutlineHAlignment(const FName ColumnId) const;
+
+	FText TableRow_GetHighlightText() const;
+	FName TableRow_GetHighlightedTimerName() const;
+
+	bool TableRow_ShouldBeEnabled(const uint32 TimerId) const;
+
+	////////////////////////////////////////////////////////////////////////////////////////////////////
+	// Filtering
+
+	/** Populates the group and stat tree with items based on the current data. */
+	void ApplyFiltering();
+
+	TSharedRef<SWidget> GetToggleButtonForTimerType(const ETimerNodeType StatType);
+	void FilterByTimerType_OnCheckStateChanged(ECheckBoxState NewRadioState, const ETimerNodeType InStatType);
+	ECheckBoxState FilterByTimerType_IsChecked(const ETimerNodeType InStatType) const;
+
+	bool SearchBox_IsEnabled() const;
+	void SearchBox_OnTextChanged(const FText& InFilterText);
+
+	////////////////////////////////////////////////////////////////////////////////////////////////////
+	// GroupBy
+
+	void CreateGroups();
+	void CreateGroupByOptionsSources();
+
+	void GroupBy_OnSelectionChanged(TSharedPtr<ETimerGroupingMode> NewGroupingMode, ESelectInfo::Type SelectInfo);
+
+	TSharedRef<SWidget> GroupBy_OnGenerateWidget(TSharedPtr<ETimerGroupingMode> InGroupingMode) const;
+
+	FText GroupBy_GetSelectedText() const;
+
+	FText GroupBy_GetSelectedTooltipText() const;
+
+	////////////////////////////////////////////////////////////////////////////////////////////////////
+	// Sorting
+
+	void SortTimers();
+
+	EColumnSortMode::Type GetSortModeForColumn(const FName ColumnId) const;
+	void SetSortModeForColumn(const FName& ColumnId, EColumnSortMode::Type SortMode);
+	void OnSortModeChanged(const EColumnSortPriority::Type SortPriority, const FName& ColumnId, const EColumnSortMode::Type SortMode);
+
+	//void TreeView_BuildSortByMenu(FMenuBuilder& MenuBuilder);
+
+	////////////////////////////////////////////////////////////////////////////////////////////////////
+	// Sorting actions
+
+	// SortMode (HeaderMenu)
+	bool HeaderMenu_SortMode_IsChecked(const FName ColumnId, const EColumnSortMode::Type InSortMode);
+	bool HeaderMenu_SortMode_CanExecute(const FName ColumnId, const EColumnSortMode::Type InSortMode) const;
+	void HeaderMenu_SortMode_Execute(const FName ColumnId, const EColumnSortMode::Type InSortMode);
+
+	// SortMode (ContextMenu)
+	bool ContextMenu_SortMode_IsChecked(const EColumnSortMode::Type InSortMode);
+	bool ContextMenu_SortMode_CanExecute(const EColumnSortMode::Type InSortMode) const;
+	void ContextMenu_SortMode_Execute(const EColumnSortMode::Type InSortMode);
+
+	// SortByColumn (ContextMenu)
+	bool ContextMenu_SortByColumn_IsChecked(const FName ColumnId);
+	bool ContextMenu_SortByColumn_CanExecute(const FName ColumnId) const;
+	void ContextMenu_SortByColumn_Execute(const FName ColumnId);
+
+	////////////////////////////////////////////////////////////////////////////////////////////////////
+	// Column visibility actions
+
+	// HideColumn (HeaderMenu)
+	bool HeaderMenu_HideColumn_CanExecute(const FName ColumnId) const;
+	void HeaderMenu_HideColumn_Execute(const FName ColumnId);
+
+	// ToggleColumn (ContextMenu)
+	bool ContextMenu_ToggleColumn_IsChecked(const FName ColumnId);
+	bool ContextMenu_ToggleColumn_CanExecute(const FName ColumnId) const;
+	void ContextMenu_ToggleColumn_Execute(const FName ColumnId);
+
+	// ShowAllColumns (ContextMenu)
+	bool ContextMenu_ShowAllColumns_CanExecute() const;
+	void ContextMenu_ShowAllColumns_Execute();
+
+	// MinMaxMedColumns (ContextMenu)
+	bool ContextMenu_ShowMinMaxMedColumns_CanExecute() const;
+	void ContextMenu_ShowMinMaxMedColumns_Execute();
+
+	// ResetColumns (ContextMenu)
+	bool ContextMenu_ResetColumns_CanExecute() const;
+	void ContextMenu_ResetColumns_Execute();
+
+protected:
+	/// A weak pointer to the profiler session used to populate this widget.
+	TSharedPtr<const Trace::IAnalysisSession>/*Weak*/ Session;
+
+	//////////////////////////////////////////////////
+	// Tree View, Columns
+
+	/// The tree widget which holds the list of timer groups and timers corresponding with each group.
+	TSharedPtr<STreeView<FTimerNodePtr>> TreeView;
+
+	/// Column metadata used to initialize column arguments, stored as PropertyName -> FEventGraphColumn.
+	TMap<FName, FTimersViewColumn> TreeViewHeaderColumns;
+
+	/// Column arguments used to initialize a new header column in the tree view, stored as column name to column arguments mapping.
+	TMap<FName, SHeaderRow::FColumn::FArguments> TreeViewHeaderColumnArgs;
+
+	/// Holds the tree view header row widget which display all columns in the tree view.
+	TSharedPtr<SHeaderRow> TreeViewHeaderRow;
+
+	/// External scrollbar used to synchronize tree view position.
+	TSharedPtr<SScrollBar> ExternalScrollbar;
+
+	//////////////////////////////////////////////////
+	// Hovered Column, Hovered Timer Node
+
+	/// Name of the column currently being hovered by the mouse.
+	FName HoveredColumnId;
+
+	/// A shared pointer to the timer currently being hovered by the mouse.
+	FTimerNodePtr HoveredTimerNodePtr;
+
+	/// Name of the timer that should be drawn as highlighted.
+	FName HighlightedTimerName;
+
+	//////////////////////////////////////////////////
+	// Timer Nodes
+
+	/// An array of group and timer nodes generated from the metadata.
+	TArray<FTimerNodePtr> GroupNodes;
+
+	/// A filtered array of group and timer nodes to be displayed in the tree widget.
+	TArray<FTimerNodePtr> FilteredGroupNodes;
+
+	/// All timer nodes.
+	TSet<FTimerNodePtr> TimerNodes;
+
+	/// All timer nodes, stored as TimerName -> FTimerNodePtr.
+	//TMap<FName, FTimerNodePtr> TimerNodesMap;
+
+	/// All timer nodes, stored as TimerId -> FTimerNodePtr.
+	TMap<uint64, FTimerNodePtr> TimerNodesIdMap;
+
+	/// Currently expanded group nodes.
+	TSet<FTimerNodePtr> ExpandedNodes;
+
+	/// If true, the expanded nodes have been saved before applying a text filter.
+	bool bExpansionSaved;
+
+	//////////////////////////////////////////////////
+	// Search box and filters
+
+	/// The search box widget used to filter items displayed in the stats and groups tree.
+	TSharedPtr<SSearchBox> SearchBox;
+
+	/// The text based filter.
+	TSharedPtr<FTimerNodeTextFilter> TextFilter;
+
+	/// The filter collection.
+	TSharedPtr<FTimerNodeFilterCollection> Filters;
+
+	/// Holds the visibility of each timer type.
+	bool bTimerNodeIsVisible[ETimerNodeType::InvalidOrMax];
+
+	//////////////////////////////////////////////////
+	// Grouping
+
+	TArray<TSharedPtr<ETimerGroupingMode>> GroupByOptionsSource;
+
+	TSharedPtr<SComboBox<TSharedPtr<ETimerGroupingMode>>> GroupByComboBox;
+
+	/// How we group the timers?
+	ETimerGroupingMode GroupingMode;
+
+	//////////////////////////////////////////////////
+	// Sorting
+
+	/// How we sort the timers?
+	EColumnSortMode::Type ColumnSortMode;
+
+	/// Name of the column currently being sorted, NAME_None if sorting is disabled.
+	FName ColumnBeingSorted;
+
+	//////////////////////////////////////////////////
+
+	double StatsStartTime;
+	double StatsEndTime;
+};
diff --git a/Engine/Source/Developer/TraceInsights/Private/Insights/Widgets/STimersViewTooltip.cpp b/Engine/Source/Developer/TraceInsights/Private/Insights/Widgets/STimersViewTooltip.cpp
new file mode 100644
--- /dev/null
+++ b/Engine/Source/Developer/TraceInsights/Private/Insights/Widgets/STimersViewTooltip.cpp
@@ -0,0 +1,558 @@
+// Copyright 1998-2019 Epic Games, Inc. All Rights Reserved.
+
+#include "STimersViewTooltip.h"
+
+#include "EditorStyleSet.h"
+#include "SlateOptMacros.h"
+#include "Widgets/Images/SImage.h"
+#include "Widgets/Layout/SGridPanel.h"
+#include "Widgets/Layout/SSeparator.h"
+#include "Widgets/SBoxPanel.h"
+#include "Widgets/SToolTip.h"
+#include "Widgets/Text/STextBlock.h"
+
+// Insights
+#include "Insights/Common/TimeUtils.h"
+#include "Insights/ViewModels/TimerNode.h"
+#include "Insights/ViewModels/TimersViewColumn.h"
+#include "Insights/ViewModels/TimerNodeHelper.h"
+
+////////////////////////////////////////////////////////////////////////////////////////////////////
+
+#define LOCTEXT_NAMESPACE "STimersViewTooltip"
+
+////////////////////////////////////////////////////////////////////////////////////////////////////
+
+BEGIN_SLATE_FUNCTION_BUILD_OPTIMIZATION
+
+TSharedPtr<SToolTip> STimersViewTooltip::GetColumnTooltip(const FTimersViewColumn& Column)
+{
+	TSharedPtr<SToolTip> ColumnTooltip =
+		SNew(SToolTip)
+		[
+			SNew(SVerticalBox)
+
+			+ SVerticalBox::Slot()
+			.AutoHeight()
+			.Padding(2.0f)
+			[
+				SNew(STextBlock)
+				.Text(Column.TitleName)
+				.TextStyle(FEditorStyle::Get(), TEXT("Profiler.TooltipBold"))
+			]
+
+			+ SVerticalBox::Slot()
+			.AutoHeight()
+			.Padding(2.0f)
+			[
+				SNew(STextBlock)
+				.Text(Column.Description)
+				.TextStyle(FEditorStyle::Get(), TEXT("Profiler.Tooltip"))
+			]
+		];
+
+	return ColumnTooltip;
+}
+
+////////////////////////////////////////////////////////////////////////////////////////////////////
+
+TSharedPtr<SToolTip> STimersViewTooltip::GetTableCellTooltip(const TSharedPtr<FTimerNode> TimerNodePtr)
+{
+	//const FSlateFontInfo TitleFont = FCoreStyle::GetDefaultFontStyle("Bold", 10);
+	//const FSlateFontInfo DescriptionFont = FCoreStyle::GetDefaultFontStyle("Regular", 8);
+	//const FSlateFontInfo DescriptionFontB = FCoreStyle::GetDefaultFontStyle("Bold", 8);
+
+	const FLinearColor DefaultColor(1.0f,1.0f,1.0f,1.0f);
+	const FLinearColor ThreadColor(5.0f, 0.0f, 0.0f, 1.0f);
+	const float Alpha = 0.0f;//TimerNodePtr->_FramePct * 0.01f;
+	const FLinearColor ColorAndOpacity = FMath::Lerp(DefaultColor, ThreadColor,Alpha);
+
+	const FTimerStats& Stats = TimerNodePtr->GetStats();
+
+	const int NumDigits = 5;
+
+	const FText TotalInclusiveTimeText = FText::FromString(TimeUtils::FormatTimeAuto(Stats.TotalInclusiveTime));
+	const FText MinInclusiveTimeText = FText::FromString(TimeUtils::FormatTimeMs(Stats.MinInclusiveTime, NumDigits, true));
+	const FText MaxInclusiveTimeText = FText::FromString(TimeUtils::FormatTimeMs(Stats.MaxInclusiveTime, NumDigits, true));
+	const FText AvgInclusiveTimeText = FText::FromString(TimeUtils::FormatTimeMs(Stats.AverageInclusiveTime, NumDigits, true));
+	const FText MedInclusiveTimeText = FText::FromString(TimeUtils::FormatTimeMs(Stats.MedianInclusiveTime, NumDigits, true));
+
+	const FText TotalExclusiveTimeText = FText::FromString(TimeUtils::FormatTimeAuto(Stats.TotalExclusiveTime));
+	const FText MinExclusiveTimeText = FText::FromString(TimeUtils::FormatTimeMs(Stats.MinExclusiveTime, NumDigits, true));
+	const FText MaxExclusiveTimeText = FText::FromString(TimeUtils::FormatTimeMs(Stats.MaxExclusiveTime, NumDigits, true));
+	const FText AvgExclusiveTimeText = FText::FromString(TimeUtils::FormatTimeMs(Stats.AverageExclusiveTime, NumDigits, true));
+	const FText MedExclusiveTimeText = FText::FromString(TimeUtils::FormatTimeMs(Stats.MedianExclusiveTime, NumDigits, true));
+
+	const FText InstanceCountText = FText::AsNumber(Stats.InstanceCount);
+
+	//TSharedPtr<SHorizontalBox> HBoxCaption;
+	TSharedPtr<SGridPanel> GridPanel;
+	TSharedPtr<SHorizontalBox> HBox;
+
+	TSharedPtr<SToolTip> TableCellTooltip =
+		SNew(SToolTip)
+		[
+			SAssignNew(HBox, SHorizontalBox)
+
+			+SHorizontalBox::Slot()
+			.AutoWidth()
+			[
+				SNew(SVerticalBox)
+
+				//+SVerticalBox::Slot()
+				//.AutoHeight()
+				//.Padding(2.0f)
+				//[
+				//	SAssignNew(HBoxCaption, SHorizontalBox)
+				//]
+
+				+SVerticalBox::Slot()
+				.AutoHeight()
+				.Padding(2.0f)
+				[
+					SNew(SSeparator)
+					.Orientation(Orient_Horizontal)
+				]
+
+				+SVerticalBox::Slot()
+				.AutoHeight()
+				.Padding(2.0f)
+				[
+					SNew(SGridPanel)
+
+					// Id: [Id]
+					+SGridPanel::Slot(0, 0)
+					.Padding(2.0f)
+					[
+						SNew(STextBlock)
+						.Text(LOCTEXT("TT_Id", "Id:"))
+						.TextStyle(FEditorStyle::Get(), TEXT("Profiler.TooltipBold"))
+					]
+					+SGridPanel::Slot(1, 0)
+					.Padding(2.0f)
+					[
+						SNew(STextBlock)
+						.Text(FText::AsNumber(TimerNodePtr->GetId()))
+						.TextStyle(FEditorStyle::Get(), TEXT("Profiler.Tooltip"))
+					]
+
+					// Name: [Name]
+					+SGridPanel::Slot(0, 1)
+					.Padding(2.0f)
+					[
+						SNew(STextBlock)
+						.Text(LOCTEXT("TT_Name", "Name:"))
+						.TextStyle(FEditorStyle::Get(), TEXT("Profiler.TooltipBold"))
+					]
+					+SGridPanel::Slot(1, 1)
+					.Padding(2.0f)
+					[
+						SNew(STextBlock)
+						.Text(FText::FromName(TimerNodePtr->GetName()))
+						.TextStyle(FEditorStyle::Get(), TEXT("Profiler.Tooltip"))
+					]
+
+					//// Group: [MetaGroupName]
+					//+SGridPanel::Slot(0, 2)
+					//.Padding(2.0f)
+					//[
+					//	SNew(STextBlock)
+					//	.Text(LOCTEXT("TT_Group", "Group:"))
+					//	.TextStyle(FEditorStyle::Get(), TEXT("Profiler.TooltipBold"))
+					//]
+					//+SGridPanel::Slot(1, 2)
+					//.Padding(2.0f)
+					//[
+					//	SNew(STextBlock)
+					//	.Text(FText::FromName(TimerNodePtr->GetMetaGroupName()))
+					//	.TextStyle(FEditorStyle::Get(), TEXT("Profiler.Tooltip"))
+					//]
+
+					// Type: [Type]
+					+ SGridPanel::Slot(0, 3)
+					.Padding(2.0f)
+					[
+						SNew(STextBlock)
+						.Text(LOCTEXT("TT_Type", "Type:"))
+						.TextStyle(FEditorStyle::Get(), TEXT("Profiler.TooltipBold"))
+					]
+					+ SGridPanel::Slot(1, 3)
+					.Padding(2.0f)
+					[
+						SNew(STextBlock)
+						.Text(TimerNodeTypeHelper::ToName(TimerNodePtr->GetType()))
+						.TextStyle(FEditorStyle::Get(), TEXT("Profiler.Tooltip"))
+					]
+				]
+
+				+SVerticalBox::Slot()
+				.AutoHeight()
+				.Padding(2.0f)
+				[
+					SNew(SSeparator)
+					.Orientation(Orient_Horizontal)
+				]
+
+				+ SVerticalBox::Slot()
+				.AutoHeight()
+				.Padding(2.0f)
+				[
+					SNew(SGridPanel)
+
+					+ SGridPanel::Slot(0, 0)
+					.Padding(2.0f)
+					[
+						SNew(STextBlock)
+						.Text(LOCTEXT("TT_NumInstances", "Num Instances:"))
+						.TextStyle(FEditorStyle::Get(), TEXT("Profiler.TooltipBold"))
+					]
+					+ SGridPanel::Slot(1, 0)
+					.Padding(2.0f)
+					[
+						SNew(STextBlock)
+						.Text(InstanceCountText)
+						.TextStyle(FEditorStyle::Get(), TEXT("Profiler.Tooltip"))
+					]
+				]
+
+				+SVerticalBox::Slot()
+				.AutoHeight()
+				.Padding(2.0f)
+				[
+					SNew(SSeparator)
+					.Orientation(Orient_Horizontal)
+				]
+
+				+ SVerticalBox::Slot()
+				.AutoHeight()
+				.Padding(2.0f)
+				[
+					SAssignNew(GridPanel, SGridPanel)
+
+					+ SGridPanel::Slot(1, 0)
+					.Padding(2.0f)
+					.HAlign(HAlign_Right)
+					[
+						SNew(STextBlock)
+						.Text(LOCTEXT("TT_InclusiveTime", "Inclusive"))
+						.TextStyle(FEditorStyle::Get(), TEXT("Profiler.TooltipBold"))
+					]
+					+ SGridPanel::Slot(2, 0)
+					.Padding(2.0f)
+					.HAlign(HAlign_Right)
+					[
+						SNew(STextBlock)
+						.Text(LOCTEXT("TT_ExclusiveTime", "Exclusive"))
+						.TextStyle(FEditorStyle::Get(), TEXT("Profiler.TooltipBold"))
+					]
+
+					// Stats are added here.
+				]
+
+				+SVerticalBox::Slot()
+				.AutoHeight()
+				.Padding(2.0f)
+				[
+					SNew(SSeparator)
+					.Orientation(Orient_Horizontal)
+				]
+			]
+		];
+
+	int32 Row = 1;
+	AddStatsRow(GridPanel, Row, LOCTEXT("TT_TotalTime",   "Total Time:"),   TotalInclusiveTimeText, TotalExclusiveTimeText);
+	AddStatsRow(GridPanel, Row, LOCTEXT("TT_MinTime",     "Min Time:"),     MinInclusiveTimeText, MinExclusiveTimeText);
+	AddStatsRow(GridPanel, Row, LOCTEXT("TT_MaxTime",     "Max Time:"),     MaxInclusiveTimeText, MaxExclusiveTimeText);
+	AddStatsRow(GridPanel, Row, LOCTEXT("TT_AverageTime", "Average Time:"), AvgInclusiveTimeText, AvgExclusiveTimeText);
+	AddStatsRow(GridPanel, Row, LOCTEXT("TT_MedianTime",  "Median Time:"),  MedInclusiveTimeText, MedExclusiveTimeText);
+
+	/*
+	//TODO: We need Stats hierarchy (not the grouping hierarchy)!
+	const bool bHasParent = TimerNodePtr->GetStats()->GetParent().IsValid();
+	const bool bHasChildren = TimerNodePtr->GetStats()->GetChildren().Num() > 0;
+
+	if (bHasParent)
+	{
+		const FText ParentName = FText::FromName(TimerNodePtr->GetGroupPtr()->GetName());
+		HBoxCaption->AddSlot()
+			.AutoWidth()
+			[
+				SNew(STextBlock)
+				.Text(ParentName)
+				.TextStyle(FEditorStyle::Get(), TEXT("Profiler.Caption"))
+			];
+
+		HBoxCaption->AddSlot()
+			.AutoWidth()
+			[
+				SNew(SImage)
+				.Image(FEditorStyle::GetBrush("BreadcrumbTrail.Delimiter"))
+			];
+	}
+
+	const FText TimerName = FText::FromName(TimerNodePtr->GetName());
+	HBoxCaption->AddSlot()
+		.AutoWidth()
+		[
+			SNew(STextBlock)
+			.Text(TimerName)
+			.TextStyle(FEditorStyle::Get(), TEXT("Profiler.CaptionBold"))
+		];
+
+	if (bHasChildren)
+	{
+		typedef TKeyValuePair<float, FName> FEventNameAndPct;
+		TArray<FEventNameAndPct> MinimalChildren;
+
+		const TArray<FTimerNodePtr>& Children = TimerNodePtr->GetChildren();
+		for(int32 ChildIndex = 0; ChildIndex < Children.Num(); ChildIndex++)
+		{
+			const FTimerNodePtr Child = Children[ChildIndex];
+			float Percent = static_cast<float>(Child->GetStats().TotalInclusiveTime / TimerNodePtr->GetStats().TotalInclusiveTime * 100.0);
+			MinimalChildren.Add(FEventNameAndPct(Child->GetStats().TotalInclusiveTime, Child->GetName()));
+		}
+
+		struct FCompareByFloatDescending
+		{
+			FORCEINLINE bool operator()(const FEventNameAndPct& A, const FEventNameAndPct& B) const
+			{
+				return A.Key > B.Key;
+			}
+		};
+		MinimalChildren.Sort(FCompareByFloatDescending());
+
+		FString ChildrenNames;
+		const int NumChildrenToDisplay = FMath::Min(MinimalChildren.Num(), 3);
+		for(int32 SortedChildIndex = 0; SortedChildIndex < NumChildrenToDisplay; SortedChildIndex++)
+		{
+			const FEventNameAndPct& MinimalChild = MinimalChildren[SortedChildIndex];
+			ChildrenNames += FString::Printf(TEXT("%s (%.1f %%)"), *MinimalChild.Value.ToString(), MinimalChild.Key);
+
+			const bool bAddDelimiter = SortedChildIndex < NumChildrenToDisplay - 1;
+			if (bAddDelimiter)
+			{
+				ChildrenNames += TEXT(", ");
+			}
+		}
+
+		HBoxCaption->AddSlot()
+			.AutoWidth()
+			[
+				SNew(SImage)
+				.Image(FEditorStyle::GetBrush("BreadcrumbTrail.Delimiter"))
+			];
+
+		HBoxCaption->AddSlot()
+			.AutoWidth()
+			[
+				SNew(STextBlock)
+				.Text(FText::FromString(ChildrenNames))
+				.TextStyle(FEditorStyle::Get(), TEXT("Profiler.Caption"))
+			];
+	}
+	*/
+
+	return TableCellTooltip;
+}
+
+////////////////////////////////////////////////////////////////////////////////////////////////////
+
+TSharedRef<SToolTip> STimersViewTooltip::GetTooltip()
+{
+	if (Session.IsValid())
+	{
+		const TSharedRef<SGridPanel> ToolTipGrid = SNew(SGridPanel);
+		int CurrentRowPos = 0;
+
+		AddHeader(ToolTipGrid, CurrentRowPos);
+		AddDescription(ToolTipGrid, CurrentRowPos);
+
+		AddNoDataInformation(ToolTipGrid, CurrentRowPos);
+
+		return SNew(SToolTip)
+			[
+				ToolTipGrid
+			];
+	}
+	else
+	{
+		return SNew(SToolTip)
+			.Text(LOCTEXT("NotImplemented", "Tooltip for multiple profiler instances has not been implemented yet"));
+	}
+}
+
+////////////////////////////////////////////////////////////////////////////////////////////////////
+
+void STimersViewTooltip::AddNoDataInformation(const TSharedRef<SGridPanel>& Grid, int32& RowPos)
+{
+	Grid->AddSlot(0, RowPos)
+		.Padding(2.0f)
+		.ColumnSpan(3)
+		[
+			SNew(STextBlock)
+			.TextStyle(FEditorStyle::Get(), TEXT("Profiler.TooltipBold"))
+			.Text(LOCTEXT("NoDataAvailable", "N/A"))
+		];
+	RowPos++;
+}
+
+////////////////////////////////////////////////////////////////////////////////////////////////////
+
+void STimersViewTooltip::AddHeader(const TSharedRef<SGridPanel>& Grid, int32& RowPos)
+{
+	const FString InstanceName = Session->GetName();
+
+	Grid->AddSlot(0, RowPos++)
+		.Padding(2.0f)
+		.ColumnSpan(3)
+		[
+			SNew(STextBlock)
+			.TextStyle(FEditorStyle::Get(), TEXT("Profiler.TooltipBold"))
+			.Text(LOCTEXT("StatInstance", "Stat information for profiler instance"))
+		];
+
+	Grid->AddSlot(0, RowPos++)
+		.Padding(2.0f)
+		.ColumnSpan(3)
+		[
+			SNew(STextBlock)
+			.TextStyle(FEditorStyle::Get(), TEXT("Profiler.Tooltip"))
+			.Text(FText::FromString(InstanceName))
+		];
+
+	AddSeparator(Grid, RowPos);
+}
+
+////////////////////////////////////////////////////////////////////////////////////////////////////
+
+void STimersViewTooltip::AddDescription(const TSharedRef<SGridPanel>& Grid, int32& RowPos)
+{
+	/*
+	const FProfilerStat& ProfilerStat = Session->GetMetaData()->GetTimerById(TimerId);
+	const ETimerNodeType SampleType = Session->GetMetaData()->GetSampleTypeForTimerId(TimerId);
+	const FSlateBrush* const StatIcon = STimersViewHelper::GetIconForStatType(SampleType);
+
+	Grid->AddSlot(0, RowPos)
+	.Padding(2.0f)
+	[
+		SNew(STextBlock)
+		.TextStyle(FEditorStyle::Get(), TEXT("Profiler.TooltipBold"))
+		.Text(LOCTEXT("GroupDesc","Group:"))
+	];
+
+	Grid->AddSlot(1, RowPos)
+	.Padding(2.0f)
+	.ColumnSpan(2)
+	[
+		SNew(STextBlock)
+		.TextStyle(FEditorStyle::Get(), TEXT("Profiler.Tooltip"))
+		.Text(FText::FromName(ProfilerStat.OwningGroup().Name()))
+	];
+	RowPos++;
+
+	Grid->AddSlot(0, RowPos)
+	.Padding(2.0f)
+	[
+		SNew(STextBlock)
+		.TextStyle(FEditorStyle::Get(), TEXT("Profiler.TooltipBold"))
+		.Text(LOCTEXT("NameDesc","Name:"))
+	];
+
+	Grid->AddSlot(1, RowPos)
+	.Padding(2.0f)
+	.ColumnSpan(2)
+	[
+		SNew(STextBlock)
+		.TextStyle(FEditorStyle::Get(), TEXT("Profiler.Tooltip"))
+		.Text(FText::FromName(ProfilerStat.Name()))
+	];
+	RowPos++;
+
+	Grid->AddSlot(0, RowPos)
+	.Padding(2.0f)
+	[
+		SNew(STextBlock)
+		.TextStyle(FEditorStyle::Get(), TEXT("Profiler.TooltipBold"))
+		.Text(LOCTEXT("TypeDesc","Type:"))
+	];
+
+	Grid->AddSlot(1, RowPos)
+	.Padding(2.0f)
+	.ColumnSpan(2)
+	[
+		SNew(SHorizontalBox)
+
+		+SHorizontalBox::Slot()
+		.AutoWidth()
+		[
+			SNew(SImage)
+			.Image(StatIcon)
+		]
+
+		+SHorizontalBox::Slot()
+		.AutoWidth()
+		.HAlign(HAlign_Left)
+		.VAlign(VAlign_Center)
+		[
+			SNew(STextBlock)
+			.Text(FText::FromString(ETimerNodeType::ToDescription(SampleType)))
+			.TextStyle(FEditorStyle::Get(), TEXT("Profiler.Tooltip"))
+		]
+	];
+	RowPos++;
+
+	AddSeparator(Grid, RowPos);
+	*/
+}
+
+////////////////////////////////////////////////////////////////////////////////////////////////////
+
+void STimersViewTooltip::AddSeparator(const TSharedRef<SGridPanel>& Grid, int32& RowPos)
+{
+	Grid->AddSlot(0, RowPos++)
+		.Padding(2.0f)
+		.ColumnSpan(3)
+		[
+			SNew(SSeparator)
+			.Orientation(Orient_Horizontal)
+		];
+}
+
+////////////////////////////////////////////////////////////////////////////////////////////////////
+
+void STimersViewTooltip::AddStatsRow(TSharedPtr<SGridPanel> Grid, int32& Row, const FText& Name, const FText& Value1, const FText& Value2)
+{
+	Grid->AddSlot(0, Row)
+		.Padding(2.0f)
+		[
+			SNew(STextBlock)
+			.Text(Name)
+			.TextStyle(FEditorStyle::Get(), TEXT("Profiler.TooltipBold"))
+		];
+
+	Grid->AddSlot(1, Row)
+		.Padding(2.0f)
+		.HAlign(HAlign_Right)
+		[
+			SNew(STextBlock)
+			.Text(Value1)
+			.TextStyle(FEditorStyle::Get(), TEXT("Profiler.Tooltip"))
+		];
+
+	Grid->AddSlot(2, Row)
+		.Padding(2.0f)
+		.HAlign(HAlign_Right)
+		[
+			SNew(STextBlock)
+			.Text(Value2)
+			.TextStyle(FEditorStyle::Get(), TEXT("Profiler.Tooltip"))
+		];
+
+	Row++;
+}
+
+END_SLATE_FUNCTION_BUILD_OPTIMIZATION
+
+////////////////////////////////////////////////////////////////////////////////////////////////////
+
+#undef LOCTEXT_NAMESPACE
diff --git a/Engine/Source/Developer/TraceInsights/Private/Insights/Widgets/STimersViewTooltip.h b/Engine/Source/Developer/TraceInsights/Private/Insights/Widgets/STimersViewTooltip.h
new file mode 100644
--- /dev/null
+++ b/Engine/Source/Developer/TraceInsights/Private/Insights/Widgets/STimersViewTooltip.h
@@ -0,0 +1,53 @@
+// Copyright 1998-2019 Epic Games, Inc. All Rights Reserved.
+
+#pragma once
+
+#include "CoreMinimal.h"
+#include "EditorStyleSet.h"
+#include "SlateOptMacros.h"
+#include "TraceServices/AnalysisService.h"
+#include "Widgets/Images/SImage.h"
+#include "Widgets/Layout/SGridPanel.h"
+#include "Widgets/Layout/SSeparator.h"
+#include "Widgets/SToolTip.h"
+#include "Widgets/Text/STextBlock.h"
+
+// Insights
+#include "Insights/InsightsManager.h"
+
+// Insights
+class FTimerNode;
+class FTimersViewColumn;
+
+#define LOCTEXT_NAMESPACE "STimersView"
+
+// Timers View Tooltip
+class STimersViewTooltip
+{
+	const uint32 TimerId;
+	TSharedPtr<const Trace::IAnalysisSession> Session;
+
+public:
+	STimersViewTooltip(const uint32 InTimerId)
+		: TimerId(InTimerId)
+	{
+		Session = FInsightsManager::Get()->GetSession();
+	}
+
+	TSharedRef<SToolTip> GetTooltip();
+
+protected:
+	void AddNoDataInformation(const TSharedRef<SGridPanel>& Grid, int32& RowPos);
+	void AddHeader(const TSharedRef<SGridPanel>& Grid, int32& RowPos);
+	void AddDescription(const TSharedRef<SGridPanel>& Grid, int32& RowPos);
+	void AddSeparator(const TSharedRef<SGridPanel>& Grid, int32& RowPos);
+
+public:
+	static TSharedPtr<SToolTip> GetColumnTooltip(const FTimersViewColumn& Column);
+	static TSharedPtr<SToolTip> GetTableCellTooltip(const TSharedPtr<FTimerNode> TimerNodePtr);
+
+protected:
+	static void AddStatsRow(TSharedPtr<SGridPanel> Grid, int32& Row, const FText& Name, const FText& Value1, const FText& Value2);
+};
+
+#undef LOCTEXT_NAMESPACE
diff --git a/Engine/Source/Developer/TraceInsights/Private/Insights/Widgets/STimingProfilerToolbar.cpp b/Engine/Source/Developer/TraceInsights/Private/Insights/Widgets/STimingProfilerToolbar.cpp
new file mode 100644
--- /dev/null
+++ b/Engine/Source/Developer/TraceInsights/Private/Insights/Widgets/STimingProfilerToolbar.cpp
@@ -0,0 +1,155 @@
+// Copyright 1998-2019 Epic Games, Inc. All Rights Reserved.
+
+#include "STimingProfilerToolbar.h"
+
+#include "EditorStyleSet.h"
+#include "Framework/Application/SlateApplication.h"
+#include "Framework/Docking/TabManager.h"
+#include "Framework/MultiBox/MultiBoxDefs.h"
+#include "Framework/MultiBox/MultiBoxBuilder.h"
+#include "Widgets/Docking/SDockTab.h"
+#include "Widgets/Layout/SBorder.h"
+#include "Widgets/SBoxPanel.h"
+
+// Insights
+#include "Insights/InsightsCommands.h"
+#include "Insights/InsightsManager.h"
+#include "Insights/TimingProfilerCommands.h"
+#include "Insights/TimingProfilerManager.h"
+
+////////////////////////////////////////////////////////////////////////////////////////////////////
+
+#define LOCTEXT_NAMESPACE "STimingProfilerToolbar"
+
+////////////////////////////////////////////////////////////////////////////////////////////////////
+
+STimingProfilerToolbar::STimingProfilerToolbar()
+{
+}
+
+////////////////////////////////////////////////////////////////////////////////////////////////////
+
+STimingProfilerToolbar::~STimingProfilerToolbar()
+{
+}
+
+////////////////////////////////////////////////////////////////////////////////////////////////////
+
+void STimingProfilerToolbar::Construct(const FArguments& InArgs)
+{
+	CreateCommands();
+
+	struct Local
+	{
+		static void FillToolbar1(FToolBarBuilder& ToolbarBuilder)
+		{
+			ToolbarBuilder.BeginSection("Session");
+			{
+				ToolbarBuilder.AddToolBarButton(FInsightsCommands::Get().InsightsManager_Live);
+				ToolbarBuilder.AddToolBarButton(FInsightsCommands::Get().InsightsManager_Load);
+				ToolbarBuilder.AddToolBarButton(FInsightsCommands::Get().InsightsManager_Mock);
+			}
+			ToolbarBuilder.EndSection();
+			ToolbarBuilder.BeginSection("View");
+			{
+				ToolbarBuilder.AddToolBarButton(FTimingProfilerCommands::Get().ToggleFramesTrackVisibility);
+				ToolbarBuilder.AddToolBarButton(FTimingProfilerCommands::Get().ToggleGraphTrackVisibility);
+				ToolbarBuilder.AddToolBarButton(FTimingProfilerCommands::Get().ToggleTimingTrackVisibility);
+				ToolbarBuilder.AddToolBarButton(FTimingProfilerCommands::Get().ToggleTimersViewVisibility);
+				ToolbarBuilder.AddToolBarButton(FTimingProfilerCommands::Get().ToggleLogViewVisibility);
+			}
+			ToolbarBuilder.EndSection();
+			ToolbarBuilder.BeginSection("Options");
+			{
+				ToolbarBuilder.AddToolBarButton(FInsightsCommands::Get().OpenSettings);
+			}
+			ToolbarBuilder.EndSection();
+		}
+
+		static void FillToolbar2(FToolBarBuilder& ToolbarBuilder)
+		{
+			ToolbarBuilder.BeginSection("Options");
+			{
+				ToolbarBuilder.AddToolBarButton(FInsightsCommands::Get().ToggleDebugInfo);
+			}
+			ToolbarBuilder.EndSection();
+		}
+	};
+
+	TSharedPtr<FUICommandList> CommandList = FInsightsManager::Get()->GetCommandList();
+
+	FToolBarBuilder ToolbarBuilder1(CommandList.ToSharedRef(), FMultiBoxCustomization::None);
+	Local::FillToolbar1(ToolbarBuilder1);
+
+	FToolBarBuilder ToolbarBuilder2(CommandList.ToSharedRef(), FMultiBoxCustomization::None);
+	Local::FillToolbar2(ToolbarBuilder2);
+
+	// Create the tool bar!
+	ChildSlot
+	[
+		SNew(SHorizontalBox)
+
+		+SHorizontalBox::Slot()
+		.HAlign(HAlign_Left)
+		.VAlign(VAlign_Center)
+		.FillWidth(1.0)
+		.Padding(0.0f)
+		[
+			SNew(SBorder)
+			.Padding(0)
+			.BorderImage(FEditorStyle::GetBrush("NoBorder"))
+			.IsEnabled(FSlateApplication::Get().GetNormalExecutionAttribute())
+			[
+				ToolbarBuilder1.MakeWidget()
+			]
+		]
+
+		+SHorizontalBox::Slot()
+		.HAlign(HAlign_Right)
+		.VAlign(VAlign_Center)
+		.FillWidth(1.0)
+		.Padding(0.0f)
+		[
+			SNew(SBorder)
+			.Padding(0)
+			.BorderImage(FEditorStyle::GetBrush("NoBorder"))
+			.IsEnabled(FSlateApplication::Get().GetNormalExecutionAttribute())
+			[
+				ToolbarBuilder2.MakeWidget()
+			]
+		]
+	];
+}
+
+////////////////////////////////////////////////////////////////////////////////////////////////////
+
+void STimingProfilerToolbar::ShowStats()
+{
+	// do nothing
+}
+
+////////////////////////////////////////////////////////////////////////////////////////////////////
+
+void STimingProfilerToolbar::ShowMemory()
+{
+	// do nothing
+}
+
+////////////////////////////////////////////////////////////////////////////////////////////////////
+
+void STimingProfilerToolbar::CreateCommands()
+{
+	//TSharedPtr<FUICommandList> ProfilerCommandList = FTimingProfilerManager::Get()->GetCommandList();
+	//const FTimingProfilerCommands& Commands = FTimingProfilerCommands::Get();
+	//
+	//// Stats command
+	//ProfilerCommandList->MapAction(Commands.StatsProfiler,
+	//	FExecuteAction::CreateRaw(this, &STimingProfilerToolbar::ShowStats),
+	//	FCanExecuteAction(),
+	//	FIsActionChecked::CreateRaw(this, &STimingProfilerToolbar::IsShowingStats)
+	//);
+}
+
+////////////////////////////////////////////////////////////////////////////////////////////////////
+
+#undef LOCTEXT_NAMESPACE
