Copy up of Unreal Insights from //Tasks/UE4/Dev-Core-Trace.

diff --git a/Engine/Source/Developer/TraceInsights/Private/Insights/Widgets/STimingProfilerWindow.cpp b/Engine/Source/Developer/TraceInsights/Private/Insights/Widgets/STimingProfilerWindow.cpp
new file mode 100644
--- /dev/null
+++ b/Engine/Source/Developer/TraceInsights/Private/Insights/Widgets/STimingProfilerWindow.cpp
@@ -0,0 +1,801 @@
+// Copyright 1998-2019 Epic Games, Inc. All Rights Reserved.
+
+#include "STimingProfilerWindow.h"
+
+#include "EditorStyleSet.h"
+#include "SlateOptMacros.h"
+#include "Widgets/Images/SImage.h"
+#include "Widgets/Layout/SBorder.h"
+#include "Widgets/Layout/SBox.h"
+#include "Widgets/Layout/SSpacer.h"
+#include "Widgets/Notifications/SNotificationList.h"
+#include "Widgets/SBoxPanel.h"
+#include "Widgets/Text/STextBlock.h"
+
+#if WITH_EDITOR
+	#include "EngineAnalytics.h"
+	#include "Runtime/Analytics/Analytics/Public/AnalyticsEventAttribute.h"
+	#include "Runtime/Analytics/Analytics/Public/Interfaces/IAnalyticsProvider.h"
+#endif // WITH_EDITOR
+
+// Insights
+#include "Insights/Version.h"
+#include "Insights/InsightsManager.h"
+#include "Insights/TimingProfilerManager.h"
+#include "Insights/Widgets/SFrameTrack.h"
+#include "Insights/Widgets/SGraphTrack.h"
+#include "Insights/Widgets/SInsightsSettings.h"
+#include "Insights/Widgets/SLogView.h"
+#include "Insights/Widgets/STimersView.h"
+#include "Insights/Widgets/STimingProfilerToolbar.h"
+#include "Insights/Widgets/STimingView.h"
+
+////////////////////////////////////////////////////////////////////////////////////////////////////
+
+#define LOCTEXT_NAMESPACE "STimingProfilerWindow"
+
+////////////////////////////////////////////////////////////////////////////////////////////////////
+
+// TODO: move this function to InsightsManager or in TraceSession.h
+static FText GetTextForNotification(const EInsightsNotificationType NotificatonType, const ELoadingProgressState ProgressState, const FString& Filename, const float ProgressPercent = 0.0f)
+{
+	FText Result;
+
+	if (NotificatonType == EInsightsNotificationType::LoadingTraceFile)
+	{
+		if (ProgressState == ELoadingProgressState::Started)
+		{
+			FFormatNamedArguments Args;
+			Args.Add(TEXT("Filename"), FText::FromString(Filename));
+			Result = FText::Format(LOCTEXT("DescF_OfflineCapture_Started", "Started loading a file {Filename}"), Args);
+		}
+		else if (ProgressState == ELoadingProgressState::InProgress)
+		{
+			FFormatNamedArguments Args;
+			Args.Add(TEXT("Filename"), FText::FromString(Filename));
+			Args.Add(TEXT("LoadingProgressPercent"), FText::AsPercent(ProgressPercent));
+			Result = FText::Format(LOCTEXT("DescF_OfflineCapture_InProgress", "Loading a file {Filename} {LoadingProgressPercent}"), Args);
+		}
+		else if (ProgressState == ELoadingProgressState::Loaded)
+		{
+			FFormatNamedArguments Args;
+			Args.Add(TEXT("Filename"), FText::FromString(Filename));
+			Result = FText::Format(LOCTEXT("DescF_OfflineCapture_Loaded", "File {Filename} has been successfully loaded"), Args);
+		}
+		else if (ProgressState == ELoadingProgressState::Failed)
+		{
+			FFormatNamedArguments Args;
+			Args.Add(TEXT("Filename"), FText::FromString(Filename));
+			Result = FText::Format(LOCTEXT("DescF_OfflineCapture_Failed", "Failed to load file {Filename}"), Args);
+		}
+	}
+
+	return Result;
+}
+
+////////////////////////////////////////////////////////////////////////////////////////////////////
+
+STimingProfilerWindow::STimingProfilerWindow()
+	: DurationActive(0.0f)
+{
+}
+
+////////////////////////////////////////////////////////////////////////////////////////////////////
+
+STimingProfilerWindow::~STimingProfilerWindow()
+{
+	// Remove ourselves from the profiler manager.
+	//if (FTimingProfilerManager::Get().IsValid())
+	//{
+	//	//...
+	//	FTimingProfilerManager::Get()->OnViewModeChanged().RemoveAll(this);
+	//}
+
+#if WITH_EDITOR
+	if (DurationActive > 0.0f && FEngineAnalytics::IsAvailable())
+	{
+		FEngineAnalytics::GetProvider().RecordEvent(TEXT("Editor.Usage.Profiler"), FAnalyticsEventAttribute(TEXT("Duration"), DurationActive));
+	}
+#endif // WITH_EDITOR
+}
+
+////////////////////////////////////////////////////////////////////////////////////////////////////
+
+BEGIN_SLATE_FUNCTION_BUILD_OPTIMIZATION
+
+TSharedRef<SWidget> STimingProfilerWindow::ConstructMultiTrackView()
+{
+	return SNew(SSplitter)
+		.Orientation(Orient_Vertical)
+
+	+ SSplitter::Slot()
+		.Value(0.1f)
+		[
+			ConstructFramesTrack()
+		]
+
+	+ SSplitter::Slot()
+		.Value(0.2f)
+		[
+			ConstructGraphTrack()
+		]
+
+	+ SSplitter::Slot()
+		.Value(0.7f)
+		[
+			ConstructTimingTrack()
+		]
+
+	; // end of Widget
+}
+
+////////////////////////////////////////////////////////////////////////////////////////////////////
+
+TSharedRef<SWidget> STimingProfilerWindow::ConstructFramesTrack()
+{
+	return SNew(SVerticalBox)
+		.Visibility(this, &STimingProfilerWindow::IsFramesTrackVisible)
+		.IsEnabled(this, &STimingProfilerWindow::IsProfilerEnabled)
+
+	// Frame Track
+	+ SVerticalBox::Slot()
+		.FillHeight(1.0)
+		.Padding(0.0f, 0.0f, 0.0f, 0.0f)
+		[
+			SNew(SBox)
+				.HeightOverride(48.0f)
+				.IsEnabled(this, &STimingProfilerWindow::IsProfilerEnabled)
+				[
+					SAssignNew(FrameTrack, SFrameTrack)
+
+					//SNew(SHorizontalBox)
+					//
+					//+ SHorizontalBox::Slot()
+					//	.FillWidth(1.0f)
+					//	.Padding(0.0f)
+					//	.HAlign(HAlign_Fill)
+					//	.VAlign(VAlign_Fill)
+					//	[
+					//
+					//	]
+				]
+		]
+
+	; // end of Widget
+}
+
+////////////////////////////////////////////////////////////////////////////////////////////////////
+
+TSharedRef<SWidget> STimingProfilerWindow::ConstructGraphTrack()
+{
+	return SNew(SVerticalBox)
+		.Visibility(this, &STimingProfilerWindow::IsGraphTrackVisible)
+		.IsEnabled(this, &STimingProfilerWindow::IsProfilerEnabled)
+
+	/*
+	// Header
+	+ SVerticalBox::Slot()
+		.AutoHeight()
+		.Padding(0.0f, 0.0f, 0.0f, 0.0f)
+		[
+			SNew(SHorizontalBox)
+
+			+ SHorizontalBox::Slot()
+				.AutoWidth()
+				[
+					SNew(SImage)
+						.Image(FEditorStyle::GetBrush(TEXT("Profiler.Tab.GraphView")))
+				]
+
+			+ SHorizontalBox::Slot()
+				.AutoWidth()
+				[
+					SNew(STextBlock)
+						.Text(LOCTEXT("GraphTrackLabel", "Overview Timing Graph"))
+						.Margin(FMargin(0.0f, 2.0f, 0.0f, 0.0f))
+				]
+		]
+	*/
+
+	// Graph Track
+	+ SVerticalBox::Slot()
+		.FillHeight(1.0)
+		.Padding(0.0f, 0.0f, 0.0f, 0.0f)
+		[
+			SNew(SBox)
+				.HeightOverride(48.0f)
+				.IsEnabled(this, &STimingProfilerWindow::IsProfilerEnabled)
+				[
+					SNew(SHorizontalBox)
+
+					+ SHorizontalBox::Slot()
+						.FillWidth(1.0f)
+						.Padding(0.0f)
+						.HAlign(HAlign_Fill)
+						.VAlign(VAlign_Fill)
+						[
+							SAssignNew(GraphTrack, SGraphTrack)
+						]
+				]
+		]
+
+	; // end of Widget
+}
+
+////////////////////////////////////////////////////////////////////////////////////////////////////
+
+TSharedRef<SWidget> STimingProfilerWindow::ConstructTimingTrack()
+{
+	return SNew(SVerticalBox)
+		.Visibility(this, &STimingProfilerWindow::IsTimingTrackVisible)
+		.IsEnabled(this, &STimingProfilerWindow::IsProfilerEnabled)
+
+	/*
+	// Header
+	+ SVerticalBox::Slot()
+		.AutoHeight()
+		.Padding(0.0f, 0.0f, 0.0f, 0.0f)
+		[
+			SNew(SHorizontalBox)
+
+			+ SHorizontalBox::Slot()
+				.AutoWidth()
+				[
+					SNew(SImage)
+						.Image(FEditorStyle::GetBrush(TEXT("Profiler.Tab.GraphView")))
+				]
+
+			+ SHorizontalBox::Slot()
+				.AutoWidth()
+				[
+					SNew(STextBlock)
+						.Text(LOCTEXT("TimingLabel", "Main Timing View"))
+						.Margin(FMargin(0.0f, 2.0f, 0.0f, 0.0f))
+				]
+		]
+	*/
+
+	// Timing Track
+	+ SVerticalBox::Slot()
+		.FillHeight(1.0)
+		.Padding(0.0f, 0.0f, 0.0f, 0.0f)
+		[
+			SNew(SBox)
+				.HeightOverride(48.0f)
+				.IsEnabled(this, &STimingProfilerWindow::IsProfilerEnabled)
+				[
+					SNew(SHorizontalBox)
+
+					+ SHorizontalBox::Slot()
+						.FillWidth(1.0f)
+						.Padding(0.0f)
+						.HAlign(HAlign_Fill)
+						.VAlign(VAlign_Fill)
+						[
+							SAssignNew(TimingView, STimingView)
+						]
+				]
+		]
+
+	//// Threads Panel
+	//+ SVerticalBox::Slot()
+	//	.FillHeight(1.0)
+	//	.Padding(0.0f, 0.0f, 0.0f, 0.0f)
+	//	[
+	//		SAssignNew(ThreadsPanel, SVerticalBox)
+	//	]
+
+	; // end of Widget
+}
+
+////////////////////////////////////////////////////////////////////////////////////////////////////
+
+TSharedRef<SWidget> STimingProfilerWindow::ConstructTimersView()
+{
+	return SNew(SVerticalBox)
+		.Visibility(this, &STimingProfilerWindow::IsTimersViewVisible)
+		.IsEnabled(this, &STimingProfilerWindow::IsProfilerEnabled)
+
+	// Header
+	+ SVerticalBox::Slot()
+		.AutoHeight()
+		[
+			SNew(SHorizontalBox)
+
+			+ SHorizontalBox::Slot()
+				.AutoWidth()
+				[
+					SNew(SImage)
+						.Image(FEditorStyle::GetBrush(TEXT("Profiler.Tab.FiltersAndPresets")))
+				]
+
+			+ SHorizontalBox::Slot()
+				.AutoWidth()
+				[
+					SNew(STextBlock)
+						.Text(LOCTEXT("TimersLabel", "Timers"))
+				]
+		]
+
+	// Timers
+	+ SVerticalBox::Slot()
+		.FillHeight(1.0f)
+		.Padding(0.0f, 2.0f, 0.0f, 0.0f)
+		[
+			SAssignNew(TimersView, STimersView)
+		]
+
+	; // end of Widget
+}
+
+////////////////////////////////////////////////////////////////////////////////////////////////////
+
+TSharedRef<SWidget> STimingProfilerWindow::ConstructLogView()
+{
+	return SAssignNew(LogView, SLogView)
+		.Visibility(this, &STimingProfilerWindow::IsLogViewVisible);
+}
+
+////////////////////////////////////////////////////////////////////////////////////////////////////
+
+void STimingProfilerWindow::Construct(const FArguments& InArgs)
+{
+	ChildSlot
+		[
+			SNew(SOverlay)
+
+			// Version
+			+ SOverlay::Slot()
+				.HAlign(HAlign_Right)
+				.VAlign(VAlign_Top)
+				.Padding(0.0f, -16.0f, 0.0f, 0.0f)
+				[
+					SNew(STextBlock)
+						.Clipping(EWidgetClipping::ClipToBoundsWithoutIntersecting)
+						.Text(LOCTEXT("UnrealInsightsVersion", UNREAL_INSIGHTS_VERSION_STRING_EX))
+						.ColorAndOpacity(FLinearColor(0.15f, 0.15f, 0.15f, 1.0f))
+				]
+
+			// Overlay slot for the main profiler window area
+			+ SOverlay::Slot()
+				.HAlign(HAlign_Fill)
+				.VAlign(VAlign_Fill)
+				[
+					SAssignNew(MainContentPanel, SVerticalBox)
+
+					// Toolbar
+					+ SVerticalBox::Slot()
+						.AutoHeight()
+						[
+							SNew(STimingProfilerToolbar)
+						]
+
+					+ SVerticalBox::Slot()
+						.HAlign(HAlign_Fill)
+						.VAlign(VAlign_Fill)
+						[
+							SNew(SSplitter)
+								.Orientation(Orient_Horizontal)
+
+							+ SSplitter::Slot()
+								.Value(0.65f)
+								[
+									SNew(SSplitter)
+										.Orientation(Orient_Vertical)
+
+									// Multi Track View
+									+ SSplitter::Slot()
+										.Value(0.75f)
+										[
+											ConstructMultiTrackView()
+										]
+
+									// Log View
+									+ SSplitter::Slot()
+										.Value(0.25f)
+										[
+											ConstructLogView()
+										]
+								]
+
+							// Timers View
+							+ SSplitter::Slot()
+								.Value(0.35f)
+								//.SizeRule(SSplitter::SizeToContent)
+								.Expose(TimersViewSlot)
+								[
+									ConstructTimersView()
+								]
+						]
+				]
+
+			// session hint overlay
+			+ SOverlay::Slot()
+				.HAlign(HAlign_Center)
+				.VAlign(VAlign_Center)
+				[
+					SNew(SBorder)
+						.Visibility(this, &STimingProfilerWindow::IsSessionOverlayVisible)
+						.BorderImage(FEditorStyle::GetBrush("NotificationList.ItemBackground"))
+						.Padding(8.0f)
+						[
+							SNew(STextBlock)
+								.Text(LOCTEXT("SelectTraceOverlayText", "Please select a trace."))
+						]
+				]
+
+			// notification area overlay
+			+ SOverlay::Slot()
+				.HAlign(HAlign_Right)
+				.VAlign(VAlign_Bottom)
+				.Padding(16.0f)
+				[
+					SAssignNew(NotificationList, SNotificationList)
+				]
+
+			// profiler settings overlay
+			+ SOverlay::Slot()
+				.HAlign(HAlign_Center)
+				.VAlign(VAlign_Center)
+				.Expose(OverlaySettingsSlot)
+		];
+
+	//ProfilerMiniView->OnSelectionBoxChanged().AddSP(GraphPanel.ToSharedRef(), &STimingProfilerGraphPanel::MiniView_OnSelectionBoxChanged);
+	//FTimingProfilerManager::Get()->OnViewModeChanged().AddSP(this, &STimingProfilerWindow::ProfilerManager_OnViewModeChanged);
+}
+
+END_SLATE_FUNCTION_BUILD_OPTIMIZATION
+
+////////////////////////////////////////////////////////////////////////////////////////////////////
+
+EVisibility STimingProfilerWindow::IsFramesTrackVisible() const
+{
+	if (FTimingProfilerManager::Get()->IsFramesTrackVisible() &&
+		FInsightsManager::Get()->GetSession().IsValid())
+	{
+		return EVisibility::Visible;
+	}
+
+	return EVisibility::Collapsed;
+}
+
+////////////////////////////////////////////////////////////////////////////////////////////////////
+
+EVisibility STimingProfilerWindow::IsGraphTrackVisible() const
+{
+	if (FTimingProfilerManager::Get()->IsGraphTrackVisible() &&
+		FInsightsManager::Get()->GetSession().IsValid())
+	{
+		return EVisibility::Visible;
+	}
+
+	return EVisibility::Collapsed;
+}
+
+////////////////////////////////////////////////////////////////////////////////////////////////////
+
+EVisibility STimingProfilerWindow::IsTimingTrackVisible() const
+{
+	if (FTimingProfilerManager::Get()->IsTimingTrackVisible() &&
+		FInsightsManager::Get()->GetSession().IsValid())
+	{
+		return EVisibility::Visible;
+	}
+
+	return EVisibility::Collapsed;
+}
+
+////////////////////////////////////////////////////////////////////////////////////////////////////
+
+EVisibility STimingProfilerWindow::IsTimersViewVisible() const
+{
+	if (FTimingProfilerManager::Get()->IsTimersViewVisible() &&
+		FInsightsManager::Get()->GetSession().IsValid())
+	{
+		return EVisibility::Visible;
+	}
+
+	return EVisibility::Collapsed;
+}
+
+////////////////////////////////////////////////////////////////////////////////////////////////////
+
+EVisibility STimingProfilerWindow::IsLogViewVisible() const
+{
+	if (FTimingProfilerManager::Get()->IsLogViewVisible() &&
+		FInsightsManager::Get()->GetSession().IsValid())
+	{
+		return EVisibility::Visible;
+	}
+
+	return EVisibility::Collapsed;
+}
+
+////////////////////////////////////////////////////////////////////////////////////////////////////
+
+EVisibility STimingProfilerWindow::IsSessionOverlayVisible() const
+{
+	if (FInsightsManager::Get()->GetSession().IsValid())
+	{
+		return EVisibility::Hidden;
+	}
+
+	return EVisibility::Visible;
+}
+
+////////////////////////////////////////////////////////////////////////////////////////////////////
+
+bool STimingProfilerWindow::IsProfilerEnabled() const
+{
+	const bool bIsActive = FInsightsManager::Get()->GetSession().IsValid();
+	return bIsActive;
+}
+
+////////////////////////////////////////////////////////////////////////////////////////////////////
+
+void STimingProfilerWindow::ManageLoadingProgressNotificationState(const FString& Filename, const EInsightsNotificationType NotificatonType, const ELoadingProgressState ProgressState, const float LoadingProgress)
+{
+	const FString BaseFilename = FPaths::GetBaseFilename(Filename);
+
+	if (ProgressState == ELoadingProgressState::Started)
+	{
+		const bool bContains = ActiveNotifications.Contains(Filename);
+		if (!bContains)
+		{
+			FNotificationInfo NotificationInfo(GetTextForNotification(NotificatonType, ProgressState, BaseFilename));
+			NotificationInfo.bFireAndForget = false;
+			NotificationInfo.bUseLargeFont = false;
+
+			// Add two buttons, one for cancel, one for loading the received file.
+			if (NotificatonType == EInsightsNotificationType::LoadingTraceFile)
+			{
+				const FText CancelButtonText = LOCTEXT("CancelButton_Text", "Cancel");
+				const FText CancelButtonTT = LOCTEXT("CancelButton_TTText", "Hides this notification");
+				const FText LoadButtonText = LOCTEXT("LoadButton_Text", "Load file");
+				const FText LoadButtonTT = LOCTEXT("LoadButton_TTText", "Loads the received file and hides this notification");
+
+				NotificationInfo.ButtonDetails.Add(FNotificationButtonInfo(CancelButtonText, CancelButtonTT,
+					FSimpleDelegate::CreateSP(this, &STimingProfilerWindow::SendingServiceSideCapture_Cancel, Filename), SNotificationItem::CS_Success));
+				NotificationInfo.ButtonDetails.Add(FNotificationButtonInfo(LoadButtonText, LoadButtonTT,
+					FSimpleDelegate::CreateSP(this, &STimingProfilerWindow::SendingServiceSideCapture_Load, Filename), SNotificationItem::CS_Success));
+			}
+
+			SNotificationItemWeak NotificationItem = NotificationList->AddNotification(NotificationInfo);
+			NotificationItem.Pin()->SetCompletionState(SNotificationItem::CS_Pending);
+			ActiveNotifications.Add(Filename, NotificationItem);
+		}
+	}
+	else if (ProgressState == ELoadingProgressState::InProgress)
+	{
+		const SNotificationItemWeak* LoadingProgressPtr = ActiveNotifications.Find(Filename);
+		if (LoadingProgressPtr)
+		{
+			TSharedPtr<SNotificationItem> LoadingProcessPinned = LoadingProgressPtr->Pin();
+			LoadingProcessPinned->SetText(GetTextForNotification(NotificatonType, ProgressState, BaseFilename, LoadingProgress));
+			LoadingProcessPinned->SetCompletionState(SNotificationItem::CS_Pending);
+		}
+	}
+	else if (ProgressState == ELoadingProgressState::Loaded)
+	{
+		const SNotificationItemWeak* LoadingProgressPtr = ActiveNotifications.Find(Filename);
+		if (LoadingProgressPtr)
+		{
+			TSharedPtr<SNotificationItem> LoadingProcessPinned = LoadingProgressPtr->Pin();
+			LoadingProcessPinned->SetText(GetTextForNotification(NotificatonType, ProgressState, BaseFilename));
+			LoadingProcessPinned->SetCompletionState(SNotificationItem::CS_Success);
+
+			// Notifications for received files are handled by the user.
+			if (NotificatonType == EInsightsNotificationType::LoadingTraceFile)
+			{
+				LoadingProcessPinned->ExpireAndFadeout();
+				ActiveNotifications.Remove(Filename);
+			}
+		}
+	}
+	else if (ProgressState == ELoadingProgressState::Failed)
+	{
+		const SNotificationItemWeak* LoadingProgressPtr = ActiveNotifications.Find(Filename);
+		if (LoadingProgressPtr)
+		{
+			TSharedPtr<SNotificationItem> LoadingProcessPinned = LoadingProgressPtr->Pin();
+			LoadingProcessPinned->SetText(GetTextForNotification(NotificatonType, ProgressState, BaseFilename));
+			LoadingProcessPinned->SetCompletionState(SNotificationItem::CS_Fail);
+
+			LoadingProcessPinned->ExpireAndFadeout();
+			ActiveNotifications.Remove(Filename);
+		}
+	}
+	else if (ProgressState == ELoadingProgressState::Cancelled)
+	{
+		const SNotificationItemWeak* LoadingProgressPtr = ActiveNotifications.Find(Filename);
+		if (LoadingProgressPtr)
+		{
+			TSharedPtr<SNotificationItem> LoadingProcessPinned = LoadingProgressPtr->Pin();
+			LoadingProcessPinned->ExpireAndFadeout();
+			ActiveNotifications.Remove(Filename);
+		}
+	}
+}
+
+////////////////////////////////////////////////////////////////////////////////////////////////////
+
+void STimingProfilerWindow::SendingServiceSideCapture_Cancel(const FString Filename)
+{
+	const SNotificationItemWeak* LoadingProgressPtr = ActiveNotifications.Find(Filename);
+	if (LoadingProgressPtr)
+	{
+		TSharedPtr<SNotificationItem> LoadingProcessPinned = LoadingProgressPtr->Pin();
+		LoadingProcessPinned->ExpireAndFadeout();
+		ActiveNotifications.Remove(Filename);
+	}
+}
+
+////////////////////////////////////////////////////////////////////////////////////////////////////
+
+void STimingProfilerWindow::SendingServiceSideCapture_Load(const FString Filename)
+{
+	const SNotificationItemWeak* LoadingProgressPtr = ActiveNotifications.Find(Filename);
+	if (LoadingProgressPtr)
+	{
+		TSharedPtr<SNotificationItem> LoadingProcessPinned = LoadingProgressPtr->Pin();
+		LoadingProcessPinned->ExpireAndFadeout();
+		ActiveNotifications.Remove(Filename);
+
+		const FString PathName = FPaths::ProfilingDir() + TEXT("UnrealStats/Received/");
+		const FString TraceFilepath = PathName + Filename;
+		FInsightsManager::Get()->LoadTraceFile(TraceFilepath);
+	}
+}
+
+////////////////////////////////////////////////////////////////////////////////////////////////////
+
+EActiveTimerReturnType STimingProfilerWindow::UpdateActiveDuration(double InCurrentTime, float InDeltaTime)
+{
+	DurationActive += InDeltaTime;
+
+	// The profiler window will explicitly unregister this active timer when the mouse leaves.
+	return EActiveTimerReturnType::Continue;
+}
+
+////////////////////////////////////////////////////////////////////////////////////////////////////
+
+void STimingProfilerWindow::OnMouseEnter(const FGeometry& MyGeometry, const FPointerEvent& MouseEvent)
+{
+	SCompoundWidget::OnMouseEnter(MyGeometry, MouseEvent);
+
+	if (!ActiveTimerHandle.IsValid())
+	{
+		ActiveTimerHandle = RegisterActiveTimer(0.f, FWidgetActiveTimerDelegate::CreateSP(this, &STimingProfilerWindow::UpdateActiveDuration));
+	}
+}
+
+////////////////////////////////////////////////////////////////////////////////////////////////////
+
+void STimingProfilerWindow::OnMouseLeave(const FPointerEvent& MouseEvent)
+{
+	SCompoundWidget::OnMouseLeave(MouseEvent);
+
+	auto PinnedActiveTimerHandle = ActiveTimerHandle.Pin();
+	if (PinnedActiveTimerHandle.IsValid())
+	{
+		UnRegisterActiveTimer(PinnedActiveTimerHandle.ToSharedRef());
+	}
+}
+
+////////////////////////////////////////////////////////////////////////////////////////////////////
+
+FReply STimingProfilerWindow::OnKeyDown(const FGeometry& MyGeometry, const FKeyEvent& InKeyEvent)
+{
+	return FTimingProfilerManager::Get()->GetCommandList()->ProcessCommandBindings(InKeyEvent) ? FReply::Handled() : FReply::Unhandled();
+}
+
+////////////////////////////////////////////////////////////////////////////////////////////////////
+
+FReply STimingProfilerWindow::OnDragOver(const FGeometry& MyGeometry, const FDragDropEvent& DragDropEvent)
+{
+	TSharedPtr<FExternalDragOperation> DragDropOp = DragDropEvent.GetOperationAs<FExternalDragOperation>();
+	if (DragDropOp.IsValid())
+	{
+		if (DragDropOp->HasFiles())
+		{
+			const TArray<FString>& Files = DragDropOp->GetFiles();
+			if (Files.Num() == 1)
+			{
+				const FString DraggedFileExtension = FPaths::GetExtension(Files[0], true);
+				if (DraggedFileExtension == TEXT(".utrace"))
+				{
+					return FReply::Handled();
+				}
+			}
+		}
+	}
+
+	return SCompoundWidget::OnDragOver(MyGeometry,DragDropEvent);
+}
+
+////////////////////////////////////////////////////////////////////////////////////////////////////
+
+FReply STimingProfilerWindow::OnDrop(const FGeometry& MyGeometry, const FDragDropEvent& DragDropEvent)
+{
+	TSharedPtr<FExternalDragOperation> DragDropOp = DragDropEvent.GetOperationAs<FExternalDragOperation>();
+	if (DragDropOp.IsValid())
+	{
+		if (DragDropOp->HasFiles())
+		{
+			// For now, only allow a single file.
+			const TArray<FString>& Files = DragDropOp->GetFiles();
+			if (Files.Num() == 1)
+			{
+				const FString DraggedFileExtension = FPaths::GetExtension(Files[0], true);
+				if (DraggedFileExtension == TEXT(".utrace"))
+				{
+					// Enqueue load operation.
+					FInsightsManager::Get()->LoadTraceFile(Files[0]);
+					return FReply::Handled();
+				}
+			}
+		}
+	}
+
+	return SCompoundWidget::OnDrop(MyGeometry,DragDropEvent);
+}
+
+////////////////////////////////////////////////////////////////////////////////////////////////////
+
+void STimingProfilerWindow::OpenProfilerSettings()
+{
+	MainContentPanel->SetEnabled(false);
+	(*OverlaySettingsSlot)
+	[
+		SNew(SBorder)
+		.BorderImage(FEditorStyle::GetBrush("NotificationList.ItemBackground"))
+		.Padding(8.0f)
+		[
+			SNew(SInsightsSettings)
+			.OnClose(this, &STimingProfilerWindow::CloseProfilerSettings)
+			.SettingPtr(&FInsightsManager::GetSettings())
+		]
+	];
+}
+
+////////////////////////////////////////////////////////////////////////////////////////////////////
+
+void STimingProfilerWindow::CloseProfilerSettings()
+{
+	// Close the profiler settings by simply replacing widget with a null one.
+	(*OverlaySettingsSlot)
+	[
+		SNullWidget::NullWidget
+	];
+	MainContentPanel->SetEnabled(true);
+}
+
+////////////////////////////////////////////////////////////////////////////////////////////////////
+
+//void STimingProfilerWindow::ProfilerManager_OnViewModeChanged(ETimingProfilerViewMode NewViewMode)
+//{
+//	if (NewViewMode == ETimingProfilerViewMode::LineIndexBased)
+//	{
+//		EventGraphPanel->SetVisibility(EVisibility::Visible);
+//		EventGraphPanel->SetEnabled(true);
+//
+//		(*FiltersAndPresetsSlot)
+//		[
+//			FiltersAndPresets.ToSharedRef()
+//		];
+//	}
+//	else if (NewViewMode == ETimingProfilerViewMode::ThreadViewTimeBased)
+//	{
+//		EventGraphPanel->SetVisibility(EVisibility::Collapsed);
+//		EventGraphPanel->SetEnabled(false);
+//
+//		(*FiltersAndPresetsSlot)
+//		[
+//			SNullWidget::NullWidget
+//		];
+//	}
+//}
+
+////////////////////////////////////////////////////////////////////////////////////////////////////
+
+#undef LOCTEXT_NAMESPACE
diff --git a/Engine/Source/Developer/TraceInsights/Private/Insights/Widgets/STimingProfilerWindow.h b/Engine/Source/Developer/TraceInsights/Private/Insights/Widgets/STimingProfilerWindow.h
new file mode 100644
--- /dev/null
+++ b/Engine/Source/Developer/TraceInsights/Private/Insights/Widgets/STimingProfilerWindow.h
@@ -0,0 +1,187 @@
+// Copyright 1998-2019 Epic Games, Inc. All Rights Reserved.
+
+#pragma once
+
+#include "CoreMinimal.h"
+#include "Input/Reply.h"
+#include "Layout/Visibility.h"
+#include "Misc/Guid.h"
+#include "SlateFwd.h"
+#include "Widgets/DeclarativeSyntaxSupport.h"
+#include "Widgets/Layout/SSplitter.h"
+#include "Widgets/SCompoundWidget.h"
+#include "Widgets/SOverlay.h"
+
+// Insights
+#include "Insights/InsightsManager.h"
+
+class FActiveTimerHandle;
+class SVerticalBox;
+class SFrameTrack;
+class SGraphTrack;
+class STimersView;
+class STimingView;
+class SLogView;
+
+/// Type definition for shared pointers to instances of SNotificationItem.
+typedef TSharedPtr<class SNotificationItem> SNotificationItemPtr;
+
+/// Type definition for shared references to instances of SNotificationItem.
+typedef TSharedRef<class SNotificationItem> SNotificationItemRef;
+
+/// Type definition for weak references to instances of SNotificationItem.
+typedef TWeakPtr<class SNotificationItem> SNotificationItemWeak;
+
+/// Implements the timing profiler window.
+class STimingProfilerWindow : public SCompoundWidget
+{
+public:
+	/// Default constructor.
+	STimingProfilerWindow();
+
+	/// Virtual destructor.
+	virtual ~STimingProfilerWindow();
+
+	SLATE_BEGIN_ARGS(STimingProfilerWindow){}
+	SLATE_END_ARGS()
+
+	/// Constructs this widget.
+	void Construct(const FArguments& InArgs);
+
+	void ManageLoadingProgressNotificationState(const FString& Filename, const EInsightsNotificationType NotificatonType, const ELoadingProgressState ProgressState, const float LoadingProgress);
+
+	void OpenProfilerSettings();
+	void CloseProfilerSettings();
+
+protected:
+	TSharedRef<SWidget> ConstructMultiTrackView();
+	TSharedRef<SWidget> ConstructFramesTrack();
+	TSharedRef<SWidget> ConstructGraphTrack();
+	TSharedRef<SWidget> ConstructTimingTrack();
+	TSharedRef<SWidget> ConstructTimersView();
+	TSharedRef<SWidget> ConstructLogView();
+
+	/// Callback for determining the visibility of the Frames track.
+	EVisibility IsFramesTrackVisible() const;
+
+	/// Callback for determining the visibility of the Graph track.
+	EVisibility IsGraphTrackVisible() const;
+
+	/// Callback for determining the visibility of the Timing track.
+	EVisibility IsTimingTrackVisible() const;
+
+	/// Callback for determining the visibility of the Timers View.
+	EVisibility IsTimersViewVisible() const;
+
+	/// Callback for determining the visibility of the Log View.
+	EVisibility IsLogViewVisible() const;
+
+	/// Callback for determining the visibility of the 'Select a session' overlay.
+	EVisibility IsSessionOverlayVisible() const;
+
+	/// Callback for getting the enabled state of the profiler window.
+	bool IsProfilerEnabled() const;
+
+	void SendingServiceSideCapture_Cancel(const FString Filename);
+
+	void SendingServiceSideCapture_Load(const FString Filename);
+
+private:
+
+	/**
+	 * Ticks this widget.  Override in derived classes, but always call the parent implementation.
+	 *
+	 * @param  AllottedGeometry The space allotted for this widget
+	 * @param  InCurrentTime  Current absolute real time
+	 * @param  InDeltaTime  Real time passed since last tick
+	 */
+	//virtual void Tick(const FGeometry& AllottedGeometry, const double InCurrentTime, const float InDeltaTime) override;
+
+	/**
+	 * The system will use this event to notify a widget that the cursor has entered it. This event is NOT bubbled.
+	 *
+	 * @param MyGeometry The Geometry of the widget receiving the event
+	 * @param MouseEvent Information about the input event
+	 */
+	virtual void OnMouseEnter(const FGeometry& MyGeometry, const FPointerEvent& MouseEvent) override;
+
+	/**
+	 * The system will use this event to notify a widget that the cursor has left it. This event is NOT bubbled.
+	 *
+	 * @param MouseEvent Information about the input event
+	 */
+	virtual void OnMouseLeave(const FPointerEvent& MouseEvent) override;
+
+	/**
+	 * Called after a key is pressed when this widget has focus
+	 *
+	 * @param MyGeometry The Geometry of the widget receiving the event
+	 * @param  InKeyEvent  Key event
+	 *
+	 * @return  Returns whether the event was handled, along with other possible actions
+	 */
+	virtual FReply OnKeyDown(const FGeometry& MyGeometry, const FKeyEvent& InKeyEvent) override;
+
+	/**
+	 * Called when the user is dropping something onto a widget; terminates drag and drop.
+	 *
+	 * @param MyGeometry      The geometry of the widget receiving the event.
+	 * @param DragDropEvent   The drag and drop event.
+	 *
+	 * @return A reply that indicated whether this event was handled.
+	 */
+	virtual FReply OnDrop(const FGeometry& MyGeometry, const FDragDropEvent& DragDropEvent) override;
+
+	/**
+	 * Called during drag and drop when the the mouse is being dragged over a widget.
+	 *
+	 * @param MyGeometry      The geometry of the widget receiving the event.
+	 * @param DragDropEvent   The drag and drop event.
+	 *
+	 * @return A reply that indicated whether this event was handled.
+	 */
+	virtual FReply OnDragOver(const FGeometry& MyGeometry, const FDragDropEvent& DragDropEvent)  override;
+
+private:
+	/// Updates the amount of time the profiler has been active
+	EActiveTimerReturnType UpdateActiveDuration(double InCurrentTime, float InDeltaTime);
+
+public:
+	/// Widget for the frame track
+	TSharedPtr<SFrameTrack> FrameTrack;
+
+	/// Widget for the graph track
+	TSharedPtr<SGraphTrack> GraphTrack;
+
+	/// Holds all timing tracks (one for each thread).
+	TSharedPtr<SVerticalBox> ThreadsPanel;
+
+	/// Widget for the timing track
+	TSharedPtr<STimingView> TimingView;
+
+	/// Holds the Timers view widget/slot.
+	SSplitter::FSlot* TimersViewSlot;
+	TSharedPtr<STimersView> TimersView;
+
+	/// Widget for the log view
+	TSharedPtr<SLogView> LogView;
+
+	/// Widget for the non-intrusive notifications.
+	TSharedPtr<SNotificationList> NotificationList;
+
+	/// Holds all active and visible notifications, stored as FGuid -> SNotificationItemWeak.
+	TMap<FString, SNotificationItemWeak> ActiveNotifications;
+
+	/// Overlay slot which contains the profiler settings widget.
+	SOverlay::FOverlaySlot* OverlaySettingsSlot;
+
+	/// The number of seconds the profiler has been active
+	float DurationActive;
+
+private:
+	/// The handle to the active update duration tick
+	TWeakPtr<FActiveTimerHandle> ActiveTimerHandle;
+
+	/// Holds all widgets for the profiler window like menu bar, toolbar and tabs.
+	TSharedPtr<SVerticalBox> MainContentPanel;
+};
