Copy up of Unreal Insights from //Tasks/UE4/Dev-Core-Trace.

diff --git a/Engine/Source/Developer/TraceInsights/Private/Insights/Widgets/STimingView.cpp b/Engine/Source/Developer/TraceInsights/Private/Insights/Widgets/STimingView.cpp
new file mode 100644
--- /dev/null
+++ b/Engine/Source/Developer/TraceInsights/Private/Insights/Widgets/STimingView.cpp
@@ -0,0 +1,2706 @@
+// Copyright 1998-2019 Epic Games, Inc. All Rights Reserved.
+
+#include "STimingView.h"
+
+#include "Brushes/SlateBorderBrush.h"
+#include "Brushes/SlateBoxBrush.h"
+#include "Brushes/SlateColorBrush.h"
+#include "Containers/ArrayBuilder.h"
+#include "Containers/MapBuilder.h"
+#include "EditorStyleSet.h"
+#include "Fonts/FontMeasure.h"
+#include "Fonts/SlateFontInfo.h"
+#include "Framework/Application/MenuStack.h"
+#include "Framework/Application/SlateApplication.h"
+#include "Framework/MultiBox/MultiBoxBuilder.h"
+#include "HAL/PlatformTime.h"
+#include "Layout/WidgetPath.h"
+#include "Misc/Paths.h"
+#include "Rendering/DrawElements.h"
+#include "SlateOptMacros.h"
+#include "Styling/CoreStyle.h"
+#include "TraceServices/AnalysisService.h"
+#include "TraceServices/SessionService.h"
+#include "Widgets/Input/SButton.h"
+#include "Widgets/Layout/SScrollBar.h"
+#include "Widgets/Layout/SSpacer.h"
+#include "Widgets/SOverlay.h"
+#include "Widgets/Text/STextBlock.h"
+
+// Insights
+#include "Insights/Common/PaintUtils.h"
+#include "Insights/Common/Stopwatch.h"
+#include "Insights/Common/TimeUtils.h"
+#include "Insights/InsightsManager.h"
+#include "Insights/TimingProfilerCommon.h"
+#include "Insights/TimingProfilerManager.h"
+#include "Insights/ViewModels/TimingViewDrawHelper.h"
+#include "Insights/Widgets/STimersView.h"
+#include "Insights/Widgets/STimingProfilerWindow.h"
+
+#include <limits>
+
+////////////////////////////////////////////////////////////////////////////////////////////////////
+
+#define DEBUG_TIMING_TRACK 0
+#define LOCTEXT_NAMESPACE "STimingView"
+
+static const float MinTooltipWidth = 110.0f;
+
+////////////////////////////////////////////////////////////////////////////////////////////////////
+
+STimingView::STimingView()
+	: TimeRulerTrack(0xFFFF0001)
+	, MarkersTrack(0xFFFF0002)
+	, GraphTrack(0xFFFF0003)
+	, TooltipWidth(MinTooltipWidth)
+	, TooltipAlpha(0.0f)
+	, WhiteBrush(FCoreStyle::Get().GetBrush("WhiteBrush"))
+	, MainFont(FCoreStyle::GetDefaultFontStyle("Regular", 8))
+{
+	Reset();
+}
+
+////////////////////////////////////////////////////////////////////////////////////////////////////
+
+STimingView::~STimingView()
+{
+	for (const auto& KV : CachedTimelines)
+	{
+		delete KV.Value;
+	}
+	CachedTimelines.Reset();
+}
+
+////////////////////////////////////////////////////////////////////////////////////////////////////
+
+void STimingView::Construct(const FArguments& InArgs)
+{
+	OnSelectionChanged = InArgs._OnSelectionChanged;
+	OnHoveredEventChanged = InArgs._OnHoveredEventChanged;
+	SelectedEventChanged = InArgs._SelectedEventChanged;
+
+	ChildSlot
+	[
+		SNew(SOverlay)
+		.Visibility(EVisibility::SelfHitTestInvisible)
+
+		+ SOverlay::Slot()
+		.VAlign(VAlign_Bottom)
+		.Padding(FMargin(0, 0, 12, 0))
+		[
+			SAssignNew(HorizontalScrollBar, SScrollBar)
+			.Orientation(Orient_Horizontal)
+			.AlwaysShowScrollbar(false)
+			.Visibility(EVisibility::Visible)
+			.Thickness(FVector2D(8.0f, 8.0f))
+			//.RenderOpacity(0.75)
+			.OnUserScrolled(this, &STimingView::HorizontalScrollBar_OnUserScrolled)
+		]
+
+		+ SOverlay::Slot()
+		.HAlign(HAlign_Right)
+		.Padding(FMargin(0, 22, 0, 12))
+		[
+			SAssignNew(VerticalScrollBar, SScrollBar)
+			.Orientation(Orient_Vertical)
+			.AlwaysShowScrollbar(false)
+			.Visibility(EVisibility::Visible)
+			.Thickness(FVector2D(8.0f, 8.0f))
+			//.RenderOpacity(0.75)
+			.OnUserScrolled(this, &STimingView::VerticalScrollBar_OnUserScrolled)
+		]
+	];
+
+	UpdateHorizontalScrollBar();
+	UpdateVerticalScrollBar();
+
+	BindCommands();
+}
+
+////////////////////////////////////////////////////////////////////////////////////////////////////
+
+void STimingView::Reset()
+{
+	bDrawTimingEvents = true;
+	bDrawIoEvents = false;
+
+	//////////////////////////////////////////////////
+
+	Viewport.Reset();
+
+	bIsViewportDirty = true;
+	bIsVerticalViewportDirty = true;
+
+	//////////////////////////////////////////////////
+
+	for (const auto& KV : CachedTimelines)
+	{
+		delete KV.Value;
+	}
+	CachedTimelines.Reset();
+
+	//TopTracks.Reset();
+	//ScrollableTracks.Reset();
+	//BottomTracks.Reset();
+
+	//////////////////////////////////////////////////
+
+	//TimingEventsTracks.Reset();
+
+	bAreTimingEventsTracksDirty = true;
+
+	bUseDownSampling = true;
+
+	//////////////////////////////////////////////////
+
+	bForceIoEventsUpdate = true;
+	bMergeIoLanes = true;
+	AllIoEvents.Reset();
+	//AllIoEvents.Reserve(1000000);
+
+	//////////////////////////////////////////////////
+
+	TimeRulerTrack.Reset();
+	MarkersTrack.Reset();
+
+	GraphTrack.Reset();
+	GraphTrack.SetVisibilityFlag(false);
+
+	//////////////////////////////////////////////////
+
+	MousePosition = FVector2D::ZeroVector;
+
+	MousePositionOnButtonDown = FVector2D::ZeroVector;
+	ViewportStartTimeOnButtonDown = 0.0;
+	ViewportScrollPosYOnButtonDown = 0.0f;
+
+	MousePositionOnButtonUp = FVector2D::ZeroVector;
+
+	bIsLMB_Pressed = false;
+	bIsRMB_Pressed = false;
+
+	bIsSpaceBarKeyPressed = false;
+	bIsDragging = false;
+
+	bIsPanning = false;
+	PanningMode = EPanningMode::None;
+
+	bIsSelecting = false;
+
+	SelectionStartTime = 0.0;
+	SelectionEndTime = 0.0;
+
+	HoveredTimingEvent.Reset();
+	SelectedTimingEvent.Reset();
+
+	LastSelectionType = ESelectionType::None;
+
+	TimeMarker = std::numeric_limits<double>::infinity();
+
+	//ThisGeometry
+
+	Layout.ForceNormalMode();
+
+	//////////////////////////////////////////////////
+
+	NumUpdatedEvents = 0;
+	TimelineCacheUpdateDurationHistory.Reset();
+	TimelineCacheUpdateDurationHistory.AddValue(0);
+	TimeMarkerCacheUpdateDurationHistory.Reset();
+	TimeMarkerCacheUpdateDurationHistory.AddValue(0);
+	DrawDurationHistory.Reset();
+	OnPaintDurationHistory.Reset();
+	LastOnPaintTime = FPlatformTime::Cycles64();
+}
+
+////////////////////////////////////////////////////////////////////////////////////////////////////
+
+void STimingView::Tick(const FGeometry& AllottedGeometry, const double InCurrentTime, const float InDeltaTime)
+{
+	ThisGeometry = AllottedGeometry;
+
+	const float ViewWidth = AllottedGeometry.GetLocalSize().X;
+	const float ViewHeight = AllottedGeometry.GetLocalSize().Y;
+	if (Viewport.UpdateSize(ViewWidth, ViewHeight))
+	{
+		bIsViewportDirty = true;
+	}
+
+	//////////////////////////////////////////////////
+	// Update Y postion of non-scrollable tracks (docked on top).
+	// Also compute total height (TopOffset) of top docked tracks.
+
+	float TopOffset = 0.0f;
+	if (TimeRulerTrack.IsVisible())
+	{
+		TimeRulerTrack.SetPosY(TopOffset);
+		TopOffset += TimeRulerTrack.GetHeight();
+	}
+	if (MarkersTrack.IsVisible())
+	{
+		MarkersTrack.SetPosY(TopOffset);
+		TopOffset += MarkersTrack.GetHeight();
+	}
+	if (GraphTrack.IsVisible())
+	{
+		GraphTrack.SetPosY(TopOffset);
+		GraphTrack.SetHeight(200.0f);
+		TopOffset += GraphTrack.GetHeight();
+	}
+	Viewport.TopOffset = TopOffset;
+
+	//////////////////////////////////////////////////
+
+	if (!bIsPanning)
+	{
+		//////////////////////////////////////////////////
+		// Elastic snap to vertical scroll limits.
+
+		const float MinY = 0.0f;// -0.5f * Viewport.Height;
+		const float MaxY = Viewport.ScrollHeight - Viewport.Height + TopOffset + 7.0f;
+		const float U = 0.5f;
+
+		float ScrollPosY = Viewport.ScrollPosY;
+		if (ScrollPosY < MinY)
+		{
+			ScrollPosY = ScrollPosY * U + (1.0f - U) * MinY;
+			if (FMath::IsNearlyEqual(ScrollPosY, MinY, 0.5f))
+			{
+				ScrollPosY = MinY;
+			}
+		}
+		else if (ScrollPosY > MaxY)
+		{
+			ScrollPosY = ScrollPosY * U + (1.0f - U) * MaxY;
+			if (FMath::IsNearlyEqual(ScrollPosY, MaxY, 0.5f))
+			{
+				ScrollPosY = MaxY;
+			}
+			if (ScrollPosY < MinY)
+			{
+				ScrollPosY = MinY;
+			}
+		}
+
+		if (Viewport.ScrollPosY != ScrollPosY)
+		{
+			Viewport.ScrollPosY = ScrollPosY;
+			UpdateVerticalScrollBar();
+			bIsVerticalViewportDirty = true;
+		}
+
+		//////////////////////////////////////////////////
+		// Elastic snap to horizontal time limits.
+
+		if (Viewport.EnforceHorizontalScrollLimits(0.5)) // 0.5 is the interpolation factor
+		{
+			UpdateHorizontalScrollBar();
+			bIsViewportDirty = true;
+		}
+	}
+
+	// Check if TimersView needs to update the list of timers.
+	TSharedPtr<STimingProfilerWindow> Wnd = FTimingProfilerManager::Get()->GetProfilerWindow();
+	if (Wnd && Wnd->TimersView)
+	{
+		Wnd->TimersView->RebuildTree(false);
+	}
+
+	// Check if horizontal scroll area has changed.
+	TSharedPtr<const Trace::IAnalysisSession> Session = FInsightsManager::Get()->GetSession();
+	if (Session)
+	{
+		double SessionTime = 0.0;
+		{
+			Trace::FAnalysisSessionReadScope _(*Session.Get());
+			SessionTime = Session->GetDurationSeconds();
+		}
+		if (SessionTime > Viewport.MaxValidTime &&
+			SessionTime != DBL_MAX &&
+			SessionTime != std::numeric_limits<double>::infinity())
+		{
+			//UE_LOG(TimingProfiler, Log, TEXT("Session Duration: %g"), DT);
+			Viewport.MaxValidTime = SessionTime;
+			UpdateHorizontalScrollBar();
+			//bIsMaxValidTimeDirty = true;
+
+			if (SessionTime >= Viewport.StartTime && SessionTime <= Viewport.EndTime)
+			{
+				bAreTimingEventsTracksDirty = true;
+				MarkersTrack.SetDirtyFlag();
+			}
+		}
+	}
+
+	// Compute total height of visible tracks.
+	float TotalScrollHeight = 0.0f;
+	for (auto& KV : CachedTimelines)
+	{
+		FTimingEventsTrack& Track = *KV.Value;
+
+		//if (Track.IsVisible)
+		{
+			//Track.Y = TotalHeight; // TODO: this is not the draw order!
+			TotalScrollHeight += Track.GetHeight();
+		}
+
+		// Reset track's hovered/selected state.
+		Track.SetHoveredState(false);
+		Track.SetSelectedFlag(false);
+	}
+	TotalScrollHeight += 1.0f; // allow 1 pixel at the bottom (for last horizontal line)
+
+	// Set hovered/selected state for actual hovered/selected tracks, if any.
+	if (HoveredTimingEvent.Track != nullptr)
+	{
+		const_cast<FTimingEventsTrack*>(HoveredTimingEvent.Track)->SetHoveredState(true);
+	}
+	if (SelectedTimingEvent.Track != nullptr)
+	{
+		const_cast<FTimingEventsTrack*>(SelectedTimingEvent.Track)->SetSelectedFlag(true);
+	}
+
+	//TimeRulerTrack.Tick(InCurrentTime, InDeltaTime);
+	MarkersTrack.Tick(InCurrentTime, InDeltaTime);
+
+	// Check if vertical scroll area has changed.
+	if (TotalScrollHeight != Viewport.ScrollHeight)
+	{
+		Viewport.ScrollHeight = TotalScrollHeight;
+		UpdateVerticalScrollBar();
+		bIsVerticalViewportDirty = true;
+	}
+
+	// Update layout change animation (i.e. compact mode <-> normal mode).
+	Layout.Update();
+
+	//////////////////////////////////////////////////
+	// Update dirty tracks.
+
+	if (bIsViewportDirty)
+	{
+		bIsViewportDirty = false;
+
+		//TimeRuler.SetDirtyFlag();
+		MarkersTrack.SetDirtyFlag();
+		bAreTimingEventsTracksDirty = true;
+	}
+
+	if (bIsVerticalViewportDirty)
+	{
+		bIsVerticalViewportDirty = false;
+
+		bAreTimingEventsTracksDirty = true;
+	}
+
+	//if (TimeRuler.IsVisible() && TimeRuler.IsDirty())
+	//{
+	//	TimeRuler.ClearDirtyFlag();
+	//
+	//	TimeRuler.OnViewportChanged(Viewport);
+	//}
+
+	if (MarkersTrack.IsVisible() && MarkersTrack.IsDirty())
+	{
+		MarkersTrack.ClearDirtyFlag();
+
+		FStopwatch Stopwatch;
+		Stopwatch.Start();
+
+		MarkersTrack.Update(Viewport);
+
+		Stopwatch.Stop();
+		TimeMarkerCacheUpdateDurationHistory.AddValue(Stopwatch.AccumulatedTime);
+	}
+
+	if (bAreTimingEventsTracksDirty)
+	{
+		bAreTimingEventsTracksDirty = false;
+
+		FStopwatch Stopwatch;
+		Stopwatch.Start();
+
+		if (GraphTrack.IsVisible())
+		{
+			GraphTrack.Update(Viewport);
+		}
+
+		//for (int TrackIndex = 0; TrackIndex < TimingEventsTracks.Num(); ++TrackIndex)
+		//{
+		//	TimingEventsTracks[TrackIndex]->Update(Viewport);
+		//}
+
+		Stopwatch.Stop();
+		TimelineCacheUpdateDurationHistory.AddValue(Stopwatch.AccumulatedTime);
+	}
+}
+
+////////////////////////////////////////////////////////////////////////////////////////////////////
+
+int32 STimingView::OnPaint(const FPaintArgs& Args, const FGeometry& AllottedGeometry, const FSlateRect& MyCullingRect, FSlateWindowElementList& OutDrawElements, int32 LayerId, const FWidgetStyle& InWidgetStyle, bool bParentEnabled) const
+{
+	const bool bEnabled = ShouldBeEnabled(bParentEnabled);
+	const ESlateDrawEffect DrawEffects = bEnabled ? ESlateDrawEffect::None : ESlateDrawEffect::DisabledEffect;
+	FDrawContext DC(AllottedGeometry, MyCullingRect, InWidgetStyle, DrawEffects, OutDrawElements, LayerId);
+
+	//OutDrawElements.GetRootDrawLayer().DrawElements.Reserve(50000);
+
+	const TSharedRef<FSlateFontMeasure> FontMeasureService = FSlateApplication::Get().GetRenderer()->GetFontMeasureService();
+
+	const float ViewWidth = AllottedGeometry.GetLocalSize().X;
+	const float ViewHeight = AllottedGeometry.GetLocalSize().Y;
+
+	/*
+	// Warming up Slate vertex/index buffers to avoid initial freezes due to multiple resizes of those buffers.
+	static bool bWarmingUp = true;
+	if (bWarmingUp)
+	{
+		bWarmingUp = false;
+
+		FRandomStream RandomStream(0);
+		const int Count = 1'000'000;
+		for (int Index = 0; Index < Count; ++Index)
+		{
+			float X = ViewWidth * RandomStream.GetFraction();
+			float Y = ViewHeight * RandomStream.GetFraction();
+			FLinearColor Color(RandomStream.GetFraction(), RandomStream.GetFraction(), RandomStream.GetFraction(), 1.0f);
+			OnPaintState.DrawBox(X, Y, 1.0f, 1.0f, WhiteBrush, Color);
+		}
+		LayerId++;
+		LayerId++;
+	}
+	*/
+
+	//////////////////////////////////////////////////
+
+	FStopwatch Stopwatch;
+	Stopwatch.Start();
+
+	FTimingViewDrawHelper Helper(DC, Viewport, Layout);
+
+	Helper.Begin();
+
+	// Draw background.
+	Helper.DrawBackground();
+
+	// Draw scrollable tracks.
+	{
+		Helper.BeginTimelines();
+
+		if (bDrawIoEvents)
+		{
+			DrawIoOverviewTrack(Helper);
+		}
+
+		if (bDrawTimingEvents)
+		{
+			DrawTimingEvents(Helper);
+		}
+
+		if (bDrawIoEvents)
+		{
+			DrawIoTrack(Helper);
+		}
+
+		Helper.EndTimelines();
+	}
+
+	if (GraphTrack.IsVisible())
+	{
+		GraphTrack.Draw(DC, Viewport);
+	}
+
+	if (!FTimingEvent::AreEquals(SelectedTimingEvent, HoveredTimingEvent))
+	{
+		// Highlight the selected timing event (if any).
+		if (SelectedTimingEvent.IsValid())
+		{
+			const float Y = Viewport.GetViewportY(SelectedTimingEvent.Track->GetPosY()) + Layout.GetLaneY(SelectedTimingEvent.Depth);
+			Helper.DrawTimingEventHighlight(SelectedTimingEvent.StartTime, SelectedTimingEvent.EndTime, Y, FTimingViewDrawHelper::EHighlightMode::Selected);
+		}
+
+		// Highlight the hovered timing event (if any).
+		if (HoveredTimingEvent.IsValid())
+		{
+			const float Y = Viewport.GetViewportY(HoveredTimingEvent.Track->GetPosY()) + Layout.GetLaneY(HoveredTimingEvent.Depth);
+			Helper.DrawTimingEventHighlight(HoveredTimingEvent.StartTime, HoveredTimingEvent.EndTime, Y, FTimingViewDrawHelper::EHighlightMode::Hovered);
+		}
+	}
+	else
+	{
+		// Highlight the selected and hovered timing event (if any).
+		if (SelectedTimingEvent.IsValid())
+		{
+			const float Y = Viewport.GetViewportY(SelectedTimingEvent.Track->GetPosY()) + Layout.GetLaneY(SelectedTimingEvent.Depth);
+			Helper.DrawTimingEventHighlight(SelectedTimingEvent.StartTime, SelectedTimingEvent.EndTime, Y, FTimingViewDrawHelper::EHighlightMode::SelectedAndHovered);
+		}
+	}
+
+	if (MarkersTrack.IsVisible())
+	{
+		// Draw the time markers (vertical lines + category & message texts).
+		MarkersTrack.Draw(DC, Viewport);
+	}
+
+	Helper.End();
+
+	//////////////////////////////////////////////////
+
+	// Draw the time ruler.
+	if (TimeRulerTrack.IsVisible())
+	{
+		TimeRulerTrack.Draw(DC, Viewport, MousePosition, bIsSelecting, SelectionStartTime, SelectionEndTime);
+	}
+
+	//////////////////////////////////////////////////
+
+	// Draw the time range selection.
+	DrawTimeRangeSelection(DC);
+
+	//////////////////////////////////////////////////
+
+	// Draw the time marker (orange vertical line).
+	//DrawTimeMarker(OnPaintState);
+	float TimeMarkerX = Viewport.TimeToSlateUnitsRounded(TimeMarker);
+	if (TimeMarkerX >= 0.0f && TimeMarkerX < Viewport.Width)
+	{
+		DC.DrawBox(TimeMarkerX, 0.0f, 1.0f, Viewport.Height, WhiteBrush, FLinearColor(0.85f, 0.5f, 0.03f, 0.5f));
+		DC.LayerId++;
+	}
+
+	//////////////////////////////////////////////////
+
+	Stopwatch.Stop();
+	DrawDurationHistory.AddValue(Stopwatch.AccumulatedTime);
+
+	//////////////////////////////////////////////////
+
+	const bool bShouldDisplayDebugInfo = FInsightsManager::Get()->IsDebugInfoEnabled();
+	if (bShouldDisplayDebugInfo)
+	{
+		const FSlateFontInfo& SummaryFont = MainFont;
+
+		const float MaxFontCharHeight = FontMeasureService->Measure(TEXT("!"), SummaryFont).Y;
+		const float DbgDY = MaxFontCharHeight;
+
+		const float DbgW = 320.0f;
+		const float DbgH = DbgDY * 9 + 3.0f;
+		const float DbgX = ViewWidth - DbgW - 20.0f;
+		float DbgY = Viewport.TopOffset + 10.0f;
+
+		DC.LayerId++;
+
+		DC.DrawBox(DbgX - 2.0f, DbgY - 2.0f, DbgW, DbgH, WhiteBrush, FLinearColor(1.0f, 1.0f, 1.0f, 0.9f));
+		DC.LayerId++;
+
+		FLinearColor DbgTextColor(0.0f, 0.0f, 0.0f, 0.9f);
+
+		//////////////////////////////////////////////////
+		// Display the "Draw" performance info.
+
+		// Time interval since last OnPaint call.
+		const uint64 CurrentTime = FPlatformTime::Cycles64();
+		const uint64 OnPaintDuration = CurrentTime - LastOnPaintTime;
+		LastOnPaintTime = CurrentTime;
+		OnPaintDurationHistory.AddValue(OnPaintDuration); // saved for last 32 OnPaint calls
+		const uint64 AvgOnPaintDuration = OnPaintDurationHistory.ComputeAverage();
+		const uint64 AvgOnPaintDurationMs = FStopwatch::Cycles64ToMilliseconds(AvgOnPaintDuration);
+		const double AvgOnPaintFps = AvgOnPaintDurationMs != 0 ? 1.0 / FStopwatch::Cycles64ToSeconds(AvgOnPaintDuration) : 0.0;
+
+		const uint64 AvgDrawDurationMs = FStopwatch::Cycles64ToMilliseconds(DrawDurationHistory.ComputeAverage());
+
+		DC.DrawText
+		(
+			DbgX, DbgY,
+			FString::Printf(TEXT("D: %llu ms + %llu ms = %llu ms (%d fps)"),
+				AvgDrawDurationMs, // drawing time
+				AvgOnPaintDurationMs - AvgDrawDurationMs, // other overhead to OnPaint calls
+				AvgOnPaintDurationMs, // average time between two OnPaint calls
+				FMath::RoundToInt(AvgOnPaintFps)), // framerate of OnPaint calls
+			SummaryFont, DbgTextColor
+		);
+		DbgY += DbgDY;
+
+		//////////////////////////////////////////////////
+		// Display the "timelines cache update" info.
+
+		const uint64 AvgTimelineUpdateDurationMs = FStopwatch::Cycles64ToMilliseconds(TimelineCacheUpdateDurationHistory.ComputeAverage());
+		const uint64 LastTimelineUpdateDurationMs = FStopwatch::Cycles64ToMilliseconds(TimelineCacheUpdateDurationHistory.GetValue(0));
+
+		DC.DrawText
+		(
+			DbgX, DbgY,
+			FString::Printf(TEXT("U1 avg: %llu ms, last: %llu ms"),
+				AvgTimelineUpdateDurationMs, // time to update the Timeline cache
+				LastTimelineUpdateDurationMs), // last time to update the Timeline cache
+			SummaryFont, DbgTextColor
+		);
+		DbgY += DbgDY;
+
+		//////////////////////////////////////////////////
+		// Display the "time markers cache update" info.
+
+		const uint64 AvgTimeMarkerUpdateDurationMs = FStopwatch::Cycles64ToMilliseconds(TimeMarkerCacheUpdateDurationHistory.ComputeAverage());
+		const uint64 LastTimeMarkerUpdateDurationMs = FStopwatch::Cycles64ToMilliseconds(TimeMarkerCacheUpdateDurationHistory.GetValue(0));
+
+		DC.DrawText
+		(
+			DbgX, DbgY,
+			FString::Printf(TEXT("U2 avg: %llu ms, last: %llu ms"),
+				AvgTimeMarkerUpdateDurationMs, // time to update the TimeMarker cache
+				LastTimeMarkerUpdateDurationMs), // last time to update the Timeline cache
+			SummaryFont, DbgTextColor
+		);
+		DbgY += DbgDY;
+
+		//////////////////////////////////////////////////
+		// Display timing events stats.
+
+		DC.DrawText
+		(
+			DbgX, DbgY,
+			FString::Format(TEXT("{0} events : {1} ({2}) boxes, {3} borders, {4} texts"),
+			{
+				FText::AsNumber(Helper.NumEvents).ToString(),
+				FText::AsNumber(Helper.NumDrawBoxes).ToString(),
+				FText::AsPercent((double)Helper.NumDrawBoxes / (Helper.NumDrawBoxes + Helper.NumMergedBoxes)).ToString(),
+				FText::AsNumber(Helper.NumDrawBorders).ToString(),
+				FText::AsNumber(Helper.NumDrawTexts).ToString(),
+				//OutDrawElements.GetRootDrawLayer().GetElementCount(),
+			}),
+			SummaryFont, DbgTextColor
+		);
+		DbgY += DbgDY;
+
+		//////////////////////////////////////////////////
+		// Display time markers stats.
+
+		DC.DrawText
+		(
+			DbgX, DbgY,
+			FString::Format(TEXT("{0}{1} logs : {2} boxes, {3} texts"),
+			{
+				MarkersTrack.IsVisible() ? TEXT("") : TEXT("*"),
+				FText::AsNumber(MarkersTrack.GetNumLogMessages()).ToString(),
+				FText::AsNumber(MarkersTrack.GetNumBoxes()).ToString(),
+				FText::AsNumber(MarkersTrack.GetNumTexts()).ToString(),
+			}),
+			SummaryFont, DbgTextColor
+		);
+		DbgY += DbgDY;
+
+		//////////////////////////////////////////////////
+		// Display Graph track stats.
+
+		DC.DrawText
+		(
+			DbgX, DbgY,
+			FString::Format(TEXT("{0} events : {1} points, {2} lines, {3} boxes"),
+				{
+					FText::AsNumber(GraphTrack.GetNumAddedEvents()).ToString(),
+					FText::AsNumber(GraphTrack.GetNumDrawPoints()).ToString(),
+					FText::AsNumber(GraphTrack.GetNumDrawLines()).ToString(),
+					FText::AsNumber(GraphTrack.GetNumDrawBoxes()).ToString(),
+				}),
+				SummaryFont, DbgTextColor
+				);
+		DbgY += DbgDY;
+
+		//////////////////////////////////////////////////
+		// Display viewport's horizontal info.
+
+		DC.DrawText
+		(
+			DbgX, DbgY,
+			FString::Printf(TEXT("SX: %g, ST: %g, ET: %s"), Viewport.ScaleX, Viewport.StartTime, *TimeUtils::FormatTimeAuto(Viewport.MaxValidTime)),
+			SummaryFont, DbgTextColor
+		);
+		DbgY += DbgDY;
+
+		//////////////////////////////////////////////////
+		// Display viewport's vertical info.
+
+		DC.DrawText
+		(
+			DbgX, DbgY,
+			FString::Printf(TEXT("Y: %.2f, H: %g, VH: %g"), Viewport.ScrollPosY, Viewport.ScrollHeight, Viewport.Height),
+			SummaryFont, DbgTextColor
+		);
+		DbgY += DbgDY;
+
+		//////////////////////////////////////////////////
+		// Display input related debug info.
+
+		FString InputStr = FString::Printf(TEXT("(%.0f, %.0f)"), MousePosition.X, MousePosition.Y);
+		if (bIsSpaceBarKeyPressed) InputStr += " Space";
+		if (bIsLMB_Pressed) InputStr += " LMB";
+		if (bIsRMB_Pressed) InputStr += " RMB";
+		if (bIsPanning) InputStr += " Panning";
+		if (bIsSelecting) InputStr += " Selecting";
+		if (bIsDragging) InputStr += " Dragging";
+		DC.DrawText(DbgX, DbgY, InputStr, SummaryFont, DbgTextColor);
+		DbgY += DbgDY;
+	}
+
+	//////////////////////////////////////////////////
+
+	{
+		// Draw info about selected event (bottom-right corner).
+		if (SelectedTimingEvent.IsValid())
+		{
+			const FTimerNodePtr TimerNodePtr = GetTimerNode(SelectedTimingEvent.TypeId);
+
+			FString Str = FString::Printf(TEXT("%s (Incl.: %s, Excl.: %s)"),
+				TimerNodePtr ? *(TimerNodePtr->GetName().ToString()) : TEXT("N/A"),
+				*TimeUtils::FormatTimeAuto(SelectedTimingEvent.Duration()),
+				*TimeUtils::FormatTimeAuto(SelectedTimingEvent.ExclusiveTime));
+
+			const FVector2D Size = FontMeasureService->Measure(Str, MainFont);
+			const float X = Viewport.Width - Size.X - 23.0f;
+			const float Y = Viewport.Height - Size.Y - 18.0f;
+
+			const FLinearColor BackgroundColor(0.05f, 0.05f, 0.05f, 1.0f);
+			const FLinearColor TextColor(0.7f, 0.7f, 0.7f, 1.0f);
+
+			DC.DrawBox(X - 8.0f, Y - 2.0f, Size.X + 16.0f, Size.Y + 4.0f, WhiteBrush, BackgroundColor);
+			DC.LayerId++;
+
+			DC.DrawText(X, Y, Str, MainFont, TextColor);
+			DC.LayerId++;
+		}
+
+		// Draw info about hovered event (like a tooltip at mouse position).
+		if (HoveredTimingEvent.IsValid() && !MousePosition.IsZero())
+		{
+			const FTimerNodePtr TimerNodePtr = GetTimerNode(HoveredTimingEvent.TypeId);
+			FString Name = TimerNodePtr ? TimerNodePtr->GetName().ToString() : TEXT("N/A");
+
+			float W = FontMeasureService->Measure(Name, MainFont).X;
+			if (W < MinTooltipWidth)
+			{
+				W = MinTooltipWidth;
+			}
+			if (TooltipWidth != W)
+			{
+				TooltipWidth = TooltipWidth * 0.75f + W * 0.25f;
+			}
+			const float MaxX = FMath::Max(0.0f, Viewport.Width - TooltipWidth - 21.0f);
+			const float X = FMath::Clamp<float>(MousePosition.X + 18.0f, 0.0f, MaxX);
+
+			const float LineH = 14.0f;
+			const float H = 4 * LineH;
+			const float MaxY = FMath::Max(0.0f, Viewport.Height - H - 18.0f);
+			float Y = FMath::Clamp<float>(MousePosition.Y + 18.0f, 0.0f, MaxY);
+
+			const float Alpha = 1.0f - FMath::Abs(TooltipWidth - W) / W;
+			if (TooltipAlpha < Alpha)
+			{
+				TooltipAlpha = TooltipAlpha * 0.9f + Alpha * 0.1f;
+			}
+			else
+			{
+				TooltipAlpha = Alpha;
+			}
+
+			const FLinearColor BackgroundColor(0.05f, 0.05f, 0.05f, TooltipAlpha);
+			const FLinearColor NameColor(0.9f, 0.9f, 0.5f, TooltipAlpha);
+			//const FLinearColor NameColor(0.85f, 0.5f, 0.03f, TooltipAlpha);
+			const FLinearColor TextColor(0.6f, 0.6f, 0.6f, TooltipAlpha);
+			const FLinearColor ValueColor(1.0f, 1.0f, 1.0f, TooltipAlpha);
+
+			DC.DrawBox(X - 6.0f, Y - 3.0f, TooltipWidth + 12.0f, H + 6.0f, WhiteBrush, BackgroundColor);
+			DC.LayerId++;
+
+			DC.DrawText(X, Y, Name, MainFont, NameColor);
+			Y += LineH;
+
+			const float ValueX = X + 58.0f;
+
+			DC.DrawText(X + 3.0f, Y, TEXT("Incl. Time:"), MainFont, TextColor);
+			FString InclStr = TimeUtils::FormatTimeAuto(HoveredTimingEvent.Duration());
+			DC.DrawText(ValueX, Y, InclStr, MainFont, ValueColor);
+			Y += LineH;
+
+			DC.DrawText(X, Y, TEXT("Excl. Time:"), MainFont, TextColor);
+			FString ExclStr = TimeUtils::FormatTimeAuto(HoveredTimingEvent.ExclusiveTime);
+			DC.DrawText(ValueX, Y, ExclStr, MainFont, ValueColor);
+			Y += LineH;
+
+			DC.DrawText(X + 24.0f, Y, TEXT("Depth:"), MainFont, TextColor);
+			FString DepthStr = FString::Printf(TEXT("%d"), HoveredTimingEvent.Depth);
+			DC.DrawText(ValueX, Y, DepthStr, MainFont, ValueColor);
+			Y += LineH;
+
+			DC.LayerId++;
+		}
+		else
+		{
+			TooltipWidth = MinTooltipWidth;
+			TooltipAlpha = 0.0f;
+		}
+	}
+
+	return SCompoundWidget::OnPaint(Args, AllottedGeometry, MyCullingRect, OutDrawElements, LayerId, InWidgetStyle, bParentEnabled && IsEnabled());
+}
+
+////////////////////////////////////////////////////////////////////////////////////////////////////
+
+FTimingEventsTrack* STimingView::GetOrAddTimingEventsTrack(uint32 TrackId)
+{
+	if (CachedTimelines.Contains(TrackId))
+	{
+		return CachedTimelines[TrackId];
+	}
+	else
+	{
+		FTimingEventsTrack* Track = new FTimingEventsTrack(TrackId);
+		CachedTimelines.Add(TrackId, Track);
+		return Track;
+	}
+}
+
+////////////////////////////////////////////////////////////////////////////////////////////////////
+
+void STimingView::DrawTimingEvents(FTimingViewDrawHelper& DC) const
+{
+	TSharedPtr<const Trace::IAnalysisSession> Session = FInsightsManager::Get()->GetSession();
+	if (Session.IsValid())
+	{
+		Trace::FAnalysisSessionReadScope _(*Session.Get());
+
+		struct FTimelineHelper
+		{
+			void DrawTimeline(uint32 TimelineIndex, const TCHAR* Name)
+			{
+				// TODO: Adding tracks should be done in Tick(), not in OnPaint().
+				FTimingEventsTrack& Track = *const_cast<STimingView&>(Outer).GetOrAddTimingEventsTrack(TimelineIndex);
+
+				if (DC.BeginTimeline(Track))
+				{
+					TimingProfilerProvider.ReadTimeline(TimelineIndex, [this](const Trace::ITimingProfilerProvider::Timeline& Timeline)
+					{
+						if (Outer.bUseDownSampling)
+						{
+							const double SecondsPerPixel = 1.0 / DC.GetViewport().ScaleX;
+							Timeline.EnumerateEventsDownSampled(DC.GetViewport().StartTime, DC.GetViewport().EndTime, SecondsPerPixel, [this](double StartTime, double EndTime, uint32 Depth, const Trace::FTimingProfilerEvent& Event)
+							{
+								DC.AddEvent(StartTime, EndTime, Depth, *Timers[Event.TimerIndex].Name);
+							});
+						}
+						else
+						{
+							Timeline.EnumerateEvents(DC.GetViewport().StartTime, DC.GetViewport().EndTime, [this](double StartTime, double EndTime, uint32 Depth, const Trace::FTimingProfilerEvent& Event)
+							{
+								DC.AddEvent(StartTime, EndTime, Depth, *Timers[Event.TimerIndex].Name);
+							});
+						}
+					});
+					DC.EndTimeline(Track, Name);
+				}
+			}
+
+			const STimingView& Outer;
+			FTimingViewDrawHelper& DC;
+			const Trace::ITimingProfilerProvider& TimingProfilerProvider;
+			const Trace::FTimingProfilerTimer* Timers;
+		};
+
+		Session->ReadTimingProfilerProvider([this, &DC, &Session](const Trace::ITimingProfilerProvider& TimingProfilerProvider)
+		{
+			TimingProfilerProvider.ReadTimers([this, &DC, &Session, &TimingProfilerProvider](const Trace::FTimingProfilerTimer* Timers, uint64 TimersCount)
+			{
+				FTimelineHelper Helper{ *this, DC, TimingProfilerProvider, Timers };
+
+				// Draw GPU track.
+				uint32 GpuTimelineIndex;
+				if (TimingProfilerProvider.GetGpuTimelineIndex(GpuTimelineIndex))
+				{
+					Helper.DrawTimeline(GpuTimelineIndex, TEXT("GPU"));
+				}
+
+				// Draw CPU tracks.
+				Session->ReadThreadProvider([this, &Helper, &TimingProfilerProvider](const Trace::IThreadProvider& ThreadProvider)
+				{
+					ThreadProvider.EnumerateThreads([this, &Helper, &TimingProfilerProvider](const Trace::FThreadInfo& ThreadInfo)
+					{
+						uint32 ThreadTimelineIndex;
+						if (TimingProfilerProvider.GetCpuThreadTimelineIndex(ThreadInfo.Id, ThreadTimelineIndex))
+						{
+							Helper.DrawTimeline(ThreadTimelineIndex, ThreadInfo.Name && *ThreadInfo.Name ? ThreadInfo.Name : *FString::Printf(TEXT("Thread %u"), ThreadInfo.Id));
+						}
+					});
+				});
+			});
+		});
+	}
+}
+
+////////////////////////////////////////////////////////////////////////////////////////////////////
+
+const TCHAR* GetFileActivityTypeName(Trace::EFileActivityType Type)
+{
+	static_assert(Trace::FileActivityType_Open == 0, "Trace::EFileActivityType enum has changed!?");
+	static_assert(Trace::FileActivityType_Close == 1, "Trace::EFileActivityType enum has changed!?");
+	static_assert(Trace::FileActivityType_Read == 2, "Trace::EFileActivityType enum has changed!?");
+	static_assert(Trace::FileActivityType_Write == 3, "Trace::EFileActivityType enum has changed!?");
+	static_assert(Trace::FileActivityType_Count == 4, "Trace::EFileActivityType enum has changed!?");
+	static const TCHAR* GFileActivityTypeNames[] =
+	{
+		TEXT("Open"),
+		TEXT("Close"),
+		TEXT("Read"),
+		TEXT("Write"),
+		TEXT("Idle"), // virtual events added for cases where Close event is more than 1s away from last Open/Read/Write event.
+		TEXT("NotClosed") // virtual events added when an Open activity never closes
+	};
+	return GFileActivityTypeNames[Type];
+}
+
+////////////////////////////////////////////////////////////////////////////////////////////////////
+
+uint32 GetFileActivityTypeColor(Trace::EFileActivityType Type)
+{
+	static const uint32 GFileActivityTypeColors[] =
+	{
+		0xFFCCAA33, // open
+		0xFF33AACC, // close
+		0xFF33AA33, // read
+		0xFFDD33CC, // write
+		0x55333333, // idle
+		0x55553333, // close
+	};
+	return GFileActivityTypeColors[Type];
+}
+
+////////////////////////////////////////////////////////////////////////////////////////////////////
+
+void STimingView::DrawIoOverviewTrack(FTimingViewDrawHelper& DC) const
+{
+	// TODO: Adding tracks should be done in Tick(), not in OnPaint().
+	STimingView& NonConstThis = const_cast<STimingView&>(*this);
+	FTimingEventsTrack& Track = *NonConstThis.GetOrAddTimingEventsTrack((0xF << 28) | 1);
+	//TODO: FIoTimingEventsTrack& Track = *GetOrAddTimingEventsTrack(ETrackType::Io, 1);
+
+	if (DC.BeginTimeline(Track))
+	{
+		// Draw the IO Overiew track using cached events (as those are sorted by Start Time).
+		for (const FIoTimingEvent& Event : AllIoEvents)
+		{
+			const Trace::EFileActivityType ActivityType = static_cast<Trace::EFileActivityType>(Event.Type & 0x0F);
+
+			if (ActivityType >= Trace::FileActivityType_Count)
+			{
+				// Ignore "Idle" and "NotClosed" events.
+				continue;
+			}
+			if (Event.EndTime <= DC.GetViewport().StartTime)
+			{
+				continue;
+			}
+			if (Event.StartTime >= DC.GetViewport().EndTime)
+			{
+				break;
+			}
+
+			uint32 Color = GetFileActivityTypeColor(ActivityType);
+
+			const bool bHasFailed = ((Event.Type & 0xF0) != 0);
+
+			if (bHasFailed)
+			{
+				FString Name = TEXT("Failed ");
+				Name += GetFileActivityTypeName(ActivityType);
+				Name += " [";
+				Name += Event.Path;
+				Name += "]";
+				Color = 0xFFAA0000;
+				DC.AddEvent(Event.StartTime, Event.EndTime, 0, *Name, Color);
+			}
+			else if (ActivityType == Trace::FileActivityType_Open)
+			{
+				FString Name = GetFileActivityTypeName(ActivityType);
+				Name += " [";
+				Name += Event.Path;
+				Name += "]";
+				DC.AddEvent(Event.StartTime, Event.EndTime, 0, *Name, Color);
+			}
+			else
+			{
+				DC.AddEvent(Event.StartTime, Event.EndTime, 0, GetFileActivityTypeName(ActivityType), Color);
+			}
+		}
+
+		DC.EndTimeline(Track, TEXT("IO Overview"));
+	}
+}
+
+////////////////////////////////////////////////////////////////////////////////////////////////////
+//PRAGMA_DISABLE_OPTIMIZATION
+
+void STimingView::DrawIoTrack(FTimingViewDrawHelper& DC) const
+{
+	// TODO: Adding tracks should be done in Tick(), not in OnPaint().
+	STimingView& NonConstThis = const_cast<STimingView&>(*this);
+	FTimingEventsTrack& Track = *NonConstThis.GetOrAddTimingEventsTrack((0xF << 28) | 2);
+	//TODO: FIoTimingEventsTrack& Track = *GetOrAddTimingEventsTrack(ETrackType::Io, 2);
+
+	if (bForceIoEventsUpdate)
+	{
+		bForceIoEventsUpdate = false;
+
+		AllIoEvents.Reset();
+
+		TSharedPtr<const Trace::IAnalysisSession> Session = FInsightsManager::Get()->GetSession();
+		if (Session.IsValid())
+		{
+			Trace::FAnalysisSessionReadScope _(*Session.Get());
+
+			Session->ReadFileActivityProvider([this, &DC](const Trace::IFileActivityProvider& FileActivityProvider)
+			{
+				// Enumerate all IO events and cache them.
+				FileActivityProvider.EnumerateFileActivity([this, &DC](const Trace::FFileInfo& FileInfo, const Trace::IFileActivityProvider::Timeline& Timeline)
+				{
+					Timeline.EnumerateEvents(-std::numeric_limits<double>::infinity(), +std::numeric_limits<double>::infinity(),
+						[this, &DC, &FileInfo, &Timeline](double EventStartTime, double EventEndTime, uint32 EventDepth, const Trace::FFileActivity& FileActivity)
+					{
+						if (bMergeIoLanes)
+						{
+							EventDepth = static_cast<uint32>(FileInfo.Id); // used by MergeLanes algorithm
+						}
+						else
+						{
+							EventDepth = FileInfo.Id % 32; // simple layout
+						}
+						uint32 Type = ((uint32)FileActivity.ActivityType & 0x0F) | (FileActivity.Failed ? 0x80 : 0);
+						AllIoEvents.Add(FIoTimingEvent{ EventStartTime, EventEndTime, EventDepth, Type, FileInfo.Path });
+					});
+
+					return true;
+				});
+			});
+		}
+
+		// Sort cached IO events by Start Time.
+		AllIoEvents.Sort([](const FIoTimingEvent& A, const FIoTimingEvent& B) { return A.StartTime < B.StartTime; });
+
+		if (bMergeIoLanes)
+		{
+			struct FIoLane
+			{
+				uint32 FileId;
+				const TCHAR* Path;
+				double LastEndTime;
+				double EndTime;
+			};
+			TArray<FIoLane> Lanes;
+
+			TArray<FIoTimingEvent> IoEventsToAdd;
+
+			for (FIoTimingEvent& Event : AllIoEvents)
+			{
+				uint64 TimelineId = Event.Depth;
+
+				int Depth = -1;
+
+				bool bIsCloseEvent = false;
+
+				const Trace::EFileActivityType ActivityType = static_cast<Trace::EFileActivityType>(Event.Type & 0x0F);
+
+				if (ActivityType == Trace::FileActivityType_Open)
+				{
+					// Find lane (avoiding overlaps with other opened files).
+					for (int LaneIndex = 0; LaneIndex < Lanes.Num(); ++LaneIndex)
+					{
+						FIoLane& Lane = Lanes[LaneIndex];
+						if (Event.StartTime >= Lane.EndTime)
+						{
+							Depth = LaneIndex;
+							break;
+						}
+					}
+
+					if (Depth < 0)
+					{
+						// Add new lane.
+						Depth = Lanes.AddDefaulted();
+					}
+
+					const bool bHasFailed = ((Event.Type & 0xF0) != 0);
+
+					FIoLane& Lane = Lanes[Depth];
+					Lane.FileId = TimelineId;
+					Lane.Path = Event.Path;
+					Lane.LastEndTime = Event.EndTime;
+					Lane.EndTime = bHasFailed ? Event.EndTime : std::numeric_limits<double>::infinity();
+				}
+				else if (ActivityType == Trace::FileActivityType_Close)
+				{
+					bIsCloseEvent = true;
+
+					// Find lane with same id.
+					for (int LaneIndex = 0; LaneIndex < Lanes.Num(); ++LaneIndex)
+					{
+						FIoLane& Lane = Lanes[LaneIndex];
+						if (Lane.FileId == TimelineId)
+						{
+							// Adds an Idle event, but only if has passed at least 1s since last open/read/write activity.
+							if (Event.StartTime - Lane.LastEndTime > 1.0)
+							{
+								IoEventsToAdd.Add(FIoTimingEvent{ Lane.LastEndTime, Event.StartTime, static_cast<uint32>(LaneIndex), Trace::FileActivityType_Count, Event.Path });
+							}
+							Lane.LastEndTime = Event.EndTime;
+							Lane.EndTime = Event.EndTime;
+							Depth = LaneIndex;
+							break;
+						}
+					}
+					ensure(Depth >= 0);
+				}
+				else
+				{
+					// All other events should be inside the virtual Open-Close event.
+					// Find lane with same id.
+					for (int LaneIndex = 0; LaneIndex < Lanes.Num(); ++LaneIndex)
+					{
+						FIoLane& Lane = Lanes[LaneIndex];
+						if (Lane.FileId == TimelineId)
+						{
+							Lane.LastEndTime = Event.EndTime;
+							Depth = LaneIndex;
+							break;
+						}
+					}
+					ensure(Depth >= 0);
+				}
+
+				if (Depth < 0) // just in case
+				{
+					// Add new lane.
+					Depth = Lanes.AddDefaulted();
+					FIoLane& Lane = Lanes[Depth];
+					Lane.FileId = TimelineId;
+					Lane.Path = Event.Path;
+					Lane.LastEndTime = Event.EndTime;
+					Lane.EndTime = bIsCloseEvent ? Event.EndTime : std::numeric_limits<double>::infinity();
+				}
+
+				Event.Depth = Depth;
+			}
+
+			for (int LaneIndex = 0; LaneIndex < Lanes.Num(); ++LaneIndex)
+			{
+				FIoLane& Lane = Lanes[LaneIndex];
+				if (Lane.EndTime == std::numeric_limits<double>::infinity())
+				{
+					IoEventsToAdd.Add(FIoTimingEvent{ Lane.LastEndTime, Lane.EndTime, static_cast<uint32>(LaneIndex), Trace::FileActivityType_Count + 1, Lane.Path });
+				}
+			}
+
+			for (const FIoTimingEvent& Event : IoEventsToAdd)
+			{
+				AllIoEvents.Add(Event);
+			}
+
+			// Sort cached IO events one more time, also by Start Time.
+			AllIoEvents.Sort([](const FIoTimingEvent& A, const FIoTimingEvent& B) { return A.StartTime < B.StartTime; });
+
+			//// Sort cached IO events again, by Depth and then by Start Time.
+			//AllIoEvents.Sort([](const FIoTimingEvent& A, const FIoTimingEvent& B)
+			//{
+			//	return A.Depth == B.Depth ? A.StartTime < B.StartTime : A.Depth < B.Depth;
+			//});
+		}
+	}
+
+	if (DC.BeginTimeline(Track))
+	{
+		// Draw IO track using cached events.
+		for (const FIoTimingEvent& Event : AllIoEvents)
+		{
+			if (Event.EndTime <= DC.GetViewport().StartTime)
+			{
+				continue;
+			}
+			if (Event.StartTime >= DC.GetViewport().EndTime)
+			{
+				break;
+			}
+
+			const Trace::EFileActivityType ActivityType = static_cast<Trace::EFileActivityType>(Event.Type & 0x0F);
+
+			uint32 Color = GetFileActivityTypeColor(ActivityType);
+
+			if (ActivityType < Trace::FileActivityType_Count)
+			{
+				const bool bHasFailed = ((Event.Type & 0xF0) != 0);
+
+				if (bHasFailed)
+				{
+					FString Name = TEXT("Failed ");
+					Name += GetFileActivityTypeName(ActivityType);
+					Name += " [";
+					Name += Event.Path;
+					Name += "]";
+					Color = 0xFFAA0000;
+					DC.AddEvent(Event.StartTime, Event.EndTime, Event.Depth, *Name, Color);
+				}
+				else if (ActivityType == Trace::FileActivityType_Open)
+				{
+					FString Name = GetFileActivityTypeName(ActivityType);
+					Name += " [";
+					Name += Event.Path;
+					Name += "]";
+					DC.AddEvent(Event.StartTime, Event.EndTime, Event.Depth, *Name, Color);
+				}
+				else
+				{
+					DC.AddEvent(Event.StartTime, Event.EndTime, Event.Depth, GetFileActivityTypeName(ActivityType), Color);
+				}
+			}
+			else
+			{
+				FString Name = GetFileActivityTypeName(ActivityType);
+				Name += " [";
+				Name += Event.Path;
+				Name += "]";
+				DC.AddEvent(Event.StartTime, Event.EndTime, Event.Depth, *Name, Color);
+			}
+		}
+
+		DC.EndTimeline(Track, TEXT("IO Activity"));
+	}
+}
+
+//PRAGMA_ENABLE_OPTIMIZATION
+////////////////////////////////////////////////////////////////////////////////////////////////////
+
+void STimingView::DrawTimeRangeSelection(FDrawContext& DC) const
+{
+	if (SelectionEndTime > SelectionStartTime)
+	{
+		float SelectionX1 = Viewport.TimeToSlateUnitsRounded(SelectionStartTime);
+		float SelectionX2 = Viewport.TimeToSlateUnitsRounded(SelectionEndTime);
+
+		if (SelectionX1 <= Viewport.Width &&
+			SelectionX2 >= 0)
+		{
+			float ClipLeft = 0.0f;
+			float ClipRight = 0.0f;
+			if (SelectionX1 < 0.0f)
+			{
+				ClipLeft = -SelectionX1;
+				SelectionX1 = 0.0f;
+			}
+			if (SelectionX2 > Viewport.Width)
+			{
+				ClipRight = SelectionX2 - Viewport.Width;
+				SelectionX2 = Viewport.Width;
+			}
+
+			// Draw selection area.
+			DC.DrawBox(SelectionX1, 0.0f, SelectionX2 - SelectionX1, Viewport.Height, WhiteBrush, FLinearColor(0.25f, 0.5f, 1.0f, 0.25f));
+			DC.LayerId++;
+
+			FColor ArrowFillColor(32, 64, 128, 255);
+			FLinearColor ArrowColor(ArrowFillColor);
+
+			if (SelectionX1 > 0.0f)
+			{
+				// Draw left side (vertical line).
+				DC.DrawBox(SelectionX1 - 1.0f, 0.0f, 1.0f, Viewport.Height, WhiteBrush, ArrowColor);
+			}
+
+			if (SelectionX2 < Viewport.Width)
+			{
+				// Draw right side (vertical line).
+				DC.DrawBox(SelectionX2, 0.0f, 1.0f, Viewport.Height, WhiteBrush, ArrowColor);
+			}
+
+			DC.LayerId++;
+
+			const float ArrowSize = 6.0f;
+			const float ArrowY = 6.0f;
+
+			if (SelectionX2 - SelectionX1 > 2 * ArrowSize)
+			{
+				// Draw horizontal line.
+				float HorizLineX1 = SelectionX1;
+				if (ClipLeft == 0.0f)
+				{
+					HorizLineX1 += 1.0f;
+				}
+				float HorizLineX2 = SelectionX2;
+				if (ClipRight == 0.0f)
+				{
+					HorizLineX2 -= 1.0f;
+				}
+				DC.DrawBox(HorizLineX1, ArrowY - 1.0f, HorizLineX2 - HorizLineX1, 3.0f, WhiteBrush, ArrowColor);
+
+				if (ClipLeft < ArrowSize)
+				{
+					// Draw left arrow.
+					for (float AH = 0.0f; AH < ArrowSize; AH += 1.0f)
+					{
+						DC.DrawBox(SelectionX1 - ClipLeft + AH, ArrowY - AH, 1.0f, 2.0f * AH + 1.0f, WhiteBrush, ArrowColor);
+					}
+				}
+
+				if (ClipRight < ArrowSize)
+				{
+					// Draw right arrow.
+					for (float AH = 0.0f; AH < ArrowSize; AH += 1.0f)
+					{
+						DC.DrawBox(SelectionX2 + ClipRight - AH - 1.0f, ArrowY - AH, 1.0f, 2.0f * AH + 1.0f, WhiteBrush, ArrowColor);
+					}
+				}
+
+				DC.LayerId++;
+
+				/*
+				//im: This should be a more efficeint way top draw the arrows, but it renders them with artifacts (missing vertical lines)!
+
+				const FSlateBrush* MyBrush = WhiteBrush;
+				FSlateShaderResourceProxy* ResourceProxy = FSlateDataPayload::ResourceManager->GetShaderResource(*MyBrush);
+				FSlateResourceHandle ResourceHandle = FSlateApplication::Get().GetRenderer()->GetResourceHandle(*MyBrush);
+
+				FVector2D AtlasOffset = ResourceProxy ? ResourceProxy->StartUV : FVector2D(0.0f, 0.0f);
+				FVector2D AtlasUVSize = ResourceProxy ? ResourceProxy->SizeUV : FVector2D(1.0f, 1.0f);
+
+				const FVector2D Pos = AllottedGeometry.GetAbsolutePosition() + FVector2D(0.0f, 40.0f);
+				const float Scale = AllottedGeometry.Scale;
+
+				FSlateRenderTransform RenderTransform;
+
+				TArray<FSlateVertex> Verts;
+				Verts.Reserve(6);
+				Verts.Add(FSlateVertex::Make<ESlateVertexRounding::Disabled>(RenderTransform, Pos + FVector2D(0.5f + SelectionX1 + ArrowSize, 0.5f + ArrowY + ArrowSize) * Scale, AtlasOffset + FVector2D(0.0f, 1.0f) * AtlasUVSize, ArrowFillColor));
+				Verts.Add(FSlateVertex::Make<ESlateVertexRounding::Disabled>(RenderTransform, Pos + FVector2D(0.5f + SelectionX1,             0.5f + ArrowY            ) * Scale, AtlasOffset + FVector2D(1.0f, 0.5f) * AtlasUVSize, ArrowFillColor));
+				Verts.Add(FSlateVertex::Make<ESlateVertexRounding::Disabled>(RenderTransform, Pos + FVector2D(0.5f + SelectionX1 + ArrowSize, 0.5f + ArrowY - ArrowSize) * Scale, AtlasOffset + FVector2D(0.0f, 0.0f) * AtlasUVSize, ArrowFillColor));
+				Verts.Add(FSlateVertex::Make<ESlateVertexRounding::Disabled>(RenderTransform, Pos + FVector2D(0.5f + SelectionX2 - ArrowSize, 0.5f + ArrowY - ArrowSize) * Scale, AtlasOffset + FVector2D(0.0f, 0.0f) * AtlasUVSize, ArrowFillColor));
+				Verts.Add(FSlateVertex::Make<ESlateVertexRounding::Disabled>(RenderTransform, Pos + FVector2D(0.5f + SelectionX2,             0.5f + ArrowY            ) * Scale, AtlasOffset + FVector2D(1.0f, 0.5f) * AtlasUVSize, ArrowFillColor));
+				Verts.Add(FSlateVertex::Make<ESlateVertexRounding::Disabled>(RenderTransform, Pos + FVector2D(0.5f + SelectionX2 - ArrowSize, 0.5f + ArrowY + ArrowSize) * Scale, AtlasOffset + FVector2D(0.0f, 1.0f) * AtlasUVSize, ArrowFillColor));
+
+				TArray<SlateIndex> Indices;
+				Indices.Reserve(6);
+				if (ClipLeft < ArrowSize)
+				{
+					Indices.Add(0);
+					Indices.Add(1);
+					Indices.Add(2);
+				}
+				if (ClipRight < ArrowSize)
+				{
+					Indices.Add(3);
+					Indices.Add(4);
+					Indices.Add(5);
+				}
+
+				FSlateDrawElement::MakeCustomVerts(
+					OutDrawElements,
+					LayerId,
+					ResourceHandle,
+					Verts,
+					Indices,
+					nullptr,
+					0,
+					0, ESlateDrawEffect::PreMultipliedAlpha);
+				DC.LayerId++;
+				*/
+			}
+
+			//////////////////////////////////////////////////
+			// Draw duration for selected time interval.
+
+			double Duration = SelectionEndTime - SelectionStartTime;
+			FString Text = TimeUtils::FormatTimeAuto(Duration);
+
+			const TSharedRef<FSlateFontMeasure> FontMeasureService = FSlateApplication::Get().GetRenderer()->GetFontMeasureService();
+			const float TextWidth = FontMeasureService->Measure(Text, MainFont).X;
+
+			const float CenterX = (SelectionX1 + SelectionX2) / 2.0f;
+
+			DC.DrawBox(CenterX - TextWidth / 2 - 2.0, ArrowY - 6.0f, TextWidth + 4.0f, 13.0f, WhiteBrush, ArrowColor);
+			DC.LayerId++;
+
+			DC.DrawText(CenterX - TextWidth / 2, ArrowY - 6.0f, Text, MainFont, FLinearColor::White);
+			DC.LayerId++;
+		}
+	}
+}
+
+////////////////////////////////////////////////////////////////////////////////////////////////////
+
+FReply STimingView::OnMouseButtonDown(const FGeometry& MyGeometry, const FPointerEvent& MouseEvent)
+{
+	FReply Reply = FReply::Unhandled();
+	MousePositionOnButtonDown = MyGeometry.AbsoluteToLocal(MouseEvent.GetScreenSpacePosition());
+
+	bool bStartPanning = false;
+
+	if (MouseEvent.GetEffectingButton() == EKeys::LeftMouseButton)
+	{
+		if (MarkersTrack.IsVisible() && MarkersTrack.IsHeaderHovered())
+		{
+			MarkersTrack.ToggleCollapsed();
+		}
+
+		if (!bIsRMB_Pressed)
+		{
+			bIsLMB_Pressed = true;
+
+			if (bIsSpaceBarKeyPressed || MousePositionOnButtonDown.Y > TimeRulerHeight)
+			{
+				bStartPanning = true;
+			}
+			else
+			{
+				bIsSelecting = true;
+				bIsDragging = false;
+
+				SelectionStartTime = Viewport.SlateUnitsToTime(MousePositionOnButtonDown.X);
+				SelectionEndTime = SelectionStartTime;
+				LastSelectionType = ESelectionType::None;
+				//TODO: SelectionChangingEvent.Broadcast(SelectionStartTime, SelectionEndTime);
+			}
+
+			// Capture mouse, so we can drag outside this widget.
+			Reply = FReply::Handled().CaptureMouse(SharedThis(this));
+		}
+	}
+	else if (MouseEvent.GetEffectingButton() == EKeys::RightMouseButton)
+	{
+		if (!bIsLMB_Pressed)
+		{
+			bIsRMB_Pressed = true;
+
+			bStartPanning = true;
+
+			// Capture mouse, so we can drag outside this widget.
+			Reply = FReply::Handled().CaptureMouse(SharedThis(this));
+		}
+	}
+
+	if (bStartPanning)
+	{
+		bIsPanning = true;
+		bIsDragging = false;
+
+		ViewportStartTimeOnButtonDown = Viewport.StartTime;
+		ViewportScrollPosYOnButtonDown = Viewport.ScrollPosY;
+
+		if (MouseEvent.GetModifierKeys().IsControlDown())
+		{
+			// Allow panning only horizontally.
+			PanningMode = EPanningMode::Horizontal;
+		}
+		else if (MouseEvent.GetModifierKeys().IsShiftDown())
+		{
+			// Allow panning only vertically.
+			PanningMode = EPanningMode::Vertical;
+		}
+		else
+		{
+			// Allow panning both horizontally and vertically.
+			PanningMode = EPanningMode::HorizontalAndVertical;
+		}
+	}
+
+	return Reply;
+}
+
+////////////////////////////////////////////////////////////////////////////////////////////////////
+
+FReply STimingView::OnMouseButtonUp(const FGeometry& MyGeometry, const FPointerEvent& MouseEvent)
+{
+	FReply Reply = FReply::Unhandled();
+	MousePositionOnButtonUp = MyGeometry.AbsoluteToLocal(MouseEvent.GetScreenSpacePosition());
+
+	const bool bIsValidForMouseClick = MousePositionOnButtonUp.Equals(MousePositionOnButtonDown, 2.0f);
+
+	if (MouseEvent.GetEffectingButton() == EKeys::LeftMouseButton)
+	{
+		if (bIsLMB_Pressed)
+		{
+			if (bIsPanning)
+			{
+				PanningMode = EPanningMode::None;
+
+				bIsPanning = false;
+			}
+			else if (bIsSelecting)
+			{
+				//TODO: SelectionChangedEvent.Broadcast(SelectionStartTime, SelectionEndTime);
+				TSharedPtr<STimingProfilerWindow> Wnd = FTimingProfilerManager::Get()->GetProfilerWindow();
+				if (Wnd.IsValid())
+				{
+					Wnd->TimersView->UpdateStats(SelectionStartTime, SelectionEndTime);
+				}
+
+				bIsSelecting = false;
+			}
+
+			if (!bIsDragging && bIsValidForMouseClick)
+			{
+				// Select the hovered timing event (if any).
+				UpdateHoveredTimingEvent(MousePositionOnButtonUp.X, MousePositionOnButtonUp.Y);
+				SelectHoveredTimingEvent();
+			}
+
+			bIsDragging = false;
+
+			// Release mouse as we no longer drag.
+			Reply = FReply::Handled().ReleaseMouseCapture();
+
+			bIsLMB_Pressed = false;
+		}
+	}
+	else if (MouseEvent.GetEffectingButton() == EKeys::RightMouseButton)
+	{
+		if (bIsRMB_Pressed)
+		{
+			if (bIsPanning)
+			{
+				PanningMode = EPanningMode::None;
+
+				bIsPanning = false;
+			}
+
+			if (!bIsDragging && !bIsSpaceBarKeyPressed && bIsValidForMouseClick)
+			{
+				ShowContextMenu(MouseEvent.GetScreenSpacePosition(), MouseEvent);
+				Reply = FReply::Handled();
+			}
+
+			bIsDragging = false;
+
+			// Release mouse as we no longer drag.
+			Reply = FReply::Handled().ReleaseMouseCapture();
+
+			bIsRMB_Pressed = false;
+		}
+	}
+
+	return Reply;
+}
+
+////////////////////////////////////////////////////////////////////////////////////////////////////
+
+FReply STimingView::OnMouseButtonDoubleClick(const FGeometry& MyGeometry, const FPointerEvent& MouseEvent)
+{
+	FReply Reply = FReply::Unhandled();
+
+	if (MouseEvent.GetEffectingButton() == EKeys::LeftMouseButton)
+	{
+		if (MarkersTrack.IsVisible() && MarkersTrack.IsHovered())
+		{
+			MarkersTrack.ToggleCollapsed();
+		}
+		else if (HoveredTimingEvent.IsValid())
+		{
+			double EndTime = Viewport.RestrictEndTime(HoveredTimingEvent.EndTime);
+			SelectTimeInterval(HoveredTimingEvent.StartTime, EndTime - HoveredTimingEvent.StartTime);
+		}
+
+		Reply = FReply::Handled();
+	}
+
+	return Reply;
+}
+
+////////////////////////////////////////////////////////////////////////////////////////////////////
+
+FReply STimingView::OnMouseMove(const FGeometry& MyGeometry, const FPointerEvent& MouseEvent)
+{
+	FReply Reply = FReply::Unhandled();
+
+	MousePosition = MyGeometry.AbsoluteToLocal(MouseEvent.GetScreenSpacePosition());
+
+	if (bIsPanning)
+	{
+		if (HasMouseCapture() && !MouseEvent.GetCursorDelta().IsZero())
+		{
+			bIsDragging = true;
+
+			if ((int32)PanningMode & (int32)EPanningMode::Horizontal)
+			{
+				ScrollAtTime(ViewportStartTimeOnButtonDown + static_cast<double>(MousePositionOnButtonDown.X - MousePosition.X) / Viewport.ScaleX);
+			}
+
+			if ((int32)PanningMode & (int32)EPanningMode::Vertical)
+			{
+				Viewport.ScrollPosY = ViewportScrollPosYOnButtonDown + (MousePositionOnButtonDown.Y - MousePosition.Y);
+				UpdateVerticalScrollBar();
+				bIsVerticalViewportDirty = true;
+			}
+		}
+
+		Reply = FReply::Handled();
+	}
+	else if (bIsSelecting)
+	{
+		if (HasMouseCapture() && !MouseEvent.GetCursorDelta().IsZero())
+		{
+			bIsDragging = true;
+
+			SelectionStartTime = Viewport.SlateUnitsToTime(MousePositionOnButtonDown.X);
+			SelectionEndTime = Viewport.SlateUnitsToTime(MousePosition.X);
+			if (SelectionStartTime > SelectionEndTime)
+			{
+				double Temp = SelectionStartTime;
+				SelectionStartTime = SelectionEndTime;
+				SelectionEndTime = Temp;
+			}
+			LastSelectionType = ESelectionType::TimeRange;
+			//TODO: SelectionChangingEvent.Broadcast(SelectionStartTime, SelectionEndTime);
+		}
+
+		Reply = FReply::Handled();
+	}
+	else
+	{
+		if (MarkersTrack.IsVisible())
+		{
+			MarkersTrack.UpdateHoveredState(MousePosition.X, MousePosition.Y, Viewport);
+		}
+
+		UpdateHoveredTimingEvent(MousePosition.X, MousePosition.Y);
+	}
+
+	return Reply; // SAssetViewItem::CreateToolTipWidget
+}
+
+////////////////////////////////////////////////////////////////////////////////////////////////////
+
+void STimingView::OnMouseEnter(const FGeometry& MyGeometry, const FPointerEvent& MouseEvent)
+{
+
+}
+
+////////////////////////////////////////////////////////////////////////////////////////////////////
+
+void STimingView::OnMouseLeave(const FPointerEvent& MouseEvent)
+{
+	if (!HasMouseCapture())
+	{
+		// No longer dragging (unless we have mouse capture).
+		bIsDragging = false;
+		bIsPanning = false;
+		bIsSelecting = false;
+
+		bIsLMB_Pressed = false;
+		bIsRMB_Pressed = false;
+
+		MousePosition = FVector2D::ZeroVector;
+	}
+}
+
+////////////////////////////////////////////////////////////////////////////////////////////////////
+
+FReply STimingView::OnMouseWheel(const FGeometry& MyGeometry, const FPointerEvent& MouseEvent)
+{
+	if (MouseEvent.GetModifierKeys().IsShiftDown())
+	{
+		// Scroll vertically.
+		static const float ScrollSpeedY = 16.0f * 3;
+		Viewport.ScrollPosY -= ScrollSpeedY * MouseEvent.GetWheelDelta();
+		UpdateVerticalScrollBar();
+		bIsVerticalViewportDirty = true;
+	}
+	else if (MouseEvent.GetModifierKeys().IsControlDown())
+	{
+		// Scroll horizontally.
+		const double ScrollSpeedX = Viewport.GetDurationForViewportDX(16.0 * 3);
+		ScrollAtTime(Viewport.StartTime - ScrollSpeedX * MouseEvent.GetWheelDelta());
+	}
+	else
+	{
+		// Zoom in/out horizontally.
+		const double Delta = MouseEvent.GetWheelDelta();
+		static const double ZoomStep = 0.25; // as percent
+		double ScaleX;
+
+		if (Delta > 0)
+		{
+			ScaleX = Viewport.ScaleX * FMath::Pow(1.0 + ZoomStep, Delta);
+		}
+		else
+		{
+			ScaleX = Viewport.ScaleX * FMath::Pow(1.0 / (1.0 + ZoomStep), -Delta);
+		}
+
+		//UE_LOG(TimingProfiler, Log, TEXT("%.2f, %.2f, %.2f"), Delta, Viewport.ScaleX, ScaleX);
+		MousePosition = MyGeometry.AbsoluteToLocal(MouseEvent.GetScreenSpacePosition());
+
+		if (Viewport.ZoomWithFixedX(ScaleX, MousePosition.X))
+		{
+			//Viewport.EnforceHorizontalScrollLimits(1.0);
+			UpdateHorizontalScrollBar();
+			bIsViewportDirty = true;
+		}
+	}
+
+	return FReply::Handled();
+}
+
+////////////////////////////////////////////////////////////////////////////////////////////////////
+
+void STimingView::OnDragEnter(const FGeometry& MyGeometry, const FDragDropEvent& DragDropEvent)
+{
+	SCompoundWidget::OnDragEnter(MyGeometry, DragDropEvent);
+
+	//TSharedPtr<FStatIDDragDropOp> Operation = DragDropEvent.GetOperationAs<FStatIDDragDropOp>();
+	//if (Operation.IsValid())
+	//{
+	//	Operation->ShowOK();
+	//}
+}
+
+////////////////////////////////////////////////////////////////////////////////////////////////////
+
+void STimingView::OnDragLeave(const FDragDropEvent& DragDropEvent)
+{
+	SCompoundWidget::OnDragLeave(DragDropEvent);
+
+	//TSharedPtr<FStatIDDragDropOp> Operation = DragDropEvent.GetOperationAs<FStatIDDragDropOp>();
+	//if (Operation.IsValid())
+	//{
+	//	Operation->ShowError();
+	//}
+}
+
+////////////////////////////////////////////////////////////////////////////////////////////////////
+
+FReply STimingView::OnDragOver(const FGeometry& MyGeometry, const FDragDropEvent& DragDropEvent)
+{
+	return SCompoundWidget::OnDragOver(MyGeometry,DragDropEvent);
+}
+
+////////////////////////////////////////////////////////////////////////////////////////////////////
+
+FReply STimingView::OnDrop(const FGeometry& MyGeometry, const FDragDropEvent& DragDropEvent)
+{
+	//TSharedPtr<FStatIDDragDropOp> Operation = DragDropEvent.GetOperationAs<FStatIDDragDropOp>();
+	//if (Operation.IsValid())
+	//{
+	//	return FReply::Handled();
+	//}
+
+	return SCompoundWidget::OnDrop(MyGeometry,DragDropEvent);
+}
+
+////////////////////////////////////////////////////////////////////////////////////////////////////
+
+FCursorReply STimingView::OnCursorQuery(const FGeometry& MyGeometry, const FPointerEvent& CursorEvent) const
+{
+	if (bIsPanning)
+	{
+		if (bIsDragging)
+		{
+			//return FCursorReply::Cursor(EMouseCursor::GrabHandClosed);
+			return FCursorReply::Cursor(EMouseCursor::GrabHand);
+		}
+	}
+	else if (bIsSelecting)
+	{
+		if (bIsDragging)
+		{
+			return FCursorReply::Cursor(EMouseCursor::ResizeLeftRight);
+		}
+	}
+	else if (bIsSpaceBarKeyPressed)
+	{
+		return FCursorReply::Cursor(EMouseCursor::GrabHand);
+	}
+
+	return FCursorReply::Unhandled();
+}
+
+////////////////////////////////////////////////////////////////////////////////////////////////////
+
+FReply STimingView::OnKeyDown(const FGeometry& MyGeometry, const FKeyEvent& InKeyEvent)
+{
+	if (InKeyEvent.GetKey() == EKeys::B)
+	{
+		// Toggle Bookmarks.
+		if (MarkersTrack.IsVisible())
+		{
+			if (!MarkersTrack.IsBookmarksTrack())
+			{
+				SetDrawOnlyBookmarks(true);
+			}
+			else
+			{
+				SetTimeMarkersVisible(false);
+			}
+		}
+		else
+		{
+			SetTimeMarkersVisible(true);
+			SetDrawOnlyBookmarks(true);
+		}
+		return FReply::Handled();
+	}
+	else if (InKeyEvent.GetKey() == EKeys::M)
+	{
+		// Toggle Time Markers.
+		if (MarkersTrack.IsVisible())
+		{
+			if (MarkersTrack.IsBookmarksTrack())
+			{
+				SetDrawOnlyBookmarks(false);
+			}
+			else
+			{
+				SetTimeMarkersVisible(false);
+			}
+		}
+		else
+		{
+			SetTimeMarkersVisible(true);
+			SetDrawOnlyBookmarks(false);
+		}
+
+		return FReply::Handled();
+	}
+	else if (InKeyEvent.GetKey() == EKeys::F)
+	{
+		FrameSelection();
+		return FReply::Handled();
+	}
+	else if (InKeyEvent.GetKey() == EKeys::C)
+	{
+		Layout.bIsCompactMode = !Layout.bIsCompactMode;
+		bAreTimingEventsTracksDirty = true;
+		return FReply::Handled();
+	}
+	else if (InKeyEvent.GetKey() == EKeys::V)
+	{
+		if (Layout.TargetMinTimelineH == 0.0f)
+		{
+			Layout.TargetMinTimelineH = RealMinTimelineH;
+		}
+		else
+		{
+			Layout.TargetMinTimelineH = 0.0f;
+		}
+
+		Layout.MinTimelineH = Layout.TargetMinTimelineH; // no layout animation
+
+		for (auto& CachedTimelineKV : CachedTimelines)
+		{
+			CachedTimelineKV.Value->SetHeight(0.0f);
+		}
+
+		Viewport.ScrollPosY = 0.0f;
+		UpdateVerticalScrollBar();
+		bIsVerticalViewportDirty = true;
+
+		return FReply::Handled();
+	}
+	else if (InKeyEvent.GetKey() == EKeys::Equals ||
+			 InKeyEvent.GetKey() == EKeys::Add)
+	{
+		// Zoom In
+		double ScaleX = Viewport.ScaleX * 1.25;
+		if (Viewport.ZoomWithFixedX(ScaleX, Viewport.Width/2))
+		{
+			//Viewport.EnforceHorizontalScrollLimits(1.0);
+			UpdateHorizontalScrollBar();
+			bIsViewportDirty = true;
+		}
+		return FReply::Handled();
+	}
+	else if (InKeyEvent.GetKey() == EKeys::Hyphen ||
+			 InKeyEvent.GetKey() == EKeys::Subtract)
+	{
+		// Zoom Out
+		double ScaleX = Viewport.ScaleX / 1.25;
+		if (Viewport.ZoomWithFixedX(ScaleX, Viewport.Width / 2))
+		{
+			//Viewport.EnforceHorizontalScrollLimits(1.0);
+			UpdateHorizontalScrollBar();
+			bIsViewportDirty = true;
+		}
+		return FReply::Handled();
+	}
+	else if (InKeyEvent.GetKey() == EKeys::Left)
+	{
+		if (InKeyEvent.GetModifierKeys().IsControlDown())
+		{
+			// Scroll Left
+			double DT = Viewport.EndTime - Viewport.StartTime;
+			//ScrollAtTime(Viewport.StartTime - DT * 0.05);
+			if (Viewport.ScrollAtTime(Viewport.StartTime - DT * 0.05))
+			{
+				//Viewport.EnforceHorizontalScrollLimits(1.0);
+				UpdateHorizontalScrollBar();
+				bIsViewportDirty = true;
+			}
+		}
+		else
+		{
+			SelectLeftTimingEvent();
+		}
+		return FReply::Handled();
+	}
+	else if (InKeyEvent.GetKey() == EKeys::Right)
+	{
+		if (InKeyEvent.GetModifierKeys().IsControlDown())
+		{
+			// Scroll Right
+			double DT = Viewport.EndTime - Viewport.StartTime;
+			//ScrollAtTime(Viewport.StartTime + DT * 0.05);
+			if (Viewport.ScrollAtTime(Viewport.StartTime + DT * 0.05))
+			{
+				//Viewport.EnforceHorizontalScrollLimits(1.0);
+				UpdateHorizontalScrollBar();
+				bIsViewportDirty = true;
+			}
+		}
+		else
+		{
+			SelectRightTimingEvent();
+		}
+		return FReply::Handled();
+	}
+	else if (InKeyEvent.GetKey() == EKeys::Up)
+	{
+		if (InKeyEvent.GetModifierKeys().IsControlDown())
+		{
+			// Scroll Up
+			Viewport.ScrollPosY -= 16.0 * 3;
+			//Viewport.EnforceVerticalScrollLimits(1.0);
+			UpdateVerticalScrollBar();
+			bIsVerticalViewportDirty = true;
+		}
+		else
+		{
+			SelectUpTimingEvent();
+		}
+		return FReply::Handled();
+	}
+	else if (InKeyEvent.GetKey() == EKeys::Down)
+	{
+		if (InKeyEvent.GetModifierKeys().IsControlDown())
+		{
+			// Scroll Down
+			Viewport.ScrollPosY += 16.0 * 3;
+			//Viewport.EnforceVerticalScrollLimits(1.0);
+			UpdateVerticalScrollBar();
+			bIsVerticalViewportDirty = true;
+		}
+		else
+		{
+			SelectDownTimingEvent();
+		}
+		return FReply::Handled();
+	}
+	else if (InKeyEvent.GetKey() == EKeys::Enter)
+	{
+		if (SelectedTimingEvent.IsValid())
+		{
+			const double Duration = Viewport.RestrictDuration(SelectedTimingEvent.StartTime, SelectedTimingEvent.EndTime);
+			SelectTimeInterval(SelectedTimingEvent.StartTime, Duration);
+		}
+		return FReply::Handled();
+	}
+	else if (InKeyEvent.GetKey() == EKeys::SpaceBar)
+	{
+		bIsSpaceBarKeyPressed = true;
+		FSlateApplication::Get().QueryCursor();
+		return FReply::Handled();
+	}
+	else if (InKeyEvent.GetKey() == EKeys::D) // debug: toggles down-sampling on/off
+	{
+		bUseDownSampling = !bUseDownSampling;
+		bAreTimingEventsTracksDirty = true;
+		return FReply::Handled();
+	}
+	else if (InKeyEvent.GetKey() == EKeys::G) // debug: toggles Graph track on/off
+	{
+		GraphTrack.SetVisibilityFlag(!GraphTrack.IsVisible());
+		bAreTimingEventsTracksDirty = true;
+		return FReply::Handled();
+	}
+	else if (InKeyEvent.GetKey() == EKeys::U) // debug: toggles Timing (GPU, CPU) tracks on/off
+	{
+		bDrawTimingEvents = !bDrawTimingEvents;
+		CachedTimelines.Reset();
+		return FReply::Handled();
+	}
+	else if (InKeyEvent.GetKey() == EKeys::I)  // debug: toggles IO tracks on/off
+	{
+		bDrawIoEvents = !bDrawIoEvents;
+		CachedTimelines.Reset();
+		return FReply::Handled();
+	}
+	else if (InKeyEvent.GetKey() == EKeys::O)  // debug: toggles IO merge lanes algorithm
+	{
+		bMergeIoLanes = !bMergeIoLanes;
+		bForceIoEventsUpdate = true;
+		return FReply::Handled();
+	}
+
+	return SCompoundWidget::OnKeyDown(MyGeometry, InKeyEvent);
+}
+
+////////////////////////////////////////////////////////////////////////////////////////////////////
+
+FReply STimingView::OnKeyUp(const FGeometry& MyGeometry, const FKeyEvent& InKeyEvent)
+{
+	if (InKeyEvent.GetKey() == EKeys::SpaceBar)
+	{
+		bIsSpaceBarKeyPressed = false;
+		FSlateApplication::Get().QueryCursor();
+		return FReply::Handled();
+	}
+
+	return SCompoundWidget::OnKeyUp(MyGeometry, InKeyEvent);
+}
+
+////////////////////////////////////////////////////////////////////////////////////////////////////
+
+void STimingView::ShowContextMenu(const FVector2D& ScreenSpacePosition, const FPointerEvent& MouseEvent)
+{
+	TSharedPtr<FUICommandList> ProfilerCommandList = FTimingProfilerManager::Get()->GetCommandList();
+	const FTimingProfilerCommands& ProfilerCommands = FTimingProfilerManager::GetCommands();
+	const FTimingProfilerActionManager& ProfilerActionManager = FTimingProfilerManager::GetActionManager();
+
+	const bool bShouldCloseWindowAfterMenuSelection = true;
+	FMenuBuilder MenuBuilder(bShouldCloseWindowAfterMenuSelection, ProfilerCommandList);
+
+	MenuBuilder.BeginSection(TEXT("Misc"), LOCTEXT("Miscellaneous", "Miscellaneous"));
+	{
+		//MenuBuilder.AddMenuEntry(FTimingProfilerManager::GetCommands().EventGraph_SelectAllFrames);
+		//MenuBuilder.AddMenuEntry(FTimingProfilerManager::GetCommands().ProfilerManager_ToggleLivePreview);
+	}
+	MenuBuilder.EndSection();
+
+	TSharedRef<SWidget> MenuWidget = MenuBuilder.MakeWidget();
+
+	FWidgetPath EventPath = MouseEvent.GetEventPath() != nullptr ? *MouseEvent.GetEventPath() : FWidgetPath();
+	FSlateApplication::Get().PushMenu(SharedThis(this), EventPath, MenuWidget, ScreenSpacePosition, FPopupTransitionEffect::ContextMenu);
+}
+
+////////////////////////////////////////////////////////////////////////////////////////////////////
+
+void STimingView::BindCommands()
+{
+	//TSharedPtr<FUICommandList> CommandList = FTimingProfilerManager::Get()->GetCommandList();
+	//const FTimingProfilerCommands& Commands = FTimingProfilerManager::GetCommands();
+}
+
+////////////////////////////////////////////////////////////////////////////////////////////////////
+
+void STimingView::HorizontalScrollBar_OnUserScrolled(float ScrollOffset)
+{
+	const double SX = 1.0 / (Viewport.MaxValidTime - Viewport.MinValidTime);
+	const float ThumbSizeFraction = FMath::Clamp<float>((Viewport.EndTime - Viewport.StartTime) * SX, 0.0f, 1.0f);
+	const float OffsetFraction = FMath::Clamp<float>(ScrollOffset, 0.0f, 1.0f - ThumbSizeFraction);
+
+	const double Time = Viewport.MinValidTime + static_cast<double>(OffsetFraction) * (Viewport.MaxValidTime - Viewport.MinValidTime);
+	if (Viewport.StartTime != Time)
+	{
+		Viewport.ScrollAtTime(Time);
+		bIsViewportDirty = true;
+	}
+
+	HorizontalScrollBar->SetState(OffsetFraction, ThumbSizeFraction);
+}
+
+////////////////////////////////////////////////////////////////////////////////////////////////////
+
+void STimingView::UpdateHorizontalScrollBar()
+{
+	const double SX = 1.0 / (Viewport.MaxValidTime - Viewport.MinValidTime);
+	const float ThumbSizeFraction = FMath::Clamp<float>((Viewport.EndTime - Viewport.StartTime) * SX, 0.0f, 1.0f);
+	const float ScrollOffset = static_cast<float>((Viewport.StartTime - Viewport.MinValidTime) * SX);
+	const float OffsetFraction = FMath::Clamp<float>(ScrollOffset, 0.0f, 1.0f - ThumbSizeFraction);
+
+	HorizontalScrollBar->SetState(OffsetFraction, ThumbSizeFraction);
+}
+
+////////////////////////////////////////////////////////////////////////////////////////////////////
+
+void STimingView::VerticalScrollBar_OnUserScrolled(float ScrollOffset)
+{
+	const float SY = 1.0 / Viewport.ScrollHeight;
+	const float H = Viewport.Height - Viewport.TopOffset;
+	const float ThumbSizeFraction = FMath::Clamp<float>(H * SY, 0.0f, 1.0f);
+	const float OffsetFraction = FMath::Clamp<float>(ScrollOffset, 0.0f, 1.0f - ThumbSizeFraction);
+
+	const float ScrollPosY = OffsetFraction * Viewport.ScrollHeight;
+	if (Viewport.ScrollPosY != ScrollPosY)
+	{
+		Viewport.ScrollPosY = ScrollPosY;
+		bIsVerticalViewportDirty = true;
+	}
+
+	VerticalScrollBar->SetState(OffsetFraction, ThumbSizeFraction);
+}
+
+////////////////////////////////////////////////////////////////////////////////////////////////////
+
+void STimingView::UpdateVerticalScrollBar()
+{
+	const float SY = 1.0 / Viewport.ScrollHeight;
+	const float H = Viewport.Height - Viewport.TopOffset;
+	const float ThumbSizeFraction = FMath::Clamp<float>(H * SY, 0.0f, 1.0f);
+	const float ScrollOffset = Viewport.ScrollPosY * SY;
+	float OffsetFraction = FMath::Clamp<float>(ScrollOffset, 0.0f, 1.0f - ThumbSizeFraction);
+
+	VerticalScrollBar->SetState(OffsetFraction, ThumbSizeFraction);
+}
+
+////////////////////////////////////////////////////////////////////////////////////////////////////
+
+void STimingView::ScrollAtTime(double StartTime)
+{
+	if (Viewport.ScrollAtTime(StartTime))
+	{
+		UpdateHorizontalScrollBar();
+		bIsViewportDirty = true;
+	}
+}
+
+////////////////////////////////////////////////////////////////////////////////////////////////////
+
+void STimingView::CenterOnTimeInterval(double IntervalStartTime, double IntervalDuration)
+{
+	if (Viewport.CenterOnTimeInterval(IntervalStartTime, IntervalDuration))
+	{
+		Viewport.EnforceHorizontalScrollLimits(1.0);
+
+		UpdateHorizontalScrollBar();
+		bIsViewportDirty = true;
+	}
+}
+
+////////////////////////////////////////////////////////////////////////////////////////////////////
+
+void STimingView::BringIntoView(double StartTime, double EndTime)
+{
+	EndTime = Viewport.RestrictEndTime(EndTime);
+
+	// Increase interval with 8% on each side.
+	const double DT = (Viewport.EndTime - Viewport.StartTime) * 0.08;
+	StartTime -= DT;
+	EndTime += DT;
+
+	double NewStartTime = Viewport.StartTime;
+
+	if (EndTime > Viewport.EndTime)
+	{
+		NewStartTime += EndTime - Viewport.EndTime;
+	}
+
+	if (StartTime < NewStartTime)
+	{
+		NewStartTime = StartTime;
+	}
+
+	ScrollAtTime(NewStartTime);
+}
+
+////////////////////////////////////////////////////////////////////////////////////////////////////
+
+void STimingView::SelectTimeInterval(double IntervalStartTime, double IntervalDuration)
+{
+	SelectionStartTime = IntervalStartTime;
+	SelectionEndTime = IntervalStartTime + IntervalDuration;
+	LastSelectionType = ESelectionType::TimeRange;
+
+	//TODO: SelectionChangedEvent.Broadcast(SelectionStartTime, SelectionEndTime);
+	TSharedPtr<STimingProfilerWindow> Wnd = FTimingProfilerManager::Get()->GetProfilerWindow();
+	if (Wnd.IsValid())
+	{
+		Wnd->TimersView->UpdateStats(SelectionStartTime, SelectionEndTime);
+	}
+}
+
+////////////////////////////////////////////////////////////////////////////////////////////////////
+
+void STimingView::SetAndCenterOnTimeMarker(double Time)
+{
+	TimeMarker = Time;
+
+	double MinT, MaxT;
+	Viewport.GetHorizontalScrollLimits(MinT, MaxT);
+	const double ViewportDuration = static_cast<double>(Viewport.Width) / Viewport.ScaleX;
+	MinT += ViewportDuration / 2;
+	MaxT += ViewportDuration / 2;
+	Time = FMath::Clamp<double>(Time, MinT, MaxT);
+
+	Time = Viewport.AlignTimeToPixel(Time);
+	CenterOnTimeInterval(Time, 0.0);
+}
+
+////////////////////////////////////////////////////////////////////////////////////////////////////
+
+void STimingView::SelectToTimeMarker(double Time)
+{
+	if (TimeMarker < Time)
+	{
+		SelectTimeInterval(TimeMarker, Time - TimeMarker);// +TimeUtils::Nanosecond);
+	}
+	else
+	{
+		SelectTimeInterval(Time, TimeMarker - Time);// +TimeUtils::Nanosecond);
+	}
+
+	TimeMarker = Time;
+}
+
+////////////////////////////////////////////////////////////////////////////////////////////////////
+
+void STimingView::SetTimeMarkersVisible(bool bIsMarkersTrackVisible)
+{
+	if (MarkersTrack.IsVisible() != bIsMarkersTrackVisible)
+	{
+		MarkersTrack.SetVisibilityFlag(bIsMarkersTrackVisible);
+
+		if (MarkersTrack.IsVisible())
+		{
+			if (Viewport.ScrollPosY != 0.0f)
+			{
+				Viewport.ScrollPosY += MarkersTrack.GetHeight();
+			}
+
+			MarkersTrack.SetDirtyFlag();
+		}
+		else
+		{
+			Viewport.ScrollPosY -= MarkersTrack.GetHeight();
+		}
+	}
+}
+
+////////////////////////////////////////////////////////////////////////////////////////////////////
+
+void STimingView::SetDrawOnlyBookmarks(bool bIsBookmarksTrack)
+{
+	if (MarkersTrack.IsBookmarksTrack() != bIsBookmarksTrack)
+	{
+		const float PrevHeight = MarkersTrack.GetHeight();
+		MarkersTrack.SetIsBookmarksTrackFlag(bIsBookmarksTrack);
+
+		if (MarkersTrack.IsVisible())
+		{
+			if (Viewport.ScrollPosY != 0.0f)
+			{
+				Viewport.ScrollPosY += MarkersTrack.GetHeight() - PrevHeight;
+			}
+
+			MarkersTrack.SetDirtyFlag();
+		}
+	}
+}
+
+////////////////////////////////////////////////////////////////////////////////////////////////////
+
+void STimingView::UpdateHoveredTimingEvent(float MX, float MY)
+{
+	HoveredTimingEvent.Track = nullptr;
+	HoveredTimingEvent.TypeId = FTimerNode::InvalidId;
+
+	if (MY >= Viewport.TopOffset && MY < Viewport.Height)
+	{
+		for (const auto& KV : CachedTimelines)
+		{
+			const FTimingEventsTrack& Track = *KV.Value;
+			const float Y = Viewport.TopOffset + Track.GetPosY() - Viewport.ScrollPosY;
+			if (MY >= Y && MY < Y + Track.GetHeight())
+			{
+				HoveredTimingEvent.Track = &Track;
+				break;
+			}
+		}
+
+		if (HoveredTimingEvent.Track)
+		{
+			const float Y0 = Viewport.TopOffset + HoveredTimingEvent.Track->GetPosY() - Viewport.ScrollPosY + 1.0f + Layout.TimelineDY;
+
+			// If mouse is not above first sub-track or below last sub-track...
+			if (MY >= Y0 && MY < Y0 + HoveredTimingEvent.Track->GetHeight() + Layout.TimelineDY)
+			{
+				int32 Depth = (MY - Y0) / (Layout.EventDY + Layout.EventH);
+				float EventMY = (MY - Y0) - Depth * (Layout.EventDY + Layout.EventH);
+
+				const double StartTime = Viewport.SlateUnitsToTime(MX);
+				const double EndTime = StartTime + 2.0 / Viewport.ScaleX; // +2px
+				const bool bStopAtFirstMatch = true; // get first one matching
+				const bool bSearchForLargestEvent = false;
+				SearchTimingEvent(StartTime, EndTime,
+					[Depth](double, double, uint32 EventDepth, uint32)
+					{
+						return EventDepth == Depth;
+					},
+					HoveredTimingEvent, bStopAtFirstMatch, bSearchForLargestEvent);
+
+				//TODO: ComputeSingleTimingEventStats(HoveredTimingEvent) --> compute ExclusiveTime
+			}
+		}
+	}
+}
+
+////////////////////////////////////////////////////////////////////////////////////////////////////
+
+bool STimingView::SearchTimingEvent(const double InStartTime,
+									 const double InEndTime,
+									 TFunctionRef<bool(double, double, uint32, uint32)> InPredicate,
+									 FTimingEvent& InOutTimingEvent,
+									 bool bInStopAtFirstMatch,
+									 bool bInSearchForLargestEvent)
+{
+	struct FSearchTimingEventContext
+	{
+		const double StartTime;
+		const double EndTime;
+		TFunctionRef<bool(double, double, uint32, uint32)> Predicate;
+		FTimingEvent& TimingEvent;
+		const bool bStopAtFirstMatch;
+		const bool bSearchForLargestEvent;
+		mutable bool bFound;
+		mutable bool bContinueSearching;
+		mutable double LargestDuration;
+
+		FSearchTimingEventContext(const double InStartTime, const double InEndTime, TFunctionRef<bool(double, double, uint32, uint32)> InPredicate, FTimingEvent& InOutTimingEvent, bool bInStopAtFirstMatch, bool bInSearchForLargestEvent)
+			: StartTime(InStartTime)
+			, EndTime(InEndTime)
+			, Predicate(InPredicate)
+			, TimingEvent(InOutTimingEvent)
+			, bStopAtFirstMatch(bInStopAtFirstMatch)
+			, bSearchForLargestEvent(bInSearchForLargestEvent)
+			, bFound(false)
+			, bContinueSearching(true)
+			, LargestDuration(-1.0)
+		{
+		}
+	};
+
+	FSearchTimingEventContext Ctx(InStartTime, InEndTime, InPredicate, InOutTimingEvent, bInStopAtFirstMatch, bInSearchForLargestEvent);
+
+	TSharedPtr<const Trace::IAnalysisSession> Session = FInsightsManager::Get()->GetSession();
+	if (Session.IsValid())
+	{
+		Trace::FAnalysisSessionReadScope _(*Session.Get());
+
+		//if (Ctx.TimingEvent.Track->GetType() == TrackType::Timing)
+		Session->ReadTimingProfilerProvider([&Ctx](const Trace::ITimingProfilerProvider& TimingProfilerProvider)
+		{
+			TimingProfilerProvider.ReadTimeline(Ctx.TimingEvent.Track->GetId(), [&Ctx](const Trace::ITimingProfilerProvider::Timeline& Timeline)
+			{
+				Timeline.EnumerateEvents(Ctx.StartTime, Ctx.EndTime, [&Ctx](double EventStartTime, double EventEndTime, uint32 EventDepth, const Trace::FTimingProfilerEvent& Event)
+				{
+					if (Ctx.bContinueSearching && Ctx.Predicate(EventStartTime, EventEndTime, EventDepth, Event.TimerIndex))
+					{
+						if (!Ctx.bSearchForLargestEvent || EventEndTime - EventStartTime > Ctx.LargestDuration)
+						{
+							Ctx.LargestDuration = EventEndTime - EventStartTime;
+
+							Ctx.TimingEvent.TypeId = Event.TimerIndex;
+							Ctx.TimingEvent.Depth = EventDepth;
+							Ctx.TimingEvent.StartTime = EventStartTime;
+							Ctx.TimingEvent.EndTime = EventEndTime;
+
+							Ctx.bFound = true;
+							Ctx.bContinueSearching = !Ctx.bStopAtFirstMatch || Ctx.bSearchForLargestEvent;
+						}
+					}
+				});
+			});
+		});
+
+		if (Ctx.bFound)
+		{
+			Session->ReadTimingProfilerProvider([&Ctx](const Trace::ITimingProfilerProvider& TimingProfilerProvider)
+			{
+				TimingProfilerProvider.ReadTimeline(Ctx.TimingEvent.Track->GetId(), [&Ctx](const Trace::ITimingProfilerProvider::Timeline& Timeline)
+				{
+					struct FEnumerationState
+					{
+						double EventStartTime;
+						double EventEndTime;
+						uint64 EventDepth;
+
+						uint64 CurrentDepth;
+						double ChildStartTime;
+						double TimeInChildren;
+						bool IsInEventScope;
+					};
+					FEnumerationState State;
+
+					State.EventStartTime = Ctx.TimingEvent.StartTime;
+					State.EventEndTime = Ctx.TimingEvent.EndTime;
+					State.EventDepth = Ctx.TimingEvent.Depth;
+
+					State.CurrentDepth = 0;
+					State.ChildStartTime = 0.0;
+					State.TimeInChildren = 0.0;
+					State.IsInEventScope = false;
+
+					Timeline.EnumerateEvents(Ctx.TimingEvent.StartTime, Ctx.TimingEvent.EndTime, [&State](bool IsEnter, double Time, const Trace::FTimingProfilerEvent& Event)
+					{
+						if (IsEnter)
+						{
+							if (State.CurrentDepth == State.EventDepth && Time == State.EventStartTime)
+							{
+								State.IsInEventScope = true;
+							}
+							else if (State.IsInEventScope && State.CurrentDepth == State.EventDepth + 1)
+							{
+								State.ChildStartTime = Time;
+							}
+							++State.CurrentDepth;
+						}
+						else
+						{
+							--State.CurrentDepth;
+							if (State.CurrentDepth == State.EventDepth && Time == State.EventEndTime)
+							{
+								State.IsInEventScope = false;
+							}
+							else if (State.IsInEventScope && State.CurrentDepth == State.EventDepth + 1)
+							{
+								State.TimeInChildren += Time - State.ChildStartTime;
+							}
+						}
+					});
+					Ctx.TimingEvent.ExclusiveTime = (Ctx.TimingEvent.EndTime - Ctx.TimingEvent.StartTime) - State.TimeInChildren;
+					check(Ctx.TimingEvent.ExclusiveTime >= 0.0);
+				});
+			});
+		}
+	}
+
+	return Ctx.bFound;
+}
+
+////////////////////////////////////////////////////////////////////////////////////////////////////
+
+void STimingView::OnSelectedTimingEventChanged()
+{
+	// Select the timer node coresponding to timing event type of selected timing event.
+	if (SelectedTimingEvent.IsValid())
+	{
+		TSharedPtr<STimingProfilerWindow> Wnd = FTimingProfilerManager::Get()->GetProfilerWindow();
+		if (Wnd.IsValid())
+		{
+			Wnd->TimersView->SelectTimerNode(SelectedTimingEvent.TypeId);
+		}
+	}
+}
+
+////////////////////////////////////////////////////////////////////////////////////////////////////
+
+void STimingView::SelectHoveredTimingEvent()
+{
+	SelectedTimingEvent = HoveredTimingEvent;
+	if (SelectedTimingEvent.IsValid())
+	{
+		LastSelectionType = ESelectionType::TimingEvent;
+		BringIntoView(SelectedTimingEvent.StartTime, SelectedTimingEvent.EndTime);
+	}
+	OnSelectedTimingEventChanged();
+}
+
+////////////////////////////////////////////////////////////////////////////////////////////////////
+
+void STimingView::SelectLeftTimingEvent()
+{
+	if (SelectedTimingEvent.IsValid())
+	{
+		const uint32 Depth = SelectedTimingEvent.Depth;
+		const double StartTime = SelectedTimingEvent.StartTime;
+		const double EndTime = SelectedTimingEvent.EndTime;
+		const bool bStopAtFirstMatch = false; // get last one matching
+		const bool bSearchForLargestEvent = false;
+		if (SearchTimingEvent(0.0, StartTime,
+			[Depth, StartTime, EndTime](double EventStartTime, double EventEndTime, uint32 EventDepth, uint32 EventTypeId)
+			{
+				return EventDepth == Depth &&
+					(EventStartTime < StartTime || EventEndTime < EndTime);
+			},
+			SelectedTimingEvent, bStopAtFirstMatch, bSearchForLargestEvent))
+		{
+			BringIntoView(SelectedTimingEvent.StartTime, SelectedTimingEvent.EndTime);
+			OnSelectedTimingEventChanged();
+		}
+	}
+}
+
+////////////////////////////////////////////////////////////////////////////////////////////////////
+
+void STimingView::SelectRightTimingEvent()
+{
+	if (SelectedTimingEvent.IsValid())
+	{
+		const uint32 Depth = SelectedTimingEvent.Depth;
+		const double StartTime = SelectedTimingEvent.StartTime;
+		const double EndTime = SelectedTimingEvent.EndTime;
+		const bool bStopAtFirstMatch = true; // get first one matching
+		const bool bSearchForLargestEvent = false;
+		if (SearchTimingEvent(EndTime, Viewport.MaxValidTime,
+			[Depth, StartTime, EndTime](double EventStartTime, double EventEndTime, uint32 EventDepth, uint32 EventTypeId)
+			{
+				return EventDepth == Depth &&
+					(EventStartTime > StartTime || EventEndTime > EndTime);
+			},
+			SelectedTimingEvent, bStopAtFirstMatch, bSearchForLargestEvent))
+		{
+			BringIntoView(SelectedTimingEvent.StartTime, SelectedTimingEvent.EndTime);
+			OnSelectedTimingEventChanged();
+		}
+	}
+}
+
+////////////////////////////////////////////////////////////////////////////////////////////////////
+
+void STimingView::SelectUpTimingEvent()
+{
+	if (SelectedTimingEvent.IsValid() &&
+		SelectedTimingEvent.Depth > 0)
+	{
+		const uint32 Depth = SelectedTimingEvent.Depth - 1;
+		double StartTime = SelectedTimingEvent.StartTime;
+		double EndTime = SelectedTimingEvent.EndTime;
+		const bool bStopAtFirstMatch = true; // get first one matching
+		const bool bSearchForLargestEvent = false;
+		if (SearchTimingEvent(StartTime, EndTime,
+			[Depth, StartTime, EndTime](double EventStartTime, double EventEndTime, uint32 EventDepth, uint32 EventTypeId)
+			{
+				return EventDepth == Depth
+					&& EventStartTime <= EndTime
+					&& EventEndTime >= StartTime;
+			},
+			SelectedTimingEvent, bStopAtFirstMatch, bSearchForLargestEvent))
+		{
+			BringIntoView(SelectedTimingEvent.StartTime, SelectedTimingEvent.EndTime);
+			OnSelectedTimingEventChanged();
+		}
+	}
+}
+
+////////////////////////////////////////////////////////////////////////////////////////////////////
+
+void STimingView::SelectDownTimingEvent()
+{
+	if (SelectedTimingEvent.IsValid())
+	{
+		const uint32 Depth = SelectedTimingEvent.Depth + 1;
+		double StartTime = SelectedTimingEvent.StartTime;
+		double EndTime = SelectedTimingEvent.EndTime;
+		const bool bStopAtFirstMatch = false; // check all timing events
+		const bool bSearchForLargestEvent = true; // get largest timing event
+		if (SearchTimingEvent(StartTime, EndTime,
+			[Depth, StartTime, EndTime](double EventStartTime, double EventEndTime, uint32 EventDepth, uint32 EventTypeId)
+			{
+				return EventDepth == Depth
+					&& EventStartTime <= EndTime
+					&& EventEndTime >= StartTime;
+			},
+			SelectedTimingEvent, bStopAtFirstMatch, bSearchForLargestEvent))
+		{
+			BringIntoView(SelectedTimingEvent.StartTime, SelectedTimingEvent.EndTime);
+			OnSelectedTimingEventChanged();
+		}
+	}
+}
+
+////////////////////////////////////////////////////////////////////////////////////////////////////
+
+void STimingView::FrameSelection()
+{
+	double StartTime, EndTime;
+
+	ESelectionType Type = ESelectionType::None;
+
+	if (LastSelectionType == ESelectionType::TimingEvent)
+	{
+		// Try framing the selected timing event.
+		if (SelectedTimingEvent.IsValid())
+		{
+			Type = ESelectionType::TimingEvent;
+		}
+
+		// Next time, try framing the selected time range.
+		LastSelectionType = ESelectionType::TimeRange;
+	}
+	else if (LastSelectionType == ESelectionType::TimeRange)
+	{
+		// Try framing the selected time range.
+		if (SelectionEndTime > SelectionStartTime)
+		{
+			Type = ESelectionType::TimeRange;
+		}
+
+		// Next time, try framing the selected timing event.
+		LastSelectionType = ESelectionType::TimingEvent;
+	}
+
+	// If no last selection or last selection is empty...
+	if (LastSelectionType == ESelectionType::None || Type == ESelectionType::None)
+	{
+		// First, try framing the selected timing event...
+		if (SelectedTimingEvent.IsValid())
+		{
+			Type = ESelectionType::TimingEvent;
+		}
+		else // ...otherwise, try framing the selected time range
+		{
+			Type = ESelectionType::TimeRange;
+		}
+	}
+
+	if (Type == ESelectionType::TimingEvent)
+	{
+		// Frame the selected event.
+		StartTime = SelectedTimingEvent.StartTime;
+		EndTime = Viewport.RestrictEndTime(SelectedTimingEvent.EndTime);
+		if (EndTime == StartTime)
+		{
+			EndTime += 1.0 / Viewport.ScaleX; // +1px
+		}
+	}
+	else
+	{
+		// Frame the selected time range.
+		StartTime = SelectionStartTime;
+		EndTime = Viewport.RestrictEndTime(SelectionEndTime);
+	}
+
+	if (EndTime > StartTime)
+	{
+		const double Duration = EndTime - StartTime;
+		if (Viewport.ZoomOnTimeInterval(StartTime - Duration * 0.1, Duration * 1.2))
+		{
+			UpdateHorizontalScrollBar();
+			bIsViewportDirty = true;
+		}
+	}
+}
+
+////////////////////////////////////////////////////////////////////////////////////////////////////
+
+const FTimerNodePtr STimingView::GetTimerNode(uint64 TypeId) const
+{
+	const FTimerNodePtr* TimerNodePtrPtr = nullptr;
+	TSharedPtr<STimingProfilerWindow> Wnd = FTimingProfilerManager::Get()->GetProfilerWindow();
+	if (Wnd.IsValid() && Wnd->TimersView.IsValid())
+	{
+		TimerNodePtrPtr = Wnd->TimersView->GetTimerNode(TypeId);
+		if (TimerNodePtrPtr == nullptr)
+		{
+			// List of timers in TimersView not up to date?
+			// Refresh and try again.
+			Wnd->TimersView->RebuildTree();
+			TimerNodePtrPtr = Wnd->TimersView->GetTimerNode(TypeId);
+		}
+	}
+	return TimerNodePtrPtr ? *TimerNodePtrPtr : nullptr;
+}
+
+////////////////////////////////////////////////////////////////////////////////////////////////////
+
+#undef LOCTEXT_NAMESPACE
