[Insights]
  - Timing view: Added filtering for CPU/GPU/IO timeline tracks.

diff --git a/Engine/Source/Developer/TraceInsights/Private/Insights/ViewModels/BaseTimingTrack.h b/Engine/Source/Developer/TraceInsights/Private/Insights/ViewModels/BaseTimingTrack.h
--- a/Engine/Source/Developer/TraceInsights/Private/Insights/ViewModels/BaseTimingTrack.h
+++ b/Engine/Source/Developer/TraceInsights/Private/Insights/ViewModels/BaseTimingTrack.h
@@ -47,6 +47,7 @@ public:
 	bool IsVisible() const { return bIsVisible; }
 	void Show() { SetVisibilityFlag(true); }
 	void Hide() { SetVisibilityFlag(false); }
+	void ToggleVisibility() { SetVisibilityFlag(!IsVisible()); }
 	virtual void SetVisibilityFlag(bool bInIsVisible) { bIsVisible = bInIsVisible; }
 
 	bool IsSelected() const { return bIsSelected; }
diff --git a/Engine/Source/Developer/TraceInsights/Private/Insights/ViewModels/TimingEventsTrack.cpp b/Engine/Source/Developer/TraceInsights/Private/Insights/ViewModels/TimingEventsTrack.cpp
--- a/Engine/Source/Developer/TraceInsights/Private/Insights/ViewModels/TimingEventsTrack.cpp
+++ b/Engine/Source/Developer/TraceInsights/Private/Insights/ViewModels/TimingEventsTrack.cpp
@@ -65,8 +65,10 @@ void FTimingEventsTrackLayout::Update()
 // FTimingEventsTrack
 ////////////////////////////////////////////////////////////////////////////////////////////////////
 
-FTimingEventsTrack::FTimingEventsTrack(uint64 InTrackId)
+FTimingEventsTrack::FTimingEventsTrack(uint64 InTrackId, ETimingEventsTrackType InType, const FString& InName)
 	: FBaseTimingTrack(InTrackId)
+	, Type(InType)
+	, Name(InName)
 	, Depth(-1)
 	, bIsCollapsed(false)
 {
diff --git a/Engine/Source/Developer/TraceInsights/Private/Insights/ViewModels/TimingEventsTrack.h b/Engine/Source/Developer/TraceInsights/Private/Insights/ViewModels/TimingEventsTrack.h
--- a/Engine/Source/Developer/TraceInsights/Private/Insights/ViewModels/TimingEventsTrack.h
+++ b/Engine/Source/Developer/TraceInsights/Private/Insights/ViewModels/TimingEventsTrack.h
@@ -41,17 +41,39 @@ struct FTimingEventsTrackLayout
 
 ////////////////////////////////////////////////////////////////////////////////////////////////////
 
+enum ETimingEventsTrackType
+{
+	Cpu,
+	Gpu,
+	Io
+};
+
+////////////////////////////////////////////////////////////////////////////////////////////////////
+
 class FTimingEventsTrack : public FBaseTimingTrack
 {
 public:
-	FTimingEventsTrack(uint64 InTrackId);
+	FTimingEventsTrack(uint64 InTrackId, ETimingEventsTrackType InType, const FString& InName);
 	virtual ~FTimingEventsTrack();
 
+	ETimingEventsTrackType GetType() const { return Type; }
+	const FString& GetName() const { return Name; }
+
+	void SetOrder(int32 InOrder) { Order = InOrder; }
+	int32 GetOrder() const { return Order; }
+
+	int32 GetDepth() const { return Depth; }
+	int32 GetNumLanes() const { return Depth + 1; }
+
 	virtual void Reset() override;
 	virtual void UpdateHoveredState(float MX, float MY, const FTimingTrackViewport& Viewport);
 
 public:
-	int32 Depth; // how many lanes has this track
+	ETimingEventsTrackType Type;
+	FString Name;
+	int32 Order;
+
+	int32 Depth; // number of lanes == Depth + 1
 	bool bIsCollapsed;
 
 	// Cached OnPaint state.
diff --git a/Engine/Source/Developer/TraceInsights/Private/Insights/ViewModels/TimingViewDrawHelper.cpp b/Engine/Source/Developer/TraceInsights/Private/Insights/ViewModels/TimingViewDrawHelper.cpp
--- a/Engine/Source/Developer/TraceInsights/Private/Insights/ViewModels/TimingViewDrawHelper.cpp
+++ b/Engine/Source/Developer/TraceInsights/Private/Insights/ViewModels/TimingViewDrawHelper.cpp
@@ -171,7 +171,7 @@ bool FTimingViewDrawHelper::BeginTimeline(FTimingEventsTrack& Track)
 
 ////////////////////////////////////////////////////////////////////////////////////////////////////
 
-void FTimingViewDrawHelper::EndTimeline(FTimingEventsTrack& Track, const TCHAR* TimelineName)
+void FTimingViewDrawHelper::EndTimeline(FTimingEventsTrack& Track)
 {
 	// Flush merged boxes.
 	for (int32 Depth = 0; Depth <= MaxDepth; ++Depth)
@@ -252,9 +252,9 @@ void FTimingViewDrawHelper::EndTimeline(FTimingEventsTrack& Track, const TCHAR*
 		// Draw name of timeline.
 		//const FString Name = FString::Printf(TEXT("%d Y: %g, H: %g, VO: %g, VY: %g"), TimelineIndex, Track.Y, Track.H, Viewport.TopOffset, Viewport.ScrollPosY); //debug
 		const TSharedRef<FSlateFontMeasure> FontMeasureService = FSlateApplication::Get().GetRenderer()->GetFontMeasureService();
-		float NameWidth = FontMeasureService->Measure(TimelineName, EventFont).X;
+		float NameWidth = FontMeasureService->Measure(Track.GetName(), EventFont).X;
 		DC.DrawBox(DC.LayerId + TimelineHeaderLayer, 0.0f, Y + 1.0f, NameWidth + 4.0f, 12.0f, WhiteBrush, Color);
-		DC.DrawText(DC.LayerId + TimelineTextLayer, 2.0f, Y, TimelineName, EventFont, TextColor);
+		DC.DrawText(DC.LayerId + TimelineTextLayer, 2.0f, Y, Track.GetName(), EventFont, TextColor);
 	}
 
 	TimelineTopY += Track.H;
diff --git a/Engine/Source/Developer/TraceInsights/Private/Insights/ViewModels/TimingViewDrawHelper.h b/Engine/Source/Developer/TraceInsights/Private/Insights/ViewModels/TimingViewDrawHelper.h
--- a/Engine/Source/Developer/TraceInsights/Private/Insights/ViewModels/TimingViewDrawHelper.h
+++ b/Engine/Source/Developer/TraceInsights/Private/Insights/ViewModels/TimingViewDrawHelper.h
@@ -69,7 +69,7 @@ public:
 	void BeginTimelines();
 	bool BeginTimeline(FTimingEventsTrack& Track);
 	void AddEvent(double StartTime, double EndTime, uint32 Depth, const TCHAR* EventName, uint32 Color = 0);
-	void EndTimeline(FTimingEventsTrack& Track, const TCHAR* TimelineName);
+	void EndTimeline(FTimingEventsTrack& Track);
 	void EndTimelines();
 
 	void DrawTimingEventHighlight(double StartTime, double EndTime, float Y, EHighlightMode Mode);
diff --git a/Engine/Source/Developer/TraceInsights/Private/Insights/Widgets/STimingView.cpp b/Engine/Source/Developer/TraceInsights/Private/Insights/Widgets/STimingView.cpp
--- a/Engine/Source/Developer/TraceInsights/Private/Insights/Widgets/STimingView.cpp
+++ b/Engine/Source/Developer/TraceInsights/Private/Insights/Widgets/STimingView.cpp
@@ -110,6 +110,44 @@ void STimingView::Construct(const FArguments& InArgs)
 			//.RenderOpacity(0.75)
 			.OnUserScrolled(this, &STimingView::VerticalScrollBar_OnUserScrolled)
 		]
+
+		+ SOverlay::Slot()
+		.HAlign(HAlign_Left)
+		.VAlign(VAlign_Top)
+		.Padding(FMargin(0.0f, 0.0f, 0.0f, 0.0f))
+		[
+			// Tracks Filter
+			SNew(SComboButton)
+			.ComboButtonStyle(FEditorStyle::Get(), "GenericFilters.ComboButtonStyle")
+			.ForegroundColor(FLinearColor::White)
+			.ToolTipText(LOCTEXT("TracksFilterToolTip", "Filter timing tracks."))
+			.OnGetMenuContent(this, &STimingView::MakeTracksFilterMenu)
+			.HasDownArrow(true)
+			.ContentPadding(FMargin(1.0f, 1.0f, 1.0f, 1.0f))
+			.ButtonContent()
+			[
+				SNew(SHorizontalBox)
+
+				+SHorizontalBox::Slot()
+				.Padding(0.0f)
+				.AutoWidth()
+				[
+					SNew(STextBlock)
+					.TextStyle(FEditorStyle::Get(), "GenericFilters.TextStyle")
+					.Font(FEditorStyle::Get().GetFontStyle("FontAwesome.9"))
+					.Text(FText::FromString(FString(TEXT("\xf0b0"))) /*fa-filter*/)
+				]
+
+				+SHorizontalBox::Slot()
+				.Padding(2.0f, 0.0f, 0.0f, 0.0f)
+				.AutoWidth()
+				[
+					SNew(STextBlock)
+					.TextStyle(FEditorStyle::Get(), "GenericFilters.TextStyle")
+					.Text(LOCTEXT("TracksFilter", "Tracks"))
+				]
+			]
+		]
 	];
 
 	UpdateHorizontalScrollBar();
@@ -122,8 +160,9 @@ void STimingView::Construct(const FArguments& InArgs)
 
 void STimingView::Reset()
 {
-	bDrawTimingEvents = true;
-	bDrawIoEvents = false;
+	bShowHideAllGpuTracks = true;
+	bShowHideAllCpuTracks = true;
+	bShowHideAllIoTracks = false;
 
 	//////////////////////////////////////////////////
 
@@ -146,7 +185,7 @@ void STimingView::Reset()
 
 	//////////////////////////////////////////////////
 
-	//TimingEventsTracks.Reset();
+	TimingEventsTracks.Reset();
 
 	bAreTimingEventsTracksDirty = true;
 
@@ -154,6 +193,9 @@ void STimingView::Reset()
 
 	//////////////////////////////////////////////////
 
+	IoOverviewTrack = nullptr;
+	IoActivityTrack = nullptr;
+
 	bForceIoEventsUpdate = true;
 	bMergeIoLanes = true;
 	AllIoEvents.Reset();
@@ -307,15 +349,15 @@ void STimingView::Tick(const FGeometry& AllottedGeometry, const double InCurrent
 		Wnd->TimersView->RebuildTree(false);
 	}
 
-	// Check if horizontal scroll area has changed.
+	UpdateIo();
+
 	TSharedPtr<const Trace::IAnalysisSession> Session = FInsightsManager::Get()->GetSession();
 	if (Session)
 	{
-		double SessionTime = 0.0;
-		{
-			Trace::FAnalysisSessionReadScope _(*Session.Get());
-			SessionTime = Session->GetDurationSeconds();
-		}
+		Trace::FAnalysisSessionReadScope _(*Session.Get());
+
+		// Check if horizontal scroll area has changed.
+		double SessionTime = Session->GetDurationSeconds();
 		if (SessionTime > Viewport.MaxValidTime &&
 			SessionTime != DBL_MAX &&
 			SessionTime != std::numeric_limits<double>::infinity())
@@ -331,17 +373,87 @@ void STimingView::Tick(const FGeometry& AllottedGeometry, const double InCurrent
 				MarkersTrack.SetDirtyFlag();
 			}
 		}
+
+		bool bIsTimingEventsTrackDirty = false;
+
+		if (IoOverviewTrack == nullptr)
+		{
+			IoOverviewTrack = AddTimingEventsTrack((0xF << 28) | 1, ETimingEventsTrackType::Io, TEXT("I/O Overview"), -1);
+			IoOverviewTrack->SetVisibilityFlag(bShowHideAllIoTracks);
+			bIsTimingEventsTrackDirty = true;
+		}
+
+		// Sync CachedTimelines and TimingEventsTracks with available timelines on analysis side.
+		// TODO: can we make this more efficient?
+		Session->ReadTimingProfilerProvider([this, &bIsTimingEventsTrackDirty, &Session](const Trace::ITimingProfilerProvider& TimingProfilerProvider)
+		{
+			// Check if we have a GPU track.
+			uint32 GpuTimelineIndex;
+			if (TimingProfilerProvider.GetGpuTimelineIndex(GpuTimelineIndex))
+			{
+				if (!CachedTimelines.Contains(GpuTimelineIndex))
+				{
+					AddTimingEventsTrack(GpuTimelineIndex, ETimingEventsTrackType::Gpu, TEXT("GPU"), 0);
+					bIsTimingEventsTrackDirty = true;
+				}
+			}
+
+			// Check available CPU tracks.
+			Session->ReadThreadProvider([this, &bIsTimingEventsTrackDirty, &TimingProfilerProvider](const Trace::IThreadProvider& ThreadProvider)
+			{
+				int Order = 1;
+				ThreadProvider.EnumerateThreads([this, &Order, &bIsTimingEventsTrackDirty, &TimingProfilerProvider](const Trace::FThreadInfo& ThreadInfo)
+				{
+					uint32 CpuTimelineIndex;
+					if (TimingProfilerProvider.GetCpuThreadTimelineIndex(ThreadInfo.Id, CpuTimelineIndex))
+					{
+						if (!CachedTimelines.Contains(CpuTimelineIndex))
+						{
+							FString TrackName(ThreadInfo.Name && *ThreadInfo.Name ? ThreadInfo.Name : FString::Printf(TEXT("Thread %u")));
+							AddTimingEventsTrack(CpuTimelineIndex, ETimingEventsTrackType::Cpu, TrackName, Order);
+							bIsTimingEventsTrackDirty = true;
+						}
+						else
+						{
+							FTimingEventsTrack* Track = CachedTimelines[CpuTimelineIndex];
+							if (Track->GetOrder() != Order)
+							{
+								Track->SetOrder(Order);
+								bIsTimingEventsTrackDirty = true;
+							}
+						}
+						Order++;
+					}
+				});
+			});
+		});
+
+		if (IoActivityTrack == nullptr)
+		{
+			IoActivityTrack = AddTimingEventsTrack((0xF << 28) | 2, ETimingEventsTrackType::Io, TEXT("I/O Activity"), 999999);
+			IoActivityTrack->SetVisibilityFlag(bShowHideAllIoTracks);
+			bIsTimingEventsTrackDirty = true;
+		}
+
+		if (bIsTimingEventsTrackDirty)
+		{
+			// The list has changed. Sort the list again.
+			TimingEventsTracks.Sort([this](const FTimingEventsTrack& A, const FTimingEventsTrack& B) -> bool
+			{
+				return A.GetOrder() < B.GetOrder();
+			});
+		}
 	}
 
 	// Compute total height of visible tracks.
 	float TotalScrollHeight = 0.0f;
-	for (auto& KV : CachedTimelines)
+	for (int TrackIndex = 0; TrackIndex < TimingEventsTracks.Num(); ++TrackIndex)
 	{
-		FTimingEventsTrack& Track = *KV.Value;
+		FTimingEventsTrack& Track = *TimingEventsTracks[TrackIndex];
 
-		//if (Track.IsVisible)
+		if (Track.IsVisible())
 		{
-			//Track.Y = TotalHeight; // TODO: this is not the draw order!
+			Track.SetPosY(TotalScrollHeight);
 			TotalScrollHeight += Track.GetHeight();
 		}
 
@@ -426,10 +538,10 @@ void STimingView::Tick(const FGeometry& AllottedGeometry, const double InCurrent
 			GraphTrack.Update(Viewport);
 		}
 
-		//for (int TrackIndex = 0; TrackIndex < TimingEventsTracks.Num(); ++TrackIndex)
-		//{
-		//	TimingEventsTracks[TrackIndex]->Update(Viewport);
-		//}
+		for (int TrackIndex = 0; TrackIndex < TimingEventsTracks.Num(); ++TrackIndex)
+		{
+			TimingEventsTracks[TrackIndex]->Update(Viewport);
+		}
 
 		Stopwatch.Stop();
 		TimelineCacheUpdateDurationHistory.AddValue(Stopwatch.AccumulatedTime);
@@ -488,19 +600,32 @@ int32 STimingView::OnPaint(const FPaintArgs& Args, const FGeometry& AllottedGeom
 	{
 		Helper.BeginTimelines();
 
-		if (bDrawIoEvents)
+		for (int TrackIndex = 0; TrackIndex < TimingEventsTracks.Num(); ++TrackIndex)
 		{
-			DrawIoOverviewTrack(Helper);
-		}
-
-		if (bDrawTimingEvents)
-		{
-			DrawTimingEvents(Helper);
-		}
-
-		if (bDrawIoEvents)
-		{
-			DrawIoTrack(Helper);
+			FTimingEventsTrack& Track = *TimingEventsTracks[TrackIndex];
+			if (Track.IsVisible())
+			{
+				//TODO: Track.Draw(Helper);
+				if (Track.GetType() == ETimingEventsTrackType::Cpu)
+				{
+					DrawCpuGpuTimelineTrack(Helper, Track);
+				}
+				else if (Track.GetType() == ETimingEventsTrackType::Gpu)
+				{
+					DrawCpuGpuTimelineTrack(Helper, Track);
+				}
+				else if (Track.GetType() == ETimingEventsTrackType::Io)
+				{
+					if (&Track == IoOverviewTrack)
+					{
+						DrawIoOverviewTrack(Helper, Track);
+					}
+					else if (&Track == IoActivityTrack)
+					{
+						DrawIoActivityTrack(Helper, Track);
+					}
+				}
+			}
 		}
 
 		Helper.EndTimelines();
@@ -553,6 +678,9 @@ int32 STimingView::OnPaint(const FPaintArgs& Args, const FGeometry& AllottedGeom
 		TimeRulerTrack.Draw(DC, Viewport, MousePosition, bIsSelecting, SelectionStartTime, SelectionEndTime);
 	}
 
+	// Fill background for the Tracks filter combobox.
+	DC.DrawBox(0, 0, 66, 22, WhiteBrush, FLinearColor(0.05f, 0.05f, 0.05f, 1.0f));
+
 	//////////////////////////////////////////////////
 
 	// Draw the time range selection.
@@ -846,93 +974,63 @@ int32 STimingView::OnPaint(const FPaintArgs& Args, const FGeometry& AllottedGeom
 
 ////////////////////////////////////////////////////////////////////////////////////////////////////
 
-FTimingEventsTrack* STimingView::GetOrAddTimingEventsTrack(uint32 TrackId)
+FTimingEventsTrack* STimingView::AddTimingEventsTrack(uint32 TrackId, ETimingEventsTrackType TrackType, const FString& TrackName, int32 Order)
 {
-	if (CachedTimelines.Contains(TrackId))
-	{
-		return CachedTimelines[TrackId];
-	}
-	else
-	{
-		FTimingEventsTrack* Track = new FTimingEventsTrack(TrackId);
-		CachedTimelines.Add(TrackId, Track);
-		return Track;
-	}
+	FTimingEventsTrack* Track = new FTimingEventsTrack(TrackId, TrackType, TrackName);
+
+	Track->SetOrder(Order);
+
+	UE_LOG(TimingProfiler, Log, TEXT("New Timing Events Track (%d) : %s (\"%s\")"),
+		TimingEventsTracks.Num() + 1,
+		TrackType == ETimingEventsTrackType::Gpu ? TEXT("GPU") :
+		TrackType == ETimingEventsTrackType::Cpu ? TEXT("CPU") :
+		TrackType == ETimingEventsTrackType::Io ? TEXT("I/O") : TEXT("unknown"),
+		*TrackName);
+
+	CachedTimelines.Add(TrackId, Track);
+	TimingEventsTracks.Add(Track);
+
+	return Track;
 }
 
 ////////////////////////////////////////////////////////////////////////////////////////////////////
 
-void STimingView::DrawTimingEvents(FTimingViewDrawHelper& DC) const
+void STimingView::DrawCpuGpuTimelineTrack(FTimingViewDrawHelper& Helper, FTimingEventsTrack& Track) const
 {
-	TSharedPtr<const Trace::IAnalysisSession> Session = FInsightsManager::Get()->GetSession();
-	if (Session.IsValid())
+	if (Helper.BeginTimeline(Track))
 	{
-		Trace::FAnalysisSessionReadScope _(*Session.Get());
-
-		struct FTimelineHelper
+		TSharedPtr<const Trace::IAnalysisSession> Session = FInsightsManager::Get()->GetSession();
+		if (Session.IsValid())
 		{
-			void DrawTimeline(uint32 TimelineIndex, const TCHAR* Name)
-			{
-				// TODO: Adding tracks should be done in Tick(), not in OnPaint().
-				FTimingEventsTrack& Track = *const_cast<STimingView&>(Outer).GetOrAddTimingEventsTrack(TimelineIndex);
+			Trace::FAnalysisSessionReadScope _(*Session.Get());
 
-				if (DC.BeginTimeline(Track))
+			Session->ReadTimingProfilerProvider([this, &Helper, &Track](const Trace::ITimingProfilerProvider& TimingProfilerProvider)
+			{
+				TimingProfilerProvider.ReadTimers([this, &Helper, &Track, &TimingProfilerProvider](const Trace::FTimingProfilerTimer* Timers, uint64 TimersCount)
 				{
-					TimingProfilerProvider.ReadTimeline(TimelineIndex, [this](const Trace::ITimingProfilerProvider::Timeline& Timeline)
+					TimingProfilerProvider.ReadTimeline(Track.GetId(), [this, &Helper, &Track, Timers](const Trace::ITimingProfilerProvider::Timeline& Timeline)
 					{
-						if (Outer.bUseDownSampling)
+						if (bUseDownSampling)
 						{
-							const double SecondsPerPixel = 1.0 / DC.GetViewport().ScaleX;
-							Timeline.EnumerateEventsDownSampled(DC.GetViewport().StartTime, DC.GetViewport().EndTime, SecondsPerPixel, [this](double StartTime, double EndTime, uint32 Depth, const Trace::FTimingProfilerEvent& Event)
+							const double SecondsPerPixel = 1.0 / Helper.GetViewport().ScaleX;
+							Timeline.EnumerateEventsDownSampled(Helper.GetViewport().StartTime, Helper.GetViewport().EndTime, SecondsPerPixel, [this, &Helper, Timers](double StartTime, double EndTime, uint32 Depth, const Trace::FTimingProfilerEvent& Event)
 							{
-								DC.AddEvent(StartTime, EndTime, Depth, *Timers[Event.TimerIndex].Name);
+								Helper.AddEvent(StartTime, EndTime, Depth, *Timers[Event.TimerIndex].Name);
 							});
 						}
 						else
 						{
-							Timeline.EnumerateEvents(DC.GetViewport().StartTime, DC.GetViewport().EndTime, [this](double StartTime, double EndTime, uint32 Depth, const Trace::FTimingProfilerEvent& Event)
+							Timeline.EnumerateEvents(Helper.GetViewport().StartTime, Helper.GetViewport().EndTime, [this, &Helper, Timers](double StartTime, double EndTime, uint32 Depth, const Trace::FTimingProfilerEvent& Event)
 							{
-								DC.AddEvent(StartTime, EndTime, Depth, *Timers[Event.TimerIndex].Name);
+								Helper.AddEvent(StartTime, EndTime, Depth, *Timers[Event.TimerIndex].Name);
 							});
 						}
 					});
-					DC.EndTimeline(Track, Name);
-				}
-			}
-
-			const STimingView& Outer;
-			FTimingViewDrawHelper& DC;
-			const Trace::ITimingProfilerProvider& TimingProfilerProvider;
-			const Trace::FTimingProfilerTimer* Timers;
-		};
-
-		Session->ReadTimingProfilerProvider([this, &DC, &Session](const Trace::ITimingProfilerProvider& TimingProfilerProvider)
-		{
-			TimingProfilerProvider.ReadTimers([this, &DC, &Session, &TimingProfilerProvider](const Trace::FTimingProfilerTimer* Timers, uint64 TimersCount)
-			{
-				FTimelineHelper Helper{ *this, DC, TimingProfilerProvider, Timers };
-
-				// Draw GPU track.
-				uint32 GpuTimelineIndex;
-				if (TimingProfilerProvider.GetGpuTimelineIndex(GpuTimelineIndex))
-				{
-					Helper.DrawTimeline(GpuTimelineIndex, TEXT("GPU"));
-				}
-
-				// Draw CPU tracks.
-				Session->ReadThreadProvider([this, &Helper, &TimingProfilerProvider](const Trace::IThreadProvider& ThreadProvider)
-				{
-					ThreadProvider.EnumerateThreads([this, &Helper, &TimingProfilerProvider](const Trace::FThreadInfo& ThreadInfo)
-					{
-						uint32 ThreadTimelineIndex;
-						if (TimingProfilerProvider.GetCpuThreadTimelineIndex(ThreadInfo.Id, ThreadTimelineIndex))
-						{
-							Helper.DrawTimeline(ThreadTimelineIndex, ThreadInfo.Name && *ThreadInfo.Name ? ThreadInfo.Name : *FString::Printf(TEXT("Thread %u"), ThreadInfo.Id));
-						}
-					});
 				});
 			});
-		});
+
+			Helper.EndTimeline(Track);
+		}
 	}
 }
 
@@ -975,14 +1073,9 @@ uint32 GetFileActivityTypeColor(Trace::EFileActivityType Type)
 
 ////////////////////////////////////////////////////////////////////////////////////////////////////
 
-void STimingView::DrawIoOverviewTrack(FTimingViewDrawHelper& DC) const
+void STimingView::DrawIoOverviewTrack(FTimingViewDrawHelper& Helper, FTimingEventsTrack& Track) const
 {
-	// TODO: Adding tracks should be done in Tick(), not in OnPaint().
-	STimingView& NonConstThis = const_cast<STimingView&>(*this);
-	FTimingEventsTrack& Track = *NonConstThis.GetOrAddTimingEventsTrack((0xF << 28) | 1);
-	//TODO: FIoTimingEventsTrack& Track = *GetOrAddTimingEventsTrack(ETrackType::Io, 1);
-
-	if (DC.BeginTimeline(Track))
+	if (Helper.BeginTimeline(Track))
 	{
 		// Draw the IO Overiew track using cached events (as those are sorted by Start Time).
 		for (const FIoTimingEvent& Event : AllIoEvents)
@@ -994,11 +1087,11 @@ void STimingView::DrawIoOverviewTrack(FTimingViewDrawHelper& DC) const
 				// Ignore "Idle" and "NotClosed" events.
 				continue;
 			}
-			if (Event.EndTime <= DC.GetViewport().StartTime)
+			if (Event.EndTime <= Helper.GetViewport().StartTime)
 			{
 				continue;
 			}
-			if (Event.StartTime >= DC.GetViewport().EndTime)
+			if (Event.StartTime >= Helper.GetViewport().EndTime)
 			{
 				break;
 			}
@@ -1015,7 +1108,7 @@ void STimingView::DrawIoOverviewTrack(FTimingViewDrawHelper& DC) const
 				Name += Event.Path;
 				Name += "]";
 				Color = 0xFFAA0000;
-				DC.AddEvent(Event.StartTime, Event.EndTime, 0, *Name, Color);
+				Helper.AddEvent(Event.StartTime, Event.EndTime, 0, *Name, Color);
 			}
 			else if (ActivityType == Trace::FileActivityType_Open)
 			{
@@ -1023,28 +1116,23 @@ void STimingView::DrawIoOverviewTrack(FTimingViewDrawHelper& DC) const
 				Name += " [";
 				Name += Event.Path;
 				Name += "]";
-				DC.AddEvent(Event.StartTime, Event.EndTime, 0, *Name, Color);
+				Helper.AddEvent(Event.StartTime, Event.EndTime, 0, *Name, Color);
 			}
 			else
 			{
-				DC.AddEvent(Event.StartTime, Event.EndTime, 0, GetFileActivityTypeName(ActivityType), Color);
+				Helper.AddEvent(Event.StartTime, Event.EndTime, 0, GetFileActivityTypeName(ActivityType), Color);
 			}
 		}
 
-		DC.EndTimeline(Track, TEXT("IO Overview"));
+		Helper.EndTimeline(Track);
 	}
 }
 
 ////////////////////////////////////////////////////////////////////////////////////////////////////
 //PRAGMA_DISABLE_OPTIMIZATION
 
-void STimingView::DrawIoTrack(FTimingViewDrawHelper& DC) const
+void STimingView::UpdateIo()
 {
-	// TODO: Adding tracks should be done in Tick(), not in OnPaint().
-	STimingView& NonConstThis = const_cast<STimingView&>(*this);
-	FTimingEventsTrack& Track = *NonConstThis.GetOrAddTimingEventsTrack((0xF << 28) | 2);
-	//TODO: FIoTimingEventsTrack& Track = *GetOrAddTimingEventsTrack(ETrackType::Io, 2);
-
 	if (bForceIoEventsUpdate)
 	{
 		bForceIoEventsUpdate = false;
@@ -1056,13 +1144,13 @@ void STimingView::DrawIoTrack(FTimingViewDrawHelper& DC) const
 		{
 			Trace::FAnalysisSessionReadScope _(*Session.Get());
 
-			Session->ReadFileActivityProvider([this, &DC](const Trace::IFileActivityProvider& FileActivityProvider)
+			Session->ReadFileActivityProvider([this](const Trace::IFileActivityProvider& FileActivityProvider)
 			{
 				// Enumerate all IO events and cache them.
-				FileActivityProvider.EnumerateFileActivity([this, &DC](const Trace::FFileInfo& FileInfo, const Trace::IFileActivityProvider::Timeline& Timeline)
+				FileActivityProvider.EnumerateFileActivity([this](const Trace::FFileInfo& FileInfo, const Trace::IFileActivityProvider::Timeline& Timeline)
 				{
 					Timeline.EnumerateEvents(-std::numeric_limits<double>::infinity(), +std::numeric_limits<double>::infinity(),
-						[this, &DC, &FileInfo, &Timeline](double EventStartTime, double EventEndTime, uint32 EventDepth, const Trace::FFileActivity& FileActivity)
+						[this, &FileInfo, &Timeline](double EventStartTime, double EventEndTime, uint32 EventDepth, const Trace::FFileActivity& FileActivity)
 					{
 						if (bMergeIoLanes)
 						{
@@ -1212,17 +1300,23 @@ void STimingView::DrawIoTrack(FTimingViewDrawHelper& DC) const
 			//});
 		}
 	}
+}
+
+//PRAGMA_ENABLE_OPTIMIZATION
+////////////////////////////////////////////////////////////////////////////////////////////////////
 
-	if (DC.BeginTimeline(Track))
+void STimingView::DrawIoActivityTrack(FTimingViewDrawHelper& Helper, FTimingEventsTrack& Track) const
+{
+	if (Helper.BeginTimeline(Track))
 	{
 		// Draw IO track using cached events.
 		for (const FIoTimingEvent& Event : AllIoEvents)
 		{
-			if (Event.EndTime <= DC.GetViewport().StartTime)
+			if (Event.EndTime <= Helper.GetViewport().StartTime)
 			{
 				continue;
 			}
-			if (Event.StartTime >= DC.GetViewport().EndTime)
+			if (Event.StartTime >= Helper.GetViewport().EndTime)
 			{
 				break;
 			}
@@ -1243,7 +1337,7 @@ void STimingView::DrawIoTrack(FTimingViewDrawHelper& DC) const
 					Name += Event.Path;
 					Name += "]";
 					Color = 0xFFAA0000;
-					DC.AddEvent(Event.StartTime, Event.EndTime, Event.Depth, *Name, Color);
+					Helper.AddEvent(Event.StartTime, Event.EndTime, Event.Depth, *Name, Color);
 				}
 				else if (ActivityType == Trace::FileActivityType_Open)
 				{
@@ -1251,11 +1345,11 @@ void STimingView::DrawIoTrack(FTimingViewDrawHelper& DC) const
 					Name += " [";
 					Name += Event.Path;
 					Name += "]";
-					DC.AddEvent(Event.StartTime, Event.EndTime, Event.Depth, *Name, Color);
+					Helper.AddEvent(Event.StartTime, Event.EndTime, Event.Depth, *Name, Color);
 				}
 				else
 				{
-					DC.AddEvent(Event.StartTime, Event.EndTime, Event.Depth, GetFileActivityTypeName(ActivityType), Color);
+					Helper.AddEvent(Event.StartTime, Event.EndTime, Event.Depth, GetFileActivityTypeName(ActivityType), Color);
 				}
 			}
 			else
@@ -1264,15 +1358,14 @@ void STimingView::DrawIoTrack(FTimingViewDrawHelper& DC) const
 				Name += " [";
 				Name += Event.Path;
 				Name += "]";
-				DC.AddEvent(Event.StartTime, Event.EndTime, Event.Depth, *Name, Color);
+				Helper.AddEvent(Event.StartTime, Event.EndTime, Event.Depth, *Name, Color);
 			}
 		}
 
-		DC.EndTimeline(Track, TEXT("IO Activity"));
+		Helper.EndTimeline(Track);
 	}
 }
 
-//PRAGMA_ENABLE_OPTIMIZATION
 ////////////////////////////////////////////////////////////////////////////////////////////////////
 
 void STimingView::DrawTimeRangeSelection(FDrawContext& DC) const
@@ -1876,26 +1969,7 @@ FReply STimingView::OnKeyDown(const FGeometry& MyGeometry, const FKeyEvent& InKe
 	}
 	else if (InKeyEvent.GetKey() == EKeys::V)
 	{
-		if (Layout.TargetMinTimelineH == 0.0f)
-		{
-			Layout.TargetMinTimelineH = RealMinTimelineH;
-		}
-		else
-		{
-			Layout.TargetMinTimelineH = 0.0f;
-		}
-
-		Layout.MinTimelineH = Layout.TargetMinTimelineH; // no layout animation
-
-		for (auto& CachedTimelineKV : CachedTimelines)
-		{
-			CachedTimelineKV.Value->SetHeight(0.0f);
-		}
-
-		Viewport.ScrollPosY = 0.0f;
-		UpdateVerticalScrollBar();
-		bIsVerticalViewportDirty = true;
-
+		ToggleAutoHideEmptyTracks();
 		return FReply::Handled();
 	}
 	else if (InKeyEvent.GetKey() == EKeys::Equals ||
@@ -2019,20 +2093,19 @@ FReply STimingView::OnKeyDown(const FGeometry& MyGeometry, const FKeyEvent& InKe
 	}
 	else if (InKeyEvent.GetKey() == EKeys::G) // debug: toggles Graph track on/off
 	{
-		GraphTrack.SetVisibilityFlag(!GraphTrack.IsVisible());
+		GraphTrack.ToggleVisibility();
 		bAreTimingEventsTracksDirty = true;
 		return FReply::Handled();
 	}
 	else if (InKeyEvent.GetKey() == EKeys::U) // debug: toggles Timing (GPU, CPU) tracks on/off
 	{
-		bDrawTimingEvents = !bDrawTimingEvents;
-		CachedTimelines.Reset();
+		ShowHideAllGpuTracks_Execute();
+		ShowHideAllCpuTracks_Execute();
 		return FReply::Handled();
 	}
 	else if (InKeyEvent.GetKey() == EKeys::I)  // debug: toggles IO tracks on/off
 	{
-		bDrawIoEvents = !bDrawIoEvents;
-		CachedTimelines.Reset();
+		ShowHideAllIoTracks_Execute();
 		return FReply::Handled();
 	}
 	else if (InKeyEvent.GetKey() == EKeys::O)  // debug: toggles IO merge lanes algorithm
@@ -2701,6 +2774,227 @@ const FTimerNodePtr STimingView::GetTimerNode(uint64 TypeId) const
 	return TimerNodePtrPtr ? *TimerNodePtrPtr : nullptr;
 }
 
+////////////////////////////////////////////////////////////////////////////////////////////////////
+
+TSharedRef<SWidget> STimingView::MakeTracksFilterMenu()
+{
+	FMenuBuilder MenuBuilder(/*bInShouldCloseWindowAfterMenuSelection=*/true, nullptr);
+
+	MenuBuilder.BeginSection("TracksFilter", LOCTEXT("TracksFilterHeading", "Quick Filter"));
+	{
+		MenuBuilder.AddMenuEntry(
+			LOCTEXT("ShowAllGpuTracks", "Show/Hide All GPU Tracks (U)"),
+			LOCTEXT("ShowAllGpuTracks_Tooltip", "Show/hide all GPU tracks"),
+			FSlateIcon(),
+			FUIAction(FExecuteAction::CreateSP(this, &STimingView::ShowHideAllGpuTracks_Execute),
+				FCanExecuteAction::CreateLambda([] { return true; }),
+				FIsActionChecked::CreateSP(this, &STimingView::ShowHideAllGpuTracks_IsChecked)),
+			NAME_None,
+			EUserInterfaceActionType::ToggleButton
+		);
+
+		MenuBuilder.AddMenuEntry(
+			LOCTEXT("ShowAllCpuTracks", "Show/Hide All CPU Tracks (U)"),
+			LOCTEXT("ShowAllCpuTracks_Tooltip", "Show/hide all CPU tracks"),
+			FSlateIcon(),
+			FUIAction(FExecuteAction::CreateSP(this, &STimingView::ShowHideAllCpuTracks_Execute),
+				FCanExecuteAction::CreateLambda([] { return true; }),
+				FIsActionChecked::CreateSP(this, &STimingView::ShowHideAllCpuTracks_IsChecked)),
+			NAME_None,
+			EUserInterfaceActionType::ToggleButton
+		);
+
+		MenuBuilder.AddMenuEntry(
+			LOCTEXT("ShowAllIoTracks", "Show/Hide All I/O Tracks (I)"),
+			LOCTEXT("ShowAllIoTracks_Tooltip", "Show/hide all I/O tracks"),
+			FSlateIcon(),
+			FUIAction(FExecuteAction::CreateSP(this, &STimingView::ShowHideAllIoTracks_Execute),
+				FCanExecuteAction::CreateLambda([] { return true; }),
+				FIsActionChecked::CreateSP(this, &STimingView::ShowHideAllIoTracks_IsChecked)),
+			NAME_None,
+			EUserInterfaceActionType::ToggleButton
+		);
+
+		MenuBuilder.AddMenuEntry(
+			LOCTEXT("AutoHideEmptyTracks", "Auto Hide Empty Tracks (V)"),
+			LOCTEXT("AutoHideEmptyTracks_Tooltip", "Auto hide empty tracks (ones without timing events in current viewport)"),
+			FSlateIcon(),
+			FUIAction(FExecuteAction::CreateSP(this, &STimingView::ToggleAutoHideEmptyTracks),
+				FCanExecuteAction::CreateLambda([] { return true; }),
+				FIsActionChecked::CreateSP(this, &STimingView::IsAutoHideEmptyTracksEnabled)),
+			NAME_None,
+			EUserInterfaceActionType::ToggleButton
+		);
+	}
+	MenuBuilder.EndSection();
+
+	MenuBuilder.BeginSection("Tracks", LOCTEXT("TracksHeading", "Tracks"));
+	CreateTracksMenu(MenuBuilder);
+	MenuBuilder.EndSection();
+
+	return MenuBuilder.MakeWidget();
+}
+
+////////////////////////////////////////////////////////////////////////////////////////////////////
+
+void STimingView::CreateTracksMenu(FMenuBuilder& MenuBuilder)
+{
+	for (int TrackIndex = 0; TrackIndex < TimingEventsTracks.Num(); ++TrackIndex)
+	{
+		const FTimingEventsTrack& Track = *TimingEventsTracks[TrackIndex];
+		if (Track.GetHeight() > 0.0f || Layout.TargetMinTimelineH > 0.0f)
+		{
+			MenuBuilder.AddMenuEntry(
+				FText::FromString(Track.GetName()),
+				TAttribute<FText>(), // no tooltip
+				FSlateIcon(),
+				FUIAction(FExecuteAction::CreateSP(this, &STimingView::ToggleTrackVisibility_Execute, Track.GetId()),
+					FCanExecuteAction::CreateLambda([] { return true; }),
+					FIsActionChecked::CreateSP(this, &STimingView::ToggleTrackVisibility_IsChecked, Track.GetId())),
+				NAME_None,
+				EUserInterfaceActionType::ToggleButton
+			);
+		}
+	}
+}
+
+////////////////////////////////////////////////////////////////////////////////////////////////////
+
+bool STimingView::ShowHideAllCpuTracks_IsChecked() const
+{
+	return bShowHideAllCpuTracks;
+}
+
+////////////////////////////////////////////////////////////////////////////////////////////////////
+
+void STimingView::ShowHideAllCpuTracks_Execute()
+{
+	bShowHideAllCpuTracks = !bShowHideAllCpuTracks;
+
+	for (int TrackIndex = 0; TrackIndex < TimingEventsTracks.Num(); ++TrackIndex)
+	{
+		FTimingEventsTrack& Track = *TimingEventsTracks[TrackIndex];
+		if (Track.GetType() == ETimingEventsTrackType::Cpu)
+		{
+			Track.SetVisibilityFlag(bShowHideAllCpuTracks);
+		}
+	}
+
+	bAreTimingEventsTracksDirty = true;
+}
+
+////////////////////////////////////////////////////////////////////////////////////////////////////
+
+bool STimingView::ShowHideAllGpuTracks_IsChecked() const
+{
+	return bShowHideAllGpuTracks;
+}
+
+////////////////////////////////////////////////////////////////////////////////////////////////////
+
+void STimingView::ShowHideAllGpuTracks_Execute()
+{
+	bShowHideAllGpuTracks = !bShowHideAllGpuTracks;
+
+	for (int TrackIndex = 0; TrackIndex < TimingEventsTracks.Num(); ++TrackIndex)
+	{
+		FTimingEventsTrack& Track = *TimingEventsTracks[TrackIndex];
+		if (Track.GetType() == ETimingEventsTrackType::Gpu)
+		{
+			Track.SetVisibilityFlag(bShowHideAllGpuTracks);
+		}
+	}
+
+	bAreTimingEventsTracksDirty = true;
+}
+
+////////////////////////////////////////////////////////////////////////////////////////////////////
+
+bool STimingView::ShowHideAllIoTracks_IsChecked() const
+{
+	return bShowHideAllIoTracks;
+}
+
+////////////////////////////////////////////////////////////////////////////////////////////////////
+
+void STimingView::ShowHideAllIoTracks_Execute()
+{
+	bShowHideAllIoTracks = !bShowHideAllIoTracks;
+
+	for (int TrackIndex = 0; TrackIndex < TimingEventsTracks.Num(); ++TrackIndex)
+	{
+		FTimingEventsTrack& Track = *TimingEventsTracks[TrackIndex];
+		if (Track.GetType() == ETimingEventsTrackType::Io)
+		{
+			Track.SetVisibilityFlag(bShowHideAllIoTracks);
+		}
+	}
+
+	bAreTimingEventsTracksDirty = true;
+
+	if (bShowHideAllIoTracks)
+	{
+		bForceIoEventsUpdate = true;
+	}
+}
+
+////////////////////////////////////////////////////////////////////////////////////////////////////
+
+bool STimingView::IsAutoHideEmptyTracksEnabled() const
+{
+	return (Layout.TargetMinTimelineH == 0.0f);
+}
+
+////////////////////////////////////////////////////////////////////////////////////////////////////
+
+void STimingView::ToggleAutoHideEmptyTracks()
+{
+	if (Layout.TargetMinTimelineH == 0.0f)
+	{
+		Layout.TargetMinTimelineH = RealMinTimelineH;
+	}
+	else
+	{
+		Layout.TargetMinTimelineH = 0.0f;
+	}
+
+	Layout.MinTimelineH = Layout.TargetMinTimelineH; // no layout animation
+
+	for (auto& CachedTimelineKV : CachedTimelines)
+	{
+		CachedTimelineKV.Value->SetHeight(0.0f);
+	}
+
+	Viewport.ScrollPosY = 0.0f;
+	UpdateVerticalScrollBar();
+	bIsVerticalViewportDirty = true;
+}
+
+////////////////////////////////////////////////////////////////////////////////////////////////////
+
+bool STimingView::ToggleTrackVisibility_IsChecked(uint64 InTrackId) const
+{
+	if (CachedTimelines.Contains(InTrackId))
+	{
+		const FTimingEventsTrack* Track = CachedTimelines[InTrackId];
+		return Track->IsVisible();
+	}
+	return false;
+}
+
+////////////////////////////////////////////////////////////////////////////////////////////////////
+
+void STimingView::ToggleTrackVisibility_Execute(uint64 InTrackId)
+{
+	if (CachedTimelines.Contains(InTrackId))
+	{
+		FTimingEventsTrack* Track = CachedTimelines[InTrackId];
+		Track->ToggleVisibility();
+	}
+}
+
+
+
 ////////////////////////////////////////////////////////////////////////////////////////////////////
 
 #undef LOCTEXT_NAMESPACE
diff --git a/Engine/Source/Developer/TraceInsights/Private/Insights/Widgets/STimingView.h b/Engine/Source/Developer/TraceInsights/Private/Insights/Widgets/STimingView.h
--- a/Engine/Source/Developer/TraceInsights/Private/Insights/Widgets/STimingView.h
+++ b/Engine/Source/Developer/TraceInsights/Private/Insights/Widgets/STimingView.h
@@ -64,6 +64,24 @@ public:
 	 */
 	void Construct(const FArguments& InArgs);
 
+	TSharedRef<SWidget> MakeTracksFilterMenu();
+	void CreateTracksMenu(FMenuBuilder& MenuBuilder);
+
+	bool ShowHideAllCpuTracks_IsChecked() const;
+	void ShowHideAllCpuTracks_Execute();
+
+	bool ShowHideAllGpuTracks_IsChecked() const;
+	void ShowHideAllGpuTracks_Execute();
+
+	bool ShowHideAllIoTracks_IsChecked() const;
+	void ShowHideAllIoTracks_Execute();
+
+	bool IsAutoHideEmptyTracksEnabled() const;
+	void ToggleAutoHideEmptyTracks();
+
+	bool ToggleTrackVisibility_IsChecked(uint64 InTrackId) const;
+	void ToggleTrackVisibility_Execute(uint64 InTrackId);
+
 	/** Resets internal widget's data to the default one. */
 	void Reset();
 
@@ -226,11 +244,14 @@ protected:
 		return FVector2D(16.0f, 16.0f);
 	}
 
-	FTimingEventsTrack* GetOrAddTimingEventsTrack(uint32 TrackId);
+	FTimingEventsTrack* AddTimingEventsTrack(uint32 TrackId, ETimingEventsTrackType TrackType, const FString& TrackName, int32 Order);
+
+	void DrawCpuGpuTimelineTrack(FTimingViewDrawHelper& Helper, FTimingEventsTrack& Track) const;
+
+	void UpdateIo();
+	void DrawIoOverviewTrack(FTimingViewDrawHelper& Helper, FTimingEventsTrack& Track) const;
+	void DrawIoActivityTrack(FTimingViewDrawHelper& Helper, FTimingEventsTrack& Track) const;
 
-	void DrawTimingEvents(FTimingViewDrawHelper& DH) const;
-	void DrawIoOverviewTrack(FTimingViewDrawHelper& DH) const;
-	void DrawIoTrack(FTimingViewDrawHelper& DH) const;
 	void DrawTimeRangeSelection(FDrawContext& DH) const;
 
 	void ShowContextMenu(const FVector2D& ScreenSpacePosition, const FPointerEvent& MouseEvent);
@@ -302,9 +323,9 @@ protected:
 	////////////////////////////////////////////////////////////////////////////////////////////////////
 
 protected:
-
-	bool bDrawTimingEvents;
-	bool bDrawIoEvents;
+	bool bShowHideAllGpuTracks;
+	bool bShowHideAllCpuTracks;
+	bool bShowHideAllIoTracks;
 
 	////////////////////////////////////////////////////////////
 
@@ -316,27 +337,28 @@ protected:
 
 	////////////////////////////////////////////////////////////
 
-	enum FTimingTrackType
-	{
-		TimeRuler,
-		Markers,
-		CpuEvents,
-		GpuEvents,
-		IoEvents,
-		Graph,
-		Other
-	};
+	//enum FTimingTrackType
+	//{
+	//	TimeRuler,
+	//	Markers,
+	//	CpuEvents,
+	//	GpuEvents,
+	//	IoEvents,
+	//	Graph,
+	//	Other
+	//};
 
 	mutable TMap<uint32, FTimingEventsTrack*> CachedTimelines; /// all tracks
 
-	//TODO: TArray<FBaseTimingTrack*> TopTracks; /// tracks docked on top, in order to be displayed
-	//TODO: TArray<FBaseTimingTrack*> ScrollableTracks; /// tracks in scrollable area, in order to be displayed
-	//TODO: TArray<FBaseTimingTrack*> BottomTracks; /// tracks docked on bottom, in order to be displayed
+	//TODO: TArray<FBaseTimingTrack*> TopTracks; /// tracks docked on top, in order to be displayed (top to bottom)
+	//TODO: TArray<FBaseTimingTrack*> ScrollableTracks; /// tracks in scrollable area, in order to be displayed (top to bottom)
+	//TODO: TArray<FBaseTimingTrack*> BottomTracks; /// tracks docked on bottom, in order to be displayed (top to bottom)
+	//TODO: TArray<FBaseTimingTrack*> ForegorundTracks; /// tracks to draw over top/scrollable/bottom tracks (can use entire area), in order to be displayed (back to front)
 
 	////////////////////////////////////////////////////////////
 	// Timing Events Tracks (GPU, CPU, IO)
 
-	//TODO: TArray<FTimingEventsTrack*> TimingEventsTracks; /// timing events tracks in order to be displayed
+	TArray<FTimingEventsTrack*> TimingEventsTracks; /// timing events tracks in order to be displayed
 
 	bool bAreTimingEventsTracksDirty;
 
@@ -345,10 +367,13 @@ protected:
 	////////////////////////////////////////////////////////////
 	// IO
 
-	bool bMergeIoLanes;
+	FTimingEventsTrack* IoOverviewTrack;
+	FTimingEventsTrack* IoActivityTrack;
 
 	mutable bool bForceIoEventsUpdate;
 
+	bool bMergeIoLanes;
+
 	struct FIoTimingEvent
 	{
 		double StartTime;
