Fixed TraceInsights' compile errors when compiling the vanilla editor target, with PCH disabled, and no unity files.

diff --git a/Engine/Source/Developer/TraceInsights/Private/Insights/ViewModels/LogFilter.h b/Engine/Source/Developer/TraceInsights/Private/Insights/ViewModels/LogFilter.h
--- a/Engine/Source/Developer/TraceInsights/Private/Insights/ViewModels/LogFilter.h
+++ b/Engine/Source/Developer/TraceInsights/Private/Insights/ViewModels/LogFilter.h
@@ -2,6 +2,7 @@
 
 #pragma once
 
+#include "Async/AsyncWork.h"
 #include "CoreMinimal.h"
 #include "Logging/LogVerbosity.h"
 #include "Misc/TextFilterExpressionEvaluator.h"
diff --git a/Engine/Source/Developer/TraceInsights/Private/Insights/ViewModels/MarkersTimingTrack.h b/Engine/Source/Developer/TraceInsights/Private/Insights/ViewModels/MarkersTimingTrack.h
--- a/Engine/Source/Developer/TraceInsights/Private/Insights/ViewModels/MarkersTimingTrack.h
+++ b/Engine/Source/Developer/TraceInsights/Private/Insights/ViewModels/MarkersTimingTrack.h
@@ -141,7 +141,7 @@ public:
 	FMarkersTimingTrack& Track;
 	const FTimingTrackViewport& Viewport;
 
-	const TSharedRef<FSlateFontMeasure> FontMeasureService;
+	const TSharedRef<class FSlateFontMeasure> FontMeasureService;
 	const FSlateFontInfo Font;
 
 	const Trace::ILogProvider* LogProviderPtr; // valid only between BeginLog() and EndLog()
