[Insights] Fixed missing include header file for FMenuBuilder used in STimingView.h  (compile error when -disableunity -nopch is used).

diff --git a/Engine/Source/Developer/TraceInsights/Private/Insights/Widgets/STimingView.h b/Engine/Source/Developer/TraceInsights/Private/Insights/Widgets/STimingView.h
--- a/Engine/Source/Developer/TraceInsights/Private/Insights/Widgets/STimingView.h
+++ b/Engine/Source/Developer/TraceInsights/Private/Insights/Widgets/STimingView.h
@@ -3,6 +3,7 @@
 #pragma once
 
 #include "CoreMinimal.h"
+#include "Framework/MultiBox/MultiBoxBuilder.h"
 #include "Input/CursorReply.h"
 #include "Input/Reply.h"
 #include "Layout/Geometry.h"
@@ -35,7 +36,7 @@ DECLARE_DELEGATE_RetVal(int32, FHoveredEventChangedDelegate);
 DECLARE_DELEGATE_RetVal(int32, FSelectedEventChangedDelegate);
 
 /** A custom widget used to display timing events. */
-class STimingView : public SCompoundWidget // TODO: Rename it as STimingView
+class STimingView : public SCompoundWidget
 {
 public:
 	/** Default constructor. */
