UnrealInsights: Single event exclusive time calculation fix

diff --git a/Engine/Source/Developer/TraceInsights/Private/Insights/Widgets/STimingView.cpp b/Engine/Source/Developer/TraceInsights/Private/Insights/Widgets/STimingView.cpp
--- a/Engine/Source/Developer/TraceInsights/Private/Insights/Widgets/STimingView.cpp
+++ b/Engine/Source/Developer/TraceInsights/Private/Insights/Widgets/STimingView.cpp
@@ -2498,8 +2498,8 @@ bool STimingView::SearchTimingEvent(const double InStartTime,
 						uint64 EventDepth;
 
 						uint64 CurrentDepth;
-						double ChildStartTime;
-						double TimeInChildren;
+						double LastTime;
+						double ExclusiveTime;
 						bool IsInEventScope;
 					};
 					FEnumerationState State;
@@ -2509,21 +2509,21 @@ bool STimingView::SearchTimingEvent(const double InStartTime,
 					State.EventDepth = Ctx.TimingEvent.Depth;
 
 					State.CurrentDepth = 0;
-					State.ChildStartTime = 0.0;
-					State.TimeInChildren = 0.0;
+					State.LastTime = 0.0;
+					State.ExclusiveTime = 0.0;
 					State.IsInEventScope = false;
 
 					Timeline.EnumerateEvents(Ctx.TimingEvent.StartTime, Ctx.TimingEvent.EndTime, [&State](bool IsEnter, double Time, const Trace::FTimingProfilerEvent& Event)
 					{
 						if (IsEnter)
 						{
-							if (State.CurrentDepth == State.EventDepth && Time == State.EventStartTime)
+							if (State.IsInEventScope && State.CurrentDepth == State.EventDepth + 1)
 							{
-								State.IsInEventScope = true;
+								State.ExclusiveTime += Time - State.LastTime;
 							}
-							else if (State.IsInEventScope && State.CurrentDepth == State.EventDepth + 1)
+							if (State.CurrentDepth == State.EventDepth && Time == State.EventStartTime)
 							{
-								State.ChildStartTime = Time;
+								State.IsInEventScope = true;
 							}
 							++State.CurrentDepth;
 						}
@@ -2533,15 +2533,12 @@ bool STimingView::SearchTimingEvent(const double InStartTime,
 							if (State.CurrentDepth == State.EventDepth && Time == State.EventEndTime)
 							{
 								State.IsInEventScope = false;
-							}
-							else if (State.IsInEventScope && State.CurrentDepth == State.EventDepth + 1)
-							{
-								State.TimeInChildren += Time - State.ChildStartTime;
+								State.ExclusiveTime += Time - State.LastTime;
 							}
 						}
+						State.LastTime = Time;
 					});
-					Ctx.TimingEvent.ExclusiveTime = (Ctx.TimingEvent.EndTime - Ctx.TimingEvent.StartTime) - State.TimeInChildren;
-					check(Ctx.TimingEvent.ExclusiveTime >= 0.0);
+					Ctx.TimingEvent.ExclusiveTime = State.ExclusiveTime;
 				});
 			});
 		}
