[Insights] Fixed TimeUtils.h compile error (linux) due to inclusion of string.h

diff --git a/Engine/Source/Developer/TraceInsights/Private/Insights/Common/TimeUtils.h b/Engine/Source/Developer/TraceInsights/Private/Insights/Common/TimeUtils.h
--- a/Engine/Source/Developer/TraceInsights/Private/Insights/Common/TimeUtils.h
+++ b/Engine/Source/Developer/TraceInsights/Private/Insights/Common/TimeUtils.h
@@ -3,7 +3,7 @@
 #pragma once
 
 #include "CoreMinimal.h"
-#include "String.h"
+#include "Containers/UnrealString.h"
 
 /*
 double is a 64 bit IEEE 754 double precision Floating Point Number
