Copying //UE4/Dev-Core to Main (//UE4/Main)

diff --git a/Engine/Plugins/Compositing/OpenColorIO/Source/OpenColorIO/Private/OpenColorIOShaderCompilationManager.cpp b/Engine/Plugins/Compositing/OpenColorIO/Source/OpenColorIO/Private/OpenColorIOShaderCompilationManager.cpp
--- a/Engine/Plugins/Compositing/OpenColorIO/Source/OpenColorIO/Private/OpenColorIOShaderCompilationManager.cpp
+++ b/Engine/Plugins/Compositing/OpenColorIO/Source/OpenColorIO/Private/OpenColorIOShaderCompilationManager.cpp
@@ -35,22 +35,27 @@ void FOpenColorIOShaderCompilationManager::Tick(float DeltaSeconds)
 
 FOpenColorIOShaderCompilationManager::FOpenColorIOShaderCompilationManager()
 {
-	// Ew. Should we just use FShaderCompilingManager's workers instead? Is that safe?
-	const int32 NumVirtualCores = FPlatformMisc::NumberOfCoresIncludingHyperthreads();
-	const uint32 NumOpenColorIOShaderCompilingThreads = FMath::Min(NumVirtualCores-1, 4);		
+	
+}
 
-	for (uint32 WorkerIndex = 0; WorkerIndex < NumOpenColorIOShaderCompilingThreads; WorkerIndex++)
+FOpenColorIOShaderCompilationManager::~FOpenColorIOShaderCompilationManager()
+{
+	for (FOpenColorIOShaderCompileWorkerInfo* Info : WorkerInfos)
 	{
-		WorkerInfos.Add(new FOpenColorIOShaderCompileWorkerInfo());
+		delete Info;
 	}
-}
 
+	WorkerInfos.Empty();
+}
 
 void FOpenColorIOShaderCompilationManager::RunCompileJobs()
 {
 #if WITH_EDITOR
 	// If we aren't compiling through workers, so we can just track the serial time here.
 //	COOK_STAT(FScopedDurationTimer CompileTimer(OpenColorIOShaderCookStats::AsyncCompileTimeSec));
+
+	InitWorkerInfo();
+
 	int32 NumActiveThreads = 0;
 
 	for (int32 WorkerIndex = 0; WorkerIndex < WorkerInfos.Num(); WorkerIndex++)
@@ -166,7 +171,20 @@ void FOpenColorIOShaderCompilationManager::RunCompileJobs()
 #endif
 }
 
+void FOpenColorIOShaderCompilationManager::InitWorkerInfo()
+{
+	if (WorkerInfos.Num() == 0) // Check to see if it has been initialized or not
+	{
+		// Ew. Should we just use FShaderCompilingManager's workers instead? Is that safe?
+		const int32 NumVirtualCores = FPlatformMisc::NumberOfCoresIncludingHyperthreads();
+		const uint32 NumOpenColorIOShaderCompilingThreads = FMath::Min(NumVirtualCores - 1, 4);
 
+		for (uint32 WorkerIndex = 0; WorkerIndex < NumOpenColorIOShaderCompilingThreads; WorkerIndex++)
+		{
+			WorkerInfos.Add(new FOpenColorIOShaderCompileWorkerInfo());
+		}
+	}	
+}
 
 OPENCOLORIO_API void FOpenColorIOShaderCompilationManager::AddJobs(TArray<FShaderCommonCompileJob*> InNewJobs)
 {
diff --git a/Engine/Plugins/Compositing/OpenColorIO/Source/OpenColorIO/Public/OpenColorIOShaderCompilationManager.h b/Engine/Plugins/Compositing/OpenColorIO/Source/OpenColorIO/Public/OpenColorIOShaderCompilationManager.h
--- a/Engine/Plugins/Compositing/OpenColorIO/Source/OpenColorIO/Public/OpenColorIOShaderCompilationManager.h
+++ b/Engine/Plugins/Compositing/OpenColorIO/Source/OpenColorIO/Public/OpenColorIOShaderCompilationManager.h
@@ -84,6 +84,7 @@ class FOpenColorIOShaderCompilationManager
 {
 public:
 	FOpenColorIOShaderCompilationManager();
+	~FOpenColorIOShaderCompilationManager();
 
 	OPENCOLORIO_API void Tick(float DeltaSeconds = 0.0f);
 	OPENCOLORIO_API void AddJobs(TArray<FShaderCommonCompileJob*> InNewJobs);
@@ -95,6 +96,8 @@ private:
 	void ProcessCompiledOpenColorIOShaderMaps(TMap<int32, FOpenColorIOShaderMapFinalizeResults>& CompiledShaderMaps, float TimeBudget);
 	void RunCompileJobs();
 
+	void InitWorkerInfo();
+
 	TArray<FShaderCommonCompileJob*> JobQueue;
 
 	/** Map from shader map Id to the compile results for that map, used to gather compiled results. */
diff --git a/Engine/Plugins/Experimental/GeometryCache/Source/GeometryCache/Private/GeometryCacheCodecBase.cpp b/Engine/Plugins/Experimental/GeometryCache/Source/GeometryCache/Private/GeometryCacheCodecBase.cpp
--- a/Engine/Plugins/Experimental/GeometryCache/Source/GeometryCache/Private/GeometryCacheCodecBase.cpp
+++ b/Engine/Plugins/Experimental/GeometryCache/Source/GeometryCache/Private/GeometryCacheCodecBase.cpp
@@ -17,7 +17,7 @@ void FStreamedGeometryCacheChunk::Serialize(FArchive& Ar, UObject* Owner, int32
 	// We force it not inline that means bulk data won't automatically be loaded when we deserialize
 	// later but only when we explicitly take action to load it
 	BulkData.SetBulkDataFlags(BULKDATA_Force_NOT_InlinePayload);
-	BulkData.Serialize(Ar, Owner, ChunkIndex);
+	BulkData.Serialize(Ar, Owner, ChunkIndex,false);
 	Ar << DataSize;
 	Ar << FirstFrame;
 	Ar << LastFrame;
diff --git a/Engine/Plugins/Experimental/GeometryCache/Source/GeometryCache/Private/GeometryCacheTrackStreamable.cpp b/Engine/Plugins/Experimental/GeometryCache/Source/GeometryCache/Private/GeometryCacheTrackStreamable.cpp
--- a/Engine/Plugins/Experimental/GeometryCache/Source/GeometryCache/Private/GeometryCacheTrackStreamable.cpp
+++ b/Engine/Plugins/Experimental/GeometryCache/Source/GeometryCache/Private/GeometryCacheTrackStreamable.cpp
@@ -11,6 +11,7 @@
 DECLARE_CYCLE_STAT(TEXT("Decode Mesh Frame"), STAT_UpdateMeshData, STATGROUP_GeometryCache);
 DECLARE_CYCLE_STAT(TEXT("Encode Mesh Frame"), STAT_AddMeshSample, STATGROUP_GeometryCache);
 
+#if WITH_EDITOR
 /**
  * Creates a totally invalid UGeometryCacheTrackStreamable instance specially set up to be very large
  * and then tries to serialize it to smoke-test the serialization of large assets and bulk data.
@@ -59,7 +60,7 @@ FAutoConsoleCommand TriggerSerializationCrashCommand(
 	TEXT("Test a crash searializing large bulk data object"),
 	FConsoleCommandDelegate::CreateStatic(UGeometryCacheTrackStreamable::TriggerSerializationCrash)
 );
-
+#endif
 
 /*-----------------------------------------------------------------------------
 UCompressedGeometryCacheTrack
diff --git a/Engine/Plugins/Experimental/GeometryCache/Source/GeometryCache/Private/StreamingGeometryCacheData.cpp b/Engine/Plugins/Experimental/GeometryCache/Source/GeometryCache/Private/StreamingGeometryCacheData.cpp
--- a/Engine/Plugins/Experimental/GeometryCache/Source/GeometryCache/Private/StreamingGeometryCacheData.cpp
+++ b/Engine/Plugins/Experimental/GeometryCache/Source/GeometryCache/Private/StreamingGeometryCacheData.cpp
@@ -27,7 +27,7 @@ DECLARE_MEMORY_STAT(TEXT("Streamed Chunks"), STAT_ChunkDataStreamed, STATGROUP_G
 DECLARE_MEMORY_STAT(TEXT("Resident Chunks"), STAT_ChunkDataResident, STATGROUP_GeometryCache);
 
 FStreamingGeometryCacheData::FStreamingGeometryCacheData(UGeometryCacheTrackStreamable* InTrack)
-	: Track(InTrack), IORequestHandle(nullptr)
+	: Track(InTrack)
 {
 }
 
@@ -48,12 +48,6 @@ FStreamingGeometryCacheData::~FStreamingGeometryCacheData()
 		RemoveResidentChunk(Iter.Value());
 	}
 	Chunks.Empty();
-
-	if (IORequestHandle)
-	{
-		delete IORequestHandle;
-		IORequestHandle = nullptr;
-	}
 }
 
 void FStreamingGeometryCacheData::ResetNeededChunks()
@@ -100,7 +94,7 @@ This is called from some random thread when reading is complete.
 void FStreamingGeometryCacheData::OnAsyncReadComplete(int32 LoadedChunkIndex, IAsyncReadRequest* ReadRequest)
 {
 	// We should do the least amount of work possible here as to not stall the async io threads.
-	// We calso cannot take the critical section here as this would lead to a deadlock between the
+	// We also cannot take the critical section here as this would lead to a deadlock between the
 	// our critical section and the async-io internal critical section. 
 	// So we just put this on queue here and then process the results later when we are on a different
 	// thread that already holds our lock.
@@ -182,7 +176,7 @@ void FStreamingGeometryCacheData::PrefetchData(UGeometryCacheComponent *Componen
 		ResidentChunk.Memory = static_cast<uint8*>(FMemory::Malloc(Chunk.DataSize));
 		INC_MEMORY_STAT_BY(STAT_ChunkDataResident, Chunk.DataSize);
 		INC_MEMORY_STAT_BY(STAT_ChunkDataStreamed, Chunk.DataSize);
-		Chunk.BulkData.GetCopy((void**)&ResidentChunk.Memory); //note: This does the actual loading internally...
+		Chunk.BulkData.GetCopy((void**)&ResidentChunk.Memory,true); //note: This does the actual loading internally...
 		ChunksAvailable.Add(ChunkId);
 	}
 }
@@ -207,7 +201,7 @@ void FStreamingGeometryCacheData::UpdateStreamingStatus()
 			// If not requested yet, then request a load
 			if (!ChunksRequested.Contains(NeededIndex))
 			{
-				const FStreamedGeometryCacheChunk& Chunk = Track->GetChunk(NeededIndex);
+				FStreamedGeometryCacheChunk& Chunk = Track->GetChunk(NeededIndex);
 
 				// This can happen in the editor if the asset hasn't been saved yet.
 				if (Chunk.BulkData.IsBulkDataLoaded())
@@ -232,19 +226,14 @@ void FStreamingGeometryCacheData::UpdateStreamingStatus()
 				EAsyncIOPriorityAndFlags AsyncIOPriority = AIOP_BelowNormal;
 
 				// Kick of a load
-				if (!IORequestHandle)
-				{
-					IORequestHandle = FPlatformFileManager::Get().GetPlatformFile().OpenAsyncRead(*Chunk.BulkData.GetFilename());
-					checkf(IORequestHandle, TEXT("Could not opan an async file")); // this generally cannot fail because it is async
-				}
 				check(Chunk.BulkData.GetBulkDataSize() == ResidentChunk.DataSize);
 
 				FAsyncFileCallBack AsyncFileCallBack = [this, NeededIndex](bool bWasCancelled, IAsyncReadRequest* Req)
 				{
 					this->OnAsyncReadComplete(NeededIndex, Req);
 				};
-
-				ResidentChunk.IORequest = IORequestHandle->ReadRequest(Chunk.BulkData.GetBulkDataOffsetInFile(), ResidentChunk.DataSize, AsyncIOPriority, &AsyncFileCallBack);
+				
+				ResidentChunk.IORequest = Chunk.BulkData.CreateStreamingRequest(AsyncIOPriority, &AsyncFileCallBack, nullptr);
 				if (!ResidentChunk.IORequest)
 				{
 					UE_LOG(LogGeoCaStreaming, Error, TEXT("Geometry cache streaming read request failed."));
@@ -339,8 +328,10 @@ void FStreamingGeometryCacheData::ProcessCompletedChunks()
 			return;
 		}
 
-		// Chunks can be queued up multiple times when scrubbing, either the request has already been fulfilled (nullptr) or it could be waiting for a different IO request
-		if (CompletedChunk.ReadRequest == Chunk->IORequest || Chunk->IORequest != nullptr)
+		// Chunks can be queued up multiple times when scrubbing but we can trust the LoadedChunkIndex of the completed chunk
+		// so all we need to check is if the FResidentChunk  has a valid IORequest pointer or not.
+		// If it is nullptr then we already processed a request for the chunk and it can be ignored.
+		if (Chunk->IORequest != nullptr)
 		{
 			// Check to see if we successfully managed to load anything
 			uint8* Mem = CompletedChunk.ReadRequest->GetReadResults();
diff --git a/Engine/Plugins/Experimental/GeometryCache/Source/GeometryCache/Private/StreamingGeometryCacheData.h b/Engine/Plugins/Experimental/GeometryCache/Source/GeometryCache/Private/StreamingGeometryCacheData.h
--- a/Engine/Plugins/Experimental/GeometryCache/Source/GeometryCache/Private/StreamingGeometryCacheData.h
+++ b/Engine/Plugins/Experimental/GeometryCache/Source/GeometryCache/Private/StreamingGeometryCacheData.h
@@ -7,8 +7,8 @@
 #include "Containers/Array.h"
 #include "Containers/Queue.h"
 
-class IAsyncReadFileHandle;
 class IAsyncReadRequest;
+struct FBulkDataIORequest;
 class UGeometryCacheTrackStreamable;
 class UGeometryCacheComponent;
 struct FStreamedGeometryCacheChunk;
@@ -21,7 +21,7 @@ struct FResidentChunk
 	void* Memory;
 	int32 DataSize;
 	int32 Refcount;
-	IAsyncReadRequest* IORequest; //null when resident, nonnull when being loaded
+	FBulkDataIORequest* IORequest; //null when resident, nonnull when being loaded
 };
 
 /**
@@ -69,8 +69,6 @@ private:
 	// The track we are associated with
 	UGeometryCacheTrackStreamable *Track;
 
-	IAsyncReadFileHandle* IORequestHandle;
-
 	// Chunks that ideally would be loaded at this point in time
 	// There may be more or less actual chunks loaded (more = cached chunks, less = we're still waiting for the disc)
 	// This should only be used from the main thread. It can be modified without taking the lock. Changes are then
diff --git a/Engine/Source/Developer/TraceInsights/Private/Insights/InsightsManager.cpp b/Engine/Source/Developer/TraceInsights/Private/Insights/InsightsManager.cpp
--- a/Engine/Source/Developer/TraceInsights/Private/Insights/InsightsManager.cpp
+++ b/Engine/Source/Developer/TraceInsights/Private/Insights/InsightsManager.cpp
@@ -4,11 +4,11 @@
 
 #include "Framework/Application/SlateApplication.h"
 #include "Modules/ModuleManager.h"
-#include "Templates/ScopedPointer.h"
 #include "Templates/UniquePtr.h"
 
 // Insights
-#include "Insights/IoProfilerManager.h"
+#include "Insights/LoadingProfiler/LoadingProfilerManager.h"
+#include "Insights/NetworkingProfiler/NetworkingProfilerManager.h"
 #include "Insights/TimingProfilerManager.h"
 #include "Insights/Widgets/SStartPageWindow.h"
 #include "Insights/Widgets/STimingProfilerWindow.h"
@@ -19,7 +19,8 @@
 
 const FName FInsightsManagerTabs::StartPageTabId(TEXT("StartPage"));
 const FName FInsightsManagerTabs::TimingProfilerTabId(TEXT("TimingProfiler"));
-const FName FInsightsManagerTabs::IoProfilerTabId(TEXT("IoProfiler"));
+const FName FInsightsManagerTabs::LoadingProfilerTabId(TEXT("LoadingProfiler"));
+const FName FInsightsManagerTabs::NetworkingProfilerTabId(TEXT("NetworkingProfiler"));
 
 TSharedPtr<FInsightsManager> FInsightsManager::Instance = nullptr;
 
@@ -85,6 +86,13 @@ TSharedPtr<const Trace::IAnalysisSession> FInsightsManager::GetSession() const
 
 ////////////////////////////////////////////////////////////////////////////////////////////////////
 
+Trace::FSessionHandle FInsightsManager::GetSessionHandle() const
+{
+	return CurrentSessionHandle;
+}
+
+////////////////////////////////////////////////////////////////////////////////////////////////////
+
 const TSharedRef<FUICommandList> FInsightsManager::GetCommandList() const
 {
 	return CommandList;
@@ -127,6 +135,7 @@ void FInsightsManager::ResetSession()
 	if (Session.IsValid())
 	{
 		Session.Reset();
+		CurrentSessionHandle = 0;
 		OnSessionChanged();
 	}
 }
@@ -141,10 +150,16 @@ void FInsightsManager::OnSessionChanged()
 		TimingProfilerManager->OnSessionChanged();
 	}
 
-	if (TSharedPtr<FIoProfilerManager> IoProfilerManager = FIoProfilerManager::Get())
+	if (TSharedPtr<FLoadingProfilerManager> LoadingProfilerManager = FLoadingProfilerManager::Get())
 	{
-		// FIXME: make IoProfilerManager to register to SessionChangedEvent instead
-		IoProfilerManager->OnSessionChanged();
+		// FIXME: make LoadingProfilerManager to register to SessionChangedEvent instead
+		LoadingProfilerManager->OnSessionChanged();
+	}
+
+	if (TSharedPtr<FNetworkingProfilerManager> NetworkingProfilerManager = FNetworkingProfilerManager::Get())
+	{
+		// FIXME: make NetworkingProfilerManager to register to SessionChangedEvent instead
+		NetworkingProfilerManager->OnSessionChanged();
 	}
 
 	SessionChangedEvent.Broadcast();
@@ -159,9 +174,21 @@ void FInsightsManager::SpawnAndActivateTabs()
 	{
 		FGlobalTabmanager::Get()->InvokeTab(FInsightsManagerTabs::TimingProfilerTabId);
 	}
-	if (FGlobalTabmanager::Get()->HasTabSpawner(FInsightsManagerTabs::IoProfilerTabId))
+	if (FGlobalTabmanager::Get()->HasTabSpawner(FInsightsManagerTabs::LoadingProfilerTabId))
+	{
+		FGlobalTabmanager::Get()->InvokeTab(FInsightsManagerTabs::LoadingProfilerTabId);
+	}
+	if (FGlobalTabmanager::Get()->HasTabSpawner(FInsightsManagerTabs::NetworkingProfilerTabId))
 	{
-		FGlobalTabmanager::Get()->InvokeTab(FInsightsManagerTabs::IoProfilerTabId);
+		// Close previous tabs.
+		TSharedPtr<SDockTab> NetworkingProfilerTab;
+		while ((NetworkingProfilerTab = FGlobalTabmanager::Get()->FindExistingLiveTab(FInsightsManagerTabs::NetworkingProfilerTabId)).IsValid())
+		{
+			NetworkingProfilerTab->RequestCloseTab();
+		}
+
+		FGlobalTabmanager::Get()->InvokeTab(FInsightsManagerTabs::NetworkingProfilerTabId);
+		FGlobalTabmanager::Get()->InvokeTab(FInsightsManagerTabs::NetworkingProfilerTabId);
 	}
 
 	// Ensure Timing Insights / Timing View is the active tab / view.
@@ -169,7 +196,7 @@ void FInsightsManager::SpawnAndActivateTabs()
 	{
 		TimingInsightsTab->ActivateInParent(ETabActivationCause::SetDirectly);
 
-		//TOOD: FTimingProfilerManager::Get()->SpawnAndActivateTabs();
+		//TODO: FTimingProfilerManager::Get()->ActivateWindow();
 		TSharedPtr<class STimingProfilerWindow> Wnd = FTimingProfilerManager::Get()->GetProfilerWindow();
 		if (Wnd)
 		{
@@ -269,6 +296,7 @@ void FInsightsManager::LoadSession(Trace::FSessionHandle SessionHandle)
 	if (DataStream)
 	{
 		Session = AnalysisService->StartAnalysis(SessionInfo.Name, MoveTemp(DataStream));
+		CurrentSessionHandle = SessionHandle;
 		SpawnAndActivateTabs();
 		OnSessionChanged();
 	}
@@ -284,6 +312,7 @@ void FInsightsManager::LoadTraceFile(const FString& TraceFilepath)
 	if (DataStream)
 	{
 		Session = AnalysisService->StartAnalysis(*TraceFilepath, MoveTemp(DataStream));
+		CurrentSessionHandle = 0;
 		SpawnAndActivateTabs();
 		OnSessionChanged();
 	}
diff --git a/Engine/Source/Developer/TraceInsights/Private/Insights/InsightsManager.h b/Engine/Source/Developer/TraceInsights/Private/Insights/InsightsManager.h
--- a/Engine/Source/Developer/TraceInsights/Private/Insights/InsightsManager.h
+++ b/Engine/Source/Developer/TraceInsights/Private/Insights/InsightsManager.h
@@ -23,7 +23,8 @@ struct FInsightsManagerTabs
 {
 	static const FName StartPageTabId;
 	static const FName TimingProfilerTabId;
-	static const FName IoProfilerTabId;
+	static const FName LoadingProfilerTabId;
+	static const FName NetworkingProfilerTabId;
 };
 
 ////////////////////////////////////////////////////////////////////////////////////////////////////
@@ -84,6 +85,9 @@ public:
 	/** @return an instance of the trace analysis session. */
 	TSharedPtr<const Trace::IAnalysisSession> GetSession() const;
 
+	/** @return the session handle of the trace analysis session. */
+	Trace::FSessionHandle GetSessionHandle() const;
+
 	/** @returns UI command list for the main manager. */
 	const TSharedRef<FUICommandList> GetCommandList() const;
 
@@ -187,6 +191,9 @@ private:
 	/** The trace analysis session. */
 	TSharedPtr<const Trace::IAnalysisSession> Session;
 
+	/** The session handle. */
+	Trace::FSessionHandle CurrentSessionHandle;
+
 	/** List of UI commands for this manager. This will be filled by this and corresponding classes. */
 	TSharedRef<FUICommandList> CommandList;
 
diff --git a/Engine/Source/Developer/TraceInsights/Private/Insights/InsightsStyle.cpp b/Engine/Source/Developer/TraceInsights/Private/Insights/InsightsStyle.cpp
--- a/Engine/Source/Developer/TraceInsights/Private/Insights/InsightsStyle.cpp
+++ b/Engine/Source/Developer/TraceInsights/Private/Insights/InsightsStyle.cpp
@@ -39,21 +39,43 @@ TSharedRef<FSlateStyleSet> FInsightsStyle::Create()
 
 	FSlateStyleSet& Style = StyleRef.Get();
 
+	//////////////////////////////////////////////////
+
+	Style.Set("WhiteBrush", new FSlateColorBrush(FLinearColor::White));
+	Style.Set("SingleBorder", new FSlateBorderBrush(NAME_None, FMargin(1.0f)));
+	Style.Set("DoubleBorder", new FSlateBorderBrush(NAME_None, FMargin(2.0f)));
+
+	Style.Set("EventBorder", new FSlateBorderBrush(NAME_None, FMargin(1.0f)));
+	Style.Set("HoveredEventBorder", new FSlateBorderBrush(NAME_None, FMargin(2.0f)));
+	Style.Set("SelectedEventBorder", new FSlateBorderBrush(NAME_None, FMargin(2.0f)));
+
 	//////////////////////////////////////////////////
 	// Icons for major components
 
 	Style.Set("StartPage.Icon.Large", new IMAGE_BRUSH("/Icons/icon_tab_Tools_16x", FVector2D(32.0f, 32.0f)));
 	Style.Set("StartPage.Icon.Small", new IMAGE_BRUSH("/Icons/icon_tab_Tools_16x", FVector2D(16.0f, 16.0f)));
 
-	Style.Set("TimingInsights.Icon.Large", new IMAGE_BRUSH("Icons/Profiler/profiler_stats_40x", FVector2D(32.0f, 32.0f)));
-	Style.Set("TimingInsights.Icon.Small", new IMAGE_BRUSH("Icons/Profiler/profiler_stats_40x", FVector2D(16.0f, 16.0f)));
+	Style.Set("TimingProfiler.Icon.Large", new IMAGE_BRUSH("Icons/Profiler/profiler_stats_40x", FVector2D(32.0f, 32.0f)));
+	Style.Set("TimingProfiler.Icon.Small", new IMAGE_BRUSH("Icons/Profiler/profiler_stats_40x", FVector2D(16.0f, 16.0f)));
+
+	Style.Set("LoadingProfiler.Icon.Large", new IMAGE_BRUSH("Icons/Profiler/profiler_stats_40x", FVector2D(32.0f, 32.0f)));
+	Style.Set("LoadingProfiler.Icon.Small", new IMAGE_BRUSH("Icons/Profiler/profiler_stats_40x", FVector2D(16.0f, 16.0f)));
 
-	Style.Set("AssetLoadingInsights.Icon.Large", new IMAGE_BRUSH("Icons/Profiler/profiler_stats_40x", FVector2D(32.0f, 32.0f)));
-	Style.Set("AssetLoadingInsights.Icon.Small", new IMAGE_BRUSH("Icons/Profiler/profiler_stats_40x", FVector2D(16.0f, 16.0f)));
+	Style.Set("NetworkingProfiler.Icon.Large", new IMAGE_BRUSH("Icons/Profiler/profiler_stats_40x", FVector2D(32.0f, 32.0f)));
+	Style.Set("NetworkingProfiler.Icon.Small", new IMAGE_BRUSH("Icons/Profiler/profiler_stats_40x", FVector2D(16.0f, 16.0f)));
 
 	Style.Set("Toolbar.Icon.Large", new IMAGE_BRUSH("/Icons/icon_tab_Tools_16x", FVector2D(32.0f, 32.0f)));
 	Style.Set("Toolbar.Icon.Small", new IMAGE_BRUSH("/Icons/icon_tab_Tools_16x", FVector2D(16.0f, 16.0f)));
 
+	//////////////////////////////////////////////////
+	// Start Page buttons
+
+	Style.Set("Open.Icon.Large", new IMAGE_BRUSH("Icons/Profiler/Profiler_LoadMultiple_Profiler_40x", FVector2D(32.0f, 32.0f)));
+	Style.Set("Open.Icon.Small", new IMAGE_BRUSH("Icons/Profiler/Profiler_Load_Profiler_40x", FVector2D(16.0f, 16.0f)));
+
+	Style.Set("OpenFile.Icon.Large", new IMAGE_BRUSH("Icons/LV_Load", FVector2D(32.0f, 32.0f)));
+	Style.Set("OpenFile.Icon.Small", new IMAGE_BRUSH("Icons/LV_Load", FVector2D(16.0f, 16.0f)));
+
 	//////////////////////////////////////////////////
 	// Timing Insights
 
@@ -75,14 +97,20 @@ TSharedRef<FSlateStyleSet> FInsightsStyle::Create()
 	Style.Set("LogView.Icon.Large", new IMAGE_BRUSH("Icons/Profiler/profiler_CopyToClipboard_32x", FVector2D(32.0f, 32.0f)));
 	Style.Set("LogView.Icon.Small", new IMAGE_BRUSH("Icons/Profiler/profiler_CopyToClipboard_32x", FVector2D(16.0f, 16.0f)));
 
+	Style.Set("TableTreeView.Icon.Large", new IMAGE_BRUSH("Icons/Profiler/Profiler_Data_Capture_40x", FVector2D(32.0f, 32.0f)));
+	Style.Set("TableTreeView.Icon.Small", new IMAGE_BRUSH("Icons/Profiler/Profiler_Data_Capture_40x", FVector2D(16.0f, 16.0f)));
+
 	//////////////////////////////////////////////////
-	// Start Page buttons
+	// Networking Insights
 
-	Style.Set("Open.Icon.Large", new IMAGE_BRUSH("Icons/Profiler/Profiler_LoadMultiple_Profiler_40x", FVector2D(32.0f, 32.0f)));
-	Style.Set("Open.Icon.Small", new IMAGE_BRUSH("Icons/Profiler/Profiler_Load_Profiler_40x", FVector2D(16.0f, 16.0f)));
+	Style.Set("PacketOveriew.Icon.Large", new IMAGE_BRUSH("Icons/Profiler/profiler_stats_40x", FVector2D(32.0f, 32.0f)));
+	Style.Set("PacketOveriew.Icon.Small", new IMAGE_BRUSH("Icons/Profiler/profiler_stats_40x", FVector2D(16.0f, 16.0f)));
 
-	Style.Set("OpenFile.Icon.Large", new IMAGE_BRUSH("Icons/LV_Load", FVector2D(32.0f, 32.0f)));
-	Style.Set("OpenFile.Icon.Small", new IMAGE_BRUSH("Icons/LV_Load", FVector2D(16.0f, 16.0f)));
+	Style.Set("PacketContentView.Icon.Large", new IMAGE_BRUSH("Icons/Profiler/profiler_stats_40x", FVector2D(32.0f, 32.0f)));
+	Style.Set("PacketContentView.Icon.Small", new IMAGE_BRUSH("Icons/Profiler/profiler_stats_40x", FVector2D(16.0f, 16.0f)));
+
+	Style.Set("NetStatsView.Icon.Large", new IMAGE_BRUSH("Icons/Profiler/profiler_stats_40x", FVector2D(32.0f, 32.0f)));
+	Style.Set("NetStatsView.Icon.Small", new IMAGE_BRUSH("Icons/Profiler/profiler_stats_40x", FVector2D(16.0f, 16.0f)));
 
 	//////////////////////////////////////////////////
 
diff --git a/Engine/Source/Developer/TraceInsights/Private/Insights/LoadingProfiler/LoadingProfilerCommands.cpp b/Engine/Source/Developer/TraceInsights/Private/Insights/LoadingProfiler/LoadingProfilerCommands.cpp
new file mode 100644
--- /dev/null
+++ b/Engine/Source/Developer/TraceInsights/Private/Insights/LoadingProfiler/LoadingProfilerCommands.cpp
@@ -0,0 +1,110 @@
+// Copyright 1998-2019 Epic Games, Inc. All Rights Reserved.
+
+#include "LoadingProfilerCommands.h"
+
+#include "DesktopPlatformModule.h"
+#include "EditorStyleSet.h"
+#include "Framework/Application/SlateApplication.h"
+#include "Framework/MultiBox/MultiBoxBuilder.h"
+
+// Insights
+#include "Insights/InsightsManager.h"
+#include "Insights/LoadingProfiler/LoadingProfilerManager.h"
+
+////////////////////////////////////////////////////////////////////////////////////////////////////
+
+#define LOCTEXT_NAMESPACE "FLoadingProfilerCommands"
+
+////////////////////////////////////////////////////////////////////////////////////////////////////
+// FLoadingProfilerMenuBuilder
+////////////////////////////////////////////////////////////////////////////////////////////////////
+
+void FLoadingProfilerMenuBuilder::AddMenuEntry(FMenuBuilder& MenuBuilder, const TSharedPtr< FUICommandInfo >& UICommandInfo, const FUIAction& UIAction)
+{
+	MenuBuilder.AddMenuEntry
+	(
+		UICommandInfo->GetLabel(),
+		UICommandInfo->GetDescription(),
+		UICommandInfo->GetIcon(),
+		UIAction,
+		NAME_None,
+		UICommandInfo->GetUserInterfaceType()
+	);
+}
+
+////////////////////////////////////////////////////////////////////////////////////////////////////
+// FLoadingProfilerCommands
+////////////////////////////////////////////////////////////////////////////////////////////////////
+
+FLoadingProfilerCommands::FLoadingProfilerCommands()
+	: TCommands<FLoadingProfilerCommands>(
+		TEXT("LoadingProfilerCommand"), // Context name for fast lookup
+		NSLOCTEXT("Contexts", "LoadingProfilerCommand", "Asset Loading Insights"), // Localized context name for displaying
+		NAME_None, // Parent
+		FEditorStyle::GetStyleSetName() // Icon Style Set
+	)
+{
+}
+
+////////////////////////////////////////////////////////////////////////////////////////////////////
+
+// UI_COMMAND takes long for the compiler to optimize
+PRAGMA_DISABLE_OPTIMIZATION
+void FLoadingProfilerCommands::RegisterCommands()
+{
+	UI_COMMAND(ToggleTimingViewVisibility, "Timing", "Toggles the visibility of the main Timing view", EUserInterfaceActionType::ToggleButton, FInputChord(EModifierKey::Control, EKeys::T));
+	UI_COMMAND(ToggleEventAggregationTreeViewVisibility, "Event Aggregation", "Toggles the visibility of the Event Aggregation table/tree view", EUserInterfaceActionType::ToggleButton, FInputChord());
+	UI_COMMAND(ToggleObjectTypeAggregationTreeViewVisibility, "Object Type Aggregation", "Toggles the visibility of the Object Type Aggregation table/tree view", EUserInterfaceActionType::ToggleButton, FInputChord());
+	UI_COMMAND(TogglePackageDetailsTreeViewVisibility, "Package Details", "Toggles the visibility of the Package Details table/tree view", EUserInterfaceActionType::ToggleButton, FInputChord());
+	UI_COMMAND(ToggleExportDetailsTreeViewVisibility, "Export Details", "Toggles the visibility of the Export Details table/tree view", EUserInterfaceActionType::ToggleButton, FInputChord());
+}
+PRAGMA_ENABLE_OPTIMIZATION
+
+////////////////////////////////////////////////////////////////////////////////////////////////////
+// Toggle Commands
+////////////////////////////////////////////////////////////////////////////////////////////////////
+
+#define IMPLEMENT_TOGGLE_COMMAND(CmdName, IsEnabled, SetIsEnabled) \
+	\
+	void FLoadingProfilerActionManager::Map_##CmdName##_Global()\
+	{\
+		This->CommandList->MapAction(This->GetCommands().CmdName, CmdName##_Custom());\
+	}\
+	\
+	const FUIAction FLoadingProfilerActionManager::CmdName##_Custom() \
+	{\
+		FUIAction UIAction;\
+		UIAction.ExecuteAction = FExecuteAction::CreateRaw(this, &FLoadingProfilerActionManager::CmdName##_Execute);\
+		UIAction.CanExecuteAction = FCanExecuteAction::CreateRaw(this, &FLoadingProfilerActionManager::CmdName##_CanExecute);\
+		UIAction.GetActionCheckState = FGetActionCheckState::CreateRaw(this, &FLoadingProfilerActionManager::CmdName##_GetCheckState);\
+		return UIAction;\
+	}\
+	\
+	void FLoadingProfilerActionManager::CmdName##_Execute()\
+	{\
+		const bool b##IsEnabled = !This->IsEnabled();\
+		This->SetIsEnabled(b##IsEnabled);\
+	}\
+	\
+	bool FLoadingProfilerActionManager::CmdName##_CanExecute() const\
+	{\
+		return FInsightsManager::Get()->GetSession().IsValid();\
+	}\
+	\
+	ECheckBoxState FLoadingProfilerActionManager::CmdName##_GetCheckState() const\
+	{\
+		const bool b##IsEnabled = This->IsEnabled();\
+		return b##IsEnabled ? ECheckBoxState::Checked : ECheckBoxState::Unchecked;\
+	}
+
+IMPLEMENT_TOGGLE_COMMAND(ToggleTimingViewVisibility, IsTimingViewVisible, ShowHideTimingView)
+IMPLEMENT_TOGGLE_COMMAND(ToggleEventAggregationTreeViewVisibility, IsEventAggregationTreeViewVisible, ShowHideEventAggregationTreeView)
+IMPLEMENT_TOGGLE_COMMAND(ToggleObjectTypeAggregationTreeViewVisibility, IsObjectTypeAggregationTreeViewVisible, ShowHideObjectTypeAggregationTreeView)
+IMPLEMENT_TOGGLE_COMMAND(TogglePackageDetailsTreeViewVisibility, IsPackageDetailsTreeViewVisible, ShowHidePackageDetailsTreeView)
+IMPLEMENT_TOGGLE_COMMAND(ToggleExportDetailsTreeViewVisibility, IsExportDetailsTreeViewVisible, ShowHideExportDetailsTreeView)
+
+#undef IMPLEMENT_TOGGLE_COMMAND
+
+////////////////////////////////////////////////////////////////////////////////////////////////////
+
+#undef LOCTEXT_NAMESPACE
diff --git a/Engine/Source/Developer/TraceInsights/Private/Insights/LoadingProfiler/LoadingProfilerCommands.h b/Engine/Source/Developer/TraceInsights/Private/Insights/LoadingProfiler/LoadingProfilerCommands.h
new file mode 100644
--- /dev/null
+++ b/Engine/Source/Developer/TraceInsights/Private/Insights/LoadingProfiler/LoadingProfilerCommands.h
@@ -0,0 +1,109 @@
+// Copyright 1998-2019 Epic Games, Inc. All Rights Reserved.
+
+#pragma once
+
+#include "CoreMinimal.h"
+#include "Framework/Commands/UIAction.h"
+#include "Framework/Commands/Commands.h"
+
+class FMenuBuilder;
+
+////////////////////////////////////////////////////////////////////////////////////////////////////
+/**
+ * Class that holds all profiler commands.
+ */
+class FLoadingProfilerCommands : public TCommands<FLoadingProfilerCommands>
+{
+public:
+	/** Default constructor. */
+	FLoadingProfilerCommands();
+
+	/** Initialize commands. */
+	virtual void RegisterCommands() override;
+
+public:
+	//////////////////////////////////////////////////
+	// Global commands need to implement following method:
+	//     void Map_<CommandName>_Global();
+	// Custom commands needs to implement also the following method:
+	//     const FUIAction <CommandName>_Custom(...) const;
+	//////////////////////////////////////////////////
+
+	/** Toggles visibility for the Timing view. Global and custom command. */
+	TSharedPtr<FUICommandInfo> ToggleTimingViewVisibility;
+
+	/** Toggles visibility for the Event Aggregation tree view. Global and custom command. */
+	TSharedPtr<FUICommandInfo> ToggleEventAggregationTreeViewVisibility;
+
+	/** Toggles visibility for the Object Type Aggregation tree view. Global and custom command. */
+	TSharedPtr<FUICommandInfo> ToggleObjectTypeAggregationTreeViewVisibility;
+
+	/** Toggles visibility for the Package Details tree view. Global and custom command. */
+	TSharedPtr<FUICommandInfo> TogglePackageDetailsTreeViewVisibility;
+
+	/** Toggles visibility for the Export Details tree view. Global and custom command. */
+	TSharedPtr<FUICommandInfo> ToggleExportDetailsTreeViewVisibility;
+};
+
+////////////////////////////////////////////////////////////////////////////////////////////////////
+/**
+ * Menu builder. Helper class for adding a customized menu entry using the global UI command info.
+ */
+class FLoadingProfilerMenuBuilder
+{
+public:
+	/**
+	 * Helper method for adding a customized menu entry using the global UI command info.
+	 * FUICommandInfo cannot be executed with custom parameters, so we need to create a custom FUIAction,
+	 * but sometime we have global and local version for the UI command, so reuse data from the global UI command info.
+	 * Ex:
+	 *     SessionInstance_ToggleCapture          - Global version will toggle capture process for all active session instances
+	 *     SessionInstance_ToggleCapture_OneParam - Local version will toggle capture process only for the specified session instance
+	 *
+	 * @param MenuBuilder The menu to add items to
+	 * @param FUICommandInfo A shared pointer to the UI command info
+	 * @param UIAction Customized version of the UI command info stored in an UI action
+	 */
+	static void AddMenuEntry(FMenuBuilder& MenuBuilder, const TSharedPtr<FUICommandInfo>& UICommandInfo, const FUIAction& UIAction);
+};
+
+////////////////////////////////////////////////////////////////////////////////////////////////////
+/**
+ * Class that provides helper functions for the commands to avoid cluttering profiler manager with many small functions.
+ * Can't contain any variables. Directly operates on the profiler manager instance.
+ */
+class FLoadingProfilerActionManager
+{
+	friend class FLoadingProfilerManager;
+
+private:
+	/** Private constructor. */
+	FLoadingProfilerActionManager(class FLoadingProfilerManager* Instance)
+		: This(Instance)
+	{}
+
+	//////////////////////////////////////////////////
+	// Toggle Commands
+
+#define DECLARE_TOGGLE_COMMAND(CmdName)\
+public:\
+	void Map_##CmdName##_Global(); /**< Maps UI command info CmdName with the specified UI command list. */\
+	const FUIAction CmdName##_Custom(); /**< UI action for CmdName command. */\
+protected:\
+	void CmdName##_Execute(); /**< Handles FExecuteAction for CmdName. */\
+	bool CmdName##_CanExecute() const; /**< Handles FCanExecuteAction for CmdName. */\
+	ECheckBoxState CmdName##_GetCheckState() const; /**< Handles FGetActionCheckState for CmdName. */
+
+	DECLARE_TOGGLE_COMMAND(ToggleTimingViewVisibility)
+	DECLARE_TOGGLE_COMMAND(ToggleEventAggregationTreeViewVisibility)
+	DECLARE_TOGGLE_COMMAND(ToggleObjectTypeAggregationTreeViewVisibility)
+	DECLARE_TOGGLE_COMMAND(TogglePackageDetailsTreeViewVisibility)
+	DECLARE_TOGGLE_COMMAND(ToggleExportDetailsTreeViewVisibility)
+#undef DECLARE_TOGGLE_COMMAND
+
+	//////////////////////////////////////////////////
+
+protected:
+	/** Reference to the global instance of the profiler manager. */
+	class FLoadingProfilerManager* This;
+};
diff --git a/Engine/Source/Developer/TraceInsights/Private/Insights/LoadingProfiler/LoadingProfilerManager.h b/Engine/Source/Developer/TraceInsights/Private/Insights/LoadingProfiler/LoadingProfilerManager.h
new file mode 100644
--- /dev/null
+++ b/Engine/Source/Developer/TraceInsights/Private/Insights/LoadingProfiler/LoadingProfilerManager.h
@@ -0,0 +1,166 @@
+// Copyright 1998-2019 Epic Games, Inc. All Rights Reserved.
+
+#pragma once
+
+#include "Containers/Ticker.h"
+#include "CoreMinimal.h"
+#include "Framework/Commands/UICommandList.h"
+
+// Insights
+#include "Insights/InsightsManager.h"
+#include "Insights/LoadingProfiler/LoadingProfilerCommands.h"
+
+class SLoadingProfilerWindow;
+
+namespace Trace
+{
+	class ISessionService;
+	class IAnalysisService;
+	class IAnalysisSession;
+}
+
+////////////////////////////////////////////////////////////////////////////////////////////////////
+/**
+ * This class manages the Asset Loading Profiler state and settings.
+ */
+class FLoadingProfilerManager : public TSharedFromThis<FLoadingProfilerManager>
+{
+	friend class FLoadingProfilerActionManager;
+
+public:
+	/** Creates the Loading Profiler (Asset Loading Insights) manager, only one instance can exist. */
+	FLoadingProfilerManager(TSharedRef<FUICommandList> InCommandList);
+
+	/** Virtual destructor. */
+	virtual ~FLoadingProfilerManager();
+
+	/** Creates an instance of the profiler manager. */
+	static TSharedPtr<FLoadingProfilerManager> Initialize()
+	{
+		if (FLoadingProfilerManager::Instance.IsValid())
+		{
+			FLoadingProfilerManager::Instance.Reset();
+		}
+
+		FLoadingProfilerManager::Instance = MakeShareable(new FLoadingProfilerManager(FInsightsManager::Get()->GetCommandList()));
+		FLoadingProfilerManager::Instance->PostConstructor();
+
+		return FLoadingProfilerManager::Instance;
+	}
+
+	/** Shutdowns the profiler manager. */
+	void Shutdown()
+	{
+		FLoadingProfilerManager::Instance.Reset();
+	}
+
+protected:
+	/** Finishes initialization of the profiler manager. */
+	void PostConstructor();
+
+	/** Binds our UI commands to delegates. */
+	void BindCommands();
+
+public:
+	/**
+	 * @return the global instance of the Loading Profiler (Asset Loading Insights) manager.
+	 * This is an internal singleton and cannot be used outside ProfilerModule.
+	 * For external use:
+	 *     IProfilerModule& ProfilerModule = FModuleManager::Get().LoadModuleChecked<IProfilerModule>("Profiler");
+	 *     ProfilerModule.GetProfilerManager();
+	 */
+	static TSharedPtr<FLoadingProfilerManager> Get();
+
+	/** @returns UI command list for the Loading Profiler manager. */
+	const TSharedRef<FUICommandList> GetCommandList() const;
+
+	/** @return an instance of the Loading Profiler commands. */
+	static const FLoadingProfilerCommands& GetCommands();
+
+	/** @return an instance of the Loading Profiler action manager. */
+	static FLoadingProfilerActionManager& GetActionManager();
+
+	void AssignProfilerWindow(const TSharedRef<SLoadingProfilerWindow>& InProfilerWindow)
+	{
+		ProfilerWindow = InProfilerWindow;
+	}
+
+	/**
+	 * Converts profiler window weak pointer to a shared pointer and returns it.
+	 * Make sure the returned pointer is valid before trying to dereference it.
+	 */
+	TSharedPtr<class SLoadingProfilerWindow> GetProfilerWindow() const
+	{
+		return ProfilerWindow.Pin();
+	}
+
+	////////////////////////////////////////////////////////////////////////////////////////////////////
+	// Getters and setters used by Toggle Commands.
+
+	/** @return true, if Timing view is visible */
+	const bool IsTimingViewVisible() const { return bIsTimingViewVisible; }
+	void SetTimingViewVisible(const bool bIsVisible) { bIsTimingViewVisible = bIsVisible; }
+	void ShowHideTimingView(const bool bIsVisible);
+
+	/** @return true, if Event Aggregation tree view is visible */
+	const bool IsEventAggregationTreeViewVisible() const { return bIsEventAggregationTreeViewVisible; }
+	void SetEventAggregationTreeViewVisible(const bool bIsVisible) { bIsEventAggregationTreeViewVisible = bIsVisible; }
+	void ShowHideEventAggregationTreeView(const bool bIsVisible);
+
+	/** @return true, if Object Type Aggregation tree view is visible */
+	const bool IsObjectTypeAggregationTreeViewVisible() const { return bIsObjectTypeAggregationTreeViewVisible; }
+	void SetObjectTypeAggregationTreeViewVisible(const bool bIsVisible) { bIsObjectTypeAggregationTreeViewVisible = bIsVisible; }
+	void ShowHideObjectTypeAggregationTreeView(const bool bIsVisible);
+
+	/** @return true, if Package Details tree view is visible */
+	const bool IsPackageDetailsTreeViewVisible() const { return bIsPackageDetailsTreeViewVisible; }
+	void SetPackageDetailsTreeViewVisible(const bool bIsVisible) { bIsPackageDetailsTreeViewVisible = bIsVisible; }
+	void ShowHidePackageDetailsTreeView(const bool bIsVisible);
+
+	/** @return true, if Export Details tree view is visible */
+	const bool IsExportDetailsTreeViewVisible() const { return bIsExportDetailsTreeViewVisible; }
+	void SetExportDetailsTreeViewVisible(const bool bIsVisible) { bIsExportDetailsTreeViewVisible = bIsVisible; }
+	void ShowHideExportDetailsTreeView(const bool bIsVisible);
+
+	////////////////////////////////////////////////////////////////////////////////////////////////////
+
+	void OnSessionChanged();
+
+protected:
+	/** Updates this manager, done through FCoreTicker. */
+	bool Tick(float DeltaTime);
+
+protected:
+	/** The delegate to be invoked when this manager ticks. */
+	FTickerDelegate OnTick;
+
+	/** Handle to the registered OnTick. */
+	FDelegateHandle OnTickHandle;
+
+	/** List of UI commands for this manager. This will be filled by this and corresponding classes. */
+	TSharedRef<FUICommandList> CommandList;
+
+	/** An instance of the Loading Profiler action manager. */
+	FLoadingProfilerActionManager ActionManager;
+
+	/** A weak pointer to the Loading Profiler window. */
+	TWeakPtr<class SLoadingProfilerWindow> ProfilerWindow;
+
+	/** If the Timing view is visible or hidden. */
+	bool bIsTimingViewVisible;
+
+	/** If the Event Aggregation tree view is visible or hidden. */
+	bool bIsEventAggregationTreeViewVisible;
+
+	/** If the Object Type Aggregation tree view is visible or hidden. */
+	bool bIsObjectTypeAggregationTreeViewVisible;
+
+	/** If the Package Details tree view is visible or hidden. */
+	bool bIsPackageDetailsTreeViewVisible;
+
+	/** If the Export Details tree view is visible or hidden. */
+	bool bIsExportDetailsTreeViewVisible;
+
+	/** A shared pointer to the global instance of the profiler manager. */
+	static TSharedPtr<FLoadingProfilerManager> Instance;
+};
diff --git a/Engine/Source/Developer/TraceInsights/Private/Insights/LoadingProfiler/Widgets/SLoadingProfilerToolbar.cpp b/Engine/Source/Developer/TraceInsights/Private/Insights/LoadingProfiler/Widgets/SLoadingProfilerToolbar.cpp
new file mode 100644
--- /dev/null
+++ b/Engine/Source/Developer/TraceInsights/Private/Insights/LoadingProfiler/Widgets/SLoadingProfilerToolbar.cpp
@@ -0,0 +1,114 @@
+// Copyright 1998-2019 Epic Games, Inc. All Rights Reserved.
+
+#include "SLoadingProfilerToolbar.h"
+
+#include "EditorStyleSet.h"
+#include "Framework/Application/SlateApplication.h"
+#include "Framework/Docking/TabManager.h"
+#include "Framework/MultiBox/MultiBoxDefs.h"
+#include "Framework/MultiBox/MultiBoxBuilder.h"
+#include "Widgets/Docking/SDockTab.h"
+#include "Widgets/Input/SCheckBox.h"
+#include "Widgets/Layout/SBorder.h"
+#include "Widgets/SBoxPanel.h"
+#include "Widgets/SToolTip.h"
+
+// Insights
+#include "Insights/InsightsCommands.h"
+#include "Insights/InsightsManager.h"
+#include "Insights/LoadingProfiler/LoadingProfilerCommands.h"
+#include "Insights/LoadingProfiler/LoadingProfilerManager.h"
+
+////////////////////////////////////////////////////////////////////////////////////////////////////
+
+#define LOCTEXT_NAMESPACE "SLoadingProfilerToolbar"
+
+////////////////////////////////////////////////////////////////////////////////////////////////////
+
+SLoadingProfilerToolbar::SLoadingProfilerToolbar()
+{
+}
+
+////////////////////////////////////////////////////////////////////////////////////////////////////
+
+SLoadingProfilerToolbar::~SLoadingProfilerToolbar()
+{
+}
+
+////////////////////////////////////////////////////////////////////////////////////////////////////
+
+void SLoadingProfilerToolbar::Construct(const FArguments& InArgs)
+{
+	struct Local
+	{
+		static void FillViewToolbar(FToolBarBuilder& ToolbarBuilder)
+		{
+			ToolbarBuilder.BeginSection("View");
+			{
+				ToolbarBuilder.AddToolBarButton(FLoadingProfilerCommands::Get().ToggleTimingViewVisibility);
+				ToolbarBuilder.AddToolBarButton(FLoadingProfilerCommands::Get().ToggleEventAggregationTreeViewVisibility);
+				ToolbarBuilder.AddToolBarButton(FLoadingProfilerCommands::Get().ToggleObjectTypeAggregationTreeViewVisibility);
+				ToolbarBuilder.AddToolBarButton(FLoadingProfilerCommands::Get().TogglePackageDetailsTreeViewVisibility);
+				ToolbarBuilder.AddToolBarButton(FLoadingProfilerCommands::Get().ToggleExportDetailsTreeViewVisibility);
+			}
+			ToolbarBuilder.EndSection();
+		}
+
+		static void FillRightSideToolbar(FToolBarBuilder& ToolbarBuilder)
+		{
+			ToolbarBuilder.BeginSection("Debug");
+			{
+				ToolbarBuilder.AddToolBarButton(FInsightsCommands::Get().ToggleDebugInfo);
+			}
+			ToolbarBuilder.EndSection();
+		}
+	};
+
+	TSharedPtr<FUICommandList> CommandList = FInsightsManager::Get()->GetCommandList();
+
+	FToolBarBuilder ToolbarBuilder(CommandList.ToSharedRef(), FMultiBoxCustomization::None);
+	Local::FillViewToolbar(ToolbarBuilder);
+
+	FToolBarBuilder RightSideToolbarBuilder(CommandList.ToSharedRef(), FMultiBoxCustomization::None);
+	Local::FillRightSideToolbar(RightSideToolbarBuilder);
+
+	// Create the tool bar!
+	ChildSlot
+	[
+		SNew(SHorizontalBox)
+
+		+SHorizontalBox::Slot()
+		.HAlign(HAlign_Left)
+		.VAlign(VAlign_Center)
+		.FillWidth(1.0)
+		.Padding(0.0f)
+		[
+			SNew(SBorder)
+			.Padding(0)
+			.BorderImage(FEditorStyle::GetBrush("NoBorder"))
+			.IsEnabled(FSlateApplication::Get().GetNormalExecutionAttribute())
+			[
+				ToolbarBuilder.MakeWidget()
+			]
+		]
+
+		+SHorizontalBox::Slot()
+		.HAlign(HAlign_Right)
+		.VAlign(VAlign_Center)
+		.AutoWidth()
+		.Padding(0.0f)
+		[
+			SNew(SBorder)
+			.Padding(0)
+			.BorderImage(FEditorStyle::GetBrush("NoBorder"))
+			.IsEnabled(FSlateApplication::Get().GetNormalExecutionAttribute())
+			[
+				RightSideToolbarBuilder.MakeWidget()
+			]
+		]
+	];
+}
+
+////////////////////////////////////////////////////////////////////////////////////////////////////
+
+#undef LOCTEXT_NAMESPACE
diff --git a/Engine/Source/Developer/TraceInsights/Private/Insights/LoadingProfiler/Widgets/SLoadingProfilerToolbar.h b/Engine/Source/Developer/TraceInsights/Private/Insights/LoadingProfiler/Widgets/SLoadingProfilerToolbar.h
new file mode 100644
--- /dev/null
+++ b/Engine/Source/Developer/TraceInsights/Private/Insights/LoadingProfiler/Widgets/SLoadingProfilerToolbar.h
@@ -0,0 +1,38 @@
+// Copyright 1998-2019 Epic Games, Inc. All Rights Reserved.
+
+
+#pragma once
+
+#include "CoreMinimal.h"
+#include "Styling/SlateTypes.h"
+#include "Widgets/DeclarativeSyntaxSupport.h"
+#include "Widgets/SCompoundWidget.h"
+
+class SBorder;
+class FToolBarBuilder;
+
+/** Ribbon based toolbar used as a main menu in the Profiler window. */
+class SLoadingProfilerToolbar : public SCompoundWidget
+{
+public:
+	/** Default constructor. */
+	SLoadingProfilerToolbar();
+
+	/** Virtual destructor. */
+	virtual ~SLoadingProfilerToolbar();
+
+	SLATE_BEGIN_ARGS( SLoadingProfilerToolbar )
+		{}
+	SLATE_END_ARGS()
+
+	/**
+	 * Construct this widget
+	 *
+	 * @param	InArgs	The declaration data for this widget
+	 */
+	void Construct( const FArguments& InArgs );
+
+private:
+	/** Create the UI commands for the toolbar */
+	void CreateCommands();
+};
diff --git a/Engine/Source/Developer/TraceInsights/Private/Insights/LoadingProfiler/Widgets/SLoadingProfilerWindow.cpp b/Engine/Source/Developer/TraceInsights/Private/Insights/LoadingProfiler/Widgets/SLoadingProfilerWindow.cpp
new file mode 100644
--- /dev/null
+++ b/Engine/Source/Developer/TraceInsights/Private/Insights/LoadingProfiler/Widgets/SLoadingProfilerWindow.cpp
@@ -0,0 +1,607 @@
+// Copyright 1998-2019 Epic Games, Inc. All Rights Reserved.
+
+#include "SLoadingProfilerWindow.h"
+
+#include "EditorStyleSet.h"
+#include "Framework/Docking/WorkspaceItem.h"
+#include "Framework/MultiBox/MultiBoxBuilder.h"
+#include "SlateOptMacros.h"
+#include "Widgets/Docking/SDockTab.h"
+#include "Widgets/Images/SImage.h"
+#include "Widgets/Layout/SBorder.h"
+#include "Widgets/Layout/SBox.h"
+#include "Widgets/Layout/SSpacer.h"
+#include "Widgets/Notifications/SNotificationList.h"
+#include "Widgets/SBoxPanel.h"
+#include "Widgets/Text/STextBlock.h"
+
+#if WITH_EDITOR
+	#include "EngineAnalytics.h"
+	#include "Runtime/Analytics/Analytics/Public/AnalyticsEventAttribute.h"
+	#include "Runtime/Analytics/Analytics/Public/Interfaces/IAnalyticsProvider.h"
+#endif // WITH_EDITOR
+
+// Insights
+#include "Insights/InsightsManager.h"
+#include "Insights/InsightsStyle.h"
+#include "Insights/LoadingProfiler/LoadingProfilerManager.h"
+#include "Insights/LoadingProfiler/Widgets/SLoadingProfilerToolbar.h"
+#include "Insights/Table/Widgets/STableTreeView.h"
+#include "Insights/Version.h"
+#include "Insights/Widgets/SInsightsSettings.h"
+#include "Insights/Widgets/STimingView.h"
+
+////////////////////////////////////////////////////////////////////////////////////////////////////
+
+#define LOCTEXT_NAMESPACE "SLoadingProfilerWindow"
+
+////////////////////////////////////////////////////////////////////////////////////////////////////
+
+const FName FLoadingProfilerTabs::ToolbarID(TEXT("Toolbar"));
+const FName FLoadingProfilerTabs::TimingViewID(TEXT("TimingView"));
+const FName FLoadingProfilerTabs::EventAggregationTreeViewID(TEXT("EventAggregation"));
+const FName FLoadingProfilerTabs::ObjectTypeAggregationTreeViewID(TEXT("ObjectTypeAggregation"));
+const FName FLoadingProfilerTabs::PackageDetailsTreeViewID(TEXT("PackageDetails"));
+const FName FLoadingProfilerTabs::ExportDetailsTreeViewID(TEXT("ExportDetails"));
+
+////////////////////////////////////////////////////////////////////////////////////////////////////
+
+SLoadingProfilerWindow::SLoadingProfilerWindow()
+	: DurationActive(0.0f)
+{
+}
+
+////////////////////////////////////////////////////////////////////////////////////////////////////
+
+SLoadingProfilerWindow::~SLoadingProfilerWindow()
+{
+#if WITH_EDITOR
+	if (DurationActive > 0.0f && FEngineAnalytics::IsAvailable())
+	{
+		FEngineAnalytics::GetProvider().RecordEvent(TEXT("Insights.Usage.LoadingProfiler"), FAnalyticsEventAttribute(TEXT("Duration"), DurationActive));
+	}
+#endif // WITH_EDITOR
+}
+
+////////////////////////////////////////////////////////////////////////////////////////////////////
+
+void SLoadingProfilerWindow::Reset()
+{
+	if (TimingView)
+	{
+		TimingView->Reset();
+	}
+
+	if (EventAggregationTreeView)
+	{
+		EventAggregationTreeView->Reset();
+	}
+
+	if (ObjectTypeAggregationTreeView)
+	{
+		ObjectTypeAggregationTreeView->Reset();
+	}
+
+	if (PackageDetailsTreeView)
+	{
+		PackageDetailsTreeView->Reset();
+	}
+
+	if (ExportDetailsTreeView)
+	{
+		ExportDetailsTreeView->Reset();
+	}
+}
+
+////////////////////////////////////////////////////////////////////////////////////////////////////
+
+BEGIN_SLATE_FUNCTION_BUILD_OPTIMIZATION
+
+TSharedRef<SDockTab> SLoadingProfilerWindow::SpawnTab_Toolbar(const FSpawnTabArgs& Args)
+{
+	const TSharedRef<SDockTab> DockTab = SNew(SDockTab)
+		.ShouldAutosize(true)
+		.TabRole(ETabRole::PanelTab)
+		[
+			SNew(SLoadingProfilerToolbar)
+		];
+
+	DockTab->SetOnTabClosed(SDockTab::FOnTabClosedCallback::CreateRaw(this, &SLoadingProfilerWindow::OnToolbarTabClosed));
+
+	return DockTab;
+}
+
+////////////////////////////////////////////////////////////////////////////////////////////////////
+
+void SLoadingProfilerWindow::OnToolbarTabClosed(TSharedRef<SDockTab> TabBeingClosed)
+{
+}
+
+////////////////////////////////////////////////////////////////////////////////////////////////////
+
+TSharedRef<SDockTab> SLoadingProfilerWindow::SpawnTab_TimingView(const FSpawnTabArgs& Args)
+{
+	FLoadingProfilerManager::Get()->SetTimingViewVisible(true);
+
+	const TSharedRef<SDockTab> DockTab = SNew(SDockTab)
+		.ShouldAutosize(false)
+		.TabRole(ETabRole::PanelTab)
+		[
+			SAssignNew(TimingView, STimingView)
+		];
+
+	TimingView->EnableAssetLoadingMode();
+
+	DockTab->SetOnTabClosed(SDockTab::FOnTabClosedCallback::CreateRaw(this, &SLoadingProfilerWindow::OnTimingViewTabClosed));
+
+	return DockTab;
+}
+
+////////////////////////////////////////////////////////////////////////////////////////////////////
+
+void SLoadingProfilerWindow::OnTimingViewTabClosed(TSharedRef<SDockTab> TabBeingClosed)
+{
+	FLoadingProfilerManager::Get()->SetTimingViewVisible(false);
+	TimingView = nullptr;
+}
+
+////////////////////////////////////////////////////////////////////////////////////////////////////
+
+TSharedRef<SDockTab> SLoadingProfilerWindow::SpawnTab_EventAggregationTreeView(const FSpawnTabArgs& Args)
+{
+	FLoadingProfilerManager::Get()->SetEventAggregationTreeViewVisible(true);
+
+	TSharedPtr<Insights::FTable> TablePtr = MakeShareable(new Insights::FTable());
+	TSharedPtr<const Trace::IAnalysisSession> Session = FInsightsManager::Get()->GetSession();
+	if (Session.IsValid() && Trace::ReadLoadTimeProfilerProvider(*Session.Get()))
+	{
+		Trace::FAnalysisSessionReadScope SessionReadScope(*Session.Get());
+		const Trace::ILoadTimeProfilerProvider& LoadTimeProfilerProvider = *Trace::ReadLoadTimeProfilerProvider(*Session.Get());
+		const double SelectionStartTime = TimingView ? TimingView->GetSelectionStartTime() : 0.0;
+		const double SelectionEndTime = TimingView ? TimingView->GetSelectionEndTime() : 0.0;
+		TSharedPtr<Trace::ITable<Trace::FLoadTimeProfilerAggregatedStats>> SourceTable = MakeShareable(LoadTimeProfilerProvider.CreateEventAggregation(SelectionStartTime, SelectionEndTime));
+		TablePtr->Init(SourceTable);
+	}
+
+	const TSharedRef<SDockTab> DockTab = SNew(SDockTab)
+		.ShouldAutosize(false)
+		.TabRole(ETabRole::PanelTab)
+		[
+			SAssignNew(EventAggregationTreeView, Insights::STableTreeView, TablePtr)
+		];
+
+	EventAggregationTreeView->Reset();
+
+	DockTab->SetOnTabClosed(SDockTab::FOnTabClosedCallback::CreateRaw(this, &SLoadingProfilerWindow::OnEventAggregationTreeViewTabClosed));
+
+	return DockTab;
+}
+
+////////////////////////////////////////////////////////////////////////////////////////////////////
+
+void SLoadingProfilerWindow::OnEventAggregationTreeViewTabClosed(TSharedRef<SDockTab> TabBeingClosed)
+{
+	FLoadingProfilerManager::Get()->SetEventAggregationTreeViewVisible(false);
+	EventAggregationTreeView = nullptr;
+}
+
+////////////////////////////////////////////////////////////////////////////////////////////////////
+
+TSharedRef<SDockTab> SLoadingProfilerWindow::SpawnTab_ObjectTypeAggregationTreeView(const FSpawnTabArgs& Args)
+{
+	FLoadingProfilerManager::Get()->SetObjectTypeAggregationTreeViewVisible(true);
+
+	TSharedPtr<Insights::FTable> TablePtr = MakeShareable(new Insights::FTable());
+	TSharedPtr<const Trace::IAnalysisSession> Session = FInsightsManager::Get()->GetSession();
+	if (Session.IsValid() && Trace::ReadLoadTimeProfilerProvider(*Session.Get()))
+	{
+		Trace::FAnalysisSessionReadScope SessionReadScope(*Session.Get());
+		const Trace::ILoadTimeProfilerProvider& LoadTimeProfilerProvider = *Trace::ReadLoadTimeProfilerProvider(*Session.Get());
+		const double SelectionStartTime = TimingView ? TimingView->GetSelectionStartTime() : 0.0;
+		const double SelectionEndTime = TimingView ? TimingView->GetSelectionEndTime() : 0.0;
+		TSharedPtr<Trace::ITable<Trace::FLoadTimeProfilerAggregatedStats>> SourceTable = MakeShareable(LoadTimeProfilerProvider.CreateObjectTypeAggregation(SelectionStartTime, SelectionEndTime));
+		TablePtr->Init(SourceTable);
+	}
+
+	const TSharedRef<SDockTab> DockTab = SNew(SDockTab)
+		.ShouldAutosize(false)
+		.TabRole(ETabRole::PanelTab)
+		[
+			SAssignNew(ObjectTypeAggregationTreeView, Insights::STableTreeView, TablePtr)
+		];
+
+	ObjectTypeAggregationTreeView->Reset();
+
+	DockTab->SetOnTabClosed(SDockTab::FOnTabClosedCallback::CreateRaw(this, &SLoadingProfilerWindow::OnObjectTypeAggregationTreeViewTabClosed));
+
+	return DockTab;
+}
+
+////////////////////////////////////////////////////////////////////////////////////////////////////
+
+void SLoadingProfilerWindow::OnObjectTypeAggregationTreeViewTabClosed(TSharedRef<SDockTab> TabBeingClosed)
+{
+	FLoadingProfilerManager::Get()->SetObjectTypeAggregationTreeViewVisible(false);
+	ObjectTypeAggregationTreeView = nullptr;
+}
+
+////////////////////////////////////////////////////////////////////////////////////////////////////
+
+TSharedRef<SDockTab> SLoadingProfilerWindow::SpawnTab_PackageDetailsTreeView(const FSpawnTabArgs& Args)
+{
+	FLoadingProfilerManager::Get()->SetPackageDetailsTreeViewVisible(true);
+
+	TSharedPtr<Insights::FTable> TablePtr = MakeShareable(new Insights::FTable());
+	TSharedPtr<const Trace::IAnalysisSession> Session = FInsightsManager::Get()->GetSession();
+	if (Session.IsValid() && Trace::ReadLoadTimeProfilerProvider(*Session.Get()))
+	{
+		Trace::FAnalysisSessionReadScope SessionReadScope(*Session.Get());
+		const Trace::ILoadTimeProfilerProvider& LoadTimeProfilerProvider = *Trace::ReadLoadTimeProfilerProvider(*Session.Get());
+		const double SelectionStartTime = TimingView ? TimingView->GetSelectionStartTime() : 0.0;
+		const double SelectionEndTime = TimingView ? TimingView->GetSelectionEndTime() : 0.0;
+		TSharedPtr<Trace::ITable<Trace::FPackagesTableRow>> SourceTable = MakeShareable(LoadTimeProfilerProvider.CreatePackageDetailsTable(SelectionStartTime, SelectionEndTime));
+		TablePtr->Init(SourceTable);
+	}
+
+	const TSharedRef<SDockTab> DockTab = SNew(SDockTab)
+		.ShouldAutosize(false)
+		.TabRole(ETabRole::PanelTab)
+		[
+			SAssignNew(PackageDetailsTreeView, Insights::STableTreeView, TablePtr)
+		];
+
+	PackageDetailsTreeView->Reset();
+
+	DockTab->SetOnTabClosed(SDockTab::FOnTabClosedCallback::CreateRaw(this, &SLoadingProfilerWindow::OnPackageDetailsTreeViewTabClosed));
+
+	return DockTab;
+}
+
+////////////////////////////////////////////////////////////////////////////////////////////////////
+
+void SLoadingProfilerWindow::OnPackageDetailsTreeViewTabClosed(TSharedRef<SDockTab> TabBeingClosed)
+{
+	FLoadingProfilerManager::Get()->SetPackageDetailsTreeViewVisible(false);
+	PackageDetailsTreeView = nullptr;
+}
+
+////////////////////////////////////////////////////////////////////////////////////////////////////
+
+TSharedRef<SDockTab> SLoadingProfilerWindow::SpawnTab_ExportDetailsTreeView(const FSpawnTabArgs& Args)
+{
+	FLoadingProfilerManager::Get()->SetExportDetailsTreeViewVisible(true);
+
+	TSharedPtr<Insights::FTable> TablePtr = MakeShareable(new Insights::FTable());
+	TSharedPtr<const Trace::IAnalysisSession> Session = FInsightsManager::Get()->GetSession();
+	if (Session.IsValid() && Trace::ReadLoadTimeProfilerProvider(*Session.Get()))
+	{
+		Trace::FAnalysisSessionReadScope SessionReadScope(*Session.Get());
+		const Trace::ILoadTimeProfilerProvider& LoadTimeProfilerProvider = *Trace::ReadLoadTimeProfilerProvider(*Session.Get());
+		const double SelectionStartTime = TimingView ? TimingView->GetSelectionStartTime() : 0.0;
+		const double SelectionEndTime = TimingView ? TimingView->GetSelectionEndTime() : 0.0;
+		TSharedPtr<Trace::ITable<Trace::FExportsTableRow>> SourceTable = MakeShareable(LoadTimeProfilerProvider.CreateExportDetailsTable(SelectionStartTime, SelectionEndTime));
+		TablePtr->Init(SourceTable);
+	}
+
+	const TSharedRef<SDockTab> DockTab = SNew(SDockTab)
+		.ShouldAutosize(false)
+		.TabRole(ETabRole::PanelTab)
+		[
+			SAssignNew(ExportDetailsTreeView, Insights::STableTreeView, TablePtr)
+		];
+
+	ExportDetailsTreeView->Reset();
+
+	DockTab->SetOnTabClosed(SDockTab::FOnTabClosedCallback::CreateRaw(this, &SLoadingProfilerWindow::OnExportDetailsTreeViewTabClosed));
+
+	return DockTab;
+}
+
+////////////////////////////////////////////////////////////////////////////////////////////////////
+
+void SLoadingProfilerWindow::OnExportDetailsTreeViewTabClosed(TSharedRef<SDockTab> TabBeingClosed)
+{
+	FLoadingProfilerManager::Get()->SetExportDetailsTreeViewVisible(false);
+	ExportDetailsTreeView = nullptr;
+}
+
+////////////////////////////////////////////////////////////////////////////////////////////////////
+
+void SLoadingProfilerWindow::Construct(const FArguments& InArgs, const TSharedRef<SDockTab>& ConstructUnderMajorTab, const TSharedPtr<SWindow>& ConstructUnderWindow)
+{
+	// Create & initialize tab manager.
+	TabManager = FGlobalTabmanager::Get()->NewTabManager(ConstructUnderMajorTab);
+	TSharedRef<FWorkspaceItem> AppMenuGroup = TabManager->AddLocalWorkspaceMenuCategory(LOCTEXT("LoadingProfilerMenuGroupName", "Asset Loading Insights"));
+
+	TabManager->RegisterTabSpawner(FLoadingProfilerTabs::ToolbarID, FOnSpawnTab::CreateRaw(this, &SLoadingProfilerWindow::SpawnTab_Toolbar))
+		.SetDisplayName(LOCTEXT("DeviceToolbarTabTitle", "Toolbar"))
+		.SetIcon(FSlateIcon(FInsightsStyle::GetStyleSetName(), "Toolbar.Icon.Small"))
+		.SetGroup(AppMenuGroup);
+
+	TabManager->RegisterTabSpawner(FLoadingProfilerTabs::TimingViewID, FOnSpawnTab::CreateRaw(this, &SLoadingProfilerWindow::SpawnTab_TimingView))
+		.SetDisplayName(LOCTEXT("LoadingProfiler.TimingViewTabTitle", "Timing View"))
+		.SetIcon(FSlateIcon(FInsightsStyle::GetStyleSetName(), "TimingView.Icon.Small"))
+		.SetGroup(AppMenuGroup);
+
+	TabManager->RegisterTabSpawner(FLoadingProfilerTabs::EventAggregationTreeViewID, FOnSpawnTab::CreateRaw(this, &SLoadingProfilerWindow::SpawnTab_EventAggregationTreeView))
+		.SetDisplayName(LOCTEXT("LoadingProfiler.EventAggregationTreeViewTabTitle", "Event Aggregation"))
+		.SetIcon(FSlateIcon(FInsightsStyle::GetStyleSetName(), "TableTreeView.Icon.Small"))
+		.SetGroup(AppMenuGroup);
+
+	TabManager->RegisterTabSpawner(FLoadingProfilerTabs::ObjectTypeAggregationTreeViewID, FOnSpawnTab::CreateRaw(this, &SLoadingProfilerWindow::SpawnTab_ObjectTypeAggregationTreeView))
+		.SetDisplayName(LOCTEXT("LoadingProfiler.ObjectTypeAggregationTreeViewTabTitle", "Object Type Aggregation"))
+		.SetIcon(FSlateIcon(FInsightsStyle::GetStyleSetName(), "TableTreeView.Icon.Small"))
+		.SetGroup(AppMenuGroup);
+
+	TabManager->RegisterTabSpawner(FLoadingProfilerTabs::PackageDetailsTreeViewID, FOnSpawnTab::CreateRaw(this, &SLoadingProfilerWindow::SpawnTab_PackageDetailsTreeView))
+		.SetDisplayName(LOCTEXT("LoadingProfiler.PackageDetailsTreeViewTabTitle", "Package Details"))
+		.SetIcon(FSlateIcon(FInsightsStyle::GetStyleSetName(), "TableTreeView.Icon.Small"))
+		.SetGroup(AppMenuGroup);
+
+	TabManager->RegisterTabSpawner(FLoadingProfilerTabs::ExportDetailsTreeViewID, FOnSpawnTab::CreateRaw(this, &SLoadingProfilerWindow::SpawnTab_ExportDetailsTreeView))
+		.SetDisplayName(LOCTEXT("LoadingProfiler.ExportDetailsTreeViewTabTitle", "Export Details"))
+		.SetIcon(FSlateIcon(FInsightsStyle::GetStyleSetName(), "TableTreeView.Icon.Small"))
+		.SetGroup(AppMenuGroup);
+
+	TSharedPtr<FLoadingProfilerManager> LoadingProfilerManager = FLoadingProfilerManager::Get();
+	ensure(LoadingProfilerManager.IsValid());
+
+	// Create tab layout.
+	const TSharedRef<FTabManager::FLayout> Layout = FTabManager::NewLayout("InsightsLoadingProfilerLayout_v1.0")
+		->AddArea
+		(
+			FTabManager::NewPrimaryArea()
+			->SetOrientation(Orient_Vertical)
+			->Split
+			(
+				FTabManager::NewStack()
+				->AddTab(FLoadingProfilerTabs::ToolbarID, ETabState::OpenedTab)
+				->SetHideTabWell(true)
+			)
+			->Split
+			(
+				FTabManager::NewSplitter()
+				->SetOrientation(Orient_Horizontal)
+				->SetSizeCoefficient(1.0f)
+				->Split
+				(
+					FTabManager::NewStack()
+					->SetSizeCoefficient(0.5f)
+					->SetHideTabWell(true)
+					->AddTab(FLoadingProfilerTabs::TimingViewID, LoadingProfilerManager->IsTimingViewVisible() ? ETabState::OpenedTab : ETabState::ClosedTab)
+				)
+				->Split
+				(
+					FTabManager::NewStack()
+					->SetSizeCoefficient(0.35f)
+					->AddTab(FLoadingProfilerTabs::EventAggregationTreeViewID, LoadingProfilerManager->IsEventAggregationTreeViewVisible() ? ETabState::OpenedTab : ETabState::ClosedTab)
+					->AddTab(FLoadingProfilerTabs::ObjectTypeAggregationTreeViewID, LoadingProfilerManager->IsObjectTypeAggregationTreeViewVisible() ? ETabState::OpenedTab : ETabState::ClosedTab)
+					->AddTab(FLoadingProfilerTabs::PackageDetailsTreeViewID, LoadingProfilerManager->IsPackageDetailsTreeViewVisible() ? ETabState::OpenedTab : ETabState::ClosedTab)
+					->AddTab(FLoadingProfilerTabs::ExportDetailsTreeViewID, LoadingProfilerManager->IsExportDetailsTreeViewVisible() ? ETabState::OpenedTab : ETabState::ClosedTab)
+					->SetForegroundTab(FLoadingProfilerTabs::EventAggregationTreeViewID)
+				)
+			)
+		);
+
+	// Create & initialize main menu.
+	FMenuBarBuilder MenuBarBuilder = FMenuBarBuilder(TSharedPtr<FUICommandList>());
+
+	MenuBarBuilder.AddPullDownMenu(
+		LOCTEXT("MenuLabel", "MENU"),
+		FText::GetEmpty(),
+		FNewMenuDelegate::CreateStatic(&SLoadingProfilerWindow::FillMenu, TabManager),
+		FName(TEXT("MENU"))
+	);
+
+	ChildSlot
+		[
+			SNew(SOverlay)
+
+			// Version
+			+ SOverlay::Slot()
+				.HAlign(HAlign_Right)
+				.VAlign(VAlign_Top)
+				.Padding(0.0f, -16.0f, 0.0f, 0.0f)
+				[
+					SNew(STextBlock)
+						.Clipping(EWidgetClipping::ClipToBoundsWithoutIntersecting)
+						.Text(LOCTEXT("UnrealInsightsVersion", UNREAL_INSIGHTS_VERSION_STRING_EX))
+						.ColorAndOpacity(FLinearColor(0.15f, 0.15f, 0.15f, 1.0f))
+				]
+
+			// Overlay slot for the main window area
+			+ SOverlay::Slot()
+				.HAlign(HAlign_Fill)
+				.VAlign(VAlign_Fill)
+				[
+					SNew(SVerticalBox)
+
+					+ SVerticalBox::Slot()
+						.AutoHeight()
+						[
+							MenuBarBuilder.MakeWidget()
+						]
+
+					+ SVerticalBox::Slot()
+						.FillHeight(1.0f)
+						[
+							TabManager->RestoreFrom(Layout, ConstructUnderWindow).ToSharedRef()
+						]
+				]
+
+			// Session hint overlay
+			+ SOverlay::Slot()
+				.HAlign(HAlign_Center)
+				.VAlign(VAlign_Center)
+				[
+					SNew(SBorder)
+						.Visibility(this, &SLoadingProfilerWindow::IsSessionOverlayVisible)
+						.BorderImage(FEditorStyle::GetBrush("NotificationList.ItemBackground"))
+						.Padding(8.0f)
+						[
+							SNew(STextBlock)
+								.Text(LOCTEXT("SelectTraceOverlayText", "Please select a trace."))
+						]
+				]
+		];
+
+	// Tell tab-manager about the global menu bar.
+	TabManager->SetMenuMultiBox(MenuBarBuilder.GetMultiBox());
+}
+
+END_SLATE_FUNCTION_BUILD_OPTIMIZATION
+
+////////////////////////////////////////////////////////////////////////////////////////////////////
+
+void SLoadingProfilerWindow::FillMenu(FMenuBuilder& MenuBuilder, const TSharedPtr<FTabManager> TabManager)
+{
+	if (!TabManager.IsValid())
+	{
+		return;
+	}
+
+#if !WITH_EDITOR
+	//TODO: FGlobalTabmanager::Get()->PopulateTabSpawnerMenu(MenuBuilder, WorkspaceMenu::GetMenuStructure().GetStructureRoot());
+#endif //!WITH_EDITOR
+
+	TabManager->PopulateLocalTabSpawnerMenu(MenuBuilder);
+}
+
+////////////////////////////////////////////////////////////////////////////////////////////////////
+
+void SLoadingProfilerWindow::ShowTab(const FName& TabID)
+{
+	if (TabManager->HasTabSpawner(TabID))
+	{
+		TabManager->InvokeTab(TabID);
+	}
+}
+
+////////////////////////////////////////////////////////////////////////////////////////////////////
+
+void SLoadingProfilerWindow::HideTab(const FName& TabID)
+{
+	TSharedPtr<SDockTab> Tab = TabManager->FindExistingLiveTab(TabID);
+	if (Tab.IsValid())
+	{
+		Tab->RequestCloseTab();
+	}
+}
+
+////////////////////////////////////////////////////////////////////////////////////////////////////
+
+EVisibility SLoadingProfilerWindow::IsSessionOverlayVisible() const
+{
+	if (FInsightsManager::Get()->GetSession().IsValid())
+	{
+		return EVisibility::Hidden;
+	}
+	else
+	{
+		return EVisibility::Visible;
+	}
+}
+
+////////////////////////////////////////////////////////////////////////////////////////////////////
+
+bool SLoadingProfilerWindow::IsProfilerEnabled() const
+{
+	return FInsightsManager::Get()->GetSession().IsValid();
+}
+
+////////////////////////////////////////////////////////////////////////////////////////////////////
+
+EActiveTimerReturnType SLoadingProfilerWindow::UpdateActiveDuration(double InCurrentTime, float InDeltaTime)
+{
+	DurationActive += InDeltaTime;
+
+	// The profiler window will explicitly unregister this active timer when the mouse leaves.
+	return EActiveTimerReturnType::Continue;
+}
+
+////////////////////////////////////////////////////////////////////////////////////////////////////
+
+void SLoadingProfilerWindow::OnMouseEnter(const FGeometry& MyGeometry, const FPointerEvent& MouseEvent)
+{
+	SCompoundWidget::OnMouseEnter(MyGeometry, MouseEvent);
+
+	if (!ActiveTimerHandle.IsValid())
+	{
+		ActiveTimerHandle = RegisterActiveTimer(0.f, FWidgetActiveTimerDelegate::CreateSP(this, &SLoadingProfilerWindow::UpdateActiveDuration));
+	}
+}
+
+////////////////////////////////////////////////////////////////////////////////////////////////////
+
+void SLoadingProfilerWindow::OnMouseLeave(const FPointerEvent& MouseEvent)
+{
+	SCompoundWidget::OnMouseLeave(MouseEvent);
+
+	auto PinnedActiveTimerHandle = ActiveTimerHandle.Pin();
+	if (PinnedActiveTimerHandle.IsValid())
+	{
+		UnRegisterActiveTimer(PinnedActiveTimerHandle.ToSharedRef());
+	}
+}
+
+////////////////////////////////////////////////////////////////////////////////////////////////////
+
+FReply SLoadingProfilerWindow::OnKeyDown(const FGeometry& MyGeometry, const FKeyEvent& InKeyEvent)
+{
+	return FLoadingProfilerManager::Get()->GetCommandList()->ProcessCommandBindings(InKeyEvent) ? FReply::Handled() : FReply::Unhandled();
+}
+
+////////////////////////////////////////////////////////////////////////////////////////////////////
+
+FReply SLoadingProfilerWindow::OnDragOver(const FGeometry& MyGeometry, const FDragDropEvent& DragDropEvent)
+{
+	TSharedPtr<FExternalDragOperation> DragDropOp = DragDropEvent.GetOperationAs<FExternalDragOperation>();
+	if (DragDropOp.IsValid())
+	{
+		if (DragDropOp->HasFiles())
+		{
+			const TArray<FString>& Files = DragDropOp->GetFiles();
+			if (Files.Num() == 1)
+			{
+				const FString DraggedFileExtension = FPaths::GetExtension(Files[0], true);
+				if (DraggedFileExtension == TEXT(".utrace"))
+				{
+					return FReply::Handled();
+				}
+			}
+		}
+	}
+
+	return SCompoundWidget::OnDragOver(MyGeometry,DragDropEvent);
+}
+
+////////////////////////////////////////////////////////////////////////////////////////////////////
+
+FReply SLoadingProfilerWindow::OnDrop(const FGeometry& MyGeometry, const FDragDropEvent& DragDropEvent)
+{
+	TSharedPtr<FExternalDragOperation> DragDropOp = DragDropEvent.GetOperationAs<FExternalDragOperation>();
+	if (DragDropOp.IsValid())
+	{
+		if (DragDropOp->HasFiles())
+		{
+			// For now, only allow a single file.
+			const TArray<FString>& Files = DragDropOp->GetFiles();
+			if (Files.Num() == 1)
+			{
+				const FString DraggedFileExtension = FPaths::GetExtension(Files[0], true);
+				if (DraggedFileExtension == TEXT(".utrace"))
+				{
+					// Enqueue load operation.
+					FInsightsManager::Get()->LoadTraceFile(Files[0]);
+					return FReply::Handled();
+				}
+			}
+		}
+	}
+
+	return SCompoundWidget::OnDrop(MyGeometry,DragDropEvent);
+}
+
+////////////////////////////////////////////////////////////////////////////////////////////////////
+
+#undef LOCTEXT_NAMESPACE
diff --git a/Engine/Source/Developer/TraceInsights/Private/Insights/LoadingProfiler/Widgets/SLoadingProfilerWindow.h b/Engine/Source/Developer/TraceInsights/Private/Insights/LoadingProfiler/Widgets/SLoadingProfilerWindow.h
new file mode 100644
--- /dev/null
+++ b/Engine/Source/Developer/TraceInsights/Private/Insights/LoadingProfiler/Widgets/SLoadingProfilerWindow.h
@@ -0,0 +1,189 @@
+// Copyright 1998-2019 Epic Games, Inc. All Rights Reserved.
+
+#pragma once
+
+#include "CoreMinimal.h"
+#include "Framework/Docking/TabManager.h"
+#include "Input/Reply.h"
+#include "Layout/Visibility.h"
+#include "Misc/Guid.h"
+#include "SlateFwd.h"
+#include "Widgets/DeclarativeSyntaxSupport.h"
+#include "Widgets/Docking/SDockTab.h"
+#include "Widgets/SCompoundWidget.h"
+
+// Insights
+#include "Insights/InsightsManager.h"
+
+////////////////////////////////////////////////////////////////////////////////////////////////////
+
+class FActiveTimerHandle;
+class FMenuBuilder;
+
+class STimingView;
+
+namespace Insights
+{
+	class STableTreeView;
+}
+
+////////////////////////////////////////////////////////////////////////////////////////////////////
+
+struct FLoadingProfilerTabs
+{
+	// Tab identifiers
+	static const FName ToolbarID;
+	static const FName TimingViewID;
+	static const FName EventAggregationTreeViewID;
+	static const FName ObjectTypeAggregationTreeViewID;
+	static const FName PackageDetailsTreeViewID;
+	static const FName ExportDetailsTreeViewID;
+};
+
+////////////////////////////////////////////////////////////////////////////////////////////////////
+
+/** Implements the timing profiler window. */
+class SLoadingProfilerWindow : public SCompoundWidget
+{
+public:
+	/** Default constructor. */
+	SLoadingProfilerWindow();
+
+	/** Virtual destructor. */
+	virtual ~SLoadingProfilerWindow();
+
+	SLATE_BEGIN_ARGS(SLoadingProfilerWindow) {}
+	SLATE_END_ARGS()
+
+	void Reset();
+
+	/** Constructs this widget. */
+	void Construct(const FArguments& InArgs, const TSharedRef<SDockTab>& ConstructUnderMajorTab, const TSharedPtr<SWindow>& ConstructUnderWindow);
+
+	void ShowTab(const FName& TabID);
+	void HideTab(const FName& TabID);
+	void ShowHideTab(const FName& TabID, bool bShow) { bShow ? ShowTab(TabID) : HideTab(TabID); }
+
+	TSharedPtr<FTabManager> GetTabManager() const { return TabManager; }
+
+	TSharedPtr<STimingView> GetTimingView() const { return TimingView; }
+	TSharedPtr<Insights::STableTreeView> GetEventAggregationTreeView() const { return EventAggregationTreeView; }
+	TSharedPtr<Insights::STableTreeView> GetObjectTypeAggregationTreeView() const { return ObjectTypeAggregationTreeView; }
+	TSharedPtr<Insights::STableTreeView> GetPackageDetailsTreeView() const { return PackageDetailsTreeView; }
+	TSharedPtr<Insights::STableTreeView> GetExportDetailsTreeView() const { return ExportDetailsTreeView; }
+
+private:
+	TSharedRef<SDockTab> SpawnTab_Toolbar(const FSpawnTabArgs& Args);
+	void OnToolbarTabClosed(TSharedRef<SDockTab> TabBeingClosed);
+
+	TSharedRef<SDockTab> SpawnTab_TimingView(const FSpawnTabArgs& Args);
+	void OnTimingViewTabClosed(TSharedRef<SDockTab> TabBeingClosed);
+
+	TSharedRef<SDockTab> SpawnTab_EventAggregationTreeView(const FSpawnTabArgs& Args);
+	void OnEventAggregationTreeViewTabClosed(TSharedRef<SDockTab> TabBeingClosed);
+
+	TSharedRef<SDockTab> SpawnTab_ObjectTypeAggregationTreeView(const FSpawnTabArgs& Args);
+	void OnObjectTypeAggregationTreeViewTabClosed(TSharedRef<SDockTab> TabBeingClosed);
+
+	TSharedRef<SDockTab> SpawnTab_PackageDetailsTreeView(const FSpawnTabArgs& Args);
+	void OnPackageDetailsTreeViewTabClosed(TSharedRef<SDockTab> TabBeingClosed);
+
+	TSharedRef<SDockTab> SpawnTab_ExportDetailsTreeView(const FSpawnTabArgs& Args);
+	void OnExportDetailsTreeViewTabClosed(TSharedRef<SDockTab> TabBeingClosed);
+
+	/**
+	 * Fill the main menu with menu items.
+	 *
+	 * @param MenuBuilder The multi-box builder that should be filled with content for this pull-down menu.
+	 * @param TabManager A Tab Manager from which to populate tab spawner menu items.
+	 */
+	static void FillMenu(FMenuBuilder& MenuBuilder, const TSharedPtr<FTabManager> TabManager);
+
+	/** Callback for determining the visibility of the 'Select a session' overlay. */
+	EVisibility IsSessionOverlayVisible() const;
+
+	/** Callback for getting the enabled state of the profiler window. */
+	bool IsProfilerEnabled() const;
+
+	/** Updates the amount of time the profiler has been active. */
+	EActiveTimerReturnType UpdateActiveDuration(double InCurrentTime, float InDeltaTime);
+
+	/**
+	 * Ticks this widget. Override in derived classes, but always call the parent implementation.
+	 *
+	 * @param  AllottedGeometry The space allotted for this widget
+	 * @param  InCurrentTime  Current absolute real time
+	 * @param  InDeltaTime  Real time passed since last tick
+	 */
+	//virtual void Tick(const FGeometry& AllottedGeometry, const double InCurrentTime, const float InDeltaTime) override;
+
+	/**
+	 * The system will use this event to notify a widget that the cursor has entered it. This event is NOT bubbled.
+	 *
+	 * @param MyGeometry The Geometry of the widget receiving the event
+	 * @param MouseEvent Information about the input event
+	 */
+	virtual void OnMouseEnter(const FGeometry& MyGeometry, const FPointerEvent& MouseEvent) override;
+
+	/**
+	 * The system will use this event to notify a widget that the cursor has left it. This event is NOT bubbled.
+	 *
+	 * @param MouseEvent Information about the input event
+	 */
+	virtual void OnMouseLeave(const FPointerEvent& MouseEvent) override;
+
+	/**
+	 * Called after a key is pressed when this widget has focus
+	 *
+	 * @param MyGeometry The Geometry of the widget receiving the event
+	 * @param  InKeyEvent  Key event
+	 *
+	 * @return  Returns whether the event was handled, along with other possible actions
+	 */
+	virtual FReply OnKeyDown(const FGeometry& MyGeometry, const FKeyEvent& InKeyEvent) override;
+
+	/**
+	 * Called when the user is dropping something onto a widget; terminates drag and drop.
+	 *
+	 * @param MyGeometry      The geometry of the widget receiving the event.
+	 * @param DragDropEvent   The drag and drop event.
+	 *
+	 * @return A reply that indicated whether this event was handled.
+	 */
+	virtual FReply OnDrop(const FGeometry& MyGeometry, const FDragDropEvent& DragDropEvent) override;
+
+	/**
+	 * Called during drag and drop when the the mouse is being dragged over a widget.
+	 *
+	 * @param MyGeometry      The geometry of the widget receiving the event.
+	 * @param DragDropEvent   The drag and drop event.
+	 *
+	 * @return A reply that indicated whether this event was handled.
+	 */
+	virtual FReply OnDragOver(const FGeometry& MyGeometry, const FDragDropEvent& DragDropEvent)  override;
+
+private:
+	/** The Timing view (multi-track) widget */
+	TSharedPtr<STimingView> TimingView;
+
+	/** The Event Aggregation tree view widget */
+	TSharedPtr<Insights::STableTreeView> EventAggregationTreeView;
+
+	/** The Object Type Aggregation tree view widget */
+	TSharedPtr<Insights::STableTreeView> ObjectTypeAggregationTreeView;
+
+	/** The Package Details tree view widget */
+	TSharedPtr<Insights::STableTreeView> PackageDetailsTreeView;
+
+	/** The Export Details tree view widget */
+	TSharedPtr<Insights::STableTreeView> ExportDetailsTreeView;
+
+	/** Holds the tab manager that manages the front-end's tabs. */
+	TSharedPtr<FTabManager> TabManager;
+
+	/** The handle to the active update duration tick */
+	TWeakPtr<FActiveTimerHandle> ActiveTimerHandle;
+
+	/** The number of seconds the profiler has been active */
+	float DurationActive;
+};
diff --git a/Engine/Source/Developer/TraceInsights/Private/Insights/Table/ViewModels/BaseTreeNode.cpp b/Engine/Source/Developer/TraceInsights/Private/Insights/Table/ViewModels/BaseTreeNode.cpp
new file mode 100644
--- /dev/null
+++ b/Engine/Source/Developer/TraceInsights/Private/Insights/Table/ViewModels/BaseTreeNode.cpp
@@ -0,0 +1,59 @@
+// Copyright 1998-2019 Epic Games, Inc. All Rights Reserved.
+
+#include "BaseTreeNode.h"
+
+// Insights
+#include "Insights/Table/ViewModels/TableCellValueSorter.h"
+
+#define LOCTEXT_NAMESPACE "Insights_TreeNode"
+
+namespace Insights
+{
+
+////////////////////////////////////////////////////////////////////////////////////////////////////
+
+const FText FBaseTreeNode::GetDisplayName() const
+{
+	FText Text = FText::GetEmpty();
+
+	if (IsGroup())
+	{
+		const int32 NumChildren = Children.Num();
+		const int32 NumFilteredChildren = FilteredChildren.Num();
+
+		if (NumFilteredChildren == NumChildren)
+		{
+			Text = FText::Format(LOCTEXT("TreeNodeGroupTextFmt1", "{0} ({1})"), FText::FromName(GetName()), FText::AsNumber(NumChildren));
+		}
+		else
+		{
+			Text = FText::Format(LOCTEXT("TreeNodeGroupTextFmt2", "{0} ({1} / {2})"), FText::FromName(GetName()), FText::AsNumber(NumFilteredChildren), FText::AsNumber(NumChildren));
+		}
+	}
+	else
+	{
+		Text = FText::FromName(GetName());
+	}
+
+	return Text;
+}
+
+////////////////////////////////////////////////////////////////////////////////////////////////////
+
+void FBaseTreeNode::SortChildrenAscending(const ITableCellValueSorter& Sorter)
+{
+	Sorter.Sort(Children, ESortMode::Ascending);
+}
+
+////////////////////////////////////////////////////////////////////////////////////////////////////
+
+void FBaseTreeNode::SortChildrenDescending(const ITableCellValueSorter& Sorter)
+{
+	Sorter.Sort(Children, ESortMode::Descending);
+}
+
+////////////////////////////////////////////////////////////////////////////////////////////////////
+
+} // namespace Insights
+
+#undef LOCTEXT_NAMESPACE
diff --git a/Engine/Source/Developer/TraceInsights/Private/Insights/Table/ViewModels/BaseTreeNode.h b/Engine/Source/Developer/TraceInsights/Private/Insights/Table/ViewModels/BaseTreeNode.h
new file mode 100644
--- /dev/null
+++ b/Engine/Source/Developer/TraceInsights/Private/Insights/Table/ViewModels/BaseTreeNode.h
@@ -0,0 +1,186 @@
+// Copyright 1998-2019 Epic Games, Inc. All Rights Reserved.
+
+#pragma once
+
+#include "CoreMinimal.h"
+
+namespace Insights
+{
+
+class ITableCellValueSorter;
+
+////////////////////////////////////////////////////////////////////////////////////////////////////
+
+class FBaseTreeNode;
+
+/** Type definition for shared pointers to instances of FBaseTreeNode. */
+typedef TSharedPtr<class FBaseTreeNode> FBaseTreeNodePtr;
+
+/** Type definition for shared references to instances of FBaseTreeNode. */
+typedef TSharedRef<class FBaseTreeNode> FBaseTreeNodeRef;
+
+/** Type definition for shared references to const instances of FBaseTreeNode. */
+typedef TSharedRef<const class FBaseTreeNode> FBaseTreeNodeRefConst;
+
+/** Type definition for weak references to instances of FBaseTreeNode. */
+typedef TWeakPtr<class FBaseTreeNode> FBaseTreeNodeWeak;
+
+////////////////////////////////////////////////////////////////////////////////////////////////////
+
+class FBaseTreeNode : public TSharedFromThis<FBaseTreeNode>
+{
+public:
+	static constexpr int32 InvalidId = -1;
+
+public:
+	/** Initialization constructor for the node. */
+	FBaseTreeNode(uint64 InId, const FName InName, bool bInIsGroup)
+		: Id(InId)
+		, Name(InName)
+		, bIsGroup(bInIsGroup)
+		, bIsExpanded(false)
+	{
+	}
+
+	virtual ~FBaseTreeNode()
+	{
+	}
+
+	virtual const FName& GetTypeName() const = 0;
+
+	/**
+	 * @return an id of this node.
+	 */
+	const int32 GetId() const
+	{
+		return Id;
+	}
+
+	/**
+	 * @return a name of this node.
+	 */
+	const FName& GetName() const
+	{
+		return Name;
+	}
+
+	/**
+	 * @return a name of this node to display in tree view. Ex.: group nodes may include additional info, like the number of visible/total children.
+	 */
+	virtual const FText GetDisplayName() const;
+
+	/**
+	 * @return true, if this node is a group node.
+	 */
+	bool IsGroup() const
+	{
+		return bIsGroup;
+	}
+
+	/**
+	 * @return a const reference to the child nodes of this group.
+	 */
+	FORCEINLINE_DEBUGGABLE const TArray<FBaseTreeNodePtr>& GetChildren() const
+	{
+		return Children;
+	}
+
+	/**
+	 * @return a const reference to the child nodes that should be visible to the UI based on filtering.
+	 */
+	FORCEINLINE_DEBUGGABLE const TArray<FBaseTreeNodePtr>& GetFilteredChildren() const
+	{
+		return FilteredChildren;
+	}
+
+	/**
+	 * @return a weak reference to the group of this node, may be invalid.
+	 */
+	FBaseTreeNodeWeak GetGroupPtr() const
+	{
+		return GroupPtr;
+	}
+
+	/**
+	 * @return a name of the group node that this node belongs to.
+	 */
+	const FName& GetGroupName() const
+	{
+		return GroupPtr.Pin()->GetName();
+	}
+
+	virtual bool IsFiltered() const
+	{
+		return false;
+	}
+
+	void SortChildrenAscending(const ITableCellValueSorter& Sorter);
+	void SortChildrenDescending(const ITableCellValueSorter& Sorter);
+
+	/** Adds specified child to the children and sets group for it. */
+	FORCEINLINE_DEBUGGABLE void AddChildAndSetGroupPtr(const FBaseTreeNodePtr& ChildPtr)
+	{
+		ChildPtr->GroupPtr = AsShared();
+		Children.Add(ChildPtr);
+	}
+
+	/** Clears children. */
+	void ClearChildren()
+	{
+		for (FBaseTreeNodePtr& NodePtr : Children)
+		{
+			NodePtr->GroupPtr = nullptr;
+		}
+		Children.Reset();
+	}
+
+	/** Adds specified child to the filtered children. */
+	FORCEINLINE_DEBUGGABLE void AddFilteredChild(const FBaseTreeNodePtr& ChildPtr)
+	{
+		FilteredChildren.Add(ChildPtr);
+	}
+
+	/** Clears filtered children. */
+	void ClearFilteredChildren()
+	{
+		FilteredChildren.Reset();
+	}
+
+	bool IsExpanded() const { return bIsExpanded; }
+	void SetExpansion(bool bOnOff) { bIsExpanded = bOnOff; }
+
+protected:
+	/**
+	 * @return a reference to the child nodes of this group.
+	 */
+	FORCEINLINE_DEBUGGABLE TArray<FBaseTreeNodePtr>& GetChildrenMutable()
+	{
+		return Children;
+	}
+
+private:
+	/** The id of this node. */
+	const int32 Id;
+
+	/** The name of this node. */
+	const FName Name;
+
+	/** Children of this node. */
+	TArray<FBaseTreeNodePtr> Children;
+
+	/** Filtered children of this node. */
+	TArray<FBaseTreeNodePtr> FilteredChildren;
+
+	/** A weak pointer to the group/parent of this node. */
+	FBaseTreeNodeWeak GroupPtr;
+
+	/** If this node is a group node or not. */
+	bool bIsGroup;
+
+	/** Whether this group node should be expanded or not. */
+	bool bIsExpanded;
+};
+
+////////////////////////////////////////////////////////////////////////////////////////////////////
+
+} // namespace Insights
diff --git a/Engine/Source/Developer/TraceInsights/Private/Insights/Table/ViewModels/Table.cpp b/Engine/Source/Developer/TraceInsights/Private/Insights/Table/ViewModels/Table.cpp
new file mode 100644
--- /dev/null
+++ b/Engine/Source/Developer/TraceInsights/Private/Insights/Table/ViewModels/Table.cpp
@@ -0,0 +1,366 @@
+// Copyright 1998-2019 Epic Games, Inc. All Rights Reserved.
+
+#include "Table.h"
+#include "TraceServices/Containers/Tables.h"
+
+// Insights
+#include "Insights/Common/TimeUtils.h"
+#include "Insights/Table/ViewModels/TableCellValueFormatter.h"
+#include "Insights/Table/ViewModels/TableCellValueGetter.h"
+#include "Insights/Table/ViewModels/TableCellValueSorter.h"
+#include "Insights/Table/ViewModels/TableColumn.h"
+#include "Insights/Table/ViewModels/TableTreeNode.h"
+
+#define LOCTEXT_NAMESPACE "Insights_Table"
+
+namespace Insights
+{
+
+////////////////////////////////////////////////////////////////////////////////////////////////////
+// FTableTreeNodeValueGetter
+////////////////////////////////////////////////////////////////////////////////////////////////////
+
+class FTableTreeNodeValueGetter : public FTableCellValueGetter
+{
+public:
+	FTableTreeNodeValueGetter(ETableCellDataType InDataType) : FTableCellValueGetter(), DataType(InDataType) {}
+
+	virtual const TOptional<FTableCellValue> GetValue(const FTableColumn& Column, const FBaseTreeNode& Node) const
+	{
+		ensure(Node.GetTypeName() == FTableTreeNode::TypeName);
+		const FTableTreeNode& TableTreeNode = static_cast<const FTableTreeNode&>(Node);
+
+		if (!Node.IsGroup()) // Table Row Node
+		{
+			const TSharedPtr<FTable> TablePtr = Column.GetParentTable().Pin();
+			if (TablePtr.IsValid())
+			{
+				TSharedPtr<Trace::IUntypedTableReader> Reader = TablePtr->GetTableReader();
+				if (Reader.IsValid() && TableTreeNode.GetRowId().HasValidIndex())
+				{
+					Reader->SetRowIndex(TableTreeNode.GetRowId().RowIndex);
+					const int32 ColumnIndex = Column.GetIndex();
+					switch (DataType)
+					{
+						case ETableCellDataType::Bool:    return TOptional<FTableCellValue>(Reader->GetValueBool(ColumnIndex));
+						case ETableCellDataType::Int64:   return TOptional<FTableCellValue>(Reader->GetValueInt(ColumnIndex));
+						case ETableCellDataType::Float:   return TOptional<FTableCellValue>(Reader->GetValueFloat(ColumnIndex));
+						case ETableCellDataType::Double:  return TOptional<FTableCellValue>(Reader->GetValueDouble(ColumnIndex));
+						case ETableCellDataType::CString: return TOptional<FTableCellValue>(Reader->GetValueCString(ColumnIndex));
+					}
+				}
+			}
+		}
+		else // Aggregated Group Node
+		{
+			if (Column.GetAggregation() != ETableColumnAggregation::None)
+			{
+				const FTableCellValue* ValuePtr = TableTreeNode.FindAggregatedValue(Column.GetId());
+				if (ValuePtr != nullptr)
+				{
+					ensure(ValuePtr->DataType == DataType);
+					return TOptional<FTableCellValue>(*ValuePtr);
+				}
+			}
+		}
+
+		return TOptional<FTableCellValue>();
+	}
+
+private:
+	ETableCellDataType DataType;
+};
+
+////////////////////////////////////////////////////////////////////////////////////////////////////
+// FTable
+////////////////////////////////////////////////////////////////////////////////////////////////////
+
+FTable::FTable()
+	: Name()
+	, Description()
+	, Columns()
+	, ColumnIdToPtrMapping()
+	, SourceTable()
+	, TableReader()
+{
+}
+
+////////////////////////////////////////////////////////////////////////////////////////////////////
+
+FTable::~FTable()
+{
+	Reset();
+}
+
+////////////////////////////////////////////////////////////////////////////////////////////////////
+
+void FTable::Reset()
+{
+	Columns.Reset();
+	ColumnIdToPtrMapping.Reset();
+
+	SourceTable.Reset();
+	TableReader.Reset();
+}
+
+////////////////////////////////////////////////////////////////////////////////////////////////////
+
+int32 FTable::GetColumnPositionIndex(const FName& ColumnId) const
+{
+	return Columns.IndexOfByPredicate([&ColumnId](const TSharedPtr<FTableColumn>& ColumnPtr) -> bool { return ColumnPtr->GetId() == ColumnId; });
+}
+
+////////////////////////////////////////////////////////////////////////////////////////////////////
+
+void FTable::Init(TSharedPtr<Trace::IUntypedTable> InSourceTable)
+{
+	Reset();
+
+	SourceTable = InSourceTable;
+
+	if (SourceTable)
+	{
+		TableReader = MakeShareable(SourceTable->CreateReader());
+		CreateColumnsFromTableLayout();
+	}
+}
+
+////////////////////////////////////////////////////////////////////////////////////////////////////
+
+bool AreTableLayoutsEqual(const Trace::ITableLayout& TableLayoutA, const Trace::ITableLayout& TableLayoutB)
+{
+	if (TableLayoutA.GetColumnCount() != TableLayoutB.GetColumnCount())
+	{
+		return false;
+	}
+
+	int32 ColumnCount = TableLayoutA.GetColumnCount();
+	for (int32 ColumnIndex = 0; ColumnIndex < ColumnCount; ++ColumnIndex)
+	{
+		if (TableLayoutA.GetColumnType(ColumnIndex) != TableLayoutB.GetColumnType(ColumnIndex))
+		{
+			return false;
+		}
+		if (FCString::Strcmp(TableLayoutA.GetColumnName(ColumnIndex), TableLayoutB.GetColumnName(ColumnIndex)) != 0)
+		{
+			return false;
+		}
+	}
+
+	return true;
+}
+
+////////////////////////////////////////////////////////////////////////////////////////////////////
+
+void FTable::UpdateSourceTable(TSharedPtr<Trace::IUntypedTable> InSourceTable)
+{
+	check(InSourceTable.IsValid() && SourceTable.IsValid());
+	check(AreTableLayoutsEqual(InSourceTable->GetLayout(), SourceTable->GetLayout()));
+	SourceTable = InSourceTable;
+	TableReader = MakeShareable(SourceTable->CreateReader());
+}
+
+////////////////////////////////////////////////////////////////////////////////////////////////////
+
+void FTable::SetColumns(const TArray<TSharedPtr<Insights::FTableColumn>>& InColumns)
+{
+	Columns.Reset(InColumns.Num());
+	ColumnIdToPtrMapping.Reset();
+	for (TSharedPtr<Insights::FTableColumn> ColumnPtr : InColumns)
+	{
+		AddColumn(ColumnPtr);
+	}
+}
+
+////////////////////////////////////////////////////////////////////////////////////////////////////
+
+void FTable::AddColumn(TSharedPtr<FTableColumn> ColumnPtr)
+{
+	ColumnPtr->SetParentTable(SharedThis(this));
+	Columns.Add(ColumnPtr);
+	ColumnIdToPtrMapping.Add(ColumnPtr->GetId(), ColumnPtr);
+}
+
+////////////////////////////////////////////////////////////////////////////////////////////////////
+
+void FTable::CreateHierarchyColumn(int32 ColumnIndex, const TCHAR* ColumnName)
+{
+	const FName HierarchyColumnId(TEXT("_Hierarchy"));
+
+	TSharedPtr<FTableColumn> ColumnPtr = MakeShareable(new FTableColumn(HierarchyColumnId));
+	FTableColumn& Column = *ColumnPtr;
+	
+	Column.SetIndex(ColumnIndex);
+
+	const FString ColumnNameStr = ColumnName ? FString::Printf(TEXT("Hierarchy (%s)"), ColumnName) : TEXT("Hierarchy");
+	const FText ColumnNameText = FText::FromString(ColumnNameStr);
+
+	Column.SetShortName(ColumnNameText);
+	Column.SetTitleName(ColumnNameText);
+	//TODO: Column.SetDescription(...);
+
+	Column.SetFlags(ETableColumnFlags::ShouldBeVisible | ETableColumnFlags::CanBeFiltered | ETableColumnFlags::IsHierarchy);
+
+	Column.SetHorizontalAlignment(HAlign_Left);
+	Column.SetInitialWidth(90.0f);
+
+	Column.SetDataType(ETableCellDataType::CString);
+
+	TSharedPtr<ITableCellValueGetter> GetterPtr = MakeShareable(new FDisplayNameValueGetter());
+	Column.SetValueGetter(GetterPtr.ToSharedRef());
+
+	TSharedPtr<ITableCellValueFormatter> FormatterPtr = MakeShareable(new FTextValueFormatter());
+	Column.SetValueFormatter(FormatterPtr.ToSharedRef());
+
+	TSharedPtr<ITableCellValueSorter> SorterPtr = MakeShareable(new FSorterByName(ColumnPtr.ToSharedRef()));
+	Column.SetValueSorter(SorterPtr);
+
+	AddColumn(ColumnPtr);
+}
+
+////////////////////////////////////////////////////////////////////////////////////////////////////
+
+void FTable::CreateColumnsFromTableLayout()
+{
+	ensure(TableReader.IsValid());
+	ensure(Columns.Num() == 0);
+
+	const Trace::ITableLayout& TableLayout = SourceTable->GetLayout();
+	const int32 ColumnCount = TableLayout.GetColumnCount();
+
+	//////////////////////////////////////////////////
+	// Hierarchy Column
+
+	int32 HierarchyColumnIndex = -1;
+	const TCHAR* HierarchyColumnName = nullptr;
+
+	// Look for first string column.
+	//for (int32 ColumnIndex = 0; ColumnIndex < ColumnCount; ++ColumnIndex)
+	//{
+	//	Trace::ETableColumnType ColumnType = TableLayout.GetColumnType(ColumnIndex);
+	//	if (ColumnType == Trace::TableColumnType_CString)
+	//	{
+	//		HierarchyColumnIndex = ColumnIndex;
+	//		HierarchyColumnName = TableLayout.GetColumnName(ColumnIndex);
+	//		break;
+	//	}
+	//}
+
+	CreateHierarchyColumn(HierarchyColumnIndex, HierarchyColumnName);
+
+	//////////////////////////////////////////////////
+
+	for (int32 ColumnIndex = 0; ColumnIndex < ColumnCount; ++ColumnIndex)
+	{
+		Trace::ETableColumnType ColumnType = TableLayout.GetColumnType(ColumnIndex);
+		const TCHAR* ColumnName = TableLayout.GetColumnName(ColumnIndex);
+
+		TSharedPtr<FTableColumn> ColumnPtr = MakeShareable(new FTableColumn(FName(ColumnName)));
+		FTableColumn& Column = *ColumnPtr;
+
+		const FString ColumnNameStr(ColumnName);
+		const FText ColumnNameText = FText::FromString(ColumnNameStr);
+
+		ETableColumnFlags ColumnFlags = ETableColumnFlags::CanBeFiltered | ETableColumnFlags::CanBeHidden;
+		if (ColumnIndex != HierarchyColumnIndex)
+		{
+			ColumnFlags |= ETableColumnFlags::ShouldBeVisible;
+		}
+
+		EHorizontalAlignment HorizontalAlignment = HAlign_Left;
+		float InitialColumnWidth = 60.0f;
+
+		ETableColumnAggregation Aggregation = ETableColumnAggregation::None;
+
+		TSharedPtr<ITableCellValueFormatter> FormatterPtr;
+		TSharedPtr<ITableCellValueSorter> SorterPtr;
+
+		switch (ColumnType)
+		{
+		case Trace::TableColumnType_Bool:
+			Column.SetDataType(ETableCellDataType::Bool);
+			HorizontalAlignment = HAlign_Right;
+			InitialColumnWidth = 40.0f;
+			//TODO: if (Hint == AsOnOff)
+			//else // if (Hint == AsTrueFalse)
+			FormatterPtr = MakeShareable(new FBoolValueFormatterAsTrueFalse());
+			SorterPtr = MakeShareable(new FSorterByBoolValue(ColumnPtr.ToSharedRef()));
+			break;
+
+		case Trace::TableColumnType_Int:
+			Column.SetDataType(ETableCellDataType::Int64);
+			Aggregation = ETableColumnAggregation::Sum;
+			HorizontalAlignment = HAlign_Right;
+			InitialColumnWidth = 60.0f;
+			//TODO: if (Hint == AsMemory)
+			//{
+			//	FormatterPtr = MakeShareable(new FInt64ValueFormatterAsMemory());
+			//}
+			//else // AsNumber
+			FormatterPtr = MakeShareable(new FInt64ValueFormatterAsNumber());
+			SorterPtr = MakeShareable(new FSorterByInt64Value(ColumnPtr.ToSharedRef()));
+			break;
+
+		case Trace::TableColumnType_Float:
+			Column.SetDataType(ETableCellDataType::Float);
+			Aggregation = ETableColumnAggregation::Sum;
+			HorizontalAlignment = HAlign_Right;
+			InitialColumnWidth = 60.0f;
+			//TODO: if (Hint == AsTimeMs)
+			//else // if (Hint == AsTimeAuto)
+			FormatterPtr = MakeShareable(new FFloatValueFormatterAsTimeAuto());
+			SorterPtr = MakeShareable(new FSorterByFloatValue(ColumnPtr.ToSharedRef()));
+			break;
+
+		case Trace::TableColumnType_Double:
+			Column.SetDataType(ETableCellDataType::Double);
+			Aggregation = ETableColumnAggregation::Sum;
+			HorizontalAlignment = HAlign_Right;
+			InitialColumnWidth = 80.0f;
+			//TODO: if (Hint == AsTimeMs)
+			//else // if (Hint == AsTimeAuto)
+			FormatterPtr = MakeShareable(new FDoubleValueFormatterAsTimeAuto());
+			SorterPtr = MakeShareable(new FSorterByDoubleValue(ColumnPtr.ToSharedRef()));
+			break;
+
+		case Trace::TableColumnType_CString:
+			Column.SetDataType(ETableCellDataType::CString);
+			HorizontalAlignment = HAlign_Left;
+			InitialColumnWidth = FMath::Max(120.0f, 6.0f * ColumnNameStr.Len());
+			FormatterPtr = MakeShareable(new FCStringValueFormatterAsText());
+			SorterPtr = MakeShareable(new FSorterByCStringValue(ColumnPtr.ToSharedRef()));
+			break;
+		}
+
+		Column.SetIndex(ColumnIndex);
+
+		Column.SetShortName(ColumnNameText);
+		Column.SetTitleName(ColumnNameText);
+
+		//TODO: Column.SetDescription(...);
+
+		Column.SetFlags(ColumnFlags);
+
+		Column.SetHorizontalAlignment(HorizontalAlignment);
+		Column.SetInitialWidth(InitialColumnWidth);
+
+		Column.SetAggregation(Aggregation);
+
+		Column.SetValueGetter(MakeShareable(new FTableTreeNodeValueGetter(Column.GetDataType())));
+
+		if (FormatterPtr.IsValid())
+		{
+			Column.SetValueFormatter(FormatterPtr.ToSharedRef());
+		}
+
+		Column.SetValueSorter(SorterPtr);
+
+		AddColumn(ColumnPtr);
+	}
+}
+
+////////////////////////////////////////////////////////////////////////////////////////////////////
+
+} // namespace Insights
+
+#undef LOCTEXT_NAMESPACE
