Copying //UE4/Dev-Core to Main (//UE4/Main)

diff --git a/Engine/Source/Developer/TraceInsights/Private/Insights/Table/ViewModels/Table.h b/Engine/Source/Developer/TraceInsights/Private/Insights/Table/ViewModels/Table.h
new file mode 100644
--- /dev/null
+++ b/Engine/Source/Developer/TraceInsights/Private/Insights/Table/ViewModels/Table.h
@@ -0,0 +1,103 @@
+// Copyright 1998-2019 Epic Games, Inc. All Rights Reserved.
+
+#pragma once
+
+#include "CoreMinimal.h"
+
+namespace Trace
+{
+	class IUntypedTable;
+	class IUntypedTableReader;
+}
+
+namespace Insights
+{
+
+class FTableColumn;
+
+////////////////////////////////////////////////////////////////////////////////////////////////////
+
+struct FTableRowId
+{
+	static constexpr int32 InvalidRowIndex = -1;
+
+	FTableRowId(int32 InRowIndex) : RowIndex(InRowIndex), Flags(0) {}
+
+	bool HasValidIndex() const { return RowIndex >= 0; }
+
+	union
+	{
+		struct
+		{
+			int32 RowIndex;
+			uint32 Flags;
+		};
+
+		void* Data;
+	};
+};
+
+////////////////////////////////////////////////////////////////////////////////////////////////////
+/**
+ * Table View Model.
+ * View model class for the STableListView and STableTreeView widgets.
+ */
+class FTable : public TSharedFromThis<FTable>
+{
+public:
+	FTable();
+	virtual ~FTable();
+
+	const FName& GetName() const { return Name; }
+	FText GetDisplayName() const { return FText::FromName(Name); }
+	const FText& GetDescription() const { return Description; }
+
+	virtual void Reset();
+
+	bool IsValid() const { return Columns.Num() > 0; }
+
+	const TArray<TSharedPtr<FTableColumn>>& GetColumns() const { return Columns; }
+	void SetColumns(const TArray<TSharedPtr<Insights::FTableColumn>>& InColumns);
+
+	TSharedPtr<FTableColumn> FindColumnChecked(const FName& ColumnId) const
+	{
+		return ColumnIdToPtrMapping.FindChecked(ColumnId);
+	}
+
+	TSharedPtr<FTableColumn> FindColumn(const FName& ColumnId) const
+	{
+		const TSharedPtr<FTableColumn>* const ColumnPtrPtr = ColumnIdToPtrMapping.Find(ColumnId);
+		return (ColumnPtrPtr != nullptr) ? *ColumnPtrPtr : nullptr;
+	}
+
+	int32 GetColumnPositionIndex(const FName& ColumnId) const;
+
+	TSharedPtr<Trace::IUntypedTable> GetSourceTable() const { return SourceTable; }
+	TSharedPtr<Trace::IUntypedTableReader> GetTableReader() const { return TableReader; }
+
+	/* Init this table to use the IUntypedTable source table. It will create new Columns array based on table layout. */
+	void Init(TSharedPtr<Trace::IUntypedTable> InSourceTable);
+	void UpdateSourceTable(TSharedPtr<Trace::IUntypedTable> InSourceTable);
+
+private:
+	void AddColumn(TSharedPtr<FTableColumn> ColumnPtr);
+	void CreateHierarchyColumn(int32 ColumnIndex, const TCHAR* ColumnName);
+	void CreateColumnsFromTableLayout();
+
+private:
+	FName Name;
+	FText Description;
+
+	/** All available columns. */
+	TArray<TSharedPtr<FTableColumn>> Columns;
+
+	/** Mapping between column Ids and FTableColumn shared pointers. */
+	TMap<FName, TSharedPtr<FTableColumn>> ColumnIdToPtrMapping;
+
+	TSharedPtr<Trace::IUntypedTable> SourceTable;
+	TSharedPtr<Trace::IUntypedTableReader> TableReader;
+};
+
+////////////////////////////////////////////////////////////////////////////////////////////////////
+
+} // namespace Insights
diff --git a/Engine/Source/Developer/TraceInsights/Private/Insights/Table/ViewModels/TableCellValue.cpp b/Engine/Source/Developer/TraceInsights/Private/Insights/Table/ViewModels/TableCellValue.cpp
new file mode 100644
--- /dev/null
+++ b/Engine/Source/Developer/TraceInsights/Private/Insights/Table/ViewModels/TableCellValue.cpp
@@ -0,0 +1,3 @@
+// Copyright 1998-2019 Epic Games, Inc. All Rights Reserved.
+
+#include "TableCellValue.h"
diff --git a/Engine/Source/Developer/TraceInsights/Private/Insights/Table/ViewModels/TableCellValue.h b/Engine/Source/Developer/TraceInsights/Private/Insights/Table/ViewModels/TableCellValue.h
new file mode 100644
--- /dev/null
+++ b/Engine/Source/Developer/TraceInsights/Private/Insights/Table/ViewModels/TableCellValue.h
@@ -0,0 +1,51 @@
+// Copyright 1998-2019 Epic Games, Inc. All Rights Reserved.
+
+#pragma once
+
+#include "CoreMinimal.h"
+
+////////////////////////////////////////////////////////////////////////////////////////////////////
+
+enum class ETableCellDataType : uint32
+{
+	Unknown,
+
+	Bool,
+	Int64,
+	Float,
+	Double,
+	CString,
+	Text,
+
+	/** Invalid enum type, may be used as a number of enumerations. */
+	InvalidOrMax,
+};
+
+////////////////////////////////////////////////////////////////////////////////////////////////////
+
+struct FTableCellValue
+{
+public:
+	FTableCellValue() : DataType(ETableCellDataType::Unknown) {}
+
+	FTableCellValue(bool Value) : DataType(ETableCellDataType::Bool), Bool(Value) {}
+	FTableCellValue(int64 Value) : DataType(ETableCellDataType::Int64), Int64(Value) {}
+	FTableCellValue(float Value) : DataType(ETableCellDataType::Float), Float(Value) {}
+	FTableCellValue(double Value) : DataType(ETableCellDataType::Double), Double(Value) {}
+	FTableCellValue(const TCHAR* Value) : DataType(ETableCellDataType::CString), CString(Value) {}
+	FTableCellValue(const FText& Value) : DataType(ETableCellDataType::Text), TextPtr(MakeShareable(new FText(Value))) {}
+
+	ETableCellDataType DataType;
+
+	union
+	{
+		bool Bool;
+		int64 Int64;
+		float Float;
+		double Double;
+		const TCHAR* CString;
+	};
+	TSharedPtr<FText> TextPtr;
+};
+
+////////////////////////////////////////////////////////////////////////////////////////////////////
diff --git a/Engine/Source/Developer/TraceInsights/Private/Insights/Table/ViewModels/TableCellValueFormatter.h b/Engine/Source/Developer/TraceInsights/Private/Insights/Table/ViewModels/TableCellValueFormatter.h
new file mode 100644
--- /dev/null
+++ b/Engine/Source/Developer/TraceInsights/Private/Insights/Table/ViewModels/TableCellValueFormatter.h
@@ -0,0 +1,149 @@
+// Copyright 1998-2019 Epic Games, Inc. All Rights Reserved.
+
+#pragma once
+
+#include "CoreMinimal.h"
+
+#include "Insights/Table/ViewModels/TableCellValue.h"
+
+namespace Insights
+{
+
+class FBaseTreeNode;
+class FTableColumn;
+
+////////////////////////////////////////////////////////////////////////////////////////////////////
+
+class ITableCellValueFormatter
+{
+public:
+	virtual FText FormatValue(const TOptional<FTableCellValue>& InValue) const = 0;
+	virtual FText FormatValueForTooltip(const TOptional<FTableCellValue>& InValue) const = 0;
+
+	virtual FText FormatValue(const FTableColumn& Column, const FBaseTreeNode& Node) const = 0;
+	virtual FText FormatValueForTooltip(const FTableColumn& Column, const FBaseTreeNode& Node) const = 0;
+};
+
+////////////////////////////////////////////////////////////////////////////////////////////////////
+
+class FTableCellValueFormatter : public ITableCellValueFormatter
+{
+public:
+	FTableCellValueFormatter() {}
+	virtual ~FTableCellValueFormatter() {}
+
+	virtual FText FormatValue(const TOptional<FTableCellValue>& InValue) const override { return FText::GetEmpty(); }
+	virtual FText FormatValueForTooltip(const TOptional<FTableCellValue>& InValue) const override { return FormatValue(InValue); }
+
+	virtual FText FormatValue(const FTableColumn& Column, const FBaseTreeNode& Node) const override; // { return FormatValue(Column.GetValue(Node)); }
+	virtual FText FormatValueForTooltip(const FTableColumn& Column, const FBaseTreeNode& Node) const override { return FormatValue(Column, Node); }
+};
+
+////////////////////////////////////////////////////////////////////////////////////////////////////
+
+class FTextValueFormatter : public FTableCellValueFormatter
+{
+public:
+	virtual FText FormatValue(const TOptional<FTableCellValue>& InValue) const override
+	{
+		if (InValue.IsSet())
+		{
+			if (InValue.GetValue().TextPtr.IsValid())
+			{
+				return *InValue.GetValue().TextPtr;
+			}
+		}
+		return FText::GetEmpty();
+	}
+};
+
+////////////////////////////////////////////////////////////////////////////////////////////////////
+
+class FBoolValueFormatterAsTrueFalse : public FTableCellValueFormatter
+{
+public:
+	virtual FText FormatValue(const TOptional<FTableCellValue>& InValue) const override;
+};
+
+////////////////////////////////////////////////////////////////////////////////////////////////////
+
+class FBoolValueFormatterAsOnOff : public FTableCellValueFormatter
+{
+public:
+	virtual FText FormatValue(const TOptional<FTableCellValue>& InValue) const override;
+};
+
+////////////////////////////////////////////////////////////////////////////////////////////////////
+
+class FInt64ValueFormatterAsNumber : public FTableCellValueFormatter
+{
+public:
+	virtual FText FormatValue(const TOptional<FTableCellValue>& InValue) const override
+	{
+		return InValue.IsSet() ? FText::AsNumber(InValue.GetValue().Int64) : FText::GetEmpty();
+	}
+};
+
+////////////////////////////////////////////////////////////////////////////////////////////////////
+
+class FInt64ValueFormatterAsMemory : public FTableCellValueFormatter
+{
+public:
+	virtual FText FormatValue(const TOptional<FTableCellValue>& InValue) const override;
+	virtual FText FormatValueForTooltip(const TOptional<FTableCellValue>& InValue) const override;
+};
+
+////////////////////////////////////////////////////////////////////////////////////////////////////
+
+class FFloatValueFormatterAsNumber : public FTableCellValueFormatter
+{
+public:
+	virtual FText FormatValue(const TOptional<FTableCellValue>& InValue) const override;
+};
+
+////////////////////////////////////////////////////////////////////////////////////////////////////
+
+class FFloatValueFormatterAsTimeAuto : public FTableCellValueFormatter
+{
+public:
+	virtual FText FormatValue(const TOptional<FTableCellValue>& InValue) const override;
+	virtual FText FormatValueForTooltip(const TOptional<FTableCellValue>& InValue) const override;
+};
+
+////////////////////////////////////////////////////////////////////////////////////////////////////
+
+class FDoubleValueFormatterAsNumber : public FTableCellValueFormatter
+{
+public:
+	virtual FText FormatValue(const TOptional<FTableCellValue>& InValue) const override;
+};
+
+////////////////////////////////////////////////////////////////////////////////////////////////////
+
+class FDoubleValueFormatterAsTimeAuto : public FTableCellValueFormatter
+{
+public:
+	virtual FText FormatValue(const TOptional<FTableCellValue>& InValue) const override;
+	virtual FText FormatValueForTooltip(const TOptional<FTableCellValue>& InValue) const override;
+};
+
+////////////////////////////////////////////////////////////////////////////////////////////////////
+
+class FDoubleValueFormatterAsTimeMs : public FTableCellValueFormatter
+{
+public:
+	virtual FText FormatValue(const TOptional<FTableCellValue>& InValue) const override;
+	virtual FText FormatValueForTooltip(const TOptional<FTableCellValue>& InValue) const override;
+};
+
+////////////////////////////////////////////////////////////////////////////////////////////////////
+
+class FCStringValueFormatterAsText : public FTableCellValueFormatter
+{
+public:
+	virtual FText FormatValue(const TOptional<FTableCellValue>& InValue) const override;
+};
+
+////////////////////////////////////////////////////////////////////////////////////////////////////
+
+} // namespace Insights
diff --git a/Engine/Source/Developer/TraceInsights/Private/Insights/Table/ViewModels/TableCellValueGetter.cpp b/Engine/Source/Developer/TraceInsights/Private/Insights/Table/ViewModels/TableCellValueGetter.cpp
new file mode 100644
--- /dev/null
+++ b/Engine/Source/Developer/TraceInsights/Private/Insights/Table/ViewModels/TableCellValueGetter.cpp
@@ -0,0 +1,30 @@
+// Copyright 1998-2019 Epic Games, Inc. All Rights Reserved.
+
+#include "TableCellValueGetter.h"
+
+#include "Insights/Table/ViewModels/BaseTreeNode.h"
+
+#define LOCTEXT_NAMESPACE "TableCellValueGetter"
+
+namespace Insights
+{
+
+////////////////////////////////////////////////////////////////////////////////////////////////////
+
+const TOptional<FTableCellValue> FNameValueGetter::GetValue(const FTableColumn& Column, const FBaseTreeNode& Node) const
+{
+	return TOptional<FTableCellValue>(FTableCellValue(FText::FromName(Node.GetName())));
+}
+
+////////////////////////////////////////////////////////////////////////////////////////////////////
+
+const TOptional<FTableCellValue> FDisplayNameValueGetter::GetValue(const FTableColumn& Column, const FBaseTreeNode& Node) const
+{
+	return TOptional<FTableCellValue>(FTableCellValue(Node.GetDisplayName()));
+}
+
+////////////////////////////////////////////////////////////////////////////////////////////////////
+
+} // namespace Insights
+
+#undef LOCTEXT_NAMESPACE
diff --git a/Engine/Source/Developer/TraceInsights/Private/Insights/Table/ViewModels/TableCellValueGetter.h b/Engine/Source/Developer/TraceInsights/Private/Insights/Table/ViewModels/TableCellValueGetter.h
new file mode 100644
--- /dev/null
+++ b/Engine/Source/Developer/TraceInsights/Private/Insights/Table/ViewModels/TableCellValueGetter.h
@@ -0,0 +1,52 @@
+// Copyright 1998-2019 Epic Games, Inc. All Rights Reserved.
+
+#pragma once
+
+#include "CoreMinimal.h"
+
+#include "Insights/Table/ViewModels/TableCellValue.h"
+
+namespace Insights
+{
+
+class FBaseTreeNode;
+class FTableColumn;
+
+////////////////////////////////////////////////////////////////////////////////////////////////////
+
+class ITableCellValueGetter
+{
+public:
+	virtual const TOptional<FTableCellValue> GetValue(const FTableColumn& Column, const FBaseTreeNode& Node) const = 0;
+};
+
+////////////////////////////////////////////////////////////////////////////////////////////////////
+
+class FTableCellValueGetter : public ITableCellValueGetter
+{
+public:
+	FTableCellValueGetter() {}
+	virtual ~FTableCellValueGetter() {}
+
+	virtual const TOptional<FTableCellValue> GetValue(const FTableColumn& Column, const FBaseTreeNode& Node) const override { return TOptional<FTableCellValue>(); }
+};
+
+////////////////////////////////////////////////////////////////////////////////////////////////////
+
+class FNameValueGetter : public FTableCellValueGetter
+{
+public:
+	virtual const TOptional<FTableCellValue> GetValue(const FTableColumn& Column, const FBaseTreeNode& Node) const override;
+};
+
+////////////////////////////////////////////////////////////////////////////////////////////////////
+
+class FDisplayNameValueGetter : public FTableCellValueGetter
+{
+public:
+	virtual const TOptional<FTableCellValue> GetValue(const FTableColumn& Column, const FBaseTreeNode& Node) const override;
+};
+
+////////////////////////////////////////////////////////////////////////////////////////////////////
+
+} // namespace Insights
diff --git a/Engine/Source/Developer/TraceInsights/Private/Insights/Table/ViewModels/TableCellValueSorter.cpp b/Engine/Source/Developer/TraceInsights/Private/Insights/Table/ViewModels/TableCellValueSorter.cpp
new file mode 100644
--- /dev/null
+++ b/Engine/Source/Developer/TraceInsights/Private/Insights/Table/ViewModels/TableCellValueSorter.cpp
@@ -0,0 +1,678 @@
+// Copyright 1998-2019 Epic Games, Inc. All Rights Reserved.
+
+#include "TableCellValueSorter.h"
+
+#include "Insights/Table/ViewModels/BaseTreeNode.h"
+
+#define LOCTEXT_NAMESPACE "TableCellValueSorter"
+
+// Sort by name (ascending). -- too slow :(
+//#define INSIGHTS_DEFAULT_SORTING_NODES(A, B) return A->GetName().LexicalLess(B->GetName());
+//#define INSIGHTS_DEFAULT_SORTING_ELEMENTS(A, B) return A.NodePtr->GetName().LexicalLess(B.NodePtr->GetName());
+
+// Sort by node id (ascending).
+#define INSIGHTS_DEFAULT_SORTING_NODES(A, B) return A->GetId() < B->GetId();
+#define INSIGHTS_DEFAULT_SORTING_ELEMENTS(A, B) return A.NodePtr->GetId() < B.NodePtr->GetId();
+
+namespace Insights
+{
+
+////////////////////////////////////////////////////////////////////////////////////////////////////
+// FTableCellValueSorter
+////////////////////////////////////////////////////////////////////////////////////////////////////
+
+FTableCellValueSorter::FTableCellValueSorter(const FName InName, const FText& InShortName, const FText& InTitleName, const FText& InDescription, TSharedRef<FTableColumn> InColumnRef)
+	: Name(InName)
+	, ShortName(InShortName)
+	, TitleName(InTitleName)
+	, Description(InDescription)
+	, ColumnRef(InColumnRef)
+	, AscendingIcon(nullptr)
+	, DescendingIcon(nullptr)
+	, AscendingCompareDelegate(nullptr)
+	, DescendingCompareDelegate(nullptr)
+{
+}
+
+////////////////////////////////////////////////////////////////////////////////////////////////////
+
+void FTableCellValueSorter::Sort(TArray<FBaseTreeNodePtr>& NodesToSort, ESortMode SortMode) const
+{
+	if (SortMode == ESortMode::Ascending)
+	{
+		if (AscendingCompareDelegate)
+		{
+			NodesToSort.Sort(AscendingCompareDelegate);
+		}
+	}
+	else
+	{
+		if (DescendingCompareDelegate)
+		{
+			NodesToSort.Sort(DescendingCompareDelegate);
+		}
+	}
+}
+
+////////////////////////////////////////////////////////////////////////////////////////////////////
+// FBaseTableColumnSorter
+////////////////////////////////////////////////////////////////////////////////////////////////////
+
+FBaseTableColumnSorter::FBaseTableColumnSorter(TSharedRef<FTableColumn> InColumnRef)
+	: FTableCellValueSorter(
+		InColumnRef->GetId(),
+		FText::Format(LOCTEXT("Sorter_ColumnValue_Name", "By {0}"), InColumnRef->GetShortName()),
+		FText::Format(LOCTEXT("Sorter_ColumnValue_Title", "Sort By {0}"), InColumnRef->GetTitleName()),
+		FText::Format(LOCTEXT("Sorter_ColumnValue_Desc", "Sort by {0} (ascending or descending), then by id (ascending)."), InColumnRef->GetShortName()),
+		InColumnRef)
+{
+}
+
+////////////////////////////////////////////////////////////////////////////////////////////////////
+// FSorterByName
+////////////////////////////////////////////////////////////////////////////////////////////////////
+
+FSorterByName::FSorterByName(TSharedRef<FTableColumn> InColumnRef)
+	: FTableCellValueSorter(
+		FName(TEXT("ByName")),
+		LOCTEXT("Sorter_ByName_Name", "By Name"),
+		LOCTEXT("Sorter_ByName_Title", "Sort By Name"),
+		LOCTEXT("Sorter_ByName_Desc", "Sort by name."),
+		InColumnRef)
+{
+	AscendingCompareDelegate = [](const FBaseTreeNodePtr& A, const FBaseTreeNodePtr& B) -> bool
+	{
+		// Sort by name (ascending).
+		return A->GetName().LexicalLess(B->GetName());
+	};
+
+	DescendingCompareDelegate = [](const FBaseTreeNodePtr& A, const FBaseTreeNodePtr& B) -> bool
+	{
+		// Sort by name (descending).
+		return B->GetName().LexicalLess(A->GetName());
+	};
+}
+
+////////////////////////////////////////////////////////////////////////////////////////////////////
+// FSorterByTypeName
+////////////////////////////////////////////////////////////////////////////////////////////////////
+
+FSorterByTypeName::FSorterByTypeName(TSharedRef<FTableColumn> InColumnRef)
+	: FTableCellValueSorter(
+		FName(TEXT("ByTypeName")),
+		LOCTEXT("Sorter_ByTypeName_Name", "By Type Name"),
+		LOCTEXT("Sorter_ByTypeName_Title", "Sort By Type Name"),
+		LOCTEXT("Sorter_ByTypeName_Desc", "Sort by type name (ascending or descending), then by id (ascending)."),
+		InColumnRef)
+{
+	AscendingCompareDelegate = [](const FBaseTreeNodePtr& A, const FBaseTreeNodePtr& B) -> bool
+	{
+		const FName& TypeA = A->GetTypeName();
+		const FName& TypeB = B->GetTypeName();
+
+		if (TypeA == TypeB)
+		{
+			INSIGHTS_DEFAULT_SORTING_NODES(A, B)
+		}
+		else
+		{
+			// Sort by type name (ascending).
+			return TypeA.FastLess(TypeB); // TypeA < TypeB;
+		}
+	};
+
+	DescendingCompareDelegate = [](const FBaseTreeNodePtr& A, const FBaseTreeNodePtr& B) -> bool
+	{
+		const FName& TypeA = A->GetTypeName();
+		const FName& TypeB = B->GetTypeName();
+
+		if (TypeA == TypeB)
+		{
+			INSIGHTS_DEFAULT_SORTING_NODES(A, B)
+		}
+		else
+		{
+			// Sort by type name (descending).
+			return TypeB.FastLess(TypeA); // TypeA > TypeB;
+		}
+	};
+}
+
+////////////////////////////////////////////////////////////////////////////////////////////////////
+// FSorterByBoolValue
+////////////////////////////////////////////////////////////////////////////////////////////////////
+
+FSorterByBoolValue::FSorterByBoolValue(TSharedRef<FTableColumn> InColumnRef)
+	: FBaseTableColumnSorter(InColumnRef)
+{
+	const FTableColumn& Column = *InColumnRef;
+
+	ensure(Column.GetDataType() == ETableCellDataType::Bool);
+
+	AscendingCompareDelegate = [&Column](const FBaseTreeNodePtr& A, const FBaseTreeNodePtr& B) -> bool
+	{
+		const TOptional<FTableCellValue> OptionalValueA = Column.GetValue(*A);
+		const TOptional<FTableCellValue> OptionalValueB = Column.GetValue(*B);
+
+		const bool ValueA = OptionalValueA.IsSet() ? OptionalValueA.GetValue().Bool : false;
+		const bool ValueB = OptionalValueB.IsSet() ? OptionalValueB.GetValue().Bool : false;
+
+		if (ValueA == ValueB)
+		{
+			INSIGHTS_DEFAULT_SORTING_NODES(A, B)
+		}
+		else
+		{
+			// Sort by value (ascending).
+			return ValueA < ValueB;
+		}
+	};
+
+	DescendingCompareDelegate = [&Column](const FBaseTreeNodePtr& A, const FBaseTreeNodePtr& B) -> bool
+	{
+		const TOptional<FTableCellValue> OptionalValueA = Column.GetValue(*A);
+		const TOptional<FTableCellValue> OptionalValueB = Column.GetValue(*B);
+
+		const bool ValueA = OptionalValueA.IsSet() ? OptionalValueA.GetValue().Bool : false;
+		const bool ValueB = OptionalValueB.IsSet() ? OptionalValueB.GetValue().Bool : false;
+
+		if (ValueA == ValueB)
+		{
+			INSIGHTS_DEFAULT_SORTING_NODES(A, B)
+		}
+		else
+		{
+			// Sort by value (descending).
+			return ValueB < ValueA;
+		}
+	};
+}
+
+////////////////////////////////////////////////////////////////////////////////////////////////////
+// FSorterByInt64Value
+////////////////////////////////////////////////////////////////////////////////////////////////////
+
+FSorterByInt64Value::FSorterByInt64Value(TSharedRef<FTableColumn> InColumnRef)
+	: FBaseTableColumnSorter(InColumnRef)
+{
+	const FTableColumn& Column = *InColumnRef;
+
+	ensure(Column.GetDataType() == ETableCellDataType::Int64);
+
+	AscendingCompareDelegate = [&Column](const FBaseTreeNodePtr& A, const FBaseTreeNodePtr& B) -> bool
+	{
+		const TOptional<FTableCellValue> OptionalValueA = Column.GetValue(*A);
+		const TOptional<FTableCellValue> OptionalValueB = Column.GetValue(*B);
+
+		const int64 ValueA = OptionalValueA.IsSet() ? OptionalValueA.GetValue().Int64 : 0;
+		const int64 ValueB = OptionalValueB.IsSet() ? OptionalValueB.GetValue().Int64 : 0;
+
+		if (ValueA == ValueB)
+		{
+			INSIGHTS_DEFAULT_SORTING_NODES(A, B)
+		}
+		else
+		{
+			// Sort by value (ascending).
+			return ValueA < ValueB;
+		}
+	};
+
+	DescendingCompareDelegate = [&Column](const FBaseTreeNodePtr& A, const FBaseTreeNodePtr& B) -> bool
+	{
+		const TOptional<FTableCellValue> OptionalValueA = Column.GetValue(*A);
+		const TOptional<FTableCellValue> OptionalValueB = Column.GetValue(*B);
+
+		const int64 ValueA = OptionalValueA.IsSet() ? OptionalValueA.GetValue().Int64 : 0;
+		const int64 ValueB = OptionalValueB.IsSet() ? OptionalValueB.GetValue().Int64 : 0;
+
+		if (ValueA == ValueB)
+		{
+			INSIGHTS_DEFAULT_SORTING_NODES(A, B)
+		}
+		else
+		{
+			// Sort by value (descending).
+			return ValueB < ValueA;
+		}
+	};
+}
+
+////////////////////////////////////////////////////////////////////////////////////////////////////
+
+void FSorterByInt64Value::Sort(TArray<FBaseTreeNodePtr>& NodesToSort, ESortMode SortMode) const
+{
+	// Note: This implementation is faster than using above delegates.
+	//       It caches the values before sorting, in this way minimizes the calls to Column.GetValue().
+
+	const FTableColumn& Column = *ColumnRef;
+
+	const int32 NumElements = NodesToSort.Num();
+
+	struct FSortElement
+	{
+		const FBaseTreeNodePtr NodePtr;
+		int64 Value;
+	};
+	TArray<FSortElement> ElementsToSort;
+
+	// Cache value for each node.
+	ElementsToSort.Reset(NumElements);
+	for (const FBaseTreeNodePtr& NodePtr : NodesToSort)
+	{
+		const TOptional<FTableCellValue> OptionalValue = Column.GetValue(*NodePtr);
+		const int64 Value = OptionalValue.IsSet() ? OptionalValue.GetValue().Int64 : 0;
+		ElementsToSort.Add({ NodePtr, Value });
+	}
+	ensure(ElementsToSort.Num() == NumElements);
+
+	if (SortMode == ESortMode::Ascending)
+	{
+		ElementsToSort.Sort([](const FSortElement& A, const FSortElement& B) -> bool
+		{
+			if (A.Value == B.Value)
+			{
+				INSIGHTS_DEFAULT_SORTING_ELEMENTS(A, B)
+			}
+			else
+			{
+				// Sort by value (ascending).
+				return A.Value < B.Value;
+			}
+		});
+	}
+	else // if (SortMode == ESortMode::Descending)
+	{
+		ElementsToSort.Sort([](const FSortElement& A, const FSortElement& B) -> bool
+		{
+			if (A.Value == B.Value)
+			{
+				INSIGHTS_DEFAULT_SORTING_ELEMENTS(A, B)
+			}
+			else
+			{
+				// Sort by value (descending).
+				return B.Value < A.Value;
+			}
+		});
+	}
+
+	for (int32 Index = 0; Index < NumElements; ++Index)
+	{
+		NodesToSort[Index] = ElementsToSort[Index].NodePtr;
+	}
+}
+
+////////////////////////////////////////////////////////////////////////////////////////////////////
+// FSorterByFloatValue
+////////////////////////////////////////////////////////////////////////////////////////////////////
+
+FSorterByFloatValue::FSorterByFloatValue(TSharedRef<FTableColumn> InColumnRef)
+	: FBaseTableColumnSorter(InColumnRef)
+{
+	const FTableColumn& Column = *InColumnRef;
+
+	ensure(Column.GetDataType() == ETableCellDataType::Float);
+
+	AscendingCompareDelegate = [&Column](const FBaseTreeNodePtr& A, const FBaseTreeNodePtr& B) -> bool
+	{
+		const TOptional<FTableCellValue> OptionalValueA = Column.GetValue(*A);
+		const TOptional<FTableCellValue> OptionalValueB = Column.GetValue(*B);
+
+		const float ValueA = OptionalValueA.IsSet() ? OptionalValueA.GetValue().Float : 0.0f;
+		const float ValueB = OptionalValueB.IsSet() ? OptionalValueB.GetValue().Float : 0.0f;
+
+		if (ValueA == ValueB)
+		{
+			INSIGHTS_DEFAULT_SORTING_NODES(A, B)
+		}
+		else
+		{
+			// Sort by value (ascending).
+			return ValueA < ValueB;
+		}
+	};
+
+	DescendingCompareDelegate = [&Column](const FBaseTreeNodePtr& A, const FBaseTreeNodePtr& B) -> bool
+	{
+		const TOptional<FTableCellValue> OptionalValueA = Column.GetValue(*A);
+		const TOptional<FTableCellValue> OptionalValueB = Column.GetValue(*B);
+
+		const float ValueA = OptionalValueA.IsSet() ? OptionalValueA.GetValue().Float : 0.0f;
+		const float ValueB = OptionalValueB.IsSet() ? OptionalValueB.GetValue().Float : 0.0f;
+
+		if (ValueA == ValueB)
+		{
+			INSIGHTS_DEFAULT_SORTING_NODES(A, B)
+		}
+		else
+		{
+			// Sort by value (descending).
+			return ValueB < ValueA;
+		}
+	};
+}
+
+////////////////////////////////////////////////////////////////////////////////////////////////////
+
+void FSorterByFloatValue::Sort(TArray<FBaseTreeNodePtr>& NodesToSort, ESortMode SortMode) const
+{
+	// Note: This implementation is faster than using above delegates.
+	//       It caches the values before sorting, in this way minimizes the calls to Column.GetValue().
+
+	const FTableColumn& Column = *ColumnRef;
+
+	const int32 NumElements = NodesToSort.Num();
+
+	struct FSortElement
+	{
+		const FBaseTreeNodePtr NodePtr;
+		float Value;
+	};
+	TArray<FSortElement> ElementsToSort;
+
+	// Cache value for each node.
+	ElementsToSort.Reset(NumElements);
+	for (const FBaseTreeNodePtr& NodePtr : NodesToSort)
+	{
+		const TOptional<FTableCellValue> OptionalValue = Column.GetValue(*NodePtr);
+		const float Value = OptionalValue.IsSet() ? OptionalValue.GetValue().Float : 0.0f;
+		ElementsToSort.Add({ NodePtr, Value });
+	}
+	ensure(ElementsToSort.Num() == NumElements);
+
+	if (SortMode == ESortMode::Ascending)
+	{
+		ElementsToSort.Sort([](const FSortElement& A, const FSortElement& B) -> bool
+		{
+			if (A.Value == B.Value)
+			{
+				INSIGHTS_DEFAULT_SORTING_ELEMENTS(A, B)
+			}
+			else
+			{
+				// Sort by value (ascending).
+				return A.Value < B.Value;
+			}
+		});
+	}
+	else // if (SortMode == ESortMode::Descending)
+	{
+		ElementsToSort.Sort([](const FSortElement& A, const FSortElement& B) -> bool
+		{
+			if (A.Value == B.Value)
+			{
+				INSIGHTS_DEFAULT_SORTING_ELEMENTS(A, B)
+			}
+			else
+			{
+				// Sort by value (descending).
+				return B.Value < A.Value;
+			}
+		});
+	}
+
+	for (int32 Index = 0; Index < NumElements; ++Index)
+	{
+		NodesToSort[Index] = ElementsToSort[Index].NodePtr;
+	}
+}
+
+////////////////////////////////////////////////////////////////////////////////////////////////////
+// FSorterByDoubleValue
+////////////////////////////////////////////////////////////////////////////////////////////////////
+
+FSorterByDoubleValue::FSorterByDoubleValue(TSharedRef<FTableColumn> InColumnRef)
+	: FBaseTableColumnSorter(InColumnRef)
+{
+	const FTableColumn& Column = *InColumnRef;
+
+	ensure(Column.GetDataType() == ETableCellDataType::Double);
+
+	AscendingCompareDelegate = [&Column](const FBaseTreeNodePtr& A, const FBaseTreeNodePtr& B) -> bool
+	{
+		const TOptional<FTableCellValue> OptionalValueA = Column.GetValue(*A);
+		const TOptional<FTableCellValue> OptionalValueB = Column.GetValue(*B);
+
+		const double ValueA = OptionalValueA.IsSet() ? OptionalValueA.GetValue().Double : 0.0;
+		const double ValueB = OptionalValueB.IsSet() ? OptionalValueB.GetValue().Double : 0.0;
+
+		if (ValueA == ValueB)
+		{
+			INSIGHTS_DEFAULT_SORTING_NODES(A, B)
+		}
+		else
+		{
+			// Sort by value (ascending).
+			return ValueA < ValueB;
+		}
+	};
+
+	DescendingCompareDelegate = [&Column](const FBaseTreeNodePtr& A, const FBaseTreeNodePtr& B) -> bool
+	{
+		const TOptional<FTableCellValue> OptionalValueA = Column.GetValue(*A);
+		const TOptional<FTableCellValue> OptionalValueB = Column.GetValue(*B);
+
+		const double ValueA = OptionalValueA.IsSet() ? OptionalValueA.GetValue().Double : 0.0;
+		const double ValueB = OptionalValueB.IsSet() ? OptionalValueB.GetValue().Double : 0.0;
+
+		if (ValueA == ValueB)
+		{
+			INSIGHTS_DEFAULT_SORTING_NODES(A, B)
+		}
+		else
+		{
+			// Sort by value (descending).
+			return ValueB < ValueA;
+		}
+	};
+}
+
+////////////////////////////////////////////////////////////////////////////////////////////////////
+
+void FSorterByDoubleValue::Sort(TArray<FBaseTreeNodePtr>& NodesToSort, ESortMode SortMode) const
+{
+	// Note: This implementation is faster than using above delegates.
+	//       It caches the values before sorting, in this way minimizes the calls to Column.GetValue().
+
+	const FTableColumn& Column = *ColumnRef;
+
+	const int32 NumElements = NodesToSort.Num();
+
+	struct FSortElement
+	{
+		const FBaseTreeNodePtr NodePtr;
+		double Value;
+	};
+	TArray<FSortElement> ElementsToSort;
+
+	// Cache value for each node.
+	ElementsToSort.Reset(NumElements);
+	for (const FBaseTreeNodePtr& NodePtr : NodesToSort)
+	{
+		const TOptional<FTableCellValue> OptionalValue = Column.GetValue(*NodePtr);
+		const double Value = OptionalValue.IsSet() ? OptionalValue.GetValue().Double : 0.0;
+		ElementsToSort.Add({ NodePtr, Value });
+	}
+	ensure(ElementsToSort.Num() == NumElements);
+
+	if (SortMode == ESortMode::Ascending)
+	{
+		ElementsToSort.Sort([](const FSortElement& A, const FSortElement& B) -> bool
+		{
+			if (A.Value == B.Value)
+			{
+				INSIGHTS_DEFAULT_SORTING_ELEMENTS(A, B)
+			}
+			else
+			{
+				// Sort by value (ascending).
+				return A.Value < B.Value;
+			}
+		});
+	}
+	else // if (SortMode == ESortMode::Descending)
+	{
+		ElementsToSort.Sort([](const FSortElement& A, const FSortElement& B) -> bool
+		{
+			if (A.Value == B.Value)
+			{
+				INSIGHTS_DEFAULT_SORTING_ELEMENTS(A, B)
+			}
+			else
+			{
+				// Sort by value (descending).
+				return B.Value < A.Value;
+			}
+		});
+	}
+
+	for (int32 Index = 0; Index < NumElements; ++Index)
+	{
+		NodesToSort[Index] = ElementsToSort[Index].NodePtr;
+	}
+}
+
+////////////////////////////////////////////////////////////////////////////////////////////////////
+// FSorterByCStringValue
+////////////////////////////////////////////////////////////////////////////////////////////////////
+
+FSorterByCStringValue::FSorterByCStringValue(TSharedRef<FTableColumn> InColumnRef)
+	: FBaseTableColumnSorter(InColumnRef)
+{
+	const FTableColumn& Column = *InColumnRef;
+
+	ensure(Column.GetDataType() == ETableCellDataType::CString);
+
+	AscendingCompareDelegate = [&Column](const FBaseTreeNodePtr& A, const FBaseTreeNodePtr& B) -> bool
+	{
+		const TOptional<FTableCellValue> OptionalValueA = Column.GetValue(*A);
+		const TOptional<FTableCellValue> OptionalValueB = Column.GetValue(*B);
+
+		const TCHAR* ValueA = OptionalValueA.IsSet() ? OptionalValueA.GetValue().CString : nullptr;
+		const TCHAR* ValueB = OptionalValueB.IsSet() ? OptionalValueB.GetValue().CString : nullptr;
+
+		if (ValueA == ValueB)
+		{
+			// Sort by name (ascending).
+			//return A->GetName().LexicalLess(B->GetName()); // so slow :(
+
+			// Sort by node id (ascending).
+			return A->GetId() < B->GetId();
+		}
+		else
+		{
+			// Sort by value (ascending).
+			return FCString::Strcmp(ValueA, ValueB) < 0;
+		}
+	};
+
+	DescendingCompareDelegate = [&Column](const FBaseTreeNodePtr& A, const FBaseTreeNodePtr& B) -> bool
+	{
+		const TOptional<FTableCellValue> OptionalValueA = Column.GetValue(*A);
+		const TOptional<FTableCellValue> OptionalValueB = Column.GetValue(*B);
+
+		const TCHAR* ValueA = OptionalValueA.IsSet() ? OptionalValueA.GetValue().CString : nullptr;
+		const TCHAR* ValueB = OptionalValueB.IsSet() ? OptionalValueB.GetValue().CString : nullptr;
+
+		if (ValueA == ValueB)
+		{
+			// Sort by name (ascending).
+			//return A->GetName().LexicalLess(B->GetName()); // so slow :(
+
+			// Sort by node id (ascending).
+			return A->GetId() < B->GetId();
+		}
+		else
+		{
+			// Sort by value (descending).
+			return FCString::Strcmp(ValueB, ValueA) < 0;
+		}
+	};
+}
+
+////////////////////////////////////////////////////////////////////////////////////////////////////
+// FSorterByTextValue
+////////////////////////////////////////////////////////////////////////////////////////////////////
+
+FSorterByTextValue::FSorterByTextValue(TSharedRef<FTableColumn> InColumnRef)
+	: FBaseTableColumnSorter(InColumnRef)
+{
+	const FTableColumn& Column = *InColumnRef;
+
+	ensure(Column.GetDataType() == ETableCellDataType::Text);
+
+	AscendingCompareDelegate = [&Column](const FBaseTreeNodePtr& A, const FBaseTreeNodePtr& B) -> bool
+	{
+		const TOptional<FTableCellValue> OptionalValueA = Column.GetValue(*A);
+		const TOptional<FTableCellValue> OptionalValueB = Column.GetValue(*B);
+
+		if (!OptionalValueB.IsSet() || !OptionalValueB.GetValue().TextPtr.IsValid())
+		{
+			return true;
+		}
+
+		if (!OptionalValueA.IsSet() || !OptionalValueA.GetValue().TextPtr.IsValid())
+		{
+			return false;
+		}
+
+		const FText& ValueA = *OptionalValueA.GetValue().TextPtr;
+		const FText& ValueB = *OptionalValueB.GetValue().TextPtr;
+
+		if (ValueA.EqualTo(ValueB))
+		{
+			// Sort by name (ascending).
+			//return A->GetName().LexicalLess(B->GetName()); // so slow :(
+
+			// Sort by node id (ascending).
+			return A->GetId() < B->GetId();
+		}
+		else
+		{
+			// Sort by value (ascending).
+			return ValueA.CompareTo(ValueB) < 0;
+		}
+	};
+
+	DescendingCompareDelegate = [&Column](const FBaseTreeNodePtr& A, const FBaseTreeNodePtr& B) -> bool
+	{
+		const TOptional<FTableCellValue> OptionalValueA = Column.GetValue(*A);
+		const TOptional<FTableCellValue> OptionalValueB = Column.GetValue(*B);
+
+		if (!OptionalValueB.IsSet() || !OptionalValueB.GetValue().TextPtr.IsValid())
+		{
+			return true;
+		}
+
+		if (!OptionalValueA.IsSet() || !OptionalValueA.GetValue().TextPtr.IsValid())
+		{
+			return false;
+		}
+
+		const FText& ValueA = *OptionalValueA.GetValue().TextPtr;
+		const FText& ValueB = *OptionalValueB.GetValue().TextPtr;
+
+		if (ValueA.EqualTo(ValueB))
+		{
+			// Sort by name (ascending).
+			//return A->GetName().LexicalLess(B->GetName()); // so slow :(
+
+			// Sort by node id (ascending).
+			return A->GetId() < B->GetId();
+		}
+		else
+		{
+			// Sort by value (descending).
+			return ValueB.CompareTo(ValueA) < 0;
+		}
+	};
+}
+
+////////////////////////////////////////////////////////////////////////////////////////////////////
+
+} // namespace Insights
+
+#undef INSIGHTS_DEFAULT_SORTING_NODES
+#undef INSIGHTS_DEFAULT_SORTING_ELEMENTS
+#undef LOCTEXT_NAMESPACE
diff --git a/Engine/Source/Developer/TraceInsights/Private/Insights/Table/ViewModels/TableTreeNode.cpp b/Engine/Source/Developer/TraceInsights/Private/Insights/Table/ViewModels/TableTreeNode.cpp
new file mode 100644
--- /dev/null
+++ b/Engine/Source/Developer/TraceInsights/Private/Insights/Table/ViewModels/TableTreeNode.cpp
@@ -0,0 +1,14 @@
+// Copyright 1998-2019 Epic Games, Inc. All Rights Reserved.
+
+#include "TableTreeNode.h"
+
+#define LOCTEXT_NAMESPACE "Insights_TableTreeNode"
+
+namespace Insights
+{
+
+const FName FTableTreeNode::TypeName(TEXT("FTableTreeNode"));
+
+} // namespace Insights
+
+#undef LOCTEXT_NAMESPACE
diff --git a/Engine/Source/Developer/TraceInsights/Private/Insights/Table/ViewModels/TableTreeNode.h b/Engine/Source/Developer/TraceInsights/Private/Insights/Table/ViewModels/TableTreeNode.h
new file mode 100644
--- /dev/null
+++ b/Engine/Source/Developer/TraceInsights/Private/Insights/Table/ViewModels/TableTreeNode.h
@@ -0,0 +1,84 @@
+// Copyright 1998-2019 Epic Games, Inc. All Rights Reserved.
+
+#pragma once
+
+#include "CoreMinimal.h"
+#include "TraceServices/AnalysisService.h"
+
+// Insights
+#include "Insights/Table/ViewModels/BaseTreeNode.h"
+#include "Insights/Table/ViewModels/Table.h" // for FTableRowId
+#include "Insights/Table/ViewModels/TableCellValue.h"
+
+namespace Insights
+{
+
+//class FTable;
+//class FTableColumn;
+
+////////////////////////////////////////////////////////////////////////////////////////////////////
+
+class FTableTreeNode;
+
+/** Type definition for shared pointers to instances of FTableTreeNode. */
+typedef TSharedPtr<class FTableTreeNode> FTableTreeNodePtr;
+
+/** Type definition for shared references to instances of FTableTreeNode. */
+typedef TSharedRef<class FTableTreeNode> FTableTreeNodeRef;
+
+/** Type definition for shared references to const instances of FTableTreeNode. */
+typedef TSharedRef<const class FTableTreeNode> FTableTreeNodeRefConst;
+
+/** Type definition for weak references to instances of FTableTreeNode. */
+typedef TWeakPtr<class FTableTreeNode> FTableTreeNodeWeak;
+
+////////////////////////////////////////////////////////////////////////////////////////////////////
+/**
+ * Table Tree Node View Model.
+ * Class used to store information about a generic table tree node (used in STableTreeView).
+ */
+class FTableTreeNode : public FBaseTreeNode
+{
+public:
+	static const FName TypeName;
+
+public:
+	/** Initialization constructor for a table record node. */
+	FTableTreeNode(uint64 InId, const FName InName, TWeakPtr<FTable> InParentTable, int32 InRowIndex)
+		: FBaseTreeNode(InId, InName, false)
+		, ParentTable(InParentTable)
+		, RowId(InRowIndex)
+	{
+	}
+
+	/** Initialization constructor for a group node. */
+	FTableTreeNode(const FName InGroupName, TWeakPtr<FTable> InParentTable)
+		: FBaseTreeNode(0, InGroupName, true)
+		, ParentTable(InParentTable)
+		, RowId(FTableRowId::InvalidRowIndex)
+	{
+	}
+
+	virtual const FName& GetTypeName() const override { return TypeName; }
+
+	TWeakPtr<FTable> GetParentTable() { return ParentTable; }
+	FTableRowId GetRowId() const { return RowId; }
+	int32 GetRowIndex() const { return RowId.RowIndex; }
+
+	void ResetAggregatedValues() { AggregatedValues.Reset(); }
+	bool HasAggregatedValue(const FName& ColumnId) const { return AggregatedValues.Contains(ColumnId); }
+	const FTableCellValue* FindAggregatedValue(const FName& ColumnId) const { return AggregatedValues.Find(ColumnId); }
+	const FTableCellValue& GetAggregatedValue(const FName& ColumnId) const { return AggregatedValues.FindChecked(ColumnId); }
+	void AddAggregatedValue(const FName& ColumnId, const FTableCellValue& Value) { AggregatedValues.Add(ColumnId, Value); }
+	void SetAggregatedValue(const FName& ColumnId, const FTableCellValue& Value) { AggregatedValues[ColumnId] = Value; }
+
+protected:
+	TWeakPtr<FTable> ParentTable;
+	FTableRowId RowId;
+
+	TMap<FName, FTableCellValue> AggregatedValues;
+};
+
+////////////////////////////////////////////////////////////////////////////////////////////////////
+
+} // namespace Insights
diff --git a/Engine/Source/Developer/TraceInsights/Private/Insights/Table/ViewModels/TreeNodeGrouping.h b/Engine/Source/Developer/TraceInsights/Private/Insights/Table/ViewModels/TreeNodeGrouping.h
new file mode 100644
--- /dev/null
+++ b/Engine/Source/Developer/TraceInsights/Private/Insights/Table/ViewModels/TreeNodeGrouping.h
@@ -0,0 +1,124 @@
+// Copyright 1998-2019 Epic Games, Inc. All Rights Reserved.
+
+#pragma once
+
+#include "CoreMinimal.h"
+#include "Styling/SlateBrush.h"
+#include "Delegates/DelegateCombinations.h"
+
+// Insights
+#include "Insights/Table/ViewModels/BaseTreeNode.h"
+#include "Insights/Table/ViewModels/TableColumn.h"
+
+namespace Insights
+{
+
+////////////////////////////////////////////////////////////////////////////////////////////////////
+
+class ITreeNodeGrouping
+{
+public:
+	virtual FText GetShortName() const = 0;
+	virtual FText GetTitleName() const = 0;
+	virtual FText GetDescription() const = 0;
+
+	virtual FName GetBrushName() const = 0;
+	virtual const FSlateBrush* GetIcon() const = 0;
+
+	virtual FName GetColumnId() const = 0;
+
+	//virtual void CreateGrouping() = 0;
+};
+
+////////////////////////////////////////////////////////////////////////////////////////////////////
+
+struct FTreeNodeGroupInfo
+{
+	FName Name;
+	bool IsExpanded;
+};
+
+////////////////////////////////////////////////////////////////////////////////////////////////////
+
+class FTreeNodeGrouping : public ITreeNodeGrouping
+{
+public:
+	FTreeNodeGrouping(FText InShortName, FText InTitleName, FText InDescription, FName InBrushName,  FSlateBrush* InIcon);
+	virtual ~FTreeNodeGrouping() {}
+
+	virtual FText GetShortName() const override { return ShortName; }
+	virtual FText GetTitleName() const override { return TitleName; }
+	virtual FText GetDescription() const override { return Description; }
+
+	virtual FName GetBrushName() const override { return BrushName; }
+	virtual const FSlateBrush* GetIcon() const override { return Icon; }
+
+	virtual FName GetColumnId() const override { return NAME_None; }
+
+	virtual FTreeNodeGroupInfo GetGroupForNode(const FBaseTreeNodePtr InNode) const { return { FName(), false }; }
+
+protected:
+	FText ShortName;
+	FText TitleName;
+	FText Description;
+	FName BrushName;
+	const FSlateBrush* Icon;
+};
+
+////////////////////////////////////////////////////////////////////////////////////////////////////
+
+/** Creates a single group for all nodes. */
+class FTreeNodeGroupingFlat : public FTreeNodeGrouping
+{
+public:
+	FTreeNodeGroupingFlat();
+	virtual ~FTreeNodeGroupingFlat() {}
+
+	virtual FTreeNodeGroupInfo GetGroupForNode(const FBaseTreeNodePtr InNode) const override;
+};
+
+////////////////////////////////////////////////////////////////////////////////////////////////////
+
+/** Creates a group for unique column value of node. */
+class FTreeNodeGroupingByUniqueValue : public FTreeNodeGrouping
+{
+public:
+	FTreeNodeGroupingByUniqueValue(TSharedRef<FTableColumn> InColumnRef);
+	virtual ~FTreeNodeGroupingByUniqueValue() {}
+
+	virtual FTreeNodeGroupInfo GetGroupForNode(const FBaseTreeNodePtr InNode) const override;
+
+	virtual FName GetColumnId() const override { return ColumnRef->GetId(); }
+	TSharedRef<FTableColumn> GetColumn() const { return ColumnRef; }
+
+private:
+	TSharedRef<FTableColumn> ColumnRef;
+};
+
+////////////////////////////////////////////////////////////////////////////////////////////////////
+
+/** Creates a group for each first letter of node names. */
+class FTreeNodeGroupingByNameFirstLetter : public FTreeNodeGrouping
+{
+public:
+	FTreeNodeGroupingByNameFirstLetter();
+	virtual ~FTreeNodeGroupingByNameFirstLetter() {}
+
+	virtual FTreeNodeGroupInfo GetGroupForNode(const FBaseTreeNodePtr InNode) const override;
+};
+
+////////////////////////////////////////////////////////////////////////////////////////////////////
+
+/** Creates a group for each type. */
+class FTreeNodeGroupingByType : public FTreeNodeGrouping
+{
+public:
+	FTreeNodeGroupingByType();
+	virtual ~FTreeNodeGroupingByType() {}
+
+	virtual FTreeNodeGroupInfo GetGroupForNode(const FBaseTreeNodePtr InNode) const override;
+};
+
+////////////////////////////////////////////////////////////////////////////////////////////////////
+
+} // namespace Insights
diff --git a/Engine/Source/Developer/TraceInsights/Private/Insights/Table/ViewModels/TreeNodeSorting.cpp b/Engine/Source/Developer/TraceInsights/Private/Insights/Table/ViewModels/TreeNodeSorting.cpp
new file mode 100644
--- /dev/null
+++ b/Engine/Source/Developer/TraceInsights/Private/Insights/Table/ViewModels/TreeNodeSorting.cpp
@@ -0,0 +1,3 @@
+// Copyright 1998-2019 Epic Games, Inc. All Rights Reserved.
+
+#include "TreeNodeSorting.h"
diff --git a/Engine/Source/Developer/TraceInsights/Private/Insights/Table/ViewModels/TreeNodeSorting.h b/Engine/Source/Developer/TraceInsights/Private/Insights/Table/ViewModels/TreeNodeSorting.h
new file mode 100644
--- /dev/null
+++ b/Engine/Source/Developer/TraceInsights/Private/Insights/Table/ViewModels/TreeNodeSorting.h
@@ -0,0 +1,15 @@
+// Copyright 1998-2019 Epic Games, Inc. All Rights Reserved.
+
+#pragma once
+
+#include "CoreMinimal.h"
+
+// Insights
+#include "Insights/Table/ViewModels/TableCellValueSorter.h"
+
+namespace Insights
+{
+
+//class FTableTreeNodeSorterByRowIndex
+
+} // namespace Insights
