Copying //UE4/Dev-Core to Main (//UE4/Main)

diff --git a/Engine/Source/Developer/TraceInsights/Private/Insights/Table/Widgets/STableTreeView.cpp b/Engine/Source/Developer/TraceInsights/Private/Insights/Table/Widgets/STableTreeView.cpp
new file mode 100644
--- /dev/null
+++ b/Engine/Source/Developer/TraceInsights/Private/Insights/Table/Widgets/STableTreeView.cpp
@@ -0,0 +1,1907 @@
+// Copyright 1998-2019 Epic Games, Inc. All Rights Reserved.
+
+#include "STableTreeView.h"
+
+#include "EditorStyleSet.h"
+#include "Framework/Application/SlateApplication.h"
+#include "Framework/MultiBox/MultiBoxBuilder.h"
+#include "Layout/WidgetPath.h"
+#include "SlateOptMacros.h"
+#include "TraceServices/AnalysisService.h"
+#include "Widgets/Input/SCheckBox.h"
+#include "Widgets/Layout/SScrollBox.h"
+#include "Widgets/Input/SSearchBox.h"
+#include "Widgets/Images/SImage.h"
+#include "Widgets/Layout/SGridPanel.h"
+#include "Widgets/Layout/SSeparator.h"
+#include "Widgets/SToolTip.h"
+#include "Widgets/Views/STableViewBase.h"
+
+// Insights
+#include "Insights/InsightsManager.h"
+#include "Insights/Table/ViewModels/Table.h"
+#include "Insights/Table/ViewModels/TableColumn.h"
+#include "Insights/Table/ViewModels/TableTreeNode.h"
+#include "Insights/Table/ViewModels/TreeNodeGrouping.h"
+#include "Insights/Table/ViewModels/TreeNodeSorting.h"
+#include "Insights/Table/Widgets/STableTreeViewTooltip.h"
+#include "Insights/Table/Widgets/STableTreeViewRow.h"
+
+#define LOCTEXT_NAMESPACE "STableTreeView"
+
+namespace Insights
+{
+
+const FName STableTreeView::RootNodeName(TEXT("Root"));
+
+////////////////////////////////////////////////////////////////////////////////////////////////////
+
+STableTreeView::STableTreeView()
+	: Table()
+	, Session(FInsightsManager::Get()->GetSession())
+	, TreeView(nullptr)
+	, TreeViewHeaderRow(nullptr)
+	, ExternalScrollbar(nullptr)
+	, HoveredColumnId()
+	, HoveredNodePtr(nullptr)
+	, HighlightedNodeName()
+	, Root(MakeShareable(new FTableTreeNode(RootNodeName, Table)))
+	, TableTreeNodes()
+	, FilteredGroupNodes()
+	, TableTreeNodesIdMap()
+	, ExpandedNodes()
+	, bExpansionSaved(false)
+	, SearchBox(nullptr)
+	, TextFilter(nullptr)
+	, Filters(nullptr)
+	, AvailableGroupings()
+	, CurrentGroupings()
+	, GroupingBreadcrumbTrail(nullptr)
+	, AvailableSorters()
+	, CurrentSorter(nullptr)
+	, ColumnBeingSorted(GetDefaultColumnBeingSorted())
+	, ColumnSortMode(GetDefaultColumnSortMode())
+	, StatsStartTime(0.0)
+	, StatsEndTime(0.0)
+{
+}
+
+////////////////////////////////////////////////////////////////////////////////////////////////////
+
+STableTreeView::~STableTreeView()
+{
+	// Remove ourselves from the Insights manager.
+	if (FInsightsManager::Get().IsValid())
+	{
+		FInsightsManager::Get()->GetSessionChangedEvent().RemoveAll(this);
+	}
+}
+
+////////////////////////////////////////////////////////////////////////////////////////////////////
+
+BEGIN_SLATE_FUNCTION_BUILD_OPTIMIZATION
+void STableTreeView::Construct(const FArguments& InArgs, TSharedPtr<FTable> InTablePtr)
+{
+	check(InTablePtr.IsValid());
+	Table = InTablePtr;
+
+	SAssignNew(ExternalScrollbar, SScrollBar)
+	.AlwaysShowScrollbar(true);
+
+	ChildSlot
+	[
+		SNew(SVerticalBox)
+
+		+SVerticalBox::Slot()
+		.VAlign(VAlign_Center)
+		.AutoHeight()
+		[
+			SNew(SBorder)
+			.BorderImage(FEditorStyle::GetBrush("ToolPanel.GroupBorder"))
+			.Padding(2.0f)
+			[
+				SNew(SVerticalBox)
+
+				// Search box
+				+SVerticalBox::Slot()
+				.VAlign(VAlign_Center)
+				.Padding(2.0f)
+				.AutoHeight()
+				[
+					SAssignNew(SearchBox, SSearchBox)
+					.HintText(LOCTEXT("SearchBoxHint", "Search timers or groups"))
+					.OnTextChanged(this, &STableTreeView::SearchBox_OnTextChanged)
+					.IsEnabled(this, &STableTreeView::SearchBox_IsEnabled)
+					.ToolTipText(LOCTEXT("FilterSearchHint", "Type here to search timer or group"))
+				]
+
+				// Group by
+				+SVerticalBox::Slot()
+				.VAlign(VAlign_Center)
+				.Padding(2.0f)
+				.AutoHeight()
+				[
+					SNew(SHorizontalBox)
+
+					+SHorizontalBox::Slot()
+					.AutoWidth()
+					.VAlign(VAlign_Center)
+					[
+						SNew(STextBlock)
+						.Text(LOCTEXT("GroupByText", "Hierarchy:"))
+						.Margin(FMargin(0.0f, 0.0f, 4.0f, 0.0f))
+					]
+
+					+SHorizontalBox::Slot()
+					.FillWidth(1.0f)
+					.VAlign(VAlign_Center)
+					[
+						SAssignNew(GroupingBreadcrumbTrail, SBreadcrumbTrail<TSharedPtr<FTreeNodeGrouping>>)
+						.ButtonContentPadding(FMargin(1.0f, 1.0f))
+						//.DelimiterImage(FEditorStyle::GetBrush("SlateFileDialogs.PathDelimiter"))
+						//.TextStyle(FEditorStyle::Get(), "Tutorials.Browser.PathText")
+						//.ShowLeadingDelimiter(true)
+						//.PersistentBreadcrumbs(true)
+						.InvertTextColorOnHover(true)
+						.OnCrumbClicked(this, &STableTreeView::OnGroupingCrumbClicked)
+						.GetCrumbMenuContent(this, &STableTreeView::GetGroupingCrumbMenuContent)
+					]
+				]
+			]
+		]
+
+		// Tree view
+		+SVerticalBox::Slot()
+		.FillHeight(1.0f)
+		.Padding(0.0f, 6.0f, 0.0f, 0.0f)
+		[
+			SNew(SHorizontalBox)
+
+			+ SHorizontalBox::Slot()
+			.FillWidth(1.0f)
+			.Padding(0.0f)
+			[
+				SNew(SScrollBox)
+				.Orientation(Orient_Horizontal)
+
+				+ SScrollBox::Slot()
+				[
+					SNew(SBorder)
+					.BorderImage(FEditorStyle::GetBrush("ToolPanel.GroupBorder"))
+					.Padding(0.0f)
+					[
+						SAssignNew(TreeView, STreeView<FTableTreeNodePtr>)
+						.ExternalScrollbar(ExternalScrollbar)
+						.SelectionMode(ESelectionMode::Multi)
+						.TreeItemsSource(&FilteredGroupNodes)
+						.OnGetChildren(this, &STableTreeView::TreeView_OnGetChildren)
+						.OnGenerateRow(this, &STableTreeView::TreeView_OnGenerateRow)
+						.OnSelectionChanged(this, &STableTreeView::TreeView_OnSelectionChanged)
+						.OnMouseButtonDoubleClick(this, &STableTreeView::TreeView_OnMouseButtonDoubleClick)
+						.OnContextMenuOpening(FOnContextMenuOpening::CreateSP(this, &STableTreeView::TreeView_GetMenuContent))
+						.ItemHeight(12.0f)
+						.HeaderRow
+						(
+							SAssignNew(TreeViewHeaderRow, SHeaderRow)
+							.Visibility(EVisibility::Visible)
+						)
+					]
+				]
+			]
+
+			+ SHorizontalBox::Slot()
+			.AutoWidth()
+			.Padding(0.0f)
+			[
+				SNew(SBox)
+				.WidthOverride(FOptionalSize(13.0f))
+				[
+					ExternalScrollbar.ToSharedRef()
+				]
+			]
+		]
+	];
+
+	InitializeAndShowHeaderColumns();
+	//BindCommands();
+
+	// Create the search filters: text based, type based etc.
+	TextFilter = MakeShareable(new FTableTreeNodeTextFilter(FTableTreeNodeTextFilter::FItemToStringArray::CreateSP(this, &STableTreeView::HandleItemToStringArray)));
+	Filters = MakeShareable(new FTableTreeNodeFilterCollection());
+	Filters->Add(TextFilter);
+
+	CreateGroupings();
+	CreateSortings();
+
+	// Register ourselves with the Insights manager.
+	FInsightsManager::Get()->GetSessionChangedEvent().AddSP(this, &STableTreeView::InsightsManager_OnSessionChanged);
+}
+END_SLATE_FUNCTION_BUILD_OPTIMIZATION
+
+////////////////////////////////////////////////////////////////////////////////////////////////////
+
+TSharedPtr<SWidget> STableTreeView::TreeView_GetMenuContent()
+{
+	const TArray<FTableTreeNodePtr> SelectedNodes = TreeView->GetSelectedItems();
+	const int32 NumSelectedNodes = SelectedNodes.Num();
+	FTableTreeNodePtr SelectedNode = NumSelectedNodes ? SelectedNodes[0] : nullptr;
+
+	const TSharedPtr<FTableColumn> HoveredColumnPtr = Table->FindColumn(HoveredColumnId);
+
+	FText SelectionStr;
+	FText PropertyName;
+	FText PropertyValue;
+
+	if (NumSelectedNodes == 0)
+	{
+		SelectionStr = LOCTEXT("NothingSelected", "Nothing selected");
+	}
+	else if (NumSelectedNodes == 1)
+	{
+		if (HoveredColumnPtr != nullptr)
+		{
+			PropertyName = HoveredColumnPtr->GetShortName();
+			PropertyValue = HoveredColumnPtr->GetValueAsTooltipText(*SelectedNode);
+		}
+		SelectionStr = FText::FromName(SelectedNode->GetName());
+	}
+	else
+	{
+		SelectionStr = LOCTEXT("MultipleSelection", "Multiple selection");
+	}
+
+	const bool bShouldCloseWindowAfterMenuSelection = true;
+	FMenuBuilder MenuBuilder(bShouldCloseWindowAfterMenuSelection, NULL);
+
+	// Selection menu
+	MenuBuilder.BeginSection("Selection", LOCTEXT("ContextMenu_Header_Selection", "Selection"));
+	{
+		struct FLocal
+		{
+			static bool ReturnFalse()
+			{
+				return false;
+			}
+		};
+
+		FUIAction DummyUIAction;
+		DummyUIAction.CanExecuteAction = FCanExecuteAction::CreateStatic(&FLocal::ReturnFalse);
+		MenuBuilder.AddMenuEntry
+		(
+			SelectionStr,
+			LOCTEXT("ContextMenu_Selection", "Currently selected items"),
+			FSlateIcon(FEditorStyle::GetStyleSetName(), "@missing.icon"), DummyUIAction, NAME_None, EUserInterfaceActionType::Button
+		);
+	}
+	MenuBuilder.EndSection();
+
+	MenuBuilder.BeginSection("Misc", LOCTEXT("ContextMenu_Header_Misc", "Miscellaneous"));
+	{
+		/*TODO
+		FUIAction Action_CopySelectedToClipboard
+		(
+			FExecuteAction::CreateSP(this, &STableTreeView::ContextMenu_CopySelectedToClipboard_Execute),
+			FCanExecuteAction::CreateSP(this, &STableTreeView::ContextMenu_CopySelectedToClipboard_CanExecute)
+		);
+		MenuBuilder.AddMenuEntry
+		(
+			LOCTEXT("ContextMenu_Header_Misc_CopySelectedToClipboard", "Copy To Clipboard"),
+			LOCTEXT("ContextMenu_Header_Misc_CopySelectedToClipboard_Desc", "Copies selection to clipboard"),
+			FSlateIcon(FEditorStyle::GetStyleSetName(), "Profiler.Misc.CopyToClipboard"), Action_CopySelectedToClipboard, NAME_None, EUserInterfaceActionType::Button
+		);
+		*/
+
+		MenuBuilder.AddSubMenu
+		(
+			LOCTEXT("ContextMenu_Header_Misc_Sort", "Sort By"),
+			LOCTEXT("ContextMenu_Header_Misc_Sort_Desc", "Sort by column"),
+			FNewMenuDelegate::CreateSP(this, &STableTreeView::TreeView_BuildSortByMenu),
+			false,
+			FSlateIcon(FEditorStyle::GetStyleSetName(), "Profiler.Misc.SortBy")
+		);
+	}
+	MenuBuilder.EndSection();
+
+	MenuBuilder.BeginSection("Columns", LOCTEXT("ContextMenu_Header_Columns", "Columns"));
+	{
+		MenuBuilder.AddSubMenu
+		(
+			LOCTEXT("ContextMenu_Header_Columns_View", "View Column"),
+			LOCTEXT("ContextMenu_Header_Columns_View_Desc", "Hides or shows columns"),
+			FNewMenuDelegate::CreateSP(this, &STableTreeView::TreeView_BuildViewColumnMenu),
+			false,
+			FSlateIcon(FEditorStyle::GetStyleSetName(), "Profiler.EventGraph.ViewColumn")
+		);
+
+		FUIAction Action_ShowAllColumns
+		(
+			FExecuteAction::CreateSP(this, &STableTreeView::ContextMenu_ShowAllColumns_Execute),
+			FCanExecuteAction::CreateSP(this, &STableTreeView::ContextMenu_ShowAllColumns_CanExecute)
+		);
+		MenuBuilder.AddMenuEntry
+		(
+			LOCTEXT("ContextMenu_Header_Columns_ShowAllColumns", "Show All Columns"),
+			LOCTEXT("ContextMenu_Header_Columns_ShowAllColumns_Desc", "Resets tree view to show all columns"),
+			FSlateIcon(FEditorStyle::GetStyleSetName(), "Profiler.EventGraph.ResetColumn"), Action_ShowAllColumns, NAME_None, EUserInterfaceActionType::Button
+		);
+
+		FUIAction Action_ResetColumns
+		(
+			FExecuteAction::CreateSP(this, &STableTreeView::ContextMenu_ResetColumns_Execute),
+			FCanExecuteAction::CreateSP(this, &STableTreeView::ContextMenu_ResetColumns_CanExecute)
+		);
+		MenuBuilder.AddMenuEntry
+		(
+			LOCTEXT("ContextMenu_Header_Columns_ResetColumns", "Reset Columns to Default"),
+			LOCTEXT("ContextMenu_Header_Columns_ResetColumns_Desc", "Resets columns to default"),
+			FSlateIcon(FEditorStyle::GetStyleSetName(), "Profiler.EventGraph.ResetColumn"), Action_ResetColumns, NAME_None, EUserInterfaceActionType::Button
+		);
+	}
+	MenuBuilder.EndSection();
+
+	return MenuBuilder.MakeWidget();
+}
+
+////////////////////////////////////////////////////////////////////////////////////////////////////
+
+void STableTreeView::TreeView_BuildSortByMenu(FMenuBuilder& MenuBuilder)
+{
+	// TODO: Refactor later @see TSharedPtr<SWidget> SCascadePreviewViewportToolBar::GenerateViewMenu() const
+
+	MenuBuilder.BeginSection("ColumnName", LOCTEXT("ContextMenu_Header_Misc_ColumnName", "Column Name"));
+
+	//TODO: for (Sorting : AvailableSortings)
+	for (const TSharedPtr<FTableColumn>& ColumnPtr : Table->GetColumns())
+	{
+		const FTableColumn& Column = *ColumnPtr;
+
+		if (Column.IsVisible() && Column.CanBeSorted())
+		{
+			FUIAction Action_SortByColumn
+			(
+				FExecuteAction::CreateSP(this, &STableTreeView::ContextMenu_SortByColumn_Execute, Column.GetId()),
+				FCanExecuteAction::CreateSP(this, &STableTreeView::ContextMenu_SortByColumn_CanExecute, Column.GetId()),
+				FIsActionChecked::CreateSP(this, &STableTreeView::ContextMenu_SortByColumn_IsChecked, Column.GetId())
+			);
+			MenuBuilder.AddMenuEntry
+			(
+				Column.GetTitleName(),
+				Column.GetDescription(),
+				FSlateIcon(), Action_SortByColumn, NAME_None, EUserInterfaceActionType::RadioButton
+			);
+		}
+	}
+
+	MenuBuilder.EndSection();
+
+	//-----------------------------------------------------------------------------
+
+	MenuBuilder.BeginSection("SortMode", LOCTEXT("ContextMenu_Header_Misc_Sort_SortMode", "Sort Mode"));
+	{
+		FUIAction Action_SortAscending
+		(
+			FExecuteAction::CreateSP(this, &STableTreeView::ContextMenu_SortMode_Execute, EColumnSortMode::Ascending),
+			FCanExecuteAction::CreateSP(this, &STableTreeView::ContextMenu_SortMode_CanExecute, EColumnSortMode::Ascending),
+			FIsActionChecked::CreateSP(this, &STableTreeView::ContextMenu_SortMode_IsChecked, EColumnSortMode::Ascending)
+		);
+		MenuBuilder.AddMenuEntry
+		(
+			LOCTEXT("ContextMenu_Header_Misc_Sort_SortAscending", "Sort Ascending"),
+			LOCTEXT("ContextMenu_Header_Misc_Sort_SortAscending_Desc", "Sorts ascending"),
+			FSlateIcon(FEditorStyle::GetStyleSetName(), "Profiler.Misc.SortAscending"), Action_SortAscending, NAME_None, EUserInterfaceActionType::RadioButton
+		);
+
+		FUIAction Action_SortDescending
+		(
+			FExecuteAction::CreateSP(this, &STableTreeView::ContextMenu_SortMode_Execute, EColumnSortMode::Descending),
+			FCanExecuteAction::CreateSP(this, &STableTreeView::ContextMenu_SortMode_CanExecute, EColumnSortMode::Descending),
+			FIsActionChecked::CreateSP(this, &STableTreeView::ContextMenu_SortMode_IsChecked, EColumnSortMode::Descending)
+		);
+		MenuBuilder.AddMenuEntry
+		(
+			LOCTEXT("ContextMenu_Header_Misc_Sort_SortDescending", "Sort Descending"),
+			LOCTEXT("ContextMenu_Header_Misc_Sort_SortDescending_Desc", "Sorts descending"),
+			FSlateIcon(FEditorStyle::GetStyleSetName(), "Profiler.Misc.SortDescending"), Action_SortDescending, NAME_None, EUserInterfaceActionType::RadioButton
+		);
+	}
+	MenuBuilder.EndSection();
+}
+
+////////////////////////////////////////////////////////////////////////////////////////////////////
+
+void STableTreeView::TreeView_BuildViewColumnMenu(FMenuBuilder& MenuBuilder)
+{
+	MenuBuilder.BeginSection("ViewColumn", LOCTEXT("ContextMenu_Header_Columns_View", "View Column"));
+
+	for (const TSharedPtr<FTableColumn>& ColumnPtr : Table->GetColumns())
+	{
+		const FTableColumn& Column = *ColumnPtr;
+
+		FUIAction Action_ToggleColumn
+		(
+			FExecuteAction::CreateSP(this, &STableTreeView::ToggleColumnVisibility, Column.GetId()),
+			FCanExecuteAction::CreateSP(this, &STableTreeView::CanToggleColumnVisibility, Column.GetId()),
+			FIsActionChecked::CreateSP(this, &STableTreeView::IsColumnVisible, Column.GetId())
+		);
+		MenuBuilder.AddMenuEntry
+		(
+			Column.GetTitleName(),
+			Column.GetDescription(),
+			FSlateIcon(), Action_ToggleColumn, NAME_None, EUserInterfaceActionType::ToggleButton
+		);
+	}
+
+	MenuBuilder.EndSection();
+}
+
+////////////////////////////////////////////////////////////////////////////////////////////////////
+
+void STableTreeView::InitializeAndShowHeaderColumns()
+{
+	for (const TSharedPtr<FTableColumn>& ColumnPtr : Table->GetColumns())
+	{
+		if (ColumnPtr->ShouldBeVisible())
+		{
+			ShowColumn(ColumnPtr->GetId());
+		}
+	}
+}
+
+////////////////////////////////////////////////////////////////////////////////////////////////////
+
+FText STableTreeView::GetColumnHeaderText(const FName ColumnId) const
+{
+	const FTableColumn& Column = *Table->FindColumnChecked(ColumnId);
+	return Column.GetShortName();
+}
+
+////////////////////////////////////////////////////////////////////////////////////////////////////
+
+TSharedRef<SWidget> STableTreeView::TreeViewHeaderRow_GenerateColumnMenu(const FTableColumn& Column)
+{
+	bool bIsMenuVisible = false;
+
+	const bool bShouldCloseWindowAfterMenuSelection = true;
+	FMenuBuilder MenuBuilder(bShouldCloseWindowAfterMenuSelection, NULL);
+	{
+		if (Column.CanBeHidden())
+		{
+			MenuBuilder.BeginSection("Column", LOCTEXT("TreeViewHeaderRow_Header_Column", "Column"));
+
+			FUIAction Action_HideColumn
+			(
+				FExecuteAction::CreateSP(this, &STableTreeView::HideColumn, Column.GetId()),
+				FCanExecuteAction::CreateSP(this, &STableTreeView::CanHideColumn, Column.GetId())
+			);
+
+			MenuBuilder.AddMenuEntry
+			(
+				LOCTEXT("TreeViewHeaderRow_HideColumn", "Hide"),
+				LOCTEXT("TreeViewHeaderRow_HideColumn_Desc", "Hides the selected column"),
+				FSlateIcon(), Action_HideColumn, NAME_None, EUserInterfaceActionType::Button
+			);
+			bIsMenuVisible = true;
+
+			MenuBuilder.EndSection();
+		}
+
+		if (Column.CanBeSorted())
+		{
+			MenuBuilder.BeginSection("SortMode", LOCTEXT("ContextMenu_Header_Misc_Sort_SortMode", "Sort Mode"));
+
+			FUIAction Action_SortAscending
+			(
+				FExecuteAction::CreateSP(this, &STableTreeView::HeaderMenu_SortMode_Execute, Column.GetId(), EColumnSortMode::Ascending),
+				FCanExecuteAction::CreateSP(this, &STableTreeView::HeaderMenu_SortMode_CanExecute, Column.GetId(), EColumnSortMode::Ascending),
+				FIsActionChecked::CreateSP(this, &STableTreeView::HeaderMenu_SortMode_IsChecked, Column.GetId(), EColumnSortMode::Ascending)
+			);
+			MenuBuilder.AddMenuEntry
+			(
+				LOCTEXT("ContextMenu_Header_Misc_Sort_SortAscending", "Sort Ascending"),
+				LOCTEXT("ContextMenu_Header_Misc_Sort_SortAscending_Desc", "Sorts ascending"),
+				FSlateIcon(FEditorStyle::GetStyleSetName(), "Profiler.Misc.SortAscending"), Action_SortAscending, NAME_None, EUserInterfaceActionType::RadioButton
+			);
+
+			FUIAction Action_SortDescending
+			(
+				FExecuteAction::CreateSP(this, &STableTreeView::HeaderMenu_SortMode_Execute, Column.GetId(), EColumnSortMode::Descending),
+				FCanExecuteAction::CreateSP(this, &STableTreeView::HeaderMenu_SortMode_CanExecute, Column.GetId(), EColumnSortMode::Descending),
+				FIsActionChecked::CreateSP(this, &STableTreeView::HeaderMenu_SortMode_IsChecked, Column.GetId(), EColumnSortMode::Descending)
+			);
+			MenuBuilder.AddMenuEntry
+			(
+				LOCTEXT("ContextMenu_Header_Misc_Sort_SortDescending", "Sort Descending"),
+				LOCTEXT("ContextMenu_Header_Misc_Sort_SortDescending_Desc", "Sorts descending"),
+				FSlateIcon(FEditorStyle::GetStyleSetName(), "Profiler.Misc.SortDescending"), Action_SortDescending, NAME_None, EUserInterfaceActionType::RadioButton
+			);
+			bIsMenuVisible = true;
+
+			MenuBuilder.EndSection();
+		}
+
+		//if (Column.CanBeFiltered())
+		//{
+		//	MenuBuilder.BeginSection("FilterMode", LOCTEXT("ContextMenu_Header_Misc_Filter_FilterMode", "Filter Mode"));
+		//	bIsMenuVisible = true;
+		//	MenuBuilder.EndSection();
+		//}
+	}
+
+	return bIsMenuVisible ? MenuBuilder.MakeWidget() : (TSharedRef<SWidget>)SNullWidget::NullWidget;
+}
+
+////////////////////////////////////////////////////////////////////////////////////////////////////
+
+void STableTreeView::InsightsManager_OnSessionChanged()
+{
+	TSharedPtr<const Trace::IAnalysisSession> NewSession = FInsightsManager::Get()->GetSession();
+
+	if (NewSession != Session)
+	{
+		Session = NewSession;
+		Reset();
+	}
+	else
+	{
+		UpdateTree();
+	}
+}
+
+////////////////////////////////////////////////////////////////////////////////////////////////////
+
+void STableTreeView::UpdateTree()
+{
+	CreateGroups();
+	SortTreeNodes();
+	ApplyFiltering();
+}
+
+////////////////////////////////////////////////////////////////////////////////////////////////////
+
+void STableTreeView::ApplyFiltering()
+{
+	// Apply filter to all groups and its children.
+	ApplyFilteringForNode(Root);
+
+	FilteredGroupNodes.Reset();
+	const TArray<FBaseTreeNodePtr>& RootChildren = Root->GetFilteredChildren();
+	const int32 NumRootChildren = RootChildren.Num();
+	for (int32 Cx = 0; Cx < NumRootChildren; ++Cx)
+	{
+		// Add a child.
+		const FTableTreeNodePtr& ChildNodePtr = StaticCastSharedPtr<FTableTreeNode>(RootChildren[Cx]);
+		if (ChildNodePtr->IsGroup())
+		{
+			FilteredGroupNodes.Add(ChildNodePtr);
+		}
+	}
+
+	// Only expand nodes if we have a text filter.
+	const bool bNonEmptyTextFilter = !TextFilter->GetRawFilterText().IsEmpty();
+	if (bNonEmptyTextFilter)
+	{
+		if (!bExpansionSaved)
+		{
+			ExpandedNodes.Empty();
+			TreeView->GetExpandedItems(ExpandedNodes);
+			bExpansionSaved = true;
+		}
+
+		for (int32 Fx = 0; Fx < FilteredGroupNodes.Num(); Fx++)
+		{
+			const FTableTreeNodePtr& GroupPtr = FilteredGroupNodes[Fx];
+			TreeView->SetItemExpansion(GroupPtr, GroupPtr->IsExpanded());
+		}
+	}
+	else
+	{
+		if (bExpansionSaved)
+		{
+			// Restore previously expanded nodes when the text filter is disabled.
+			TreeView->ClearExpandedItems();
+			for (auto It = ExpandedNodes.CreateConstIterator(); It; ++It)
+			{
+				TreeView->SetItemExpansion(*It, true);
+			}
+			bExpansionSaved = false;
+		}
+	}
+
+	// Request tree refresh
+	TreeView->RequestTreeRefresh();
+}
+
+////////////////////////////////////////////////////////////////////////////////////////////////////
+
+bool STableTreeView::ApplyFilteringForNode(FTableTreeNodePtr NodePtr)
+{
+	const bool bIsNodeVisible = Filters->PassesAllFilters(NodePtr);
+
+	if (NodePtr->IsGroup())
+	{
+		NodePtr->ClearFilteredChildren();
+
+		const TArray<FBaseTreeNodePtr>& GroupChildren = NodePtr->GetChildren();
+		const int32 NumChildren = GroupChildren.Num();
+		int32 NumVisibleChildren = 0;
+		for (int32 Cx = 0; Cx < NumChildren; ++Cx)
+		{
+			// Add a child.
+			const FTableTreeNodePtr& ChildNodePtr = StaticCastSharedPtr<FTableTreeNode>(GroupChildren[Cx]);
+			if (ApplyFilteringForNode(ChildNodePtr))
+			{
+				NodePtr->AddFilteredChild(ChildNodePtr);
+				NumVisibleChildren++;
+			}
+		}
+
+		const bool bIsGroupNodeVisible = bIsNodeVisible || NumVisibleChildren > 0;
+
+		if (bIsGroupNodeVisible)
+		{
+			// Add a group.
+			NodePtr->SetExpansion(true);
+		}
+		else
+		{
+			NodePtr->SetExpansion(false);
+		}
+
+		return bIsGroupNodeVisible;
+	}
+	else
+	{
+		return bIsNodeVisible;
+	}
+}
+
+////////////////////////////////////////////////////////////////////////////////////////////////////
+
+void STableTreeView::HandleItemToStringArray(const FTableTreeNodePtr& FTableTreeNodePtr, TArray<FString>& OutSearchStrings) const
+{
+	OutSearchStrings.Add(FTableTreeNodePtr->GetName().GetPlainNameString());
+}
+
+////////////////////////////////////////////////////////////////////////////////////////////////////
+// TreeView
+////////////////////////////////////////////////////////////////////////////////////////////////////
+
+void STableTreeView::TreeView_Refresh()
+{
+	if (TreeView.IsValid())
+	{
+		TreeView->RequestTreeRefresh();
+	}
+}
+
+////////////////////////////////////////////////////////////////////////////////////////////////////
+
+void STableTreeView::TreeView_OnSelectionChanged(FTableTreeNodePtr SelectedItem, ESelectInfo::Type SelectInfo)
+{
+}
+
+////////////////////////////////////////////////////////////////////////////////////////////////////
+
+void STableTreeView::TreeView_OnGetChildren(FTableTreeNodePtr InParent, TArray<FTableTreeNodePtr>& OutChildren)
+{
+	const TArray<FBaseTreeNodePtr>& FilteredChildren = InParent->GetFilteredChildren();
+	for (const FBaseTreeNodePtr& NodePtr : FilteredChildren)
+	{
+		OutChildren.Add(StaticCastSharedPtr<FTableTreeNode>(NodePtr));
+	}
+}
+
+////////////////////////////////////////////////////////////////////////////////////////////////////
+
+void STableTreeView::TreeView_OnMouseButtonDoubleClick(FTableTreeNodePtr NodePtr)
+{
+	if (NodePtr->IsGroup())
+	{
+		const bool bIsGroupExpanded = TreeView->IsItemExpanded(NodePtr);
+		TreeView->SetItemExpansion(NodePtr, !bIsGroupExpanded);
+	}
+}
+
+////////////////////////////////////////////////////////////////////////////////////////////////////
+// Tree View's Table Row
+////////////////////////////////////////////////////////////////////////////////////////////////////
+
+TSharedRef<ITableRow> STableTreeView::TreeView_OnGenerateRow(FTableTreeNodePtr NodePtr, const TSharedRef<STableViewBase>& OwnerTable)
+{
+	TSharedRef<ITableRow> TableRow =
+		SNew(STableTreeViewRow, OwnerTable)
+		.OnShouldBeEnabled(this, &STableTreeView::TableRow_ShouldBeEnabled)
+		.OnIsColumnVisible(this, &STableTreeView::IsColumnVisible)
+		.OnSetHoveredCell(this, &STableTreeView::TableRow_SetHoveredCell)
+		.OnGetColumnOutlineHAlignmentDelegate(this, &STableTreeView::TableRow_GetColumnOutlineHAlignment)
+		.HighlightText(this, &STableTreeView::TableRow_GetHighlightText)
+		.HighlightedNodeName(this, &STableTreeView::TableRow_GetHighlightedNodeName)
+		.TablePtr(Table)
+		.TableTreeNodePtr(NodePtr);
+
+	return TableRow;
+}
+
+////////////////////////////////////////////////////////////////////////////////////////////////////
+
+void STableTreeView::TableRow_SetHoveredCell(TSharedPtr<FTable> InTablePtr, TSharedPtr<FTableColumn> InColumnPtr, const FTableTreeNodePtr InNodePtr)
+{
+	HoveredColumnId = InColumnPtr ? InColumnPtr->GetId() : FName();
+
+	const bool bIsAnyMenusVisible = FSlateApplication::Get().AnyMenusVisible();
+	if (!HasMouseCapture() && !bIsAnyMenusVisible)
+	{
+		HoveredNodePtr = InNodePtr;
+	}
+}
+
+////////////////////////////////////////////////////////////////////////////////////////////////////
+
+EHorizontalAlignment STableTreeView::TableRow_GetColumnOutlineHAlignment(const FName ColumnId) const
+{
+	const TIndirectArray<SHeaderRow::FColumn>& Columns = TreeViewHeaderRow->GetColumns();
+	const int32 LastColumnIdx = Columns.Num() - 1;
+
+	// First column
+	if (Columns[0].ColumnId == ColumnId)
+	{
+		return HAlign_Left;
+	}
+	// Last column
+	else if (Columns[LastColumnIdx].ColumnId == ColumnId)
+	{
+		return HAlign_Right;
+	}
+	// Middle columns
+	{
+		return HAlign_Center;
+	}
+}
+
+////////////////////////////////////////////////////////////////////////////////////////////////////
+
+FText STableTreeView::TableRow_GetHighlightText() const
+{
+	return SearchBox->GetText();
+}
+
+////////////////////////////////////////////////////////////////////////////////////////////////////
+
+FName STableTreeView::TableRow_GetHighlightedNodeName() const
+{
+	return HighlightedNodeName;
+}
+
+////////////////////////////////////////////////////////////////////////////////////////////////////
+
+bool STableTreeView::TableRow_ShouldBeEnabled(const uint32 TimerId) const
+{
+	return true;//im:TODO: Session->GetAggregatedStat(TimerId) != nullptr;
+}
+
+////////////////////////////////////////////////////////////////////////////////////////////////////
+// SearchBox
+////////////////////////////////////////////////////////////////////////////////////////////////////
+
+void STableTreeView::SearchBox_OnTextChanged(const FText& InFilterText)
+{
+	TextFilter->SetRawFilterText(InFilterText);
+	SearchBox->SetError(TextFilter->GetFilterErrorText());
+	ApplyFiltering();
+}
+
+////////////////////////////////////////////////////////////////////////////////////////////////////
+
+bool STableTreeView::SearchBox_IsEnabled() const
+{
+	return TableTreeNodes.Num() > 0;
+}
+
+////////////////////////////////////////////////////////////////////////////////////////////////////
+// Grouping
+////////////////////////////////////////////////////////////////////////////////////////////////////
+
+void STableTreeView::CreateGroups()
+{
+	GroupNodesRec(TableTreeNodes, *Root, 0);
+
+	ResetAggregatedValuesRec(*Root);
+	for (const TSharedPtr<FTableColumn> ColumnPtr : Table->GetColumns())
+	{
+		FTableColumn& Column = *ColumnPtr;
+		if (Column.GetAggregation() == ETableColumnAggregation::Sum)
+		{
+			switch (Column.GetDataType())
+			{
+				case ETableCellDataType::Int64:
+					UpdateInt64SumAggregationRec(Column, *Root);
+					break;
+
+				case ETableCellDataType::Float:
+					UpdateFloatSumAggregationRec(Column, *Root);
+					break;
+
+				case ETableCellDataType::Double:
+					UpdateDoubleSumAggregationRec(Column, *Root);
+					break;
+			}
+		}
+	}
+}
+
+////////////////////////////////////////////////////////////////////////////////////////////////////
+
+void STableTreeView::GroupNodesRec(const TArray<FTableTreeNodePtr>& Nodes, FTableTreeNode& ParentGroup, int32 GroupingDepth)
+{
+	ensure(CurrentGroupings.Num() > 0);
+
+	FTreeNodeGrouping& Grouping = *CurrentGroupings[GroupingDepth];
+
+	TMap<FName, FTableTreeNodePtr> GroupMap;
+
+	ParentGroup.ClearChildren();
+
+	for (FTableTreeNodePtr NodePtr : Nodes)
+	{
+		ensure(!NodePtr->IsGroup());
+
+		FTableTreeNodePtr GroupPtr = nullptr;
+
+		FTreeNodeGroupInfo GroupInfo = Grouping.GetGroupForNode(NodePtr);
+		FTableTreeNodePtr* GroupPtrPtr = GroupMap.Find(GroupInfo.Name);
+		if (!GroupPtrPtr)
+		{
+			GroupPtr = MakeShareable(new FTableTreeNode(GroupInfo.Name, Table));
+
+			GroupMap.Add(GroupInfo.Name, GroupPtr);
+			ParentGroup.AddChildAndSetGroupPtr(GroupPtr);
+
+			GroupPtr->SetExpansion(GroupInfo.IsExpanded);
+			TreeView->SetItemExpansion(GroupPtr, GroupInfo.IsExpanded);
+		}
+		else
+		{
+			GroupPtr = *GroupPtrPtr;
+		}
+
+		GroupPtr->AddChildAndSetGroupPtr(NodePtr);
+	}
+
+	if (GroupingDepth < CurrentGroupings.Num() - 1)
+	{
+		TArray<FTableTreeNodePtr> ChildNodes;
+
+		for (FBaseTreeNodePtr GroupPtr : ParentGroup.GetChildren())
+		{
+			ensure(GroupPtr->IsGroup());
+			FTableTreeNode& Group = *StaticCastSharedPtr<FTableTreeNode>(GroupPtr);
+
+			// Make a copy of the child nodes.
+			ChildNodes.Reset();
+			for (FBaseTreeNodePtr ChildPtr : Group.GetChildren())
+			{
+				ChildNodes.Add(StaticCastSharedPtr<FTableTreeNode>(ChildPtr));
+			}
+
+			GroupNodesRec(ChildNodes, Group, GroupingDepth + 1);
+		}
+	}
+}
+
+////////////////////////////////////////////////////////////////////////////////////////////////////
+
+void STableTreeView::ResetAggregatedValuesRec(FTableTreeNode& GroupNode)
+{
+	ensure(GroupNode.IsGroup());
+
+	GroupNode.ResetAggregatedValues();
+
+	for (FBaseTreeNodePtr ChildPtr : GroupNode.GetChildren())
+	{
+		if (ChildPtr->IsGroup())
+		{
+			ResetAggregatedValuesRec(*StaticCastSharedPtr<FTableTreeNode>(ChildPtr));
+		}
+	}
+}
+
+////////////////////////////////////////////////////////////////////////////////////////////////////
+
+void STableTreeView::UpdateInt64SumAggregationRec(FTableColumn& Column, FTableTreeNode& GroupNode)
+{
+	int64 AggregatedValue = 0;
+
+	for (FBaseTreeNodePtr NodePtr : GroupNode.GetChildren())
+	{
+		FTableTreeNode& TableNode = *StaticCastSharedPtr<FTableTreeNode>(NodePtr);
+		if (TableNode.IsGroup())
+		{
+			UpdateInt64SumAggregationRec(Column, TableNode);
+		}
+
+		const TOptional<FTableCellValue> OptionalValue = Column.GetValue(TableNode);
+		if (OptionalValue.IsSet())
+		{
+			ensure(OptionalValue.GetValue().DataType == ETableCellDataType::Int64);
+			AggregatedValue += OptionalValue.GetValue().Int64;
+		}
+	}
+
+	GroupNode.AddAggregatedValue(Column.GetId(), FTableCellValue(AggregatedValue));
+}
+
+////////////////////////////////////////////////////////////////////////////////////////////////////
+
+void STableTreeView::UpdateFloatSumAggregationRec(FTableColumn& Column, FTableTreeNode& GroupNode)
+{
+	float AggregatedValue = 0.0f;
+
+	for (FBaseTreeNodePtr NodePtr : GroupNode.GetChildren())
+	{
+		FTableTreeNode& TableNode = *StaticCastSharedPtr<FTableTreeNode>(NodePtr);
+		if (TableNode.IsGroup())
+		{
+			UpdateFloatSumAggregationRec(Column, TableNode);
+		}
+
+		const TOptional<FTableCellValue> OptionalValue = Column.GetValue(TableNode);
+		if (OptionalValue.IsSet())
+		{
+			ensure(OptionalValue.GetValue().DataType == ETableCellDataType::Float);
+			AggregatedValue += OptionalValue.GetValue().Float;
+		}
+	}
+
+	GroupNode.AddAggregatedValue(Column.GetId(), FTableCellValue(AggregatedValue));
+}
+
+////////////////////////////////////////////////////////////////////////////////////////////////////
+
+void STableTreeView::UpdateDoubleSumAggregationRec(FTableColumn& Column, FTableTreeNode& GroupNode)
+{
+	double AggregatedValue = 0.0;
+
+	for (FBaseTreeNodePtr NodePtr : GroupNode.GetChildren())
+	{
+		FTableTreeNode& TableNode = *StaticCastSharedPtr<FTableTreeNode>(NodePtr);
+		if (TableNode.IsGroup())
+		{
+			UpdateDoubleSumAggregationRec(Column, TableNode);
+		}
+
+		const TOptional<FTableCellValue> OptionalValue = Column.GetValue(TableNode);
+		if (OptionalValue.IsSet())
+		{
+			ensure(OptionalValue.GetValue().DataType == ETableCellDataType::Double);
+			AggregatedValue += OptionalValue.GetValue().Double;
+		}
+	}
+
+	GroupNode.AddAggregatedValue(Column.GetId(), FTableCellValue(AggregatedValue));
+}
+
+////////////////////////////////////////////////////////////////////////////////////////////////////
+
+void STableTreeView::CreateGroupings()
+{
+	AvailableGroupings.Reset(3);
+
+	AvailableGroupings.Add(MakeShareable(new FTreeNodeGroupingFlat()));
+	//AvailableGroupings.Add(MakeShareable(new FTreeNodeGroupingByNameFirstLetter()));
+	//AvailableGroupings.Add(MakeShareable(new FTreeNodeGroupingByType()));
+
+	for (const TSharedPtr<FTableColumn> ColumnPtr : Table->GetColumns())
+	{
+		if (!ColumnPtr->IsHierarchy())
+		{
+			AvailableGroupings.Add(MakeShareable(new FTreeNodeGroupingByUniqueValue(ColumnPtr.ToSharedRef())));
+		}
+	}
+
+	CurrentGroupings.Add(AvailableGroupings[0]);
+
+	RebuildGroupingCrumbs();
+}
+
+////////////////////////////////////////////////////////////////////////////////////////////////////
+
+void STableTreeView::PreChangeGroupings()
+{
+	for (TSharedPtr<FTreeNodeGrouping> GroupingPtr : CurrentGroupings)
+	{
+		const FName& ColumnId = GroupingPtr->GetColumnId();
+		if (ColumnId != NAME_None)
+		{
+			// Show columns used in previous groupings.
+			ShowColumn(ColumnId);
+		}
+	}
+}
+
+////////////////////////////////////////////////////////////////////////////////////////////////////
+
+void STableTreeView::PostChangeGroupings()
+{
+	constexpr float HierarchyMinWidth = 60.0f;
+	constexpr float HierarchyIndentation = 10.0f;
+	constexpr float DefaultHierarchyColumnWidth = 90.0f;
+
+	float HierarchyColumnWidth = DefaultHierarchyColumnWidth;
+	//FString GroupingStr;
+
+	int32 GroupingDepth = 0;
+	for (TSharedPtr<FTreeNodeGrouping> GroupingPtr : CurrentGroupings)
+	{
+		const FName& ColumnId = GroupingPtr->GetColumnId();
+
+		if (ColumnId != NAME_None)
+		{
+			// Compute width for Hierarchy column based on column used in grouping and its indentation.
+			const int32 NumColumns = TreeViewHeaderRow->GetColumns().Num();
+			for (int32 ColumnIndex = 0; ColumnIndex < NumColumns; ++ColumnIndex)
+			{
+				const SHeaderRow::FColumn& CurrentColumn = TreeViewHeaderRow->GetColumns()[ColumnIndex];
+				if (CurrentColumn.ColumnId == ColumnId)
+				{
+					const float Width = HierarchyMinWidth + GroupingDepth * HierarchyIndentation + CurrentColumn.GetWidth();
+					if (Width > HierarchyColumnWidth)
+					{
+						HierarchyColumnWidth = Width;
+					}
+					break;
+				}
+			}
+
+			// Hide columns used in groupings.
+			HideColumn(ColumnId);
+		}
+
+		// Compute name of the Hierarchy column.
+		//if (!GroupingStr.IsEmpty())
+		//{
+		//	GroupingStr.Append(TEXT(" / "));
+		//}
+		//GroupingStr.Append(GroupingPtr->GetShortName().ToString());
+
+		++GroupingDepth;
+	}
+
+	//////////////////////////////////////////////////
+
+	// Set with for the Hierarchy column.
+	SHeaderRow::FColumn& HierarchyColumn = const_cast<SHeaderRow::FColumn&>(TreeViewHeaderRow->GetColumns()[0]);
+	HierarchyColumn.SetWidth(HierarchyColumnWidth);
+
+	// Set name for the Hierarchy column.
+	//FTableColumn& HierarchyTableColumn = *Table->FindColumnChecked(HierarchyColumn.ColumnId);
+	//if (!GroupingStr.IsEmpty())
+	//{
+	//	const FText HierarchyColumnName = FText::Format(LOCTEXT("HierarchyShortNameFmt", "Hierarchy ({0})"), FText::FromString(GroupingStr));
+	//	HierarchyTableColumn.SetShortName(HierarchyColumnName);
+	//}
+	//else
+	//{
+	//	const FText HierarchyColumnName(LOCTEXT("HierarchyShortName", "Hierarchy"));
+	//	HierarchyTableColumn.SetShortName(HierarchyColumnName);
+	//}
+
+	//////////////////////////////////////////////////
+
+	TreeViewHeaderRow->RefreshColumns();
+
+	CreateGroups();
+	SortTreeNodes();
+	ApplyFiltering();
+
+	RebuildGroupingCrumbs();
+}
+
+////////////////////////////////////////////////////////////////////////////////////////////////////
+
+void STableTreeView::RebuildGroupingCrumbs()
+{
+	GroupingBreadcrumbTrail->ClearCrumbs();
+
+	for (const TSharedPtr<FTreeNodeGrouping> Grouping : CurrentGroupings)
+	{
+		GroupingBreadcrumbTrail->PushCrumb(Grouping->GetShortName(), Grouping);
+	}
+}
+
+////////////////////////////////////////////////////////////////////////////////////////////////////
+
+int32 STableTreeView::GetGroupingDepth(const TSharedPtr<FTreeNodeGrouping>& Grouping) const
+{
+	for (int32 GroupingDepth = CurrentGroupings.Num() - 1; GroupingDepth >= 0; --GroupingDepth)
+	{
+		if (Grouping == CurrentGroupings[GroupingDepth])
+		{
+			return GroupingDepth;
+		}
+	}
+	return -1;
+}
+
+////////////////////////////////////////////////////////////////////////////////////////////////////
+
+void STableTreeView::OnGroupingCrumbClicked(const TSharedPtr<FTreeNodeGrouping>& CrumbGrouping)
+{
+	const int32 CrumbGroupingDepth = GetGroupingDepth(CrumbGrouping);
+	if (CrumbGroupingDepth >= 0 && CrumbGroupingDepth < CurrentGroupings.Num() - 1)
+	{
+		PreChangeGroupings();
+
+		CurrentGroupings.RemoveAt(CrumbGroupingDepth + 1, CurrentGroupings.Num() - CrumbGroupingDepth - 1);
+
+		PostChangeGroupings();
+	}
+}
+
+////////////////////////////////////////////////////////////////////////////////////////////////////
+
+void STableTreeView::BuildGroupingSubMenu_Change(FMenuBuilder& MenuBuilder, const TSharedPtr<FTreeNodeGrouping> CrumbGrouping)
+{
+	MenuBuilder.BeginSection("ChangeGrouping");
+	{
+		for (const TSharedPtr<FTreeNodeGrouping>& Grouping : AvailableGroupings)
+		{
+			FUIAction Action_Change
+			(
+				FExecuteAction::CreateSP(this, &STableTreeView::GroupingCrumbMenu_Change_Execute, CrumbGrouping, Grouping),
+				FCanExecuteAction::CreateSP(this, &STableTreeView::GroupingCrumbMenu_Change_CanExecute, CrumbGrouping, Grouping)
+			);
+			MenuBuilder.AddMenuEntry
+			(
+				Grouping->GetTitleName(),
+				Grouping->GetDescription(),
+				FSlateIcon(), Action_Change, NAME_None, EUserInterfaceActionType::Button
+			);
+		}
+	}
+	MenuBuilder.EndSection();
+}
+
+////////////////////////////////////////////////////////////////////////////////////////////////////
+
+void STableTreeView::BuildGroupingSubMenu_Add(FMenuBuilder& MenuBuilder, const TSharedPtr<FTreeNodeGrouping> CrumbGrouping)
+{
+	MenuBuilder.BeginSection("AddGrouping");
+	{
+		for (const TSharedPtr<FTreeNodeGrouping>& Grouping : AvailableGroupings)
+		{
+			FUIAction Action_Add
+			(
+				FExecuteAction::CreateSP(this, &STableTreeView::GroupingCrumbMenu_Add_Execute, Grouping, CrumbGrouping),
+				FCanExecuteAction::CreateSP(this, &STableTreeView::GroupingCrumbMenu_Add_CanExecute, Grouping, CrumbGrouping)
+			);
+			MenuBuilder.AddMenuEntry
+			(
+				Grouping->GetTitleName(),
+				Grouping->GetDescription(),
+				FSlateIcon(), Action_Add, NAME_None, EUserInterfaceActionType::Button
+			);
+		}
+	}
+	MenuBuilder.EndSection();
+}
+
+////////////////////////////////////////////////////////////////////////////////////////////////////
+
+TSharedRef<SWidget> STableTreeView::GetGroupingCrumbMenuContent(const TSharedPtr<FTreeNodeGrouping>& CrumbGrouping)
+{
+	const bool bShouldCloseWindowAfterMenuSelection = true;
+	FMenuBuilder MenuBuilder(bShouldCloseWindowAfterMenuSelection, NULL);
+
+	const int32 CrumbGroupingDepth = GetGroupingDepth(CrumbGrouping);
+
+	MenuBuilder.BeginSection("InsertOrAdd");
+	{
+		const FText AddGroupingText = (CrumbGroupingDepth == CurrentGroupings.Num() - 1) ? // after last one
+			LOCTEXT("AddGrouping_Section", "Add Grouping...") :
+			LOCTEXT("InsertGrouping_Section", "Insert Grouping...");
+		MenuBuilder.AddSubMenu
+		(
+			AddGroupingText,
+			LOCTEXT("GroupingMenu_Add_Desc", "Add or insert new grouping."),
+			FNewMenuDelegate::CreateSP(this, &STableTreeView::BuildGroupingSubMenu_Add, CrumbGrouping),
+			false,
+			FSlateIcon()
+		);
+	}
+	MenuBuilder.EndSection();
+
+	if (CrumbGroupingDepth >= 0)
+	{
+		MenuBuilder.BeginSection("CrumbGrouping", CrumbGrouping->GetTitleName());
+		{
+			MenuBuilder.AddSubMenu
+			(
+				LOCTEXT("ChangeGrouping_Section", "Change To..."),
+				LOCTEXT("GroupingMenu_Add_Desc", "Change selected grouping."),
+				FNewMenuDelegate::CreateSP(this, &STableTreeView::BuildGroupingSubMenu_Change, CrumbGrouping),
+				false,
+				FSlateIcon()
+			);
+
+			if (CrumbGroupingDepth > 0)
+			{
+				FUIAction Action_MoveLeft
+				(
+					FExecuteAction::CreateSP(this, &STableTreeView::GroupingCrumbMenu_MoveLeft_Execute, CrumbGrouping),
+					FCanExecuteAction()
+				);
+				MenuBuilder.AddMenuEntry
+				(
+					LOCTEXT("GroupingMenu_MoveLeft", "Move Left"),
+					LOCTEXT("GroupingMenu_MoveLeft_Desc", "Move selected grouping to the left."),
+					FSlateIcon(), Action_MoveLeft, NAME_None, EUserInterfaceActionType::Button
+				);
+			}
+
+			if (CrumbGroupingDepth < CurrentGroupings.Num() - 1)
+			{
+				FUIAction Action_MoveRight
+				(
+					FExecuteAction::CreateSP(this, &STableTreeView::GroupingCrumbMenu_MoveRight_Execute, CrumbGrouping),
+					FCanExecuteAction()
+				);
+				MenuBuilder.AddMenuEntry
+				(
+					LOCTEXT("GroupingMenu_MoveRight", "Move Right"),
+					LOCTEXT("GroupingMenu_MoveRight_Desc", "Move selected grouping to the right."),
+					FSlateIcon(), Action_MoveRight, NAME_None, EUserInterfaceActionType::Button
+				);
+			}
+
+			if (CurrentGroupings.Num() > 1)
+			{
+				FUIAction Action_Remove
+				(
+					FExecuteAction::CreateSP(this, &STableTreeView::GroupingCrumbMenu_Remove_Execute, CrumbGrouping),
+					FCanExecuteAction()
+				);
+				MenuBuilder.AddMenuEntry
+				(
+					LOCTEXT("GroupingMenu_Remove", "Remove"),
+					LOCTEXT("GroupingMenu_Remove_Desc", "Remove selected grouping."),
+					FSlateIcon(), Action_Remove, NAME_None, EUserInterfaceActionType::Button
+				);
+			}
+		}
+		MenuBuilder.EndSection();
+	}
+
+	if (CurrentGroupings.Num() > 1 || CurrentGroupings[0] != AvailableGroupings[0])
+	{
+		MenuBuilder.BeginSection("ResetGroupings");
+		{
+			FUIAction Action_Reset
+			(
+				FExecuteAction::CreateSP(this, &STableTreeView::GroupingCrumbMenu_Reset_Execute),
+				FCanExecuteAction()
+			);
+			MenuBuilder.AddMenuEntry
+			(
+				LOCTEXT("GroupingMenu_Reset", "Reset"),
+				LOCTEXT("GroupingMenu_Reset_Desc", "Reset groupings to default."),
+				FSlateIcon(), Action_Reset, NAME_None, EUserInterfaceActionType::Button
+			);
+		}
+		MenuBuilder.EndSection();
+	}
+
+	return MenuBuilder.MakeWidget();
+}
+
+////////////////////////////////////////////////////////////////////////////////////////////////////
+
+void STableTreeView::GroupingCrumbMenu_Reset_Execute()
+{
+	PreChangeGroupings();
+
+	CurrentGroupings.Reset();
+	CurrentGroupings.Add(AvailableGroupings[0]);
+
+	PostChangeGroupings();
+}
+
+////////////////////////////////////////////////////////////////////////////////////////////////////
+
+void STableTreeView::GroupingCrumbMenu_Remove_Execute(const TSharedPtr<FTreeNodeGrouping> Grouping)
+{
+	const int32 GroupingDepth = GetGroupingDepth(Grouping);
+	if (GroupingDepth >= 0)
+	{
+		PreChangeGroupings();
+
+		CurrentGroupings.RemoveAt(GroupingDepth);
+
+		PostChangeGroupings();
+	}
+}
+
+////////////////////////////////////////////////////////////////////////////////////////////////////
+
+void STableTreeView::GroupingCrumbMenu_MoveLeft_Execute(const TSharedPtr<FTreeNodeGrouping> Grouping)
+{
+	const int32 GroupingDepth = GetGroupingDepth(Grouping);
+	if (GroupingDepth > 0)
+	{
+		PreChangeGroupings();
+
+		CurrentGroupings[GroupingDepth] = CurrentGroupings[GroupingDepth - 1];
+		CurrentGroupings[GroupingDepth - 1] = Grouping;
+
+		PostChangeGroupings();
+	}
+}
+
+////////////////////////////////////////////////////////////////////////////////////////////////////
+
+void STableTreeView::GroupingCrumbMenu_MoveRight_Execute(const TSharedPtr<FTreeNodeGrouping> Grouping)
+{
+	const int32 GroupingDepth = GetGroupingDepth(Grouping);
+	if (GroupingDepth < CurrentGroupings.Num() - 1)
+	{
+		PreChangeGroupings();
+
+		CurrentGroupings[GroupingDepth] = CurrentGroupings[GroupingDepth + 1];
+		CurrentGroupings[GroupingDepth + 1] = Grouping;
+
+		PostChangeGroupings();
+	}
+}
+
+////////////////////////////////////////////////////////////////////////////////////////////////////
+
+void STableTreeView::GroupingCrumbMenu_Change_Execute(const TSharedPtr<FTreeNodeGrouping> OldGrouping, const TSharedPtr<FTreeNodeGrouping> NewGrouping)
+{
+	const int32 OldGroupingDepth = GetGroupingDepth(OldGrouping);
+	if (OldGroupingDepth >= 0)
+	{
+		PreChangeGroupings();
+
+		const int32 NewGroupingDepth = GetGroupingDepth(NewGrouping);
+
+		if (NewGroupingDepth >= 0 && NewGroupingDepth != OldGroupingDepth) // NewGrouping already exists
+		{
+			CurrentGroupings.RemoveAt(NewGroupingDepth);
+
+			if (NewGroupingDepth < OldGroupingDepth)
+			{
+				CurrentGroupings[OldGroupingDepth - 1] = NewGrouping;
+			}
+			else
+			{
+				CurrentGroupings[OldGroupingDepth] = NewGrouping;
+			}
+		}
+		else
+		{
+			CurrentGroupings[OldGroupingDepth] = NewGrouping;
+		}
+
+		PostChangeGroupings();
+	}
+}
+
+////////////////////////////////////////////////////////////////////////////////////////////////////
+
+bool STableTreeView::GroupingCrumbMenu_Change_CanExecute(const TSharedPtr<FTreeNodeGrouping> OldGrouping, const TSharedPtr<FTreeNodeGrouping> NewGrouping) const
+{
+	return NewGrouping != OldGrouping;
+}
+
+////////////////////////////////////////////////////////////////////////////////////////////////////
+
+void STableTreeView::GroupingCrumbMenu_Add_Execute(const TSharedPtr<FTreeNodeGrouping> Grouping, const TSharedPtr<FTreeNodeGrouping> AfterGrouping)
+{
+	PreChangeGroupings();
+
+	if (AfterGrouping.IsValid())
+	{
+		const int32 AfterGroupingDepth = GetGroupingDepth(AfterGrouping);
+		ensure(AfterGroupingDepth >= 0);
+
+		const int32 GroupingDepth = GetGroupingDepth(Grouping);
+
+		if (GroupingDepth >= 0) // Grouping already exists
+		{
+			CurrentGroupings.RemoveAt(GroupingDepth);
+
+			if (GroupingDepth <= AfterGroupingDepth)
+			{
+				CurrentGroupings.Insert(Grouping, AfterGroupingDepth);
+			}
+			else
+			{
+				CurrentGroupings.Insert(Grouping, AfterGroupingDepth + 1);
+			}
+		}
+		else
+		{
+			CurrentGroupings.Insert(Grouping, AfterGroupingDepth + 1);
+		}
+	}
+	else
+	{
+		CurrentGroupings.Remove(Grouping);
+		CurrentGroupings.Add(Grouping);
+	}
+
+	PostChangeGroupings();
+}
+
+////////////////////////////////////////////////////////////////////////////////////////////////////
+
+bool STableTreeView::GroupingCrumbMenu_Add_CanExecute(const TSharedPtr<FTreeNodeGrouping> Grouping, const TSharedPtr<FTreeNodeGrouping> AfterGrouping) const
+{
+	if (AfterGrouping.IsValid())
+	{
+		const int32 AfterGroupingDepth = GetGroupingDepth(AfterGrouping);
+		ensure(AfterGroupingDepth >= 0);
+
+		const int32 GroupingDepth = GetGroupingDepth(Grouping);
+
+		return GroupingDepth < AfterGroupingDepth || GroupingDepth > AfterGroupingDepth + 1;
+	}
+	else
+	{
+		return Grouping != CurrentGroupings.Last();
+	}
+}
+
+////////////////////////////////////////////////////////////////////////////////////////////////////
+// Sorting
+////////////////////////////////////////////////////////////////////////////////////////////////////
+
+////////////////////////////////////////////////////////////////////////////////////////////////////
+
+const EColumnSortMode::Type STableTreeView::GetDefaultColumnSortMode()
+{
+	return EColumnSortMode::Descending;
+}
+
+////////////////////////////////////////////////////////////////////////////////////////////////////
+
+const FName STableTreeView::GetDefaultColumnBeingSorted()
+{
+	return NAME_None;
+}
+
+////////////////////////////////////////////////////////////////////////////////////////////////////
+
+void STableTreeView::CreateSortings()
+{
+	AvailableSorters.Reset();
+	CurrentSorter = nullptr;
+
+	for (const TSharedPtr<FTableColumn> ColumnPtr : Table->GetColumns())
+	{
+		if (ColumnPtr->CanBeSorted())
+		{
+			TSharedPtr<Insights::ITableCellValueSorter> SorterPtr = ColumnPtr->GetValueSorter();
+			if (ensure(SorterPtr.IsValid()))
+			{
+				AvailableSorters.Add(SorterPtr);
+			}
+		}
+	}
+
+	UpdateCurrentSortingByColumn();
+}
+
+////////////////////////////////////////////////////////////////////////////////////////////////////
+
+void STableTreeView::UpdateCurrentSortingByColumn()
+{
+	TSharedPtr<Insights::FTableColumn> ColumnPtr = Table->FindColumn(ColumnBeingSorted);
+	CurrentSorter = ColumnPtr.IsValid() ? ColumnPtr->GetValueSorter() : nullptr;
+}
+
+////////////////////////////////////////////////////////////////////////////////////////////////////
+
+void STableTreeView::SortTreeNodes()
+{
+	if (CurrentSorter.IsValid())
+	{
+		SortTreeNodesRec(*Root, *CurrentSorter);
+	}
+}
+
+////////////////////////////////////////////////////////////////////////////////////////////////////
+
+void STableTreeView::SortTreeNodesRec(FTableTreeNode& GroupNode, const ITableCellValueSorter& Sorter)
+{
+	if (ColumnSortMode == EColumnSortMode::Type::Descending)
+	{
+		GroupNode.SortChildrenDescending(Sorter);
+	}
+	else // if (ColumnSortMode == EColumnSortMode::Type::Ascending)
+	{
+		GroupNode.SortChildrenAscending(Sorter);
+	}
+
+	for (FBaseTreeNodePtr ChildPtr : GroupNode.GetChildren())
+	{
+		if (ChildPtr->IsGroup())
+		{
+			SortTreeNodesRec(*StaticCastSharedPtr<FTableTreeNode>(ChildPtr), Sorter);
+		}
+	}
+}
+
+////////////////////////////////////////////////////////////////////////////////////////////////////
+
+EColumnSortMode::Type STableTreeView::GetSortModeForColumn(const FName ColumnId) const
+{
+	if (ColumnBeingSorted != ColumnId)
+	{
+		return EColumnSortMode::None;
+	}
+
+	return ColumnSortMode;
+}
+
+////////////////////////////////////////////////////////////////////////////////////////////////////
+
+void STableTreeView::SetSortModeForColumn(const FName& ColumnId, const EColumnSortMode::Type SortMode)
+{
+	ColumnBeingSorted = ColumnId;
+	ColumnSortMode = SortMode;
+	UpdateCurrentSortingByColumn();
+
+	SortTreeNodes();
+	ApplyFiltering();
+}
+
+////////////////////////////////////////////////////////////////////////////////////////////////////
+
+void STableTreeView::OnSortModeChanged(const EColumnSortPriority::Type SortPriority, const FName& ColumnId, const EColumnSortMode::Type SortMode)
+{
+	SetSortModeForColumn(ColumnId, SortMode);
+	TreeView_Refresh();
+}
+
+////////////////////////////////////////////////////////////////////////////////////////////////////
+// SortMode action (HeaderMenu)
+////////////////////////////////////////////////////////////////////////////////////////////////////
+
+bool STableTreeView::HeaderMenu_SortMode_IsChecked(const FName ColumnId, const EColumnSortMode::Type InSortMode)
+{
+	return ColumnBeingSorted == ColumnId && ColumnSortMode == InSortMode;
+}
+
+////////////////////////////////////////////////////////////////////////////////////////////////////
+
+bool STableTreeView::HeaderMenu_SortMode_CanExecute(const FName ColumnId, const EColumnSortMode::Type InSortMode) const
+{
+	const FTableColumn& Column = *Table->FindColumnChecked(ColumnId);
+	return Column.CanBeSorted();
+}
+
+////////////////////////////////////////////////////////////////////////////////////////////////////
+
+void STableTreeView::HeaderMenu_SortMode_Execute(const FName ColumnId, const EColumnSortMode::Type InSortMode)
+{
+	SetSortModeForColumn(ColumnId, InSortMode);
+	TreeView_Refresh();
+}
+
+////////////////////////////////////////////////////////////////////////////////////////////////////
+// SortMode action (ContextMenu)
+////////////////////////////////////////////////////////////////////////////////////////////////////
+
+bool STableTreeView::ContextMenu_SortMode_IsChecked(const EColumnSortMode::Type InSortMode)
+{
+	return ColumnSortMode == InSortMode;
+}
+
+////////////////////////////////////////////////////////////////////////////////////////////////////
+
+bool STableTreeView::ContextMenu_SortMode_CanExecute(const EColumnSortMode::Type InSortMode) const
+{
+	return true; //ColumnSortMode != InSortMode;
+}
+
+////////////////////////////////////////////////////////////////////////////////////////////////////
+
+void STableTreeView::ContextMenu_SortMode_Execute(const EColumnSortMode::Type InSortMode)
+{
+	SetSortModeForColumn(ColumnBeingSorted, InSortMode);
+	TreeView_Refresh();
+}
+
+////////////////////////////////////////////////////////////////////////////////////////////////////
+// SortByColumn action (ContextMenu)
+////////////////////////////////////////////////////////////////////////////////////////////////////
+
+bool STableTreeView::ContextMenu_SortByColumn_IsChecked(const FName ColumnId)
+{
+	return ColumnId == ColumnBeingSorted;
+}
+
+////////////////////////////////////////////////////////////////////////////////////////////////////
+
+bool STableTreeView::ContextMenu_SortByColumn_CanExecute(const FName ColumnId) const
+{
+	return true; //ColumnId != ColumnBeingSorted;
+}
+
+////////////////////////////////////////////////////////////////////////////////////////////////////
+
+void STableTreeView::ContextMenu_SortByColumn_Execute(const FName ColumnId)
+{
+	SetSortModeForColumn(ColumnId, EColumnSortMode::Descending);
+	TreeView_Refresh();
+}
+
+////////////////////////////////////////////////////////////////////////////////////////////////////
+// ShowColumn action
+////////////////////////////////////////////////////////////////////////////////////////////////////
+
+bool STableTreeView::CanShowColumn(const FName ColumnId) const
+{
+	return true;
+}
+
+////////////////////////////////////////////////////////////////////////////////////////////////////
+
+void STableTreeView::ShowColumn(const FName ColumnId)
+{
+	FTableColumn& Column = *Table->FindColumnChecked(ColumnId);
+	Column.Show();
+
+	SHeaderRow::FColumn::FArguments ColumnArgs;
+	ColumnArgs
+		.ColumnId(Column.GetId())
+		.DefaultLabel(Column.GetShortName())
+		.HAlignHeader(HAlign_Fill)
+		.VAlignHeader(VAlign_Fill)
+		.HeaderContentPadding(FMargin(2.0f))
+		.HAlignCell(HAlign_Fill)
+		.VAlignCell(VAlign_Fill)
+		.SortMode(this, &STableTreeView::GetSortModeForColumn, Column.GetId())
+		.OnSort(this, &STableTreeView::OnSortModeChanged)
+		.ManualWidth(Column.GetInitialWidth())
+		.FixedWidth(Column.IsFixedWidth() ? Column.GetInitialWidth() : TOptional<float>())
+		.HeaderContent()
+		[
+			SNew(SBox)
+			.ToolTip(STableTreeViewTooltip::GetColumnTooltip(Column))
+			.HAlign(Column.GetHorizontalAlignment())
+			.VAlign(VAlign_Center)
+			[
+				SNew(STextBlock)
+				.Text(this, &STableTreeView::GetColumnHeaderText, Column.GetId())
+			]
+		]
+		.MenuContent()
+		[
+			TreeViewHeaderRow_GenerateColumnMenu(Column)
+		];
+
+	int32 ColumnIndex = 0;
+	const int32 NewColumnPosition = Table->GetColumnPositionIndex(ColumnId);
+	const int32 NumColumns = TreeViewHeaderRow->GetColumns().Num();
+	for (; ColumnIndex < NumColumns; ColumnIndex++)
+	{
+		const SHeaderRow::FColumn& CurrentColumn = TreeViewHeaderRow->GetColumns()[ColumnIndex];
+		const int32 CurrentColumnPosition = Table->GetColumnPositionIndex(CurrentColumn.ColumnId);
+		if (NewColumnPosition < CurrentColumnPosition)
+		{
+			break;
+		}
+	}
+
+	TreeViewHeaderRow->InsertColumn(ColumnArgs, ColumnIndex);
+}
+
+////////////////////////////////////////////////////////////////////////////////////////////////////
+// HideColumn action
+////////////////////////////////////////////////////////////////////////////////////////////////////
+
+bool STableTreeView::CanHideColumn(const FName ColumnId) const
+{
+	const FTableColumn& Column = *Table->FindColumnChecked(ColumnId);
+	return Column.CanBeHidden();
+}
+
+////////////////////////////////////////////////////////////////////////////////////////////////////
+
+void STableTreeView::HideColumn(const FName ColumnId)
+{
+	FTableColumn& Column = *Table->FindColumnChecked(ColumnId);
+	Column.Hide();
+
+	TreeViewHeaderRow->RemoveColumn(ColumnId);
+}
+
+////////////////////////////////////////////////////////////////////////////////////////////////////
+// ToggleColumn action
+////////////////////////////////////////////////////////////////////////////////////////////////////
+
+bool STableTreeView::IsColumnVisible(const FName ColumnId)
+{
+	const FTableColumn& Column = *Table->FindColumnChecked(ColumnId);
+	return Column.IsVisible();
+}
+
+////////////////////////////////////////////////////////////////////////////////////////////////////
+
+bool STableTreeView::CanToggleColumnVisibility(const FName ColumnId) const
+{
+	const FTableColumn& Column = *Table->FindColumnChecked(ColumnId);
+	return !Column.IsVisible() || Column.CanBeHidden();
+}
+
+////////////////////////////////////////////////////////////////////////////////////////////////////
+
+void STableTreeView::ToggleColumnVisibility(const FName ColumnId)
+{
+	const FTableColumn& Column = *Table->FindColumnChecked(ColumnId);
+	if (Column.IsVisible())
+	{
+		HideColumn(ColumnId);
+	}
+	else
+	{
+		ShowColumn(ColumnId);
+	}
+}
+
+////////////////////////////////////////////////////////////////////////////////////////////////////
+// "Show All Columns" action (ContextMenu)
+////////////////////////////////////////////////////////////////////////////////////////////////////
+
+bool STableTreeView::ContextMenu_ShowAllColumns_CanExecute() const
+{
+	return true;
+}
+
+////////////////////////////////////////////////////////////////////////////////////////////////////
+
+void STableTreeView::ContextMenu_ShowAllColumns_Execute()
+{
+	ColumnBeingSorted = GetDefaultColumnBeingSorted();
+	ColumnSortMode = GetDefaultColumnSortMode();
+	UpdateCurrentSortingByColumn();
+
+	for (const TSharedPtr<FTableColumn>& ColumnPtr : Table->GetColumns())
+	{
+		const FTableColumn& Column = *ColumnPtr;
+
+		if (!Column.IsVisible())
+		{
+			ShowColumn(Column.GetId());
+		}
+	}
+}
+
+////////////////////////////////////////////////////////////////////////////////////////////////////
+// ResetColumns action (ContextMenu)
+////////////////////////////////////////////////////////////////////////////////////////////////////
+
+bool STableTreeView::ContextMenu_ResetColumns_CanExecute() const
+{
+	return true;
+}
+
+////////////////////////////////////////////////////////////////////////////////////////////////////
+
+void STableTreeView::ContextMenu_ResetColumns_Execute()
+{
+	ColumnBeingSorted = GetDefaultColumnBeingSorted();
+	ColumnSortMode = GetDefaultColumnSortMode();
+	UpdateCurrentSortingByColumn();
+
+	for (const TSharedPtr<FTableColumn>& ColumnPtr : Table->GetColumns())
+	{
+		const FTableColumn& Column = *ColumnPtr;
+
+		if (Column.ShouldBeVisible() && !Column.IsVisible())
+		{
+			ShowColumn(Column.GetId());
+		}
+		else if (!Column.ShouldBeVisible() && Column.IsVisible())
+		{
+			HideColumn(Column.GetId());
+		}
+	}
+}
+
+////////////////////////////////////////////////////////////////////////////////////////////////////
+
+void STableTreeView::Reset()
+{
+	StatsStartTime = 0.0;
+	StatsEndTime = 0.0;
+
+	RebuildTree(true);
+}
+
+////////////////////////////////////////////////////////////////////////////////////////////////////
+
+void STableTreeView::RebuildTree(bool bResync)
+{
+	bool bListHasChanged = false;
+
+	if (bResync)
+	{
+		const int32 PreviousNodeCount = TableTreeNodes.Num();
+		TableTreeNodes.Empty(PreviousNodeCount);
+		TableTreeNodesIdMap.Empty(PreviousNodeCount);
+		bListHasChanged = true;
+	}
+
+	TSharedPtr<Trace::IUntypedTable> SourceTable = Table->GetSourceTable();
+	TSharedPtr<Trace::IUntypedTableReader> TableReader = Table->GetTableReader();
+
+	if (Session.IsValid() && SourceTable.IsValid() && TableReader.IsValid())
+	{
+		int32 TotalRowCount = SourceTable->GetRowCount();
+
+		if (TotalRowCount != TableTreeNodes.Num())
+		{
+			bResync = true;
+		}
+
+		if (bResync)
+		{
+			const int32 PreviousNodeCount = TableTreeNodes.Num();
+			TableTreeNodes.Empty(PreviousNodeCount);
+			TableTreeNodesIdMap.Empty(PreviousNodeCount);
+			bListHasChanged = true;
+
+			for (int32 RowIndex = 0; RowIndex < TotalRowCount; ++RowIndex)
+			{
+				TableReader->SetRowIndex(RowIndex);
+				uint64 NodeId = static_cast<uint64>(RowIndex);
+				FName NodeName(*FString::Printf(TEXT("row %d"), RowIndex));
+				FTableTreeNodePtr NodePtr = MakeShareable(new FTableTreeNode(NodeId, NodeName, Table, RowIndex));
+				TableTreeNodes.Add(NodePtr);
+				TableTreeNodesIdMap.Add(NodeId, NodePtr);
+			}
+		}
+	}
+
+	if (bListHasChanged)
+	{
+		UpdateTree();
+
+		TreeView->RebuildList();
+	}
+}
+
+////////////////////////////////////////////////////////////////////////////////////////////////////
+
+void STableTreeView::SelectNodeByNodeId(uint64 Id)
+{
+	FTableTreeNodePtr* NodePtrPtr = TableTreeNodesIdMap.Find(Id);
+	if (NodePtrPtr != nullptr)
+	{
+		FTableTreeNodePtr NodePtr = *NodePtrPtr;
+
+		TreeView->SetSelection(NodePtr);
+		TreeView->RequestScrollIntoView(NodePtr);
+	}
+}
+
+////////////////////////////////////////////////////////////////////////////////////////////////////
+
+void STableTreeView::SelectNodeByTableRowIndex(int32 RowIndex)
+{
+	if (RowIndex >= 0 && RowIndex < TableTreeNodes.Num())
+	{
+		FTableTreeNodePtr NodePtr = TableTreeNodes[RowIndex];
+
+		TreeView->SetSelection(NodePtr);
+		TreeView->RequestScrollIntoView(NodePtr);
+	}
+}
+
+////////////////////////////////////////////////////////////////////////////////////////////////////
+
+} // namespace Insights
+
+#undef LOCTEXT_NAMESPACE
diff --git a/Engine/Source/Developer/TraceInsights/Private/Insights/Table/Widgets/STableTreeView.h b/Engine/Source/Developer/TraceInsights/Private/Insights/Table/Widgets/STableTreeView.h
new file mode 100644
--- /dev/null
+++ b/Engine/Source/Developer/TraceInsights/Private/Insights/Table/Widgets/STableTreeView.h
@@ -0,0 +1,347 @@
+// Copyright 1998-2019 Epic Games, Inc. All Rights Reserved.
+
+#pragma once
+
+#include "CoreMinimal.h"
+#include "Misc/FilterCollection.h"
+#include "Misc/TextFilter.h"
+#include "SlateFwd.h"
+#include "Widgets/DeclarativeSyntaxSupport.h"
+#include "Widgets/Navigation/SBreadcrumbTrail.h"
+#include "Widgets/SCompoundWidget.h"
+#include "Widgets/SWidget.h"
+#include "Widgets/Views/SHeaderRow.h"
+#include "Widgets/Views/STableRow.h"
+#include "Widgets/Views/STableViewBase.h"
+#include "Widgets/Views/STreeView.h"
+
+// Insights
+#include "Insights/Table/ViewModels/TableTreeNode.h"
+
+class FMenuBuilder;
+
+namespace Trace
+{
+	class IAnalysisSession;
+	struct FTimelineEvent;
+}
+
+namespace Insights
+{
+
+class FTable;
+class FTableColumn;
+class FTreeNodeGrouping;
+class ITableCellValueSorter;
+
+////////////////////////////////////////////////////////////////////////////////////////////////////
+
+/** The filter collection - used for updating the list of tree nodes. */
+typedef TFilterCollection<const FTableTreeNodePtr&> FTableTreeNodeFilterCollection;
+
+/** The text based filter - used for updating the list of tree nodes. */
+typedef TTextFilter<const FTableTreeNodePtr&> FTableTreeNodeTextFilter;
+
+////////////////////////////////////////////////////////////////////////////////////////////////////
+/**
+ * A custom widget used to display the list of tree nodes.
+ */
+class STableTreeView : public SCompoundWidget
+{
+public:
+	/** Default constructor. */
+	STableTreeView();
+
+	/** Virtual destructor. */
+	virtual ~STableTreeView();
+
+	SLATE_BEGIN_ARGS(STableTreeView) {}
+	SLATE_END_ARGS()
+
+	/**
+	 * Construct this widget
+	 * @param InArgs - The declaration data for this widget
+	 */
+	void Construct(const FArguments& InArgs, TSharedPtr<Insights::FTable> InTablePtr);
+
+	TSharedPtr<Insights::FTable> GetTable() const { return Table; }
+
+	virtual void Reset();
+
+	/**
+	 * Rebuilds the tree (if necessary).
+	 * @param bResync - If true, it forces a resync even if the list did not changed since last sync.
+	 */
+	virtual void RebuildTree(bool bResync);
+
+	const FTableTreeNodePtr* FindNode(uint64 Id) const { return TableTreeNodesIdMap.Find(Id); }
+
+	void SelectNodeByNodeId(uint64 Id);
+	void SelectNodeByTableRowIndex(int32 RowIndex);
+
+protected:
+	void UpdateTree();
+
+	/** Called when the analysis session has changed. */
+	void InsightsManager_OnSessionChanged();
+
+	/**
+	 * Populates OutSearchStrings with the strings that should be used in searching.
+	 *
+	 * @param GroupOrStatNodePtr - the group and stat node to get a text description from.
+	 * @param OutSearchStrings   - an array of strings to use in searching.
+	 *
+	 */
+	void HandleItemToStringArray(const FTableTreeNodePtr& GroupOrStatNodePtr, TArray<FString>& OutSearchStrings) const;
+
+	////////////////////////////////////////////////////////////////////////////////////////////////////
+	// Tree View - Context Menu
+
+	TSharedPtr<SWidget> TreeView_GetMenuContent();
+	void TreeView_BuildSortByMenu(FMenuBuilder& MenuBuilder);
+	void TreeView_BuildViewColumnMenu(FMenuBuilder& MenuBuilder);
+
+	////////////////////////////////////////////////////////////////////////////////////////////////////
+	// Tree View - Columns' Header
+
+	void InitializeAndShowHeaderColumns();
+
+	FText GetColumnHeaderText(const FName ColumnId) const;
+
+	TSharedRef<SWidget> TreeViewHeaderRow_GenerateColumnMenu(const FTableColumn& Column);
+
+	////////////////////////////////////////////////////////////////////////////////////////////////////
+	// Tree View - Misc
+
+	void TreeView_Refresh();
+
+	/**
+	 * Called by STreeView to retrieves the children for the specified parent item.
+	 * @param InParent    - The parent node to retrieve the children from.
+	 * @param OutChildren - List of children for the parent node.
+	 */
+	void TreeView_OnGetChildren(FTableTreeNodePtr InParent, TArray<FTableTreeNodePtr>& OutChildren);
+
+	/** Called by STreeView when selection has changed. */
+	void TreeView_OnSelectionChanged(FTableTreeNodePtr SelectedItem, ESelectInfo::Type SelectInfo);
+
+	/** Called by STreeView when a tree item is double clicked. */
+	void TreeView_OnMouseButtonDoubleClick(FTableTreeNodePtr TreeNode);
+
+	////////////////////////////////////////////////////////////////////////////////////////////////////
+	// Tree View - Table Row
+
+	/** Called by STreeView to generate a table row for the specified item. */
+	TSharedRef<ITableRow> TreeView_OnGenerateRow(FTableTreeNodePtr TreeNode, const TSharedRef<STableViewBase>& OwnerTable);
+
+	void TableRow_SetHoveredCell(TSharedPtr<FTable> TablePtr, TSharedPtr<FTableColumn> ColumnPtr, const FTableTreeNodePtr NodePtr);
+	EHorizontalAlignment TableRow_GetColumnOutlineHAlignment(const FName ColumnId) const;
+
+	FText TableRow_GetHighlightText() const;
+	FName TableRow_GetHighlightedNodeName() const;
+
+	bool TableRow_ShouldBeEnabled(const uint32 TreeNodeId) const;
+
+	////////////////////////////////////////////////////////////////////////////////////////////////////
+	// Filtering
+
+	/** Populates the group and stat tree with items based on the current data. */
+	void ApplyFiltering();
+
+	bool ApplyFilteringForNode(FTableTreeNodePtr NodePtr);
+
+	bool SearchBox_IsEnabled() const;
+	void SearchBox_OnTextChanged(const FText& InFilterText);
+
+	////////////////////////////////////////////////////////////////////////////////////////////////////
+	// Grouping
+
+	void CreateGroupings();
+
+	void CreateGroups();
+	void GroupNodesRec(const TArray<FTableTreeNodePtr>& Nodes, FTableTreeNode& ParentGroup, int32 GroupingDepth);
+
+	void ResetAggregatedValuesRec(FTableTreeNode& GroupNode);
+	void UpdateInt64SumAggregationRec(FTableColumn& Column, FTableTreeNode& GroupNode);
+	void UpdateFloatSumAggregationRec(FTableColumn& Column, FTableTreeNode& GroupNode);
+	void UpdateDoubleSumAggregationRec(FTableColumn& Column, FTableTreeNode& GroupNode);
+
+	void RebuildGroupingCrumbs();
+	void OnGroupingCrumbClicked(const TSharedPtr<FTreeNodeGrouping>& InEntry);
+	void BuildGroupingSubMenu_Change(FMenuBuilder& MenuBuilder, const TSharedPtr<FTreeNodeGrouping> CrumbGrouping);
+	void BuildGroupingSubMenu_Add(FMenuBuilder& MenuBuilder, const TSharedPtr<FTreeNodeGrouping> CrumbGrouping);
+	TSharedRef<SWidget> GetGroupingCrumbMenuContent(const TSharedPtr<FTreeNodeGrouping>& CrumbGrouping);
+
+	void PreChangeGroupings();
+	void PostChangeGroupings();
+	int32 GetGroupingDepth(const TSharedPtr<FTreeNodeGrouping>& Grouping) const;
+
+	void GroupingCrumbMenu_Reset_Execute();
+	void GroupingCrumbMenu_Remove_Execute(const TSharedPtr<FTreeNodeGrouping> Grouping);
+	void GroupingCrumbMenu_MoveLeft_Execute(const TSharedPtr<FTreeNodeGrouping> Grouping);
+	void GroupingCrumbMenu_MoveRight_Execute(const TSharedPtr<FTreeNodeGrouping> Grouping);
+	void GroupingCrumbMenu_Change_Execute(const TSharedPtr<FTreeNodeGrouping> OldGrouping, const TSharedPtr<FTreeNodeGrouping> NewGrouping);
+	bool GroupingCrumbMenu_Change_CanExecute(const TSharedPtr<FTreeNodeGrouping> OldGrouping, const TSharedPtr<FTreeNodeGrouping> NewGrouping) const;
+	void GroupingCrumbMenu_Add_Execute(const TSharedPtr<FTreeNodeGrouping> Grouping, const TSharedPtr<FTreeNodeGrouping> AfterGrouping);
+	bool GroupingCrumbMenu_Add_CanExecute(const TSharedPtr<FTreeNodeGrouping> Grouping, const TSharedPtr<FTreeNodeGrouping> AfterGrouping) const;
+
+	////////////////////////////////////////////////////////////////////////////////////////////////////
+	// Sorting
+
+	static const EColumnSortMode::Type GetDefaultColumnSortMode();
+	static const FName GetDefaultColumnBeingSorted();
+
+	void CreateSortings();
+
+	void UpdateCurrentSortingByColumn();
+	void SortTreeNodes();
+	void SortTreeNodesRec(FTableTreeNode& GroupNode, const ITableCellValueSorter& Sorter);
+
+	EColumnSortMode::Type GetSortModeForColumn(const FName ColumnId) const;
+	void SetSortModeForColumn(const FName& ColumnId, EColumnSortMode::Type SortMode);
+	void OnSortModeChanged(const EColumnSortPriority::Type SortPriority, const FName& ColumnId, const EColumnSortMode::Type SortMode);
+
+	////////////////////////////////////////////////////////////////////////////////////////////////////
+	// Sorting actions
+
+	// SortMode (HeaderMenu)
+	bool HeaderMenu_SortMode_IsChecked(const FName ColumnId, const EColumnSortMode::Type InSortMode);
+	bool HeaderMenu_SortMode_CanExecute(const FName ColumnId, const EColumnSortMode::Type InSortMode) const;
+	void HeaderMenu_SortMode_Execute(const FName ColumnId, const EColumnSortMode::Type InSortMode);
+
+	// SortMode (ContextMenu)
+	bool ContextMenu_SortMode_IsChecked(const EColumnSortMode::Type InSortMode);
+	bool ContextMenu_SortMode_CanExecute(const EColumnSortMode::Type InSortMode) const;
+	void ContextMenu_SortMode_Execute(const EColumnSortMode::Type InSortMode);
+
+	// SortByColumn (ContextMenu)
+	bool ContextMenu_SortByColumn_IsChecked(const FName ColumnId);
+	bool ContextMenu_SortByColumn_CanExecute(const FName ColumnId) const;
+	void ContextMenu_SortByColumn_Execute(const FName ColumnId);
+
+	////////////////////////////////////////////////////////////////////////////////////////////////////
+	// Column visibility actions
+
+	// ShowColumn
+	bool CanShowColumn(const FName ColumnId) const;
+	void ShowColumn(const FName ColumnId);
+
+	// HideColumn
+	bool CanHideColumn(const FName ColumnId) const;
+	void HideColumn(const FName ColumnId);
+
+	// ToggleColumnVisibility
+	bool IsColumnVisible(const FName ColumnId);
+	bool CanToggleColumnVisibility(const FName ColumnId) const;
+	void ToggleColumnVisibility(const FName ColumnId);
+
+	// ShowAllColumns (ContextMenu)
+	bool ContextMenu_ShowAllColumns_CanExecute() const;
+	void ContextMenu_ShowAllColumns_Execute();
+
+	// ResetColumns (ContextMenu)
+	bool ContextMenu_ResetColumns_CanExecute() const;
+	void ContextMenu_ResetColumns_Execute();
+
+	////////////////////////////////////////////////////////////////////////////////////////////////////
+
+protected:
+	/** Table view model. */
+	TSharedPtr<Insights::FTable> Table;
+
+	/** A weak pointer to the profiler session used to populate this widget. */
+	TSharedPtr<const Trace::IAnalysisSession>/*Weak*/ Session;
+
+	//////////////////////////////////////////////////
+	// Tree View, Columns
+
+	/** The child STreeView widget. */
+	TSharedPtr<STreeView<FTableTreeNodePtr>> TreeView;
+
+	/** Holds the tree view header row widget which display all columns in the tree view. */
+	TSharedPtr<SHeaderRow> TreeViewHeaderRow;
+
+	/** External scrollbar used to synchronize tree view position. */
+	TSharedPtr<SScrollBar> ExternalScrollbar;
+
+	//////////////////////////////////////////////////
+	// Hovered Column, Hovered Tree Node
+
+	/** Name of the column currently being hovered by the mouse. */
+	FName HoveredColumnId;
+
+	/** A shared pointer to the tree node currently being hovered by the mouse. */
+	FTableTreeNodePtr HoveredNodePtr;
+
+	/** Name of the tree node that should be drawn as highlighted. */
+	FName HighlightedNodeName;
+
+	//////////////////////////////////////////////////
+	// Tree Nodes
+
+	static const FName RootNodeName;
+
+	/** The root node of the tree. */
+	FTableTreeNodePtr Root;
+
+	/** Table (row) nodes. Each node corresponds to a table row. */
+	TArray<FTableTreeNodePtr> TableTreeNodes;
+
+	/** A filtered array of group and nodes to be displayed in the tree widget. */
+	TArray<FTableTreeNodePtr> FilteredGroupNodes;
+
+	/** All nodes, stored as Id -> FTableTreeNodePtr. */
+	TMap<uint64, FTableTreeNodePtr> TableTreeNodesIdMap;
+
+	/** Currently expanded group nodes. */
+	TSet<FTableTreeNodePtr> ExpandedNodes;
+
+	/** If true, the expanded nodes have been saved before applying a text filter. */
+	bool bExpansionSaved;
+
+	//////////////////////////////////////////////////
+	// Search box and filters
+
+	/** The search box widget used to filter items displayed in the stats and groups tree. */
+	TSharedPtr<SSearchBox> SearchBox;
+
+	/** The text based filter. */
+	TSharedPtr<FTableTreeNodeTextFilter> TextFilter;
+
+	/** The filter collection. */
+	TSharedPtr<FTableTreeNodeFilterCollection> Filters;
+
+	//////////////////////////////////////////////////
+	// Grouping
+
+	TArray<TSharedPtr<FTreeNodeGrouping>> AvailableGroupings;
+
+	/** How we group the tree nodes? */
+	TArray<TSharedPtr<FTreeNodeGrouping>> CurrentGroupings;
+
+	TSharedPtr<SBreadcrumbTrail<TSharedPtr<FTreeNodeGrouping>>> GroupingBreadcrumbTrail;
+
+	//////////////////////////////////////////////////
+	// Sorting
+
+	/** All available sorters. */
+	TArray<TSharedPtr<ITableCellValueSorter>> AvailableSorters;
+
+	/** Current sorter. It is nullptr if sorting is disabled. */
+	TSharedPtr<ITableCellValueSorter> CurrentSorter;
+
+	/** Name of the column currently being sorted. Can be NAME_None if sorting is disabled (CurrentSorting == nullptr) or if a complex sorting is used (CurrentSorting != nullptr). */
+	FName ColumnBeingSorted;
+
+	/** How we sort the nodes? Ascending or Descending. */
+	EColumnSortMode::Type ColumnSortMode;
+
+	//////////////////////////////////////////////////
+
+	double StatsStartTime;
+	double StatsEndTime;
+};
+
+////////////////////////////////////////////////////////////////////////////////////////////////////
+
+} // namespace Insights
diff --git a/Engine/Source/Developer/TraceInsights/Private/Insights/Table/Widgets/STableTreeViewCell.cpp b/Engine/Source/Developer/TraceInsights/Private/Insights/Table/Widgets/STableTreeViewCell.cpp
new file mode 100644
--- /dev/null
+++ b/Engine/Source/Developer/TraceInsights/Private/Insights/Table/Widgets/STableTreeViewCell.cpp
@@ -0,0 +1,135 @@
+// Copyright 1998-2019 Epic Games, Inc. All Rights Reserved.
+
+#include "STableTreeViewCell.h"
+
+#include "EditorStyleSet.h"
+#include "SlateOptMacros.h"
+#include "Widgets/Images/SImage.h"
+#include "Widgets/Input/SButton.h"
+#include "Widgets/SBoxPanel.h"
+#include "Widgets/SToolTip.h"
+#include "Widgets/Text/STextBlock.h"
+#include "Widgets/Views/SExpanderArrow.h"
+
+// Insights
+#include "Insights/Table/ViewModels/Table.h"
+#include "Insights/Table/ViewModels/TableColumn.h"
+#include "Insights/Table/Widgets/STableTreeViewTooltip.h"
+
+#define LOCTEXT_NAMESPACE "STableTreeView"
+
+namespace Insights
+{
+
+////////////////////////////////////////////////////////////////////////////////////////////////////
+
+BEGIN_SLATE_FUNCTION_BUILD_OPTIMIZATION
+
+void STableTreeViewCell::Construct(const FArguments& InArgs, const TSharedRef<ITableRow>& TableRow)
+{
+	TablePtr = InArgs._TablePtr;
+	ColumnPtr = InArgs._ColumnPtr;
+	TableTreeNodePtr = InArgs._TableTreeNodePtr;
+
+	ensure(TablePtr.IsValid());
+	ensure(ColumnPtr.IsValid());
+	ensure(TableTreeNodePtr.IsValid());
+
+	SetHoveredCellDelegate = InArgs._OnSetHoveredCell;
+
+	ChildSlot
+		[
+			GenerateWidgetForColumn(InArgs, TableRow)
+		];
+}
+
+////////////////////////////////////////////////////////////////////////////////////////////////////
+
+TSharedRef<SWidget> STableTreeViewCell::GenerateWidgetForColumn(const FArguments& InArgs, const TSharedRef<ITableRow>& TableRow)
+{
+	if (InArgs._IsNameColumn)
+	{
+		return GenerateWidgetForNameColumn(InArgs, TableRow);
+	}
+	else
+	{
+		return GenerateWidgetForTableColumn(InArgs, TableRow);
+	}
+}
+
+////////////////////////////////////////////////////////////////////////////////////////////////////
+
+TSharedRef<SWidget> STableTreeViewCell::GenerateWidgetForNameColumn(const FArguments& InArgs, const TSharedRef<ITableRow>& TableRow)
+{
+	return
+		SNew(SHorizontalBox)
+
+		+ SHorizontalBox::Slot()
+		.AutoWidth()
+		.HAlign(HAlign_Right)
+		.VAlign(VAlign_Center)
+		[
+			SNew(SExpanderArrow, TableRow)
+		]
+
+		// Info icon + tooltip
+		+ SHorizontalBox::Slot()
+		.AutoWidth()
+		.HAlign(HAlign_Center)
+		.VAlign(VAlign_Center)
+		[
+			SNew(SImage)
+			.Visibility(this, &STableTreeViewCell::GetHintIconVisibility)
+			.Image(FEditorStyle::GetBrush("Profiler.Tooltip.HintIcon10"))
+			.ToolTip(STableTreeViewTooltip::GetCellTooltip(TableTreeNodePtr, ColumnPtr))
+		]
+
+		// Name
+		+ SHorizontalBox::Slot()
+		.AutoWidth()
+		.VAlign(VAlign_Center)
+		.HAlign(ColumnPtr->GetHorizontalAlignment())
+		.Padding(FMargin(2.0f, 0.0f))
+		[
+			SNew(STextBlock)
+			.Text(this, &STableTreeViewCell::GetDisplayName)
+			.HighlightText(InArgs._HighlightText)
+			.TextStyle(FEditorStyle::Get(), TEXT("Profiler.Tooltip"))
+			.ColorAndOpacity(this, &STableTreeViewCell::GetColorAndOpacity)
+			.ShadowColorAndOpacity(this, &STableTreeViewCell::GetShadowColorAndOpacity)
+		]
+	;
+}
+
+////////////////////////////////////////////////////////////////////////////////////////////////////
+
+TSharedRef<SWidget> STableTreeViewCell::GenerateWidgetForTableColumn(const FArguments& InArgs, const TSharedRef<ITableRow>& TableRow)
+{
+	const FText CellText = ColumnPtr->GetValueAsText(*TableTreeNodePtr);
+
+	return
+		SNew(SHorizontalBox)
+
+		// Value
+		+ SHorizontalBox::Slot()
+		.FillWidth(1.0f)
+		.VAlign(VAlign_Center)
+		.HAlign(ColumnPtr->GetHorizontalAlignment())
+		.Padding(FMargin(2.0f, 0.0f))
+		[
+			SNew(STextBlock)
+			.Text(CellText)
+			.TextStyle(FEditorStyle::Get(), TEXT("Profiler.Tooltip"))
+			.ColorAndOpacity(this, &STableTreeViewCell::GetStatsColorAndOpacity)
+			.ShadowColorAndOpacity(this, &STableTreeViewCell::GetShadowColorAndOpacity)
+		]
+	;
+}
+
+END_SLATE_FUNCTION_BUILD_OPTIMIZATION
+
+////////////////////////////////////////////////////////////////////////////////////////////////////
+
+} // namespace Insights
+
+#undef LOCTEXT_NAMESPACE
diff --git a/Engine/Source/Developer/TraceInsights/Private/Insights/Table/Widgets/STableTreeViewCell.h b/Engine/Source/Developer/TraceInsights/Private/Insights/Table/Widgets/STableTreeViewCell.h
new file mode 100644
--- /dev/null
+++ b/Engine/Source/Developer/TraceInsights/Private/Insights/Table/Widgets/STableTreeViewCell.h
@@ -0,0 +1,150 @@
+// Copyright 1998-2019 Epic Games, Inc. All Rights Reserved.
+
+#pragma once
+
+#include "CoreMinimal.h"
+#include "Widgets/DeclarativeSyntaxSupport.h"
+#include "Widgets/SCompoundWidget.h"
+
+// Insights
+#include "Insights/Table/ViewModels/TableTreeNode.h"
+
+class ITableRow;
+
+namespace Insights
+{
+
+class FTable;
+class FTableColumn;
+
+DECLARE_DELEGATE_ThreeParams(FSetHoveredTableTreeViewCell, TSharedPtr<FTable> /*TablePtr*/, TSharedPtr<FTableColumn> /*ColumnPtr*/, const FTableTreeNodePtr /*TableTreeNodePtr*/);
+
+class STableTreeViewCell : public SCompoundWidget
+{
+public:
+	SLATE_BEGIN_ARGS(STableTreeViewCell) {}
+		SLATE_EVENT(FSetHoveredTableTreeViewCell, OnSetHoveredCell)
+		SLATE_ATTRIBUTE(FText, HighlightText)
+		SLATE_ARGUMENT(TSharedPtr<FTable>, TablePtr)
+		SLATE_ARGUMENT(TSharedPtr<FTableColumn>, ColumnPtr)
+		SLATE_ARGUMENT(FTableTreeNodePtr, TableTreeNodePtr)
+		SLATE_ARGUMENT(bool, IsNameColumn)
+	SLATE_END_ARGS()
+
+	void Construct(const FArguments& InArgs, const TSharedRef<ITableRow>& TableRow);
+
+protected:
+	TSharedRef<SWidget> GenerateWidgetForColumn(const FArguments& InArgs, const TSharedRef<ITableRow>& TableRow);
+	TSharedRef<SWidget> GenerateWidgetForNameColumn(const FArguments& InArgs, const TSharedRef<ITableRow>& TableRow);
+	TSharedRef<SWidget> GenerateWidgetForTableColumn(const FArguments& InArgs, const TSharedRef<ITableRow>& TableRow);
+
+	/**
+	 * The system will use this event to notify a widget that the cursor has entered it. This event is NOT bubbled.
+	 *
+	 * @param MyGeometry The Geometry of the widget receiving the event
+	 * @param MouseEvent Information about the input event
+	 */
+	virtual void OnMouseEnter(const FGeometry& MyGeometry, const FPointerEvent& MouseEvent) override
+	{
+		SCompoundWidget::OnMouseEnter(MyGeometry, MouseEvent);
+		SetHoveredCellDelegate.ExecuteIfBound(TablePtr, ColumnPtr, TableTreeNodePtr);
+	}
+
+	/**
+	 * The system will use this event to notify a widget that the cursor has left it. This event is NOT bubbled.
+	 *
+	 * @param MouseEvent Information about the input event
+	 */
+	virtual void OnMouseLeave(const FPointerEvent& MouseEvent) override
+	{
+		SCompoundWidget::OnMouseLeave(MouseEvent);
+		SetHoveredCellDelegate.ExecuteIfBound(nullptr, nullptr, nullptr);
+	}
+
+	/**
+	 * Called during drag and drop when the drag enters a widget.
+	 *
+	 * Enter/Leave events in slate are meant as lightweight notifications.
+	 * So we do not want to capture mouse or set focus in response to these.
+	 * However, OnDragEnter must also support external APIs (e.g. OLE Drag/Drop)
+	 * Those require that we let them know whether we can handle the content
+	 * being dragged OnDragEnter.
+	 *
+	 * The concession is to return a can_handled/cannot_handle
+	 * boolean rather than a full FReply.
+	 *
+	 * @param MyGeometry      The geometry of the widget receiving the event.
+	 * @param DragDropEvent   The drag and drop event.
+	 *
+	 * @return A reply that indicated whether the contents of the DragDropEvent can potentially be processed by this widget.
+	 */
+	virtual void OnDragEnter(const FGeometry& MyGeometry, const FDragDropEvent& DragDropEvent) override
+	{
+		SCompoundWidget::OnDragEnter(MyGeometry, DragDropEvent);
+		SetHoveredCellDelegate.ExecuteIfBound(TablePtr, ColumnPtr, TableTreeNodePtr);
+	}
+
+	/**
+	 * Called during drag and drop when the drag leaves a widget.
+	 *
+	 * @param DragDropEvent   The drag and drop event.
+	 */
+	virtual void OnDragLeave(const FDragDropEvent& DragDropEvent)  override
+	{
+		SCompoundWidget::OnDragLeave(DragDropEvent);
+		SetHoveredCellDelegate.ExecuteIfBound(nullptr, nullptr, nullptr);
+	}
+
+	EVisibility GetHintIconVisibility() const
+	{
+		return IsHovered() ? EVisibility::Visible : EVisibility::Hidden;
+	}
+
+	FText GetDisplayName() const
+	{
+		return TableTreeNodePtr->GetDisplayName();
+	}
+
+	FSlateColor GetColorAndOpacity() const
+	{
+		const FLinearColor TextColor =
+			TableTreeNodePtr->IsFiltered() ?
+				FLinearColor(1.0f, 1.0f, 1.0f, 0.5f) :
+				FLinearColor(1.0f, 1.0f, 1.0f, 1.0f);
+		return TextColor;
+	}
+
+	FSlateColor GetStatsColorAndOpacity() const
+	{
+		const FLinearColor TextColor =
+			TableTreeNodePtr->IsGroup() ?
+				FLinearColor(0.7f, 0.7f, 0.7f, 1.0f) :
+			TableTreeNodePtr->IsFiltered() ?
+				FLinearColor(1.0f, 1.0f, 1.0f, 0.5f) :
+				FLinearColor(1.0f, 1.0f, 1.0f, 1.0f);
+		return TextColor;
+	}
+
+	FLinearColor GetShadowColorAndOpacity() const
+	{
+		const FLinearColor ShadowColor =
+			TableTreeNodePtr->IsFiltered() ?
+				FLinearColor(0.f, 0.f, 0.f, 0.25f) :
+				FLinearColor(0.0f, 0.0f, 0.0f, 0.5f);
+		return ShadowColor;
+	}
+
+protected:
+	/** A shared pointer to the table view model. */
+	TSharedPtr<FTable> TablePtr;
+
+	/** A shared pointer to the table column view model. */
+	TSharedPtr<FTableColumn> ColumnPtr;
+
+	/** A shared pointer to the tree node. */
+	FTableTreeNodePtr TableTreeNodePtr;
+
+	FSetHoveredTableTreeViewCell SetHoveredCellDelegate;
+};
+
+} // namespace Insights
diff --git a/Engine/Source/Developer/TraceInsights/Private/Insights/Table/Widgets/STableTreeViewRow.cpp b/Engine/Source/Developer/TraceInsights/Private/Insights/Table/Widgets/STableTreeViewRow.cpp
new file mode 100644
--- /dev/null
+++ b/Engine/Source/Developer/TraceInsights/Private/Insights/Table/Widgets/STableTreeViewRow.cpp
@@ -0,0 +1,191 @@
+// Copyright 1998-2019 Epic Games, Inc. All Rights Reserved.
+
+#include "STableTreeViewRow.h"
+
+#include "EditorStyleSet.h"
+#include "SlateOptMacros.h"
+#include "Widgets/SOverlay.h"
+#include "Widgets/Images/SImage.h"
+
+// Insights
+#include "Insights/Common/TimeUtils.h"
+#include "Insights/Table/ViewModels/Table.h"
+#include "Insights/Table/ViewModels/TableColumn.h"
+#include "Insights/Table/Widgets/STableTreeViewCell.h"
+
+#define LOCTEXT_NAMESPACE "STableTreeView"
+
+namespace Insights
+{
+
+////////////////////////////////////////////////////////////////////////////////////////////////////
+
+BEGIN_SLATE_FUNCTION_BUILD_OPTIMIZATION
+
+void STableTreeViewRow::Construct(const FArguments& InArgs, const TSharedRef<STableViewBase>& InOwnerTableView)
+{
+	OnShouldBeEnabled = InArgs._OnShouldBeEnabled;
+	IsColumnVisibleDelegate = InArgs._OnIsColumnVisible;
+	SetHoveredCellDelegate = InArgs._OnSetHoveredCell;
+	GetColumnOutlineHAlignmentDelegate = InArgs._OnGetColumnOutlineHAlignmentDelegate;
+
+	HighlightText = InArgs._HighlightText;
+	HighlightedNodeName = InArgs._HighlightedNodeName;
+
+	TablePtr = InArgs._TablePtr;
+	TableTreeNodePtr = InArgs._TableTreeNodePtr;
+
+	SetEnabled(TAttribute<bool>(this, &STableTreeViewRow::HandleShouldBeEnabled));
+
+	SMultiColumnTableRow<FTableTreeNodePtr>::Construct(SMultiColumnTableRow<FTableTreeNodePtr>::FArguments(), InOwnerTableView);
+}
+
+////////////////////////////////////////////////////////////////////////////////////////////////////
+
+TSharedRef<SWidget> STableTreeViewRow::GenerateWidgetForColumn(const FName& ColumnId)
+{
+	TSharedPtr<FTableColumn> ColumnPtr = TablePtr->FindColumnChecked(ColumnId);
+
+	return
+		SNew(SOverlay)
+		.Visibility(EVisibility::SelfHitTestInvisible)
+
+		+SOverlay::Slot()
+		.Padding(0.0f)
+		[
+			SNew(SImage)
+			.Image(FEditorStyle::GetBrush("Profiler.LineGraphArea"))
+			.ColorAndOpacity(this, &STableTreeViewRow::GetBackgroundColorAndOpacity)
+		]
+
+		+SOverlay::Slot()
+		.Padding(0.0f)
+		[
+			SNew(SImage)
+			.Image(this, &STableTreeViewRow::GetOutlineBrush, ColumnId)
+			.ColorAndOpacity(this, &STableTreeViewRow::GetOutlineColorAndOpacity)
+		]
+
+		+SOverlay::Slot()
+		[
+			SNew(STableTreeViewCell, SharedThis(this))
+			.Visibility(this, &STableTreeViewRow::IsColumnVisible, ColumnId)
+			.TablePtr(TablePtr)
+			.ColumnPtr(ColumnPtr)
+			.TableTreeNodePtr(TableTreeNodePtr)
+			.HighlightText(HighlightText)
+			.IsNameColumn(ColumnPtr->IsHierarchy())
+			.OnSetHoveredCell(this, &STableTreeViewRow::OnSetHoveredCell)
+		];
+}
+
+END_SLATE_FUNCTION_BUILD_OPTIMIZATION
+
+////////////////////////////////////////////////////////////////////////////////////////////////////
+
+FReply STableTreeViewRow::OnDragDetected(const FGeometry& MyGeometry, const FPointerEvent& MouseEvent)
+{
+	return SMultiColumnTableRow<FTableTreeNodePtr>::OnDragDetected(MyGeometry, MouseEvent);
+}
+
+////////////////////////////////////////////////////////////////////////////////////////////////////
+
+FSlateColor STableTreeViewRow::GetBackgroundColorAndOpacity() const
+{
+	return FLinearColor(0.0f, 0.0f, 0.0f, 1.0f);
+}
+
+////////////////////////////////////////////////////////////////////////////////////////////////////
+
+FSlateColor STableTreeViewRow::GetBackgroundColorAndOpacity(double Time) const
+{
+	const FLinearColor Color =	Time > TimeUtils::Second      ? FLinearColor(0.3f, 0.0f, 0.0f, 1.0f) :
+								Time > TimeUtils::Milisecond  ? FLinearColor(0.3f, 0.1f, 0.0f, 1.0f) :
+								Time > TimeUtils::Microsecond ? FLinearColor(0.0f, 0.1f, 0.0f, 1.0f) :
+																FLinearColor(0.0f, 0.0f, 0.0f, 1.0f);
+	return Color;
+}
+
+////////////////////////////////////////////////////////////////////////////////////////////////////
+
+FSlateColor STableTreeViewRow::GetOutlineColorAndOpacity() const
+{
+	const FLinearColor NoColor(0.0f, 0.0f, 0.0f, 0.0f);
+	const bool bShouldBeHighlighted = TableTreeNodePtr->GetName() == HighlightedNodeName.Get();
+	const FLinearColor OutlineColorAndOpacity = bShouldBeHighlighted ? FLinearColor(FColorList::SlateBlue) : NoColor;
+	return OutlineColorAndOpacity;
+}
+
+////////////////////////////////////////////////////////////////////////////////////////////////////
+
+const FSlateBrush* STableTreeViewRow::GetOutlineBrush(const FName ColumnId) const
+{
+	EHorizontalAlignment Result = HAlign_Center;
+	if (IsColumnVisibleDelegate.IsBound())
+	{
+		Result = GetColumnOutlineHAlignmentDelegate.Execute(ColumnId);
+	}
+
+	const FSlateBrush* Brush = nullptr;
+	if (Result == HAlign_Left)
+	{
+		Brush = FEditorStyle::GetBrush("Profiler.EventGraph.Border.L");
+	}
+	else if(Result == HAlign_Right)
+	{
+		Brush = FEditorStyle::GetBrush("Profiler.EventGraph.Border.R");
+	}
+	else
+	{
+		Brush = FEditorStyle::GetBrush("Profiler.EventGraph.Border.TB");
+	}
+	return Brush;
+}
+
+////////////////////////////////////////////////////////////////////////////////////////////////////
+
+bool STableTreeViewRow::HandleShouldBeEnabled() const
+{
+	bool bResult = false;
+
+	if (TableTreeNodePtr->IsGroup())
+	{
+		bResult = true;
+	}
+	else
+	{
+		if (OnShouldBeEnabled.IsBound())
+		{
+			bResult = OnShouldBeEnabled.Execute(TableTreeNodePtr->GetId());
+		}
+	}
+
+	return bResult;
+}
+
+////////////////////////////////////////////////////////////////////////////////////////////////////
+
+EVisibility STableTreeViewRow::IsColumnVisible(const FName ColumnId) const
+{
+	EVisibility Result = EVisibility::Collapsed;
+
+	if (IsColumnVisibleDelegate.IsBound())
+	{
+		Result = IsColumnVisibleDelegate.Execute(ColumnId) ? EVisibility::Visible : EVisibility::Collapsed;
+	}
+
+	return Result;
+}
+
+////////////////////////////////////////////////////////////////////////////////////////////////////
+
+void STableTreeViewRow::OnSetHoveredCell(TSharedPtr<FTable> InTablePtr, TSharedPtr<FTableColumn> InColumnPtr, const FTableTreeNodePtr InTreeNodePtr)
+{
+	SetHoveredCellDelegate.ExecuteIfBound(InTablePtr, InColumnPtr, InTreeNodePtr);
+}
+
+////////////////////////////////////////////////////////////////////////////////////////////////////
+
+} // namespace Insights
+
+#undef LOCTEXT_NAMESPACE
diff --git a/Engine/Source/Developer/TraceInsights/Private/Insights/Table/Widgets/STableTreeViewRow.h b/Engine/Source/Developer/TraceInsights/Private/Insights/Table/Widgets/STableTreeViewRow.h
new file mode 100644
--- /dev/null
+++ b/Engine/Source/Developer/TraceInsights/Private/Insights/Table/Widgets/STableTreeViewRow.h
@@ -0,0 +1,90 @@
+// Copyright 1998-2019 Epic Games, Inc. All Rights Reserved.
+
+#pragma once
+
+#include "CoreMinimal.h"
+#include "Widgets/DeclarativeSyntaxSupport.h"
+#include "Widgets/Views/SHeaderRow.h"
+#include "Widgets/Views/STableViewBase.h"
+#include "Widgets/Views/STableRow.h"
+#include "Widgets/Views/STreeView.h"
+
+// Insights
+#include "Insights/Table/ViewModels/TableTreeNode.h"
+
+namespace Insights
+{
+
+class FTable;
+class FTableColumn;
+
+DECLARE_DELEGATE_RetVal_OneParam(bool, FShouldBeEnabledDelegate, const uint32 /*TreeNodeId*/);
+DECLARE_DELEGATE_RetVal_OneParam(bool, FIsColumnVisibleDelegate, const FName /*ColumnId*/);
+DECLARE_DELEGATE_RetVal_OneParam(EHorizontalAlignment, FGetColumnOutlineHAlignmentDelegate, const FName /*ColumnId*/);
+DECLARE_DELEGATE_ThreeParams(FSetHoveredTableTreeViewCell, TSharedPtr<FTable> /*TablePtr*/, TSharedPtr<FTableColumn> /*ColumnPtr*/, const FTableTreeNodePtr /*NodePtr*/);
+
+/** Widget that represents a table row in the tree control. Generates widgets for each column on demand. */
+class STableTreeViewRow : public SMultiColumnTableRow<FTableTreeNodePtr>
+{
+public:
+	SLATE_BEGIN_ARGS(STableTreeViewRow) {}
+		SLATE_EVENT(FShouldBeEnabledDelegate, OnShouldBeEnabled)
+		SLATE_EVENT(FIsColumnVisibleDelegate, OnIsColumnVisible)
+		SLATE_EVENT(FGetColumnOutlineHAlignmentDelegate, OnGetColumnOutlineHAlignmentDelegate)
+		SLATE_EVENT(FSetHoveredTableTreeViewCell, OnSetHoveredCell)
+		SLATE_ATTRIBUTE(FText, HighlightText)
+		SLATE_ATTRIBUTE(FName, HighlightedNodeName)
+		SLATE_ARGUMENT(TSharedPtr<FTable>, TablePtr)
+		SLATE_ARGUMENT(FTableTreeNodePtr, TableTreeNodePtr)
+	SLATE_END_ARGS()
+
+public:
+	void Construct(const FArguments& InArgs, const TSharedRef<STableViewBase>& InOwnerTableView);
+
+	virtual TSharedRef<SWidget> GenerateWidgetForColumn(const FName& ColumnId) override;
+
+	/**
+	 * Called when Slate detects that a widget started to be dragged.
+	 * Usage:
+	 * A widget can ask Slate to detect a drag.
+	 * OnMouseDown() reply with FReply::Handled().DetectDrag(SharedThis(this)).
+	 * Slate will either send an OnDragDetected() event or do nothing.
+	 * If the user releases a mouse button or leaves the widget before
+	 * a drag is triggered (maybe user started at the very edge) then no event will be
+	 * sent.
+	 *
+	 * @param  InMyGeometry  Widget geometry
+	 * @param  InMouseEvent  MouseMove that triggered the drag
+	 *
+	 */
+	virtual FReply OnDragDetected(const FGeometry& MyGeometry, const FPointerEvent& MouseEvent) override;
+
+protected:
+	FSlateColor GetBackgroundColorAndOpacity() const;
+	FSlateColor GetBackgroundColorAndOpacity(double Time) const;
+	FSlateColor GetOutlineColorAndOpacity() const;
+	const FSlateBrush* GetOutlineBrush(const FName ColumnId) const;
+	bool HandleShouldBeEnabled() const;
+	EVisibility IsColumnVisible(const FName ColumnId) const;
+	void OnSetHoveredCell(TSharedPtr<FTable> InTablePtr, TSharedPtr<FTableColumn> InColumnPtr, const FTableTreeNodePtr InTreeNodePtr);
+
+protected:
+	/** A shared pointer to the table view model. */
+	TSharedPtr<FTable> TablePtr;
+
+	/** Data context for this table row. */
+	FTableTreeNodePtr TableTreeNodePtr;
+
+	FShouldBeEnabledDelegate OnShouldBeEnabled;
+	FIsColumnVisibleDelegate IsColumnVisibleDelegate;
+	FSetHoveredTableTreeViewCell SetHoveredCellDelegate;
+	FGetColumnOutlineHAlignmentDelegate GetColumnOutlineHAlignmentDelegate;
+
+	/** Text to be highlighted on tree node's name. */
+	TAttribute<FText> HighlightText;
+
+	/** Name of the tree node that should be drawn as highlighted. */
+	TAttribute<FName> HighlightedNodeName;
+};
+
+} // namespace Insights
