Copying //UE4/Dev-Core to Main (//UE4/Main)

diff --git a/Engine/Source/Developer/TraceInsights/Private/Insights/Table/Widgets/STableTreeViewTooltip.cpp b/Engine/Source/Developer/TraceInsights/Private/Insights/Table/Widgets/STableTreeViewTooltip.cpp
new file mode 100644
--- /dev/null
+++ b/Engine/Source/Developer/TraceInsights/Private/Insights/Table/Widgets/STableTreeViewTooltip.cpp
@@ -0,0 +1,242 @@
+// Copyright 1998-2019 Epic Games, Inc. All Rights Reserved.
+
+#include "STableTreeViewTooltip.h"
+
+#include "EditorStyleSet.h"
+#include "SlateOptMacros.h"
+#include "Widgets/Images/SImage.h"
+#include "Widgets/Layout/SGridPanel.h"
+#include "Widgets/Layout/SSeparator.h"
+#include "Widgets/SBoxPanel.h"
+#include "Widgets/SToolTip.h"
+#include "Widgets/Text/STextBlock.h"
+
+// Insights
+#include "Insights/Common/TimeUtils.h"
+#include "Insights/Table/ViewModels/Table.h"
+#include "Insights/Table/ViewModels/TableColumn.h"
+#include "Insights/Table/ViewModels/TableTreeNode.h"
+
+#define LOCTEXT_NAMESPACE "STableTreeView"
+
+namespace Insights
+{
+
+////////////////////////////////////////////////////////////////////////////////////////////////////
+
+BEGIN_SLATE_FUNCTION_BUILD_OPTIMIZATION
+
+TSharedPtr<SToolTip> STableTreeViewTooltip::GetTableTooltip(const FTable& Table)
+{
+	TSharedPtr<SToolTip> ColumnTooltip =
+		SNew(SToolTip)
+		[
+			SNew(SVerticalBox)
+
+			+ SVerticalBox::Slot()
+			.AutoHeight()
+			.Padding(2.0f)
+			[
+				SNew(STextBlock)
+				.Text(Table.GetDisplayName())
+				.TextStyle(FEditorStyle::Get(), TEXT("Profiler.TooltipBold"))
+			]
+
+			+ SVerticalBox::Slot()
+			.AutoHeight()
+			.Padding(2.0f)
+			[
+				SNew(STextBlock)
+				.Text(Table.GetDescription())
+				.TextStyle(FEditorStyle::Get(), TEXT("Profiler.Tooltip"))
+			]
+		];
+
+	return ColumnTooltip;
+}
+
+////////////////////////////////////////////////////////////////////////////////////////////////////
+
+TSharedPtr<SToolTip> STableTreeViewTooltip::GetColumnTooltip(const FTableColumn& Column)
+{
+	TSharedPtr<SToolTip> ColumnTooltip =
+		SNew(SToolTip)
+		[
+			SNew(SVerticalBox)
+
+			+ SVerticalBox::Slot()
+			.AutoHeight()
+			.Padding(2.0f)
+			[
+				SNew(STextBlock)
+				.Text(Column.GetTitleName())
+				.TextStyle(FEditorStyle::Get(), TEXT("Profiler.TooltipBold"))
+			]
+
+			+ SVerticalBox::Slot()
+			.AutoHeight()
+			.Padding(2.0f)
+			[
+				SNew(STextBlock)
+				.Text(Column.GetDescription())
+				.TextStyle(FEditorStyle::Get(), TEXT("Profiler.Tooltip"))
+			]
+		];
+
+	return ColumnTooltip;
+}
+
+////////////////////////////////////////////////////////////////////////////////////////////////////
+
+TSharedPtr<SToolTip> STableTreeViewTooltip::GetCellTooltip(const TSharedPtr<FTableTreeNode> TreeNodePtr, const TSharedPtr<FTableColumn> ColumnPtr)
+{
+	TSharedPtr<SGridPanel> GridPanel;
+	TSharedPtr<SHorizontalBox> HBox;
+
+	TSharedPtr<SToolTip> TableCellTooltip =
+		SNew(SToolTip)
+		[
+			SAssignNew(HBox, SHorizontalBox)
+
+			+ SHorizontalBox::Slot()
+			.AutoWidth()
+			[
+				SNew(SVerticalBox)
+
+				+ SVerticalBox::Slot()
+				.AutoHeight()
+				.Padding(2.0f)
+				[
+					SNew(SSeparator)
+					.Orientation(Orient_Horizontal)
+				]
+
+				+ SVerticalBox::Slot()
+				.AutoHeight()
+				.Padding(2.0f)
+				[
+					SNew(SGridPanel)
+
+					// Id: [Id]
+					+ SGridPanel::Slot(0, 0)
+					.Padding(2.0f)
+					[
+						SNew(STextBlock)
+						.Text(LOCTEXT("TT_Id", "Id:"))
+						.TextStyle(FEditorStyle::Get(), TEXT("Profiler.TooltipBold"))
+					]
+					+ SGridPanel::Slot(1, 0)
+					.Padding(2.0f)
+					[
+						SNew(STextBlock)
+						.Text(FText::AsNumber(TreeNodePtr->GetId()))
+						.TextStyle(FEditorStyle::Get(), TEXT("Profiler.Tooltip"))
+					]
+
+					// Item Type: [Type]
+					+ SGridPanel::Slot(0, 1)
+					.Padding(2.0f)
+					[
+						SNew(STextBlock)
+						.Text(LOCTEXT("TT_Type", "Item Type:"))
+						.TextStyle(FEditorStyle::Get(), TEXT("Profiler.TooltipBold"))
+					]
+					+ SGridPanel::Slot(1, 1)
+					.Padding(2.0f)
+					[
+						SNew(STextBlock)
+						.Text(TreeNodePtr->IsGroup() ? LOCTEXT("TT_Type_Group", "Group Node") : LOCTEXT("TT_Type_TableRow", "Table Row"))
+						.TextStyle(FEditorStyle::Get(), TEXT("Profiler.Tooltip"))
+					]
+
+					// Item Name: [Name]
+					+ SGridPanel::Slot(0, 2)
+					.Padding(2.0f)
+					[
+						SNew(STextBlock)
+						.Text(LOCTEXT("TT_Name", "Item Name:"))
+						.TextStyle(FEditorStyle::Get(), TEXT("Profiler.TooltipBold"))
+					]
+					+ SGridPanel::Slot(1, 2)
+					.Padding(2.0f)
+					[
+						SNew(STextBlock)
+						.Text(FText::FromName(TreeNodePtr->GetName()))
+						.TextStyle(FEditorStyle::Get(), TEXT("Profiler.Tooltip"))
+					]
+				]
+
+				+ SVerticalBox::Slot()
+				.AutoHeight()
+				.Padding(2.0f)
+				[
+					SNew(SSeparator)
+					.Orientation(Orient_Horizontal)
+				]
+
+				+ SVerticalBox::Slot()
+				.AutoHeight()
+				.Padding(2.0f)
+				[
+					SAssignNew(GridPanel, SGridPanel)
+
+					// Values for each table column are added here.
+				]
+
+				+ SVerticalBox::Slot()
+				.AutoHeight()
+				.Padding(2.0f)
+				[
+					SNew(SSeparator)
+					.Orientation(Orient_Horizontal)
+				]
+			]
+		];
+
+	TSharedPtr<FTable> Table = TreeNodePtr->GetParentTable().Pin();
+	if (Table.IsValid())
+	{
+		int32 Row = 0;
+		for (const TSharedPtr<FTableColumn>& Column : Table->GetColumns())
+		{
+			if (!Column->IsHierarchy())
+			{
+				FText Name = FText::Format(LOCTEXT("TooltipValueFormat", "{0}:"), Column->GetTitleName());
+				AddGridRow(GridPanel, Row, Name, Column->GetValueAsTooltipText(*TreeNodePtr));
+			}
+		}
+	}
+
+	return TableCellTooltip;
+}
+
+////////////////////////////////////////////////////////////////////////////////////////////////////
+
+void STableTreeViewTooltip::AddGridRow(TSharedPtr<SGridPanel> Grid, int32& Row, const FText& Name, const FText& Value)
+{
+	Grid->AddSlot(0, Row)
+		.Padding(2.0f)
+		[
+			SNew(STextBlock)
+			.Text(Name)
+			.TextStyle(FEditorStyle::Get(), TEXT("Profiler.TooltipBold"))
+		];
+
+	Grid->AddSlot(1, Row)
+		.Padding(2.0f)
+		[
+			SNew(STextBlock)
+			.Text(Value)
+			.TextStyle(FEditorStyle::Get(), TEXT("Profiler.Tooltip"))
+		];
+
+	Row++;
+}
+
+END_SLATE_FUNCTION_BUILD_OPTIMIZATION
+
+////////////////////////////////////////////////////////////////////////////////////////////////////
+
+} // namespace Insights
+
+#undef LOCTEXT_NAMESPACE
diff --git a/Engine/Source/Developer/TraceInsights/Private/Insights/Table/Widgets/STableTreeViewTooltip.h b/Engine/Source/Developer/TraceInsights/Private/Insights/Table/Widgets/STableTreeViewTooltip.h
new file mode 100644
--- /dev/null
+++ b/Engine/Source/Developer/TraceInsights/Private/Insights/Table/Widgets/STableTreeViewTooltip.h
@@ -0,0 +1,31 @@
+// Copyright 1998-2019 Epic Games, Inc. All Rights Reserved.
+
+#pragma once
+
+#include "CoreMinimal.h"
+
+class SToolTip;
+class SGridPanel;
+
+namespace Insights
+{
+
+class FTable;
+class FTableColumn;
+class FTableTreeNode;
+
+/** Tooltip for STableTreeView widget. */
+class STableTreeViewTooltip
+{
+public:
+	STableTreeViewTooltip() = delete;
+
+	static TSharedPtr<SToolTip> GetTableTooltip(const FTable& Table);
+	static TSharedPtr<SToolTip> GetColumnTooltip(const FTableColumn& Column);
+	static TSharedPtr<SToolTip> GetCellTooltip(const TSharedPtr<FTableTreeNode> TreeNodePtr, const TSharedPtr<FTableColumn> ColumnPtr);
+
+private:
+	static void AddGridRow(TSharedPtr<SGridPanel> Grid, int32& Row, const FText& Name, const FText& Value);
+};
+
+} // namespace Insights
diff --git a/Engine/Source/Developer/TraceInsights/Private/Insights/TraceInsightsModule.cpp b/Engine/Source/Developer/TraceInsights/Private/Insights/TraceInsightsModule.cpp
--- a/Engine/Source/Developer/TraceInsights/Private/Insights/TraceInsightsModule.cpp
+++ b/Engine/Source/Developer/TraceInsights/Private/Insights/TraceInsightsModule.cpp
@@ -1,7 +1,6 @@
 // Copyright 1998-2019 Epic Games, Inc. All Rights Reserved.
 
 #include "CoreMinimal.h"
-#include "EditorStyleSet.h"
 #include "HAL/PlatformApplicationMisc.h"
 #include "Modules/ModuleManager.h"
 #include "TraceServices/ITraceServicesModule.h"
@@ -9,23 +8,22 @@
 #include "Widgets/DeclarativeSyntaxSupport.h"
 #include "Widgets/Docking/SDockTab.h"
 #include "Widgets/SWidget.h"
-#include "Widgets/Text/STextBlock.h"
-//#include "WorkspaceMenuStructure.h"
-//#include "WorkspaceMenuStructureModule.h"
 
 // Insights
 #include "Insights/InsightsManager.h"
 #include "Insights/InsightsStyle.h"
-#include "Insights/IoProfilerManager.h"
 #include "Insights/IUnrealInsightsModule.h"
+#include "Insights/LoadingProfiler/LoadingProfilerManager.h"
+#include "Insights/LoadingProfiler/Widgets/SLoadingProfilerWindow.h"
+#include "Insights/NetworkingProfiler/NetworkingProfilerManager.h"
+#include "Insights/NetworkingProfiler/Widgets/SNetworkingProfilerWindow.h"
 #include "Insights/TimingProfilerManager.h"
-#include "Insights/Widgets/SIoProfilerWindow.h"
 #include "Insights/Widgets/SStartPageWindow.h"
 #include "Insights/Widgets/STimingProfilerWindow.h"
 
 ////////////////////////////////////////////////////////////////////////////////////////////////////
 /**
- * Implements the Trace Insights module (TimingProfiler, IoProfiler, etc.).
+ * Implements the Trace Insights module (TimingProfiler, LoadingProfiler, etc.).
  */
 class FTraceInsightsModule : public IUnrealInsightsModule
 {
@@ -48,8 +46,11 @@ protected:
 	/** Callback called when the Timing Profiler major tab is closed. */
 	void OnTimingProfilerTabBeingClosed(TSharedRef<SDockTab> TabBeingClosed);
 
-	/** Callback called when the Io Profiler major tab is closed. */
-	void OnIoProfilerTabBeingClosed(TSharedRef<SDockTab> TabBeingClosed);
+	/** Callback called when the Loading Profiler major tab is closed. */
+	void OnLoadingProfilerTabBeingClosed(TSharedRef<SDockTab> TabBeingClosed);
+
+	/** Callback called when the Networking Profiler major tab is closed. */
+	void OnNetworkingProfilerTabBeingClosed(TSharedRef<SDockTab> TabBeingClosed);
 
 	/** Start Page */
 	TSharedRef<SDockTab> SpawnStartPageTab(const FSpawnTabArgs& Args);
@@ -57,8 +58,11 @@ protected:
 	/** Timing Profiler */
 	TSharedRef<SDockTab> SpawnTimingProfilerTab(const FSpawnTabArgs& Args);
 
-	/** I/O Profiler */
-	TSharedRef<SDockTab> SpawnIoProfilerTab(const FSpawnTabArgs& Args);
+	/** Loading Profiler */
+	TSharedRef<SDockTab> SpawnLoadingProfilerTab(const FSpawnTabArgs& Args);
+
+	/** Networking Profiler */
+	TSharedRef<SDockTab> SpawnNetworkingProfilerTab(const FSpawnTabArgs& Args);
 
 protected:
 	TSharedPtr<Trace::IAnalysisService> TraceAnalysisService;
@@ -70,6 +74,13 @@ protected:
 
 IMPLEMENT_MODULE(FTraceInsightsModule, TraceInsights);
 
+////////////////////////////////////////////////////////////////////////////////////////////////////
+
+static TSharedPtr<SDockTab> NeverReuse(const FTabId&)
+{
+	return TSharedPtr<SDockTab>();
+}
+
 ////////////////////////////////////////////////////////////////////////////////////////////////////
 // FTraceInsightsModule
 ////////////////////////////////////////////////////////////////////////////////////////////////////
@@ -88,21 +99,28 @@ void FTraceInsightsModule::StartupModule()
 
 	FInsightsManager::Initialize(TraceAnalysisService.ToSharedRef(), TraceSessionService.ToSharedRef(), TraceModuleService.ToSharedRef());
 	FTimingProfilerManager::Initialize();
-	FIoProfilerManager::Initialize();
+	FLoadingProfilerManager::Initialize();
+	FNetworkingProfilerManager::Initialize();
 
 	// Register tab spawner for the Start Page.
-	auto& StartPageTabSpawnerEntry = FGlobalTabmanager::Get()->RegisterTabSpawner(FInsightsManagerTabs::StartPageTabId,
+	auto& StartPageTabSpawnerEntry = FGlobalTabmanager::Get()->RegisterNomadTabSpawner(FInsightsManagerTabs::StartPageTabId,
 		FOnSpawnTab::CreateRaw(this, &FTraceInsightsModule::SpawnStartPageTab))
 		.SetDisplayName(NSLOCTEXT("FTraceInsightsModule", "StartPageTabTitle", "Unreal Insights"))
 		.SetTooltipText(NSLOCTEXT("FTraceInsightsModule", "StartPageTooltipText", "Open the start page for Unreal Insights."))
 		.SetIcon(FSlateIcon(FInsightsStyle::GetStyleSetName(), "StartPage.Icon.Small"));
 
+//#if WITH_EDITOR
+//	StartPageTabSpawnerEntry.SetGroup(WorkspaceMenu::GetMenuStructure().GetDeveloperToolsMiscCategory());
+//#else
+//	StartPageTabSpawnerEntry.SetGroup(WorkspaceMenu::GetMenuStructure().GetToolsCategory());
+//#endif
+
 	// Register tab spawner for the Timing Insights.
-	auto& TimingProfilerTabSpawnerEntry = FGlobalTabmanager::Get()->RegisterTabSpawner(FInsightsManagerTabs::TimingProfilerTabId,
+	auto& TimingProfilerTabSpawnerEntry = FGlobalTabmanager::Get()->RegisterNomadTabSpawner(FInsightsManagerTabs::TimingProfilerTabId,
 		FOnSpawnTab::CreateRaw(this, &FTraceInsightsModule::SpawnTimingProfilerTab))
 		.SetDisplayName(NSLOCTEXT("FTraceInsightsModule", "TimingProfilerTabTitle", "Timing Insights"))
 		.SetTooltipText(NSLOCTEXT("FTraceInsightsModule", "TimingProfilerTooltipText", "Open the Timing Insights tab."))
-		.SetIcon(FSlateIcon(FInsightsStyle::GetStyleSetName(), "TimingInsights.Icon.Small"));
+		.SetIcon(FSlateIcon(FInsightsStyle::GetStyleSetName(), "TimingProfiler.Icon.Small"));
 
 //#if WITH_EDITOR
 //	TimingProfilerTabSpawnerEntry.SetGroup(WorkspaceMenu::GetMenuStructure().GetDeveloperToolsMiscCategory());
@@ -111,16 +129,30 @@ void FTraceInsightsModule::StartupModule()
 //#endif
 
 	// Register tab spawner for the Asset Loading Insights.
-	auto& IoProfilerTabSpawnerEntry = FGlobalTabmanager::Get()->RegisterTabSpawner(FInsightsManagerTabs::IoProfilerTabId,
-		FOnSpawnTab::CreateRaw(this, &FTraceInsightsModule::SpawnIoProfilerTab))
-		.SetDisplayName(NSLOCTEXT("FTraceInsightsModule", "IoProfilerTabTitle", "Asset Loading Insights"))
-		.SetTooltipText(NSLOCTEXT("FTraceInsightsModule", "IoProfilerTooltipText", "Open the Asset Loading Insights tab."))
-		.SetIcon(FSlateIcon(FInsightsStyle::GetStyleSetName(), "AssetLoadingInsights.Icon.Small"));
+	auto& LoadingProfilerTabSpawnerEntry = FGlobalTabmanager::Get()->RegisterNomadTabSpawner(FInsightsManagerTabs::LoadingProfilerTabId,
+		FOnSpawnTab::CreateRaw(this, &FTraceInsightsModule::SpawnLoadingProfilerTab))
+		.SetDisplayName(NSLOCTEXT("FTraceInsightsModule", "LoadingProfilerTabTitle", "Asset Loading Insights"))
+		.SetTooltipText(NSLOCTEXT("FTraceInsightsModule", "LoadingProfilerTooltipText", "Open the Asset Loading Insights tab."))
+		.SetIcon(FSlateIcon(FInsightsStyle::GetStyleSetName(), "LoadingProfiler.Icon.Small"));
 
 //#if WITH_EDITOR
-//	IoProfilerTabSpawnerEntry.SetGroup(WorkspaceMenu::GetMenuStructure().GetDeveloperToolsMiscCategory());
+//	LoadingProfilerTabSpawnerEntry.SetGroup(WorkspaceMenu::GetMenuStructure().GetDeveloperToolsMiscCategory());
 //#else
-//	IoProfilerTabSpawnerEntry.SetGroup(WorkspaceMenu::GetMenuStructure().GetToolsCategory());
+//	LoadingProfilerTabSpawnerEntry.SetGroup(WorkspaceMenu::GetMenuStructure().GetToolsCategory());
+//#endif
+
+	// Register tab spawner for the Networking Insights.
+	auto& NetworkingProfilerTabSpawnerEntry = FGlobalTabmanager::Get()->RegisterNomadTabSpawner(FInsightsManagerTabs::NetworkingProfilerTabId,
+		FOnSpawnTab::CreateRaw(this, &FTraceInsightsModule::SpawnNetworkingProfilerTab))
+		.SetReuseTabMethod(FOnFindTabToReuse::CreateStatic(&NeverReuse))
+		.SetDisplayName(NSLOCTEXT("FTraceInsightsModule", "NetworkingProfilerTabTitle", "Networking Insights"))
+		.SetTooltipText(NSLOCTEXT("FTraceInsightsModule", "NetworkingProfilerTooltipText", "Open the Networking Insights tab."))
+		.SetIcon(FSlateIcon(FInsightsStyle::GetStyleSetName(), "NetworkingProfiler.Icon.Small"));
+
+//#if WITH_EDITOR
+//	NetworkingProfilerTabSpawnerEntry.SetGroup(WorkspaceMenu::GetMenuStructure().GetDeveloperToolsMiscCategory());
+//#else
+//	NetworkingProfilerTabSpawnerEntry.SetGroup(WorkspaceMenu::GetMenuStructure().GetToolsCategory());
 //#endif
 }
 
@@ -130,20 +162,27 @@ void FTraceInsightsModule::ShutdownModule()
 {
 	TraceSessionService->StopRecorderServer();
 
-	FGlobalTabmanager::Get()->UnregisterNomadTabSpawner(FInsightsManagerTabs::StartPageTabId);
+	FGlobalTabmanager::Get()->UnregisterNomadTabSpawner(FInsightsManagerTabs::NetworkingProfilerTabId);
+	FGlobalTabmanager::Get()->UnregisterNomadTabSpawner(FInsightsManagerTabs::LoadingProfilerTabId);
 	FGlobalTabmanager::Get()->UnregisterNomadTabSpawner(FInsightsManagerTabs::TimingProfilerTabId);
-	FGlobalTabmanager::Get()->UnregisterNomadTabSpawner(FInsightsManagerTabs::IoProfilerTabId);
+	FGlobalTabmanager::Get()->UnregisterNomadTabSpawner(FInsightsManagerTabs::StartPageTabId);
 
-	if (FTimingProfilerManager::Get().IsValid())
+	if (FNetworkingProfilerManager::Get().IsValid())
 	{
-		// Shutdown the Timing Profiler manager.
-		FTimingProfilerManager::Get()->Shutdown();
+		// Shutdown the NetworkingProfiler (Networking Insights) manager.
+		FNetworkingProfilerManager::Get()->Shutdown();
+	}
+
+	if (FLoadingProfilerManager::Get().IsValid())
+	{
+		// Shutdown the LoadingProfiler (Asset Loading Insights) manager.
+		FLoadingProfilerManager::Get()->Shutdown();
 	}
 
-	if (FIoProfilerManager::Get().IsValid())
+	if (FTimingProfilerManager::Get().IsValid())
 	{
-		// Shutdown the I/O Profiler manager.
-		FIoProfilerManager::Get()->Shutdown();
+		// Shutdown the TimingProfiler (Timing Insights) manager.
+		FTimingProfilerManager::Get()->Shutdown();
 	}
 
 	if (FInsightsManager::Get().IsValid())
@@ -166,7 +205,8 @@ void FTraceInsightsModule::OnNewLayout(TSharedRef<FTabManager::FLayout> NewLayou
 			FTabManager::NewStack()
 			->AddTab(FInsightsManagerTabs::StartPageTabId, ETabState::OpenedTab)
 			->AddTab(FInsightsManagerTabs::TimingProfilerTabId, ETabState::ClosedTab)
-			->AddTab(FInsightsManagerTabs::IoProfilerTabId, ETabState::ClosedTab)
+			->AddTab(FInsightsManagerTabs::LoadingProfilerTabId, ETabState::ClosedTab)
+			->AddTab(FInsightsManagerTabs::NetworkingProfilerTabId, ETabState::ClosedTab)
 			->SetForegroundTab(FTabId(FInsightsManagerTabs::StartPageTabId))
 		)
 	);
@@ -220,17 +260,17 @@ void FTraceInsightsModule::OnTimingProfilerTabBeingClosed(TSharedRef<SDockTab> T
 
 ////////////////////////////////////////////////////////////////////////////////////////////////////
 
-TSharedRef<SDockTab> FTraceInsightsModule::SpawnIoProfilerTab(const FSpawnTabArgs& Args)
+TSharedRef<SDockTab> FTraceInsightsModule::SpawnLoadingProfilerTab(const FSpawnTabArgs& Args)
 {
 	const TSharedRef<SDockTab> DockTab = SNew(SDockTab)
 		.TabRole(ETabRole::NomadTab);
 
 	// Register OnTabClosed to handle I/O profiler manager shutdown.
-	DockTab->SetOnTabClosed(SDockTab::FOnTabClosedCallback::CreateRaw(this, &FTraceInsightsModule::OnIoProfilerTabBeingClosed));
+	DockTab->SetOnTabClosed(SDockTab::FOnTabClosedCallback::CreateRaw(this, &FTraceInsightsModule::OnLoadingProfilerTabBeingClosed));
 
-	// Create the SIoProfilerWindow widget.
-	TSharedRef<SIoProfilerWindow> Window = SNew(SIoProfilerWindow);
-	FIoProfilerManager::Get()->AssignProfilerWindow(Window);
+	// Create the SLoadingProfilerWindow widget.
+	TSharedRef<SLoadingProfilerWindow> Window = SNew(SLoadingProfilerWindow, DockTab, Args.GetOwnerWindow());
+	FLoadingProfilerManager::Get()->AssignProfilerWindow(Window);
 	DockTab->SetContent(Window);
 
 	return DockTab;
@@ -238,8 +278,37 @@ TSharedRef<SDockTab> FTraceInsightsModule::SpawnIoProfilerTab(const FSpawnTabArg
 
 ////////////////////////////////////////////////////////////////////////////////////////////////////
 
-void FTraceInsightsModule::OnIoProfilerTabBeingClosed(TSharedRef<SDockTab> TabBeingClosed)
+void FTraceInsightsModule::OnLoadingProfilerTabBeingClosed(TSharedRef<SDockTab> TabBeingClosed)
+{
+	// Disable TabClosed delegate.
+	TabBeingClosed->SetOnTabClosed(SDockTab::FOnTabClosedCallback());
+}
+
+////////////////////////////////////////////////////////////////////////////////////////////////////
+
+TSharedRef<SDockTab> FTraceInsightsModule::SpawnNetworkingProfilerTab(const FSpawnTabArgs& Args)
 {
+	const TSharedRef<SDockTab> DockTab = SNew(SDockTab)
+		.TabRole(ETabRole::NomadTab);
+
+	// Register OnTabClosed to handle I/O profiler manager shutdown.
+	DockTab->SetOnTabClosed(SDockTab::FOnTabClosedCallback::CreateRaw(this, &FTraceInsightsModule::OnNetworkingProfilerTabBeingClosed));
+
+	// Create the SNetworkingProfilerWindow widget.
+	TSharedRef<SNetworkingProfilerWindow> Window = SNew(SNetworkingProfilerWindow, DockTab, Args.GetOwnerWindow());
+	FNetworkingProfilerManager::Get()->AddProfilerWindow(Window);
+	DockTab->SetContent(Window);
+
+	return DockTab;
+}
+
+////////////////////////////////////////////////////////////////////////////////////////////////////
+
+void FTraceInsightsModule::OnNetworkingProfilerTabBeingClosed(TSharedRef<SDockTab> TabBeingClosed)
+{
+	TSharedRef<SNetworkingProfilerWindow> Window = StaticCastSharedRef<SNetworkingProfilerWindow>(TabBeingClosed->GetContent());
+	FNetworkingProfilerManager::Get()->RemoveProfilerWindow(Window);
+
 	// Disable TabClosed delegate.
 	TabBeingClosed->SetOnTabClosed(SDockTab::FOnTabClosedCallback());
 }
diff --git a/Engine/Source/Developer/TraceInsights/Private/Insights/ViewModels/DrawHelpers.h b/Engine/Source/Developer/TraceInsights/Private/Insights/ViewModels/DrawHelpers.h
new file mode 100644
--- /dev/null
+++ b/Engine/Source/Developer/TraceInsights/Private/Insights/ViewModels/DrawHelpers.h
@@ -0,0 +1,90 @@
+// Copyright 1998-2019 Epic Games, Inc. All Rights Reserved.
+
+#pragma once
+
+#include "CoreMinimal.h"
+#include "Fonts/SlateFontInfo.h"
+#include "Math/Color.h"
+
+////////////////////////////////////////////////////////////////////////////////////////////////////
+
+struct FDrawContext;
+struct FSlateBrush;
+
+class FTimingTrackViewport;
+
+////////////////////////////////////////////////////////////////////////////////////////////////////
+
+class FDrawHelpers
+{
+public:
+	/**
+	 * Draw background.
+	 * X1-X2 specifies the valid area. X0-X1 and X2-X3 specifies the invalid areas.
+	 *
+	 *  X0    X1    X2     X3
+	 *  ++++++|*****|++++++
+	 *  ++++++|*****|++++++
+	 *  ++++++|*****|++++++
+	 */
+	static void DrawBackground(const FDrawContext& DrawContext,
+							   const FSlateBrush* BackgroundAreaBrush,
+							   const FLinearColor& ValidAreaColor,
+							   const FLinearColor& InvalidAreaColor,
+							   const FLinearColor& EdgeColor,
+							   const float X0,
+							   const float X1,
+							   const float X2,
+							   const float X3,
+							   const float Y,
+							   const float H,
+							   float& OutValidAreaX,
+							   float& OutValidAreaW);
+
+	static void DrawBackground(const FDrawContext& DrawContext,
+							   const FSlateBrush* BackgroundAreaBrush,
+							   const float X0,
+							   const float X1,
+							   const float X2,
+							   const float X3,
+							   const float Y,
+							   const float H);
+
+	static void DrawBackground(const FDrawContext& DrawContext,
+							   const FSlateBrush* BackgroundAreaBrush,
+							   const FTimingTrackViewport& Viewport,
+							   const float Y,
+							   const float H);
+
+	static void DrawBackground(const FDrawContext& DrawContext,
+							   const FSlateBrush* BackgroundAreaBrush,
+							   const FTimingTrackViewport& Viewport,
+							   const float Y,
+							   const float H,
+							   float& OutValidAreaX,
+							   float& OutValidAreaW);
+
+	/**
+	 * Draw time range selection.
+	 */
+	static void DrawTimeRangeSelection(const FDrawContext& DrawContext,
+									   const FTimingTrackViewport& Viewport,
+									   const double StartTime,
+									   const double EndTime,
+									   const FSlateBrush* Brush,
+									   const FSlateFontInfo& Font);
+
+	static void DrawSelection(const FDrawContext& DrawContext,
+							  const float MinX,
+							  const float MaxX,
+							  float SelectionX1,
+							  float SelectionX2,
+							  const float Y,
+							  const float H,
+							  const float TextY,
+							  const FString Text,
+							  const FSlateBrush* Brush,
+							  const FSlateFontInfo& Font);
+};
+
+////////////////////////////////////////////////////////////////////////////////////////////////////
diff --git a/Engine/Source/Developer/TraceInsights/Private/Insights/ViewModels/FrameTrackHelper.cpp b/Engine/Source/Developer/TraceInsights/Private/Insights/ViewModels/FrameTrackHelper.cpp
--- a/Engine/Source/Developer/TraceInsights/Private/Insights/ViewModels/FrameTrackHelper.cpp
+++ b/Engine/Source/Developer/TraceInsights/Private/Insights/ViewModels/FrameTrackHelper.cpp
@@ -2,9 +2,6 @@
 
 #include "FrameTrackHelper.h"
 
-#include "Brushes/SlateBorderBrush.h"
-#include "Brushes/SlateColorBrush.h"
-#include "EditorStyleSet.h"
 #include "Fonts/FontMeasure.h"
 #include "Fonts/SlateFontInfo.h"
 #include "Framework/Application/SlateApplication.h"
@@ -13,32 +10,34 @@
 
 // Insights
 #include "Insights/Common/PaintUtils.h"
+#include "Insights/InsightsStyle.h"
+#include "Insights/ViewModels/DrawHelpers.h"
 #include "Insights/ViewModels/FrameTrackViewport.h"
 
 #include <limits>
 
 ////////////////////////////////////////////////////////////////////////////////////////////////////
-// FFrameTrackTimelineBuilder
+// FFrameTrackSeriesBuilder
 ////////////////////////////////////////////////////////////////////////////////////////////////////
 
-FFrameTrackTimelineBuilder::FFrameTrackTimelineBuilder(FFrameTrackTimeline& InTimeline, const FFrameTrackViewport& InViewport)
-	: Timeline(InTimeline)
+FFrameTrackSeriesBuilder::FFrameTrackSeriesBuilder(FFrameTrackSeries& InSeries, const FFrameTrackViewport& InViewport)
+	: Series(InSeries)
 	, Viewport(InViewport)
 	, NumAddedFrames(0)
 {
 	SampleW = Viewport.GetSampleWidth();
 	FramesPerSample = Viewport.GetNumFramesPerSample();
-	NumSamples = FMath::Max(0, FMath::CeilToInt(Viewport.Width / SampleW));
+	NumSamples = FMath::Max(0, FMath::CeilToInt(Viewport.GetWidth() / SampleW));
 	FirstFrameIndex = Viewport.GetFirstFrameIndex();
 
-	Timeline.NumAggregatedFrames = 0;
-	Timeline.Samples.Reset();
-	Timeline.Samples.AddDefaulted(NumSamples);
+	Series.NumAggregatedFrames = 0;
+	Series.Samples.Reset();
+	Series.Samples.AddDefaulted(NumSamples);
 }
 
 ////////////////////////////////////////////////////////////////////////////////////////////////////
 
-void FFrameTrackTimelineBuilder::AddFrame(const Trace::FFrame& Frame)
+void FFrameTrackSeriesBuilder::AddFrame(const Trace::FFrame& Frame)
 {
 	NumAddedFrames++;
 
@@ -47,7 +46,7 @@ void FFrameTrackTimelineBuilder::AddFrame(const Trace::FFrame& Frame)
 	int32 SampleIndex = (FrameIndex - FirstFrameIndex) / FramesPerSample;
 	if (SampleIndex >= 0 && SampleIndex < NumSamples)
 	{
-		FFrameTrackSample& Sample = Timeline.Samples[SampleIndex];
+		FFrameTrackSample& Sample = Series.Samples[SampleIndex];
 		Sample.NumFrames++;
 
 		double Duration = Frame.EndTime - Frame.StartTime;
@@ -67,7 +66,7 @@ void FFrameTrackTimelineBuilder::AddFrame(const Trace::FFrame& Frame)
 			Sample.LargestFrameDuration = Duration;
 		}
 
-		Timeline.NumAggregatedFrames++;
+		Series.NumAggregatedFrames++;
 	}
 }
 
@@ -78,8 +77,8 @@ void FFrameTrackTimelineBuilder::AddFrame(const Trace::FFrame& Frame)
 FFrameTrackDrawHelper::FFrameTrackDrawHelper(const FDrawContext& InDrawContext, const FFrameTrackViewport& InViewport)
 	: DrawContext(InDrawContext)
 	, Viewport(InViewport)
-	, WhiteBrush(FCoreStyle::Get().GetBrush("WhiteBrush"))
-	, BorderBrush(FEditorStyle::GetBrush("PlainBorder"))
+	, WhiteBrush(FInsightsStyle::Get().GetBrush("WhiteBrush"))
+	, HoveredFrameBorderBrush(FInsightsStyle::Get().GetBrush("HoveredEventBorder"))
 	, NumFrames(0)
 	, NumDrawSamples(0)
 {
@@ -89,62 +88,34 @@ FFrameTrackDrawHelper::FFrameTrackDrawHelper(const FDrawContext& InDrawContext,
 
 void FFrameTrackDrawHelper::DrawBackground() const
 {
-	const FSlateBrush* AreaBrush = WhiteBrush;
-	const FLinearColor ValidAreaColor(0.07f, 0.07f, 0.07f, 1.0f);
-	const FLinearColor InvalidAreaColor(0.1f, 0.07f, 0.07f, 1.0f);
+	const FAxisViewportInt32& ViewportX = Viewport.GetHorizontalAxisViewport();
 
-	float X0 = Viewport.MinX - Viewport.PosX;
-	float X1 = Viewport.MaxX - Viewport.PosX;
-	const float W = FMath::CeilToFloat(Viewport.Width);
-	const float H = FMath::CeilToFloat(Viewport.Height);
+	const float X0 = 0.0f;
+	const float X1 = ViewportX.GetMinPos() - ViewportX.GetPos();
+	const float X2 = ViewportX.GetMaxPos() - ViewportX.GetPos();
+	const float X3 = FMath::CeilToFloat(Viewport.GetWidth());
 
-	if (X0 >= W || X1 <= 0.0f)
-	{
-		// Draw invalid area (entire view).
-		DrawContext.DrawBox(0.0f, 0.0f, W, H, AreaBrush, InvalidAreaColor);
-	}
-	else // X0 < W && X1 > 0
-	{
-		if (X0 > 0.0f)
-		{
-			// Draw invalid area (left).
-			DrawContext.DrawBox(0.0f, 0.0f, X0, H, AreaBrush, InvalidAreaColor);
-		}
-
-		if (X1 < W)
-		{
-			// Draw invalid area (right).
-			DrawContext.DrawBox(X1, 0.0f, W - X1, H, AreaBrush, InvalidAreaColor);
-		}
-
-		X0 = FMath::Max(X0, 0.0f);
-		X1 = FMath::Min(X1, W);
+	const float Y = 0.0f;
+	const float H = FMath::CeilToFloat(Viewport.GetHeight());
 
-		if (X1 > X0)
-		{
-			// Draw valid area.
-			DrawContext.DrawBox(X0, 0.0f, X1 - X0, H, AreaBrush, ValidAreaColor);
-		}
-	}
-
-	DrawContext.LayerId++;
+	FDrawHelpers::DrawBackground(DrawContext, WhiteBrush, X0, X1, X2, X3, Y, H);
 }
 
 ////////////////////////////////////////////////////////////////////////////////////////////////////
 
-FLinearColor FFrameTrackDrawHelper::GetColorById(int32 Id) const
+FLinearColor FFrameTrackDrawHelper::GetColorByFrameType(int32 FrameType)
 {
-	const float Alpha = 0.9;
-	switch (Id)
+	const float Alpha = 0.9f;
+	switch (FrameType)
 	{
-	case 0: // Game Frames
+	case TraceFrameType_Game:
 		return FLinearColor(0.75, 1.0, 1.0, Alpha);
 
-	case 1: // Rendering Frames
+	case TraceFrameType_Rendering:
 		return FLinearColor(1.0, 0.75, 0.75, Alpha);
 
-	case 2:
-		return FLinearColor(0.75, 0.75, 1.0, Alpha);
+	//case 2:
+	//	return FLinearColor(0.75, 0.75, 1.0, Alpha);
 
 	default:
 		return FLinearColor(1.0, 1.0, 1.0, Alpha);
@@ -153,25 +124,28 @@ FLinearColor FFrameTrackDrawHelper::GetColorById(int32 Id) const
 
 ////////////////////////////////////////////////////////////////////////////////////////////////////
 
-void FFrameTrackDrawHelper::DrawCached(const FFrameTrackTimeline& Timeline) const
+void FFrameTrackDrawHelper::DrawCached(const FFrameTrackSeries& Series) const
 {
-	if (Timeline.NumAggregatedFrames == 0)
+	if (Series.NumAggregatedFrames == 0)
 	{
 		return;
 	}
 
-	NumFrames += Timeline.NumAggregatedFrames;
+	NumFrames += Series.NumAggregatedFrames;
 
-	FLinearColor TimelineColor = GetColorById(static_cast<int32>(Timeline.Id));
+	FLinearColor SeriesColor = GetColorByFrameType(Series.FrameType);
 
 	const float SampleW = Viewport.GetSampleWidth();
-	const int32 NumSamples = Timeline.Samples.Num();
+	const int32 NumSamples = Series.Samples.Num();
+
+	const FAxisViewportDouble& ViewportY = Viewport.GetVerticalAxisViewport();
 
-	const float Y1 = FMath::RoundToFloat(Viewport.GetViewportYForValue(0.0));
+	const float ViewHeight = FMath::RoundToFloat(Viewport.GetHeight());
+	const float BaselineY = FMath::RoundToFloat(ViewportY.GetOffsetForValue(0.0));
 
 	for (int32 SampleIndex = 0; SampleIndex < NumSamples; SampleIndex++)
 	{
-		const FFrameTrackSample& Sample = Timeline.Samples[SampleIndex];
+		const FFrameTrackSample& Sample = Series.Samples[SampleIndex];
 		if (Sample.NumFrames == 0)
 		{
 			continue;
@@ -180,20 +154,20 @@ void FFrameTrackDrawHelper::DrawCached(const FFrameTrackTimeline& Timeline) cons
 		NumDrawSamples++;
 
 		const float X = SampleIndex * SampleW;
-		float Y2;
+		float ValueY;
 
-		FLinearColor ColorFill = TimelineColor;
+		FLinearColor ColorFill = SeriesColor;
 
 		if (Sample.LargestFrameDuration == std::numeric_limits<double>::infinity())
 		{
-			Y2 = Viewport.Height;
+			ValueY = ViewHeight;
 			ColorFill.R = 0.0f;
 			ColorFill.G = 0.0f;
 			ColorFill.B = 0.0f;
 		}
 		else
 		{
-			Y2 = FMath::RoundToFloat(Viewport.GetViewportYForValue(Sample.LargestFrameDuration));
+			ValueY = FMath::RoundToFloat(ViewportY.GetOffsetForValue(Sample.LargestFrameDuration));
 			if (Sample.LargestFrameDuration > 1.0 / 30.0)
 			{
 				ColorFill.G *= 0.5f;
@@ -205,8 +179,8 @@ void FFrameTrackDrawHelper::DrawCached(const FFrameTrackTimeline& Timeline) cons
 			}
 		}
 
-		const float Y = FMath::RoundToFloat(Viewport.Height) - Y2;
-		const float H = Y2 - Y1;
+		const float H = ValueY - BaselineY;
+		const float Y = ViewHeight - H;
 
 		const FLinearColor ColorBorder(ColorFill.R * 0.75f, ColorFill.G * 0.75f, ColorFill.B * 0.75f, 1.0);
 
@@ -239,38 +213,42 @@ void FFrameTrackDrawHelper::DrawHoveredSample(const FFrameTrackSample& Sample) c
 	const int32 SampleIndex = (Sample.LargestFrameIndex - FirstFrameIndex) / FramesPerSample;
 	const float X = SampleIndex * SampleW;
 
-	const float Y1 = FMath::RoundToFloat(Viewport.GetViewportYForValue(0.0));
-	float Y2;
+	const FAxisViewportDouble& ViewportY = Viewport.GetVerticalAxisViewport();
+
+	const float ViewHeight = FMath::RoundToFloat(Viewport.GetHeight());
+	const float BaselineY = FMath::RoundToFloat(ViewportY.GetOffsetForValue(0.0));
+
+	float ValueY;
 	if (Sample.LargestFrameDuration == std::numeric_limits<double>::infinity())
 	{
-		Y2 = Viewport.Height;
+		ValueY = ViewHeight;
 	}
 	else
 	{
-		Y2 = FMath::RoundToFloat(Viewport.GetViewportYForValue(Sample.LargestFrameDuration));
+		ValueY = FMath::RoundToFloat(ViewportY.GetOffsetForValue(Sample.LargestFrameDuration));
 	}
-	const float Y = FMath::RoundToFloat(Viewport.Height) - Y2;
-	const float H = Y2 - Y1;
+	const float H = ValueY - BaselineY;
+	const float Y = ViewHeight - H;
 
 	const FLinearColor ColorBorder(1.0f, 1.0f, 0.0f, 1.0);
-	DrawContext.DrawBox(X - 1.0f, Y - 1.0f, SampleW + 2.0f, H + 2.0f, BorderBrush, ColorBorder);
+	DrawContext.DrawBox(X - 1.0f, Y - 1.0f, SampleW + 2.0f, H + 2.0f, HoveredFrameBorderBrush, ColorBorder);
 	DrawContext.LayerId++;
 }
 
 ////////////////////////////////////////////////////////////////////////////////////////////////////
 
-void FFrameTrackDrawHelper::DrawHighlightedInterval(const FFrameTrackTimeline& Timeline, const double StartTime, const double EndTime) const
+void FFrameTrackDrawHelper::DrawHighlightedInterval(const FFrameTrackSeries& Series, const double StartTime, const double EndTime) const
 {
-	const int32 NumSamples = Timeline.Samples.Num();
+	const int32 NumSamples = Series.Samples.Num();
 
 	//TODO: binary search
 	int32 Index1 = 0;
 	int32 Index2 = NumSamples - 1;
-	while (Index1 < NumSamples && Timeline.Samples[Index1].EndTime < StartTime)
+	while (Index1 < NumSamples && Series.Samples[Index1].EndTime < StartTime)
 	{
 		Index1++;
 	}
-	while (Index2 >= Index1 && Timeline.Samples[Index2].StartTime > EndTime)
+	while (Index2 >= Index1 && Series.Samples[Index2].StartTime > EndTime)
 	{
 		Index2--;
 	}
@@ -279,23 +257,23 @@ void FFrameTrackDrawHelper::DrawHighlightedInterval(const FFrameTrackTimeline& T
 	{
 		const float SampleW = Viewport.GetSampleWidth();
 		float X1 = Index1 * SampleW;
-		float X2 = Index2 * SampleW;
+		float X2 = (Index2 + 1) * SampleW;
 
-		constexpr float Y1 = 12.0f; // allows 12px for the horizontal scrollbar (one displayed on top of the track)
-		const float Y2 = Viewport.Height;
+		constexpr float Y1 = 0.0f; // allows 12px for the horizontal scrollbar (one displayed on top of the track)
+		const float Y2 = Viewport.GetHeight();
 		constexpr float D = 2.0f; // line thickness (for both horizontal and vertical lines)
 		constexpr float H = 10.0f; // height of corner lines
 
 		const FLinearColor Color(1.0f, 1.0f, 1.0f, 1.0f);
 
-		if (X1 >= 0.0f && X1 < Viewport.Width - 2.0f)
+		if (X1 >= 0.0f && X1 < Viewport.GetWidth() - 2.0f)
 		{
 			// Draw left side vertical lines.
 			DrawContext.DrawBox(X1 - D, Y1, D, H, WhiteBrush, Color);
 			DrawContext.DrawBox(X1 - D, Y2 - H, D, H, WhiteBrush, Color);
 		}
 
-		if (X2 >= -2.0f && X2 < Viewport.Width)
+		if (X2 >= -2.0f && X2 < Viewport.GetWidth())
 		{
 			// Draw right side vertical lines.
 			DrawContext.DrawBox(X2, Y1, D, H, WhiteBrush, Color);
@@ -306,9 +284,9 @@ void FFrameTrackDrawHelper::DrawHighlightedInterval(const FFrameTrackTimeline& T
 		{
 			X1 = 0.0f;
 		}
-		if (X2 > Viewport.Width)
+		if (X2 > Viewport.GetWidth())
 		{
-			X2 = Viewport.Width;
+			X2 = Viewport.GetWidth();
 		}
 		if (X1 < X2)
 		{
diff --git a/Engine/Source/Developer/TraceInsights/Private/Insights/ViewModels/FrameTrackHelper.h b/Engine/Source/Developer/TraceInsights/Private/Insights/ViewModels/FrameTrackHelper.h
--- a/Engine/Source/Developer/TraceInsights/Private/Insights/ViewModels/FrameTrackHelper.h
+++ b/Engine/Source/Developer/TraceInsights/Private/Insights/ViewModels/FrameTrackHelper.h
@@ -21,23 +21,13 @@ namespace Trace
 
 ////////////////////////////////////////////////////////////////////////////////////////////////////
 
-enum class EFrameType
-{
-	Unknown = -1,
-
-	Game,
-	Render
-};
-
-////////////////////////////////////////////////////////////////////////////////////////////////////
-
 struct FFrameTrackSample
 {
 	int32 NumFrames;
 	double TotalDuration; // sum of durations of all frames in this sample
 	double StartTime; // min start time of all frames in this sample
 	double EndTime; // max end time of all frames in this sample
-	double LargestFrameIndex; // index of the largest frame
+	int32 LargestFrameIndex; // index of the largest frame
 	double LargestFrameStartTime; // start time of the largest frame
 	double LargestFrameDuration; // duration of the largest frame
 
@@ -50,46 +40,73 @@ struct FFrameTrackSample
 		, LargestFrameStartTime(0.0)
 		, LargestFrameDuration(0.0)
 	{}
+
+	FFrameTrackSample(const FFrameTrackSample&) = default;
+	FFrameTrackSample& operator=(const FFrameTrackSample&) = default;
+
+	FFrameTrackSample(FFrameTrackSample&&) = default;
+	FFrameTrackSample& operator=(FFrameTrackSample&&) = default;
+
+	bool Equals(const FFrameTrackSample& Other) const
+	{
+		return NumFrames == Other.NumFrames
+			&& TotalDuration == Other.TotalDuration
+			&& StartTime == Other.StartTime
+			&& EndTime == Other.EndTime
+			&& LargestFrameIndex == Other.LargestFrameIndex
+			&& LargestFrameStartTime == Other.LargestFrameStartTime
+			&& LargestFrameDuration == Other.LargestFrameDuration;
+	}
+
+	static bool AreEquals(const FFrameTrackSample& A, const FFrameTrackSample& B)
+	{
+		return A.Equals(B);
+	}
 };
 
 ////////////////////////////////////////////////////////////////////////////////////////////////////
 
-struct FFrameTrackTimeline
+struct FFrameTrackSeries
 {
-	uint64 Id;
-	//TODO: EFrameType Type;
-
-	//TODO: double MinValue; // min value
-	//TODO: double MaxValue; // max value
-
-	//TODO: float SampleW; // width of a sample, in Slate units
-	//TODO: int32 FramesPerSample; // number of frames in a sample
-	//TODO: int32 FirstFrameIndex; // index of first frame in first sample; can be negative
-	//TODO: int32 NumSamples; // total number of samples
-	int32 NumAggregatedFrames; // total number of frames aggregated in samples; ie. sum of all Sample.NumFrames
-
-	TArray<FFrameTrackSample> Samples;
+	int32 FrameType;
+	bool bIsVisible;
+	int32 NumAggregatedFrames; // total number of frames aggregated in samples; i.e. sum of all Sample.NumFrames
+	TArray<FFrameTrackSample> Samples; // the aggregated samples
+
+	explicit FFrameTrackSeries(int32 InFrameType)
+		: FrameType(InFrameType)
+		, bIsVisible(true)
+		, NumAggregatedFrames(0)
+		, Samples()
+	{
+	}
+
+	void Reset()
+	{
+		NumAggregatedFrames = 0;
+		Samples.Reset();
+	}
 };
 
 ////////////////////////////////////////////////////////////////////////////////////////////////////
 
-class FFrameTrackTimelineBuilder
+class FFrameTrackSeriesBuilder
 {
 public:
-	FFrameTrackTimelineBuilder(FFrameTrackTimeline& InTimeline, const FFrameTrackViewport& InViewport);
+	explicit FFrameTrackSeriesBuilder(FFrameTrackSeries& InSeries, const FFrameTrackViewport& InViewport);
 
 	/**
 	 * Non-copyable
 	 */
-	FFrameTrackTimelineBuilder(const FFrameTrackTimelineBuilder&) = delete;
-	FFrameTrackTimelineBuilder& operator=(const FFrameTrackTimelineBuilder&) = delete;
+	FFrameTrackSeriesBuilder(const FFrameTrackSeriesBuilder&) = delete;
+	FFrameTrackSeriesBuilder& operator=(const FFrameTrackSeriesBuilder&) = delete;
 
 	void AddFrame(const Trace::FFrame& Frame);
 
 	int32 GetNumAddedFrames() const { return NumAddedFrames; }
 
 private:
-	FFrameTrackTimeline& Timeline;
+	FFrameTrackSeries& Series; // series to update
 	const FFrameTrackViewport& Viewport;
 
 	float SampleW; // width of a sample, in Slate units
@@ -106,7 +123,7 @@ private:
 class FFrameTrackDrawHelper
 {
 public:
-	FFrameTrackDrawHelper(const FDrawContext& InDrawContext, const FFrameTrackViewport& InViewport);
+	explicit FFrameTrackDrawHelper(const FDrawContext& InDrawContext, const FFrameTrackViewport& InViewport);
 
 	/**
 	 * Non-copyable
@@ -115,11 +132,11 @@ public:
 	FFrameTrackDrawHelper& operator=(const FFrameTrackDrawHelper&) = delete;
 
 	void DrawBackground() const;
-	void DrawCached(const FFrameTrackTimeline& Timeline) const;
+	void DrawCached(const FFrameTrackSeries& Series) const;
 	void DrawHoveredSample(const FFrameTrackSample& Sample) const;
-	void DrawHighlightedInterval(const FFrameTrackTimeline& Timeline, const double StartTime, const double EndTime) const;
+	void DrawHighlightedInterval(const FFrameTrackSeries& Series, const double StartTime, const double EndTime) const;
 
-	FLinearColor GetColorById(int32 Id) const;
+	static FLinearColor GetColorByFrameType(int32 FrameType);
 
 	int32 GetNumFrames() const { return NumFrames; }
 	int32 GetNumDrawSamples() const { return NumDrawSamples; }
@@ -129,7 +146,7 @@ private:
 	const FFrameTrackViewport& Viewport;
 
 	const FSlateBrush* WhiteBrush;
-	const FSlateBrush* BorderBrush;
+	const FSlateBrush* HoveredFrameBorderBrush;
 
 	// Debug stats.
 	mutable int32 NumFrames;
diff --git a/Engine/Source/Developer/TraceInsights/Private/Insights/ViewModels/MarkersTimingTrack.cpp b/Engine/Source/Developer/TraceInsights/Private/Insights/ViewModels/MarkersTimingTrack.cpp
--- a/Engine/Source/Developer/TraceInsights/Private/Insights/ViewModels/MarkersTimingTrack.cpp
+++ b/Engine/Source/Developer/TraceInsights/Private/Insights/ViewModels/MarkersTimingTrack.cpp
@@ -2,7 +2,6 @@
 
 #include "MarkersTimingTrack.h"
 
-#include "Brushes/SlateColorBrush.h"
 #include "Fonts/FontMeasure.h"
 #include "Framework/Application/SlateApplication.h"
 #include "Styling/CoreStyle.h"
@@ -11,6 +10,7 @@
 // Insights
 #include "Insights/Common/PaintUtils.h"
 #include "Insights/InsightsManager.h"
+#include "Insights/InsightsStyle.h"
 #include "Insights/TimingProfilerManager.h"
 #include "Insights/ViewModels/TimingTrackViewport.h"
 
@@ -28,13 +28,12 @@ FMarkersTimingTrack::FMarkersTimingTrack(uint64 InTrackId)
 	//, TimeMarkerTexts()
 	, bIsCollapsed(false)
 	, bUseOnlyBookmarks(true)
-	, bIsDirty(true)
 	, TargetHoveredAnimPercent(0.0f)
 	, CurrentHoveredAnimPercent(0.0f)
 	, NumLogMessages(0)
 	, NumDrawBoxes(0)
 	, NumDrawTexts(0)
-	, WhiteBrush(FCoreStyle::Get().GetBrush("WhiteBrush"))
+	, WhiteBrush(FInsightsStyle::Get().GetBrush("WhiteBrush"))
 	, Font(FCoreStyle::GetDefaultFontStyle("Regular", 8))
 {
 }
@@ -56,7 +55,6 @@ void FMarkersTimingTrack::Reset()
 
 	bIsCollapsed = false;
 	bUseOnlyBookmarks = true;
-	bIsDirty = true;
 
 	TargetHoveredAnimPercent = 0.0f;
 	CurrentHoveredAnimPercent = 0.0f;
@@ -147,8 +145,8 @@ void FMarkersTimingTrack::Update(const FTimingTrackViewport& InViewport)
 		Builder.BeginLog(LogProvider);
 
 		LogProvider.EnumerateMessages(
-			Builder.GetViewport().StartTime,
-			Builder.GetViewport().EndTime,
+			Builder.GetViewport().GetStartTime(),
+			Builder.GetViewport().GetEndTime(),
 			[&Builder](const Trace::FLogMessage& Message) { Builder.AddLogMessage(Message); });
 
 		Builder.EndLog();
@@ -159,8 +157,10 @@ void FMarkersTimingTrack::Update(const FTimingTrackViewport& InViewport)
 
 void FMarkersTimingTrack::Draw(FDrawContext& DrawContext, const FTimingTrackViewport& Viewport) const
 {
+	const FLinearColor BackgroundColor(0.04f, 0.04f, 0.04f, 1.0f);
+
 	// Draw background.
-	DrawContext.DrawBox(0.0f, GetPosY(), Viewport.Width, GetHeight(), WhiteBrush, FLinearColor(0.04f, 0.04f, 0.04f, 1.0f));
+	DrawContext.DrawBox(0.0f, GetPosY(), Viewport.GetWidth(), GetHeight(), WhiteBrush, BackgroundColor);
 	DrawContext.LayerId++;
 
 	// Draw the track's header, in background.
@@ -182,10 +182,10 @@ void FMarkersTimingTrack::Draw(FDrawContext& DrawContext, const FTimingTrackView
 	else
 	{
 		BoxY = 0.0f;
-		BoxH = Viewport.Height;
+		BoxH = Viewport.GetHeight();
 	}
 
-	int32 NumBoxes = TimeMarkerBoxes.Num();
+	const int32 NumBoxes = TimeMarkerBoxes.Num();
 	for (int32 BoxIndex = 0; BoxIndex < NumBoxes; BoxIndex++)
 	{
 		const FTimeMarkerBoxInfo& Box = TimeMarkerBoxes[BoxIndex];
@@ -200,7 +200,7 @@ void FMarkersTimingTrack::Draw(FDrawContext& DrawContext, const FTimingTrackView
 	const float CategoryY = GetPosY() + 2.0f;
 	const float MessageY = GetPosY() + (IsBookmarksTrack() ? 1.0f : 14.0f);
 
-	int32 NumTexts = TimeMarkerTexts.Num();
+	const int32 NumTexts = TimeMarkerTexts.Num();
 	for (int32 TextIndex = 0; TextIndex < NumTexts; TextIndex++)
 	{
 		const FTimeMarkerTextInfo& TextInfo = TimeMarkerTexts[TextIndex];
@@ -454,7 +454,7 @@ void FTimeMarkerTrackBuilder::Flush(float AvailableTextW)
 
 void FTimeMarkerTrackBuilder::AddTimeMarker(const float X, const uint64 LogIndex, const ELogVerbosity::Type Verbosity, const TCHAR* const Category, const TCHAR* Message)
 {
-	float W = X - LastX2;
+	const float W = X - LastX2;
 
 	if (W > 0.0f) // There is at least 1px from previous box?
 	{
@@ -518,7 +518,7 @@ void FTimeMarkerTrackBuilder::AddTimeMarker(const float X, const uint64 LogIndex
 
 void FTimeMarkerTrackBuilder::EndLog()
 {
-	Flush(Viewport.Width - LastX2);
+	Flush(Viewport.GetWidth() - LastX2);
 }
 
 ////////////////////////////////////////////////////////////////////////////////////////////////////
diff --git a/Engine/Source/Developer/TraceInsights/Private/Insights/ViewModels/MarkersTimingTrack.h b/Engine/Source/Developer/TraceInsights/Private/Insights/ViewModels/MarkersTimingTrack.h
--- a/Engine/Source/Developer/TraceInsights/Private/Insights/ViewModels/MarkersTimingTrack.h
+++ b/Engine/Source/Developer/TraceInsights/Private/Insights/ViewModels/MarkersTimingTrack.h
@@ -44,7 +44,7 @@ class FMarkersTimingTrack : public FBaseTimingTrack
 	friend class FTimeMarkerTrackBuilder;
 
 public:
-	FMarkersTimingTrack(uint64 InTrackId);
+	explicit FMarkersTimingTrack(uint64 InTrackId);
 	virtual ~FMarkersTimingTrack();
 
 	virtual void Reset() override;
@@ -59,10 +59,6 @@ public:
 		UpdateHeight();
 	}
 
-	bool IsDirty() const { return bIsDirty; }
-	void SetDirtyFlag() { bIsDirty = true; }
-	void ClearDirtyFlag() { bIsDirty = false; }
-
 	// Stats
 	int32 GetNumLogMessages() const { return NumLogMessages; }
 	int32 GetNumBoxes() const { return TimeMarkerBoxes.Num(); }
@@ -91,7 +87,6 @@ private:
 
 	bool bIsCollapsed; // If false, the vertical lines will extend to entire viewport height; otherwise will be limited to this track's height.
 	bool bUseOnlyBookmarks; // If true, uses only bookmarks; otherwise it uses all log messages.
-	bool bIsDirty; // Source list of bookmarks/markers has changed. Cached Boxes and Texts are dirty.
 
 	float TargetHoveredAnimPercent; // [0.0 .. 1.0], 0.0 = hidden, 1.0 = visible
 	float CurrentHoveredAnimPercent;
@@ -111,7 +106,7 @@ private:
 class FTimeMarkerTrackBuilder
 {
 public:
-	FTimeMarkerTrackBuilder(FMarkersTimingTrack& InTrack, const FTimingTrackViewport& InViewport);
+	explicit FTimeMarkerTrackBuilder(FMarkersTimingTrack& InTrack, const FTimingTrackViewport& InViewport);
 
 	/**
 	 * Non-copyable
diff --git a/Engine/Source/Developer/TraceInsights/Private/Insights/ViewModels/StatsNode.cpp b/Engine/Source/Developer/TraceInsights/Private/Insights/ViewModels/StatsNode.cpp
--- a/Engine/Source/Developer/TraceInsights/Private/Insights/ViewModels/StatsNode.cpp
+++ b/Engine/Source/Developer/TraceInsights/Private/Insights/ViewModels/StatsNode.cpp
@@ -7,6 +7,8 @@
 
 #define LOCTEXT_NAMESPACE "StatsNode"
 
+const FName FStatsNode::TypeName(TEXT("FStatsNode"));
+
 ////////////////////////////////////////////////////////////////////////////////////////////////////
 
 void FStatsNode::ResetAggregatedStats()
@@ -152,40 +154,4 @@ const FText FStatsNode::GetTextForAggregatedStatsUpperQuartile() const
 
 ////////////////////////////////////////////////////////////////////////////////////////////////////
 
-const FText FStatsNode::GetNameEx() const
-{
-	FText Text = FText::GetEmpty();
-
-	if (IsGroup())
-	{
-		const int32 NumChildren = Children.Num();
-		const int32 NumFilteredChildren = FilteredChildren.Num();
-
-		if (NumFilteredChildren == NumChildren)
-		{
-			Text = FText::Format(LOCTEXT("StatsNodeGroupTextFmt1", "{0} ({1})"), FText::FromName(Name), FText::AsNumber(NumChildren));
-		}
-		else
-		{
-			Text = FText::Format(LOCTEXT("StatsNodeGroupTextFmt2", "{0} ({1} / {2})"), FText::FromName(Name), FText::AsNumber(NumFilteredChildren), FText::AsNumber(NumChildren));
-		}
-	}
-	else
-	{
-		//const bool bIsStatTracked = FProfilerManager::Get()->IsStatTracked(GroupOrStatNode->GetStatID());
-		//if (bIsStatTracked)
-		//{
-		//	Text = FText::Format(LOCTEXT("StatsNodeTrackedTextFmt", "{0}*"), FText::FromName(Name));
-		//}
-		//else
-		{
-			Text = FText::FromName(Name);
-		}
-	}
-
-	return Text;
-}
-
-////////////////////////////////////////////////////////////////////////////////////////////////////
-
 #undef LOCTEXT_NAMESPACE
diff --git a/Engine/Source/Developer/TraceInsights/Private/Insights/ViewModels/StatsNode.h b/Engine/Source/Developer/TraceInsights/Private/Insights/ViewModels/StatsNode.h
--- a/Engine/Source/Developer/TraceInsights/Private/Insights/ViewModels/StatsNode.h
+++ b/Engine/Source/Developer/TraceInsights/Private/Insights/ViewModels/StatsNode.h
@@ -4,19 +4,10 @@
 
 #include "CoreMinimal.h"
 
-class FStatsNode;
-
-/** Type definition for shared pointers to instances of FStatsNode. */
-typedef TSharedPtr<class FStatsNode> FStatsNodePtr;
-
-/** Type definition for shared references to instances of FStatsNode. */
-typedef TSharedRef<class FStatsNode> FStatsNodeRef;
-
-/** Type definition for shared references to const instances of FStatsNode. */
-typedef TSharedRef<const class FStatsNode> FStatsNodeRefConst;
+// Insights
+#include "Insights/Table/ViewModels/BaseTreeNode.h"
 
-/** Type definition for weak references to instances of FStatsNode. */
-typedef TWeakPtr<class FStatsNode> FStatsNodeWeak;
+////////////////////////////////////////////////////////////////////////////////////////////////////
 
 enum class EStatsNodeType
 {
@@ -33,6 +24,8 @@ enum class EStatsNodeType
 	InvalidOrMax,
 };
 
+////////////////////////////////////////////////////////////////////////////////////////////////////
+
 template<typename Type>
 struct TAggregatedStats
 {
@@ -74,59 +67,63 @@ struct TAggregatedStats
 typedef TAggregatedStats<double> FAggregatedStats;
 typedef TAggregatedStats<int64> FAggregatedIntegerStats;
 
+////////////////////////////////////////////////////////////////////////////////////////////////////
+
+class FStatsNode;
+
+/** Type definition for shared pointers to instances of FStatsNode. */
+typedef TSharedPtr<class FStatsNode> FStatsNodePtr;
+
+/** Type definition for shared references to instances of FStatsNode. */
+typedef TSharedRef<class FStatsNode> FStatsNodeRef;
+
+/** Type definition for shared references to const instances of FStatsNode. */
+typedef TSharedRef<const class FStatsNode> FStatsNodeRefConst;
+
+/** Type definition for weak references to instances of FStatsNode. */
+typedef TWeakPtr<class FStatsNode> FStatsNodeWeak;
+
+////////////////////////////////////////////////////////////////////////////////////////////////////
 /**
- * Class used to store information about a timer node (used in timers' tree view).
+ * Class used to store information about a stats counter node (used in SStatsView).
  */
-class FStatsNode : public TSharedFromThis<FStatsNode>
+class FStatsNode : public Insights::FBaseTreeNode
 {
 public:
+	static const FName TypeName;
 	static const uint64 InvalidId = -1;
 
 public:
-	/** Initialization constructor for the timer node. */
+	/** Initialization constructor for the stats node. */
 	FStatsNode(uint64 InId, const FName InName, const FName InMetaGroupName, EStatsNodeType InType)
-		: Id(InId)
-		, Name(InName)
+		: FBaseTreeNode(InId, InName, InType == EStatsNodeType::Group)
 		, MetaGroupName(InMetaGroupName)
 		, Type(InType)
-		, bForceExpandGroupNode(false)
+		, bIsAddedToGraph(false)
 	{
+		const uint32 HashColor = GetId() * 0x2c2c57ed;
+		Color.R = ((HashColor >> 16) & 0xFF) / 255.0f;
+		Color.G = ((HashColor >> 8) & 0xFF) / 255.0f;
+		Color.B = ((HashColor) & 0xFF) / 255.0f;
+		Color.A = 1.0;
+
 		ResetAggregatedStats();
 	}
 
 	/** Initialization constructor for the group node. */
 	FStatsNode(const FName InGroupName)
-		: Id(0)
-		, Name(InGroupName)
+		: FBaseTreeNode(0, InGroupName, true)
 		, Type(EStatsNodeType::Group)
-		, bForceExpandGroupNode(false)
+		, Color(0.0, 0.0, 0.0, 1.0)
+		, bIsAddedToGraph(false)
 	{
 		ResetAggregatedStats();
 	}
 
-	/**
-	 * @return an Id of this timer, valid only for timer nodes.
-	 */
-	const uint64 GetId() const
-	{
-		return Id;
-	}
-
-	/**
-	 * @return a name of this node, group or timer.
-	 */
-	const FName& GetName() const
-	{
-		return Name;
-	}
-
-	/**
-	 * @return a name of this node, group or timer + addditional info, to display in Stats tree view.
-	*/
-	const FText GetNameEx() const;
+	virtual const FName& GetTypeName() const override { return TypeName; }
 
 	/**
-	 * @return a name of the group that this timer node belongs to, taken from the metadata.
+	 * @return a name of the meta group that this stats node belongs to, taken from the metadata.
 	 */
 	const FName& GetMetaGroupName() const
 	{
@@ -134,7 +131,7 @@ public:
 	}
 
 	/**
-	 * @return a type of this timer or EStatsNodeType::Group for group nodes.
+	 * @return a type of this stats node or EStatsNodeType::Group for group nodes.
 	 */
 	const EStatsNodeType& GetType() const
 	{
@@ -142,15 +139,15 @@ public:
 	}
 
 	/**
-	 * @return true, if this node is a group node.
+	 * @return color of the node. Used when showing a graph series for a stats counter.
 	 */
-	bool IsGroup() const
+	FLinearColor GetColor() const
 	{
-		return Type == EStatsNodeType::Group;
+		return Color;
 	}
 
 	/**
-	 * @return the aggregated stats of this counter (if counter is a "float number" type).
+	 * @return the aggregated stats of this stats counter (if counter is a "float number" type).
 	 */
 	const FAggregatedStats& GetAggregatedStats() const
 	{
@@ -162,7 +159,7 @@ public:
 	void SetAggregatedStats(FAggregatedStats& AggregatedStats);
 
 	/**
-	 * @return the aggregated stats of this counter  (if counter is an "integer number" type).
+	 * @return the aggregated stats of this stats counter (if counter is an "integer number" type).
 	 */
 	const FAggregatedIntegerStats& GetAggregatedIntegerStats() const
 	{
@@ -186,107 +183,44 @@ public:
 	const FText GetTextForAggregatedStatsLowerQuartile() const;
 	const FText GetTextForAggregatedStatsUpperQuartile() const;
 
-	/**
-	 * @return a const reference to the child nodes of this group.
-	 */
-	FORCEINLINE_DEBUGGABLE const TArray<FStatsNodePtr>& GetChildren() const
-	{
-		return Children;
-	}
-
-	/**
-	 * @return a const reference to the child nodes that should be visible to the UI based on filtering.
-	 */
-	FORCEINLINE_DEBUGGABLE const TArray<FStatsNodePtr>& GetFilteredChildren() const
-	{
-		return FilteredChildren;
-	}
-
-	/**
-	 * @return a weak reference to the group of this timer node, may be invalid.
-	 */
-	FStatsNodeWeak GetGroupPtr() const
-	{
-		return GroupPtr;
-	}
-
-	/**
-	 * @return a name of the fake group that this timer node belongs to.
-	 */
-	const FName& GetGroupName() const
-	{
-		return GroupPtr.Pin()->Name;
-	}
-
-	/**
-	 * @return a name of the fake group that this timer node belongs to.
-	 */
-	//const FName& GetGroupNameSafe() const
-	//{
-	//	return GroupPtr.IsValid() ? GroupPtr.Pin()->Name : NAME_None;
-	//}
-
-	bool IsFiltered() const
-	{
-		return false; // TODO
-	}
-
-public:
 	/** Sorts children using the specified class instance. */
 	template<typename TSortingClass>
 	void SortChildren(const TSortingClass& Instance)
 	{
-		Children.Sort(Instance);
-	}
-
-	/** Adds specified child to the children and sets group for it. */
-	FORCEINLINE_DEBUGGABLE void AddChildAndSetGroupPtr(const FStatsNodePtr& ChildPtr)
-	{
-		ChildPtr->GroupPtr = AsShared();
-		Children.Add(ChildPtr);
+		auto Projection = [](Insights::FBaseTreeNodePtr Node) -> FStatsNodePtr
+		{
+			return StaticCastSharedPtr<FStatsNode, Insights::FBaseTreeNode>(Node);
+		};
+		Algo::SortBy(GetChildrenMutable(), Projection, Instance);
 	}
 
-	/** Adds specified child to the filtered children. */
-	FORCEINLINE_DEBUGGABLE void AddFilteredChild(const FStatsNodePtr& ChildPtr)
+	bool IsAddedToGraph() const
 	{
-		FilteredChildren.Add(ChildPtr);
+		return bIsAddedToGraph;
 	}
 
-	/** Clears filtered children. */
-	void ClearFilteredChildren()
+	void SetAddedToGraphFlag(bool bOnOff)
 	{
-		FilteredChildren.Reset();
+		bIsAddedToGraph = bOnOff;
 	}
 
-protected:
-	/** The Id of this timer or group. */
-	const uint64 Id;
-
-	/** The name of this timer or group. */
-	const FName Name;
-
-	/** The name of the group that this timer belongs to, based on the timer's metadata; only valid for timer nodes. */
+private:
+	/** The name of the meta group that this stats counter belongs to, based on the stats' metadata; only valid for stats counter nodes. */
 	const FName MetaGroupName;
 
-	/** Holds the type of this timer; for the group, this is Group. */
+	/** Holds the type of this stats counter. */
 	const EStatsNodeType Type;
 
+	/** Color of the node. */
+	FLinearColor Color;
+
+	bool bIsAddedToGraph;
+
 	/** Aggregated stats (double). */
 	FAggregatedStats AggregatedStats;
 
 	/** Aggregated stats (int64). */
 	FAggregatedIntegerStats AggregatedIntegerStats;
-
-	/** Children of this node. */
-	TArray<FStatsNodePtr> Children;
-
-	/** Filtered children of this node. */
-	TArray<FStatsNodePtr> FilteredChildren;
-
-	/** A weak pointer to the group/parent of this node. */
-	FStatsNodeWeak GroupPtr;
-
-public:
-	/** Whether this group node should be expanded when the text filtering is enabled. */
-	bool bForceExpandGroupNode;
 };
+
+////////////////////////////////////////////////////////////////////////////////////////////////////
diff --git a/Engine/Source/Developer/TraceInsights/Private/Insights/ViewModels/TimeRulerTrack.cpp b/Engine/Source/Developer/TraceInsights/Private/Insights/ViewModels/TimeRulerTrack.cpp
--- a/Engine/Source/Developer/TraceInsights/Private/Insights/ViewModels/TimeRulerTrack.cpp
+++ b/Engine/Source/Developer/TraceInsights/Private/Insights/ViewModels/TimeRulerTrack.cpp
@@ -2,15 +2,17 @@
 
 #include "TimeRulerTrack.h"
 
-#include "Brushes/SlateColorBrush.h"
 #include "Fonts/FontMeasure.h"
 #include "Fonts/SlateFontInfo.h"
 #include "Framework/Application/SlateApplication.h"
 #include "Rendering/DrawElements.h"
+#include "Styling/CoreStyle.h"
 
 // Insights
 #include "Insights/Common/PaintUtils.h"
 #include "Insights/Common/TimeUtils.h"
+#include "Insights/InsightsStyle.h"
+#include "Insights/ViewModels/DrawHelpers.h"
 #include "Insights/ViewModels/TimingTrackViewport.h"
 
 #define LOCTEXT_NAMESPACE "TimeRulerTrack"
@@ -19,7 +21,7 @@
 
 FTimeRulerTrack::FTimeRulerTrack(uint64 InTrackId)
 	: FBaseTimingTrack(InTrackId)
-	, WhiteBrush(FCoreStyle::Get().GetBrush("WhiteBrush"))
+	, WhiteBrush(FInsightsStyle::Get().GetBrush("WhiteBrush"))
 	, Font(FCoreStyle::GetDefaultFontStyle("Regular", 8))
 {
 }
@@ -62,44 +64,7 @@ void FTimeRulerTrack::UpdateHoveredState(float MouseX, float MouseY, const FTimi
 
 void FTimeRulerTrack::DrawBackground(FDrawContext& DrawContext, const FTimingTrackViewport& Viewport) const
 {
-	const float X0 = Viewport.TimeToSlateUnitsRounded(0.0);
-	const float X1 = Viewport.TimeToSlateUnitsRounded(Viewport.MaxValidTime);
-	const float W = FMath::CeilToFloat(Viewport.Width);
-	const float H = GetHeight();
-
-	const FLinearColor InvalidAreaColor(0.08f, 0.07f, 0.07f, 1.0f);
-	const FLinearColor ValidAreaColor(0.09f, 0.09f, 0.09f, 1.0f);
-
-	if (X0 >= W || X1 <= 0.0f)
-	{
-		// Draw invalid area (entire view).
-		DrawContext.DrawBox(0.0f, 0.0f, W, H, WhiteBrush, InvalidAreaColor);
-	}
-	else // X0 < W && X1 > 0
-	{
-		if (X0 > 0.0f)
-		{
-			// Draw invalid area (left).
-			DrawContext.DrawBox(0.0f, 0.0f, X0, H, WhiteBrush, InvalidAreaColor);
-		}
-
-		if (X1 < W)
-		{
-			// Draw invalid area (right).
-			DrawContext.DrawBox(X1, 0.0f, W - X1, H, WhiteBrush, InvalidAreaColor);
-		}
-
-		float ValidX0 = FMath::Max(X0, 0.0f);
-		float ValidX1 = FMath::Min(X1, W);
-
-		if (ValidX1 > ValidX0)
-		{
-			// Draw valid area.
-			DrawContext.DrawBox(ValidX0, 0.0f, ValidX1 - ValidX0, H, WhiteBrush, ValidAreaColor);
-		}
-	}
-
-	DrawContext.LayerId++;
+	FDrawHelpers::DrawBackground(DrawContext, WhiteBrush, Viewport, GetPosY(), GetHeight());
 }
 
 ////////////////////////////////////////////////////////////////////////////////////////////////////
@@ -117,7 +82,7 @@ void FTimeRulerTrack::Draw(FDrawContext& DrawContext, const FTimingTrackViewport
 	double MinorTickMarkTime = Viewport.GetDurationForViewportDX(MinorTickMark);
 	double MajorTickMarkTime = Viewport.GetDurationForViewportDX(MajorTickMark);
 
-	double VX = Viewport.StartTime * Viewport.ScaleX;
+	double VX = Viewport.GetStartTime() * Viewport.GetScaleX();
 	double MinorN = FMath::FloorToDouble(VX / static_cast<double>(MinorTickMark));
 	double MajorN = FMath::FloorToDouble(VX / static_cast<double>(MajorTickMark));
 	float MinorOX = static_cast<float>(FMath::RoundToDouble(MinorN * static_cast<double>(MinorTickMark) - VX));
@@ -127,7 +92,7 @@ void FTimeRulerTrack::Draw(FDrawContext& DrawContext, const FTimingTrackViewport
 	DrawBackground(DrawContext, Viewport);
 
 	// Draw the minor tick marks.
-	for (float X = MinorOX; X < Viewport.Width; X += MinorTickMark)
+	for (float X = MinorOX; X < Viewport.GetWidth(); X += MinorTickMark)
 	{
 		const bool bIsTenth = ((int32)(((X - MajorOX) / MinorTickMark) + 0.4f) % 2 == 0);
 		const float MinorTickH = bIsTenth ? MinorTickMarkHeight : MinorTickMarkHeight - 1.0f;
@@ -135,25 +100,25 @@ void FTimeRulerTrack::Draw(FDrawContext& DrawContext, const FTimingTrackViewport
 			bIsTenth ? FLinearColor(0.3f, 0.3f, 0.3f, 1.0f) : FLinearColor(0.25f, 0.25f, 0.25f, 1.0f));
 	}
 	// Draw the major tick marks.
-	for (float X = MajorOX; X < Viewport.Width; X += MajorTickMark)
+	for (float X = MajorOX; X < Viewport.GetWidth(); X += MajorTickMark)
 	{
 		DrawContext.DrawBox(X, GetPosY(), 1.0f, MajorTickMarkHeight, WhiteBrush, FLinearColor(0.4f, 0.4f, 0.4f, 1.0f));
 	}
 	DrawContext.LayerId++;
 
-	const double DT = static_cast<double>(MajorTickMark) / Viewport.ScaleX;
+	const double DT = static_cast<double>(MajorTickMark) / Viewport.GetScaleX();
 	const double Precision = FMath::Max(DT / 10.0, TimeUtils::Nanosecond);
 
 	const TSharedRef<FSlateFontMeasure> FontMeasureService = FSlateApplication::Get().GetRenderer()->GetFontMeasureService();
 
 	// Draw the time at major tick marks.
-	for (float X = MajorOX; X < Viewport.Width + MajorTickMark; X += MajorTickMark)
+	for (float X = MajorOX; X < Viewport.GetWidth() + MajorTickMark; X += MajorTickMark)
 	{
 		const double T = Viewport.SlateUnitsToTime(X);
 		FString Text = TimeUtils::FormatTime(T, Precision);
 		const float TextWidth = FontMeasureService->Measure(Text, Font).X;
 		DrawContext.DrawText(X - TextWidth / 2, TextY, Text, Font,
-			(T < 0 || T >= Viewport.MaxValidTime) ? FLinearColor(0.7f, 0.5f, 0.5f, 1.0f) : FLinearColor(0.8f, 0.8f, 0.8f, 1.0f));
+			(T < Viewport.GetMinValidTime() || T >= Viewport.GetMaxValidTime()) ? FLinearColor(0.7f, 0.5f, 0.5f, 1.0f) : FLinearColor(0.8f, 0.8f, 0.8f, 1.0f));
 	}
 	DrawContext.LayerId++;
 
@@ -193,13 +158,14 @@ void FTimeRulerTrack::Draw(FDrawContext& DrawContext, const FTimingTrackViewport
 		float W = CrtMousePosTextWidth + 4.0f;
 		if (bIsSelecting && SelectionStartTime < SelectionEndTime)
 		{
+			// While selecting, display the current time on either left or right side of the selected time range (i.e. to not overlap the selection arrows).
 			float SelectionX1 = Viewport.TimeToSlateUnitsRounded(SelectionStartTime);
 			float SelectionX2 = Viewport.TimeToSlateUnitsRounded(SelectionEndTime);
-			if (X - SelectionX1 > 1.0f)
+			if (FMath::Abs(X - SelectionX1) > FMath::Abs(SelectionX2 - X))
 			{
 				X = SelectionX2 + W / 2;
 			}
-			if (SelectionX2 - X > 1.0f)
+			else
 			{
 				X = SelectionX1 - W / 2;
 			}
@@ -211,7 +177,7 @@ void FTimeRulerTrack::Draw(FDrawContext& DrawContext, const FTimingTrackViewport
 			//DrawContext.DrawBox(0.0f, MousePosition.Y, Viewport.Width, 1.0f, WhiteBrush, MousePosLineColor);
 
 			// Draw vertical line at mouse position.
-			DrawContext.DrawBox(MousePosition.X, 0.0f, 1.0f, Viewport.Height, WhiteBrush, MousePosLineColor);
+			DrawContext.DrawBox(MousePosition.X, 0.0f, 1.0f, Viewport.GetHeight(), WhiteBrush, MousePosLineColor);
 
 			// Stroke the vertical line above current time box.
 			DrawContext.DrawBox(MousePosition.X, 0.0f, 1.0f, TextY, WhiteBrush, MousePosTextBackgroundColor);
diff --git a/Engine/Source/Developer/TraceInsights/Private/Insights/ViewModels/TimeRulerTrack.h b/Engine/Source/Developer/TraceInsights/Private/Insights/ViewModels/TimeRulerTrack.h
--- a/Engine/Source/Developer/TraceInsights/Private/Insights/ViewModels/TimeRulerTrack.h
+++ b/Engine/Source/Developer/TraceInsights/Private/Insights/ViewModels/TimeRulerTrack.h
@@ -17,7 +17,7 @@ class FTimingTrackViewport;
 class FTimeRulerTrack : public FBaseTimingTrack
 {
 public:
-	FTimeRulerTrack(uint64 InTrackId);
+	explicit FTimeRulerTrack(uint64 InTrackId);
 	virtual ~FTimeRulerTrack();
 
 	virtual void Reset() override;
diff --git a/Engine/Source/Developer/TraceInsights/Private/Insights/ViewModels/TimerGroupingAndSorting.cpp b/Engine/Source/Developer/TraceInsights/Private/Insights/ViewModels/TimerGroupingAndSorting.cpp
--- a/Engine/Source/Developer/TraceInsights/Private/Insights/ViewModels/TimerGroupingAndSorting.cpp
+++ b/Engine/Source/Developer/TraceInsights/Private/Insights/ViewModels/TimerGroupingAndSorting.cpp
@@ -2,8 +2,304 @@
 
 #include "TimerGroupingAndSorting.h"
 
+#include "Insights/Table/ViewModels/TableColumn.h"
+#include "Insights/ViewModels/TimerNodeHelper.h"
 
 #define LOCTEXT_NAMESPACE "TimerNode"
 
+// Sort by name (ascending).
+#define INSIGHTS_DEFAULT_SORTING_NODES(A, B) return A->GetName().LexicalLess(B->GetName());
 
+////////////////////////////////////////////////////////////////////////////////////////////////////
+// Sorting by Timer Type
+////////////////////////////////////////////////////////////////////////////////////////////////////
+
+FTimerNodeSortingByTimerType::FTimerNodeSortingByTimerType(TSharedRef<Insights::FTableColumn> InColumnRef)
+	: Insights::FTableCellValueSorter(
+		FName(TEXT("ByTimerType")),
+		LOCTEXT("Sorting_ByTimerType_Name", "By Type"),
+		LOCTEXT("Sorting_ByTimerType_Title", "Sort By Type"),
+		LOCTEXT("Sorting_ByTimerType_Desc", "Sort by timer type (ascending or descending), then by name (ascending)."),
+		InColumnRef)
+{
+}
+
+////////////////////////////////////////////////////////////////////////////////////////////////////
+
+void FTimerNodeSortingByTimerType::Sort(TArray<Insights::FBaseTreeNodePtr>& NodesToSort, Insights::ESortMode SortMode) const
+{
+	if (SortMode == Insights::ESortMode::Ascending)
+	{
+		NodesToSort.Sort([](const Insights::FBaseTreeNodePtr& A, const Insights::FBaseTreeNodePtr& B) -> bool
+		{
+			ensure(A.IsValid() && A->GetTypeName() == FTimerNode::TypeName);
+			const FTimerNodePtr TimerNodeA = StaticCastSharedPtr<FTimerNode, Insights::FBaseTreeNode>(A);
+
+			ensure(B.IsValid() && B->GetTypeName() == FTimerNode::TypeName);
+			const FTimerNodePtr TimerNodeB = StaticCastSharedPtr<FTimerNode, Insights::FBaseTreeNode>(B);
+
+			if (TimerNodeA->GetType() == TimerNodeB->GetType())
+			{
+				INSIGHTS_DEFAULT_SORTING_NODES(A, B)
+			}
+			else
+			{
+				// Sort by timer type (ascending).
+				return TimerNodeA->GetType() < TimerNodeB->GetType();
+			}
+		});
+	}
+	else // if (SortMode == ESortMode::Descending)
+	{
+		NodesToSort.Sort([](const Insights::FBaseTreeNodePtr& A, const Insights::FBaseTreeNodePtr& B) -> bool
+		{
+			ensure(A.IsValid() && A->GetTypeName() == FTimerNode::TypeName);
+			const FTimerNodePtr TimerNodeA = StaticCastSharedPtr<FTimerNode, Insights::FBaseTreeNode>(A);
+
+			ensure(B.IsValid() && B->GetTypeName() == FTimerNode::TypeName);
+			const FTimerNodePtr TimerNodeB = StaticCastSharedPtr<FTimerNode, Insights::FBaseTreeNode>(B);
+
+			if (TimerNodeA->GetType() == TimerNodeB->GetType())
+			{
+				INSIGHTS_DEFAULT_SORTING_NODES(A, B)
+			}
+			else
+			{
+				// Sort by timer type (descending).
+				return TimerNodeB->GetType() < TimerNodeA->GetType();
+			}
+		});
+	}
+}
+
+////////////////////////////////////////////////////////////////////////////////////////////////////
+// Sort by Instance Count
+////////////////////////////////////////////////////////////////////////////////////////////////////
+
+FTimerNodeSortingByInstanceCount::FTimerNodeSortingByInstanceCount(TSharedRef<Insights::FTableColumn> InColumnRef)
+	: Insights::FTableCellValueSorter(
+		FName(TEXT("ByInstanceCount")),
+		LOCTEXT("Sorting_ByInstanceCount_Name", "By Instance Count"),
+		LOCTEXT("Sorting_ByInstanceCount_Title", "Sort By Instance Count"),
+		LOCTEXT("Sorting_ByInstanceCount_Desc", "Sort by aggregated instance count (ascending or descending), then by name (ascending)."),
+		InColumnRef)
+{
+}
+
+////////////////////////////////////////////////////////////////////////////////////////////////////
+
+void FTimerNodeSortingByInstanceCount::Sort(TArray<Insights::FBaseTreeNodePtr>& NodesToSort, Insights::ESortMode SortMode) const
+{
+	if (SortMode == Insights::ESortMode::Ascending)
+	{
+		NodesToSort.Sort([](const Insights::FBaseTreeNodePtr& A, const Insights::FBaseTreeNodePtr& B) -> bool
+		{
+			ensure(A.IsValid() && A->GetTypeName() == FTimerNode::TypeName);
+			const FTimerNodePtr TimerNodeA = StaticCastSharedPtr<FTimerNode, Insights::FBaseTreeNode>(A);
+			const uint64 ValueA = TimerNodeA->GetAggregatedStats().InstanceCount;
+
+			ensure(B.IsValid() && B->GetTypeName() == FTimerNode::TypeName);
+			const FTimerNodePtr TimerNodeB = StaticCastSharedPtr<FTimerNode, Insights::FBaseTreeNode>(B);
+			const uint64 ValueB = TimerNodeB->GetAggregatedStats().InstanceCount;
+
+			if (ValueA == ValueB)
+			{
+				INSIGHTS_DEFAULT_SORTING_NODES(A, B)
+			}
+			else
+			{
+				// Sort by instance count (ascending).
+				return ValueA < ValueB;
+			}
+		});
+	}
+	else // if (SortMode == ESortMode::Descending)
+	{
+		NodesToSort.Sort([](const Insights::FBaseTreeNodePtr& A, const Insights::FBaseTreeNodePtr& B) -> bool
+		{
+			ensure(A.IsValid() && A->GetTypeName() == FTimerNode::TypeName);
+			const FTimerNodePtr TimerNodeA = StaticCastSharedPtr<FTimerNode, Insights::FBaseTreeNode>(A);
+			const uint64 ValueA = TimerNodeA->GetAggregatedStats().InstanceCount;
+
+			ensure(B.IsValid() && B->GetTypeName() == FTimerNode::TypeName);
+			const FTimerNodePtr TimerNodeB = StaticCastSharedPtr<FTimerNode, Insights::FBaseTreeNode>(B);
+			const uint64 ValueB = TimerNodeB->GetAggregatedStats().InstanceCount;
+
+			if (ValueA == ValueB)
+			{
+				INSIGHTS_DEFAULT_SORTING_NODES(A, B)
+			}
+			else
+			{
+				// Sort by instance count (descending).
+				return ValueB < ValueA;
+			}
+		});
+	}
+}
+
+////////////////////////////////////////////////////////////////////////////////////////////////////
+// Sort by Total Inclusive Time
+////////////////////////////////////////////////////////////////////////////////////////////////////
+
+FTimerNodeSortingByTotalInclusiveTime::FTimerNodeSortingByTotalInclusiveTime(TSharedRef<Insights::FTableColumn> InColumnRef)
+	: Insights::FTableCellValueSorter(
+		FName(TEXT("ByTotalInclusiveTime")),
+		LOCTEXT("Sorting_ByTotalInclusiveTime_Name", "By Total Inclusive Time"),
+		LOCTEXT("Sorting_ByTotalInclusiveTime_Title", "Sort By Total Inclusive Time"),
+		LOCTEXT("Sorting_ByTotalInclusiveTime_Desc", "Sort by aggregated total inclusive time (ascending or descending), then by name (ascending)."),
+		InColumnRef)
+{
+}
+
+////////////////////////////////////////////////////////////////////////////////////////////////////
+
+void FTimerNodeSortingByTotalInclusiveTime::Sort(TArray<Insights::FBaseTreeNodePtr>& NodesToSort, Insights::ESortMode SortMode) const
+{
+	if (SortMode == Insights::ESortMode::Ascending)
+	{
+		NodesToSort.Sort([](const Insights::FBaseTreeNodePtr& A, const Insights::FBaseTreeNodePtr& B) -> bool
+		{
+			ensure(A.IsValid() && A->GetTypeName() == FTimerNode::TypeName);
+			const FTimerNodePtr TimerNodeA = StaticCastSharedPtr<FTimerNode, Insights::FBaseTreeNode>(A);
+			const double ValueA = TimerNodeA->GetAggregatedStats().TotalInclusiveTime;
+
+			ensure(B.IsValid() && B->GetTypeName() == FTimerNode::TypeName);
+			const FTimerNodePtr TimerNodeB = StaticCastSharedPtr<FTimerNode, Insights::FBaseTreeNode>(B);
+			const double ValueB = TimerNodeB->GetAggregatedStats().TotalInclusiveTime;
+
+			if (ValueA == ValueB)
+			{
+				INSIGHTS_DEFAULT_SORTING_NODES(A, B)
+			}
+			else
+			{
+				// Sort by total inclusive time (ascending).
+				return ValueA < ValueB;
+			}
+		});
+	}
+	else // if (SortMode == ESortMode::Descending)
+	{
+		NodesToSort.Sort([](const Insights::FBaseTreeNodePtr& A, const Insights::FBaseTreeNodePtr& B) -> bool
+		{
+			ensure(A.IsValid() && A->GetTypeName() == FTimerNode::TypeName);
+			const FTimerNodePtr TimerNodeA = StaticCastSharedPtr<FTimerNode, Insights::FBaseTreeNode>(A);
+			const double ValueA = TimerNodeA->GetAggregatedStats().TotalInclusiveTime;
+
+			ensure(B.IsValid() && B->GetTypeName() == FTimerNode::TypeName);
+			const FTimerNodePtr TimerNodeB = StaticCastSharedPtr<FTimerNode, Insights::FBaseTreeNode>(B);
+			const double ValueB = TimerNodeB->GetAggregatedStats().TotalInclusiveTime;
+
+			if (ValueA == ValueB)
+			{
+				INSIGHTS_DEFAULT_SORTING_NODES(A, B)
+			}
+			else
+			{
+				// Sort by total inclusive time (descending).
+				return ValueB < ValueA;
+			}
+		});
+	}
+}
+
+////////////////////////////////////////////////////////////////////////////////////////////////////
+// Sort by Total Exclusive Time
+////////////////////////////////////////////////////////////////////////////////////////////////////
+
+FTimerNodeSortingByTotalExclusiveTime::FTimerNodeSortingByTotalExclusiveTime(TSharedRef<Insights::FTableColumn> InColumnRef)
+	: Insights::FTableCellValueSorter(
+		FName(TEXT("ByTotalExclusiveTime")),
+		LOCTEXT("Sorting_ByTotalExclusiveTime_Name", "By Total Exclusive Time"),
+		LOCTEXT("Sorting_ByTotalExclusiveTime_Title", "Sort By Total Exclusive Time"),
+		LOCTEXT("Sorting_ByTotalExclusiveTime_Desc", "Sort by aggregated total exclusive time (ascending or descending), then by name (ascending)."),
+		InColumnRef)
+{
+}
+
+////////////////////////////////////////////////////////////////////////////////////////////////////
+
+void FTimerNodeSortingByTotalExclusiveTime::Sort(TArray<Insights::FBaseTreeNodePtr>& NodesToSort, Insights::ESortMode SortMode) const
+{
+	if (SortMode == Insights::ESortMode::Ascending)
+	{
+		NodesToSort.Sort([](const Insights::FBaseTreeNodePtr& A, const Insights::FBaseTreeNodePtr& B) -> bool
+		{
+			ensure(A.IsValid() && A->GetTypeName() == FTimerNode::TypeName);
+			const FTimerNodePtr TimerNodeA = StaticCastSharedPtr<FTimerNode, Insights::FBaseTreeNode>(A);
+			const double ValueA = TimerNodeA->GetAggregatedStats().TotalExclusiveTime;
+
+			ensure(B.IsValid() && B->GetTypeName() == FTimerNode::TypeName);
+			const FTimerNodePtr TimerNodeB = StaticCastSharedPtr<FTimerNode, Insights::FBaseTreeNode>(B);
+			const double ValueB = TimerNodeB->GetAggregatedStats().TotalExclusiveTime;
+
+			if (ValueA == ValueB)
+			{
+				INSIGHTS_DEFAULT_SORTING_NODES(A, B)
+			}
+			else
+			{
+				// Sort by total exclusive time (ascending).
+				return ValueA < ValueB;
+			}
+		});
+	}
+	else // if (SortMode == ESortMode::Descending)
+	{
+		NodesToSort.Sort([](const Insights::FBaseTreeNodePtr& A, const Insights::FBaseTreeNodePtr& B) -> bool
+		{
+			ensure(A.IsValid() && A->GetTypeName() == FTimerNode::TypeName);
+			const FTimerNodePtr TimerNodeA = StaticCastSharedPtr<FTimerNode, Insights::FBaseTreeNode>(A);
+			const double ValueA = TimerNodeA->GetAggregatedStats().TotalExclusiveTime;
+
+			ensure(B.IsValid() && B->GetTypeName() == FTimerNode::TypeName);
+			const FTimerNodePtr TimerNodeB = StaticCastSharedPtr<FTimerNode, Insights::FBaseTreeNode>(B);
+			const double ValueB = TimerNodeB->GetAggregatedStats().TotalExclusiveTime;
+
+			if (ValueA == ValueB)
+			{
+				INSIGHTS_DEFAULT_SORTING_NODES(A, B)
+			}
+			else
+			{
+				// Sort by total exclusive time (descending).
+				return ValueB < ValueA;
+			}
+		});
+	}
+}
+
+////////////////////////////////////////////////////////////////////////////////////////////////////
+// Grouping by Timer Type
+////////////////////////////////////////////////////////////////////////////////////////////////////
+
+FTimerNodeGroupingByTimerType::FTimerNodeGroupingByTimerType()
+	: FTreeNodeGrouping(
+		LOCTEXT("Grouping_ByTimerType_ShortName", "Type"),
+		LOCTEXT("Grouping_ByTimerType_TitleName", "Type"),
+		LOCTEXT("Grouping_ByTimerType_Desc", "Creates a group for each timer type."),
+		TEXT("Profiler.FiltersAndPresets.StatTypeIcon"), //TODO
+		nullptr)
+{
+}
+
+////////////////////////////////////////////////////////////////////////////////////////////////////
+
+Insights::FTreeNodeGroupInfo FTimerNodeGroupingByTimerType::GetGroupForNode(const Insights::FBaseTreeNodePtr InNodePtr) const
+{
+	if (InNodePtr->GetTypeName() == FTimerNode::TypeName)
+	{
+		const FTimerNodePtr& TimerNodePtr = StaticCastSharedPtr<FTimerNode>(InNodePtr);
+		return { FName(*TimerNodeTypeHelper::ToText(TimerNodePtr->GetType()).ToString()), true };
+	}
+	else
+	{
+		return { FName("Unknown"), true };
+	}
+}
+
+////////////////////////////////////////////////////////////////////////////////////////////////////
+
+#undef INSIGHTS_DEFAULT_SORTING_NODES
 #undef LOCTEXT_NAMESPACE
diff --git a/Engine/Source/Developer/TraceInsights/Private/Insights/ViewModels/TimerGroupingAndSorting.h b/Engine/Source/Developer/TraceInsights/Private/Insights/ViewModels/TimerGroupingAndSorting.h
--- a/Engine/Source/Developer/TraceInsights/Private/Insights/ViewModels/TimerGroupingAndSorting.h
+++ b/Engine/Source/Developer/TraceInsights/Private/Insights/ViewModels/TimerGroupingAndSorting.h
@@ -3,10 +3,58 @@
 #pragma once
 
 #include "CoreMinimal.h"
+
+// Insights
 #include "Insights/ViewModels/TimerNode.h"
+#include "Insights/Table/ViewModels/TableColumn.h"
+#include "Insights/Table/ViewModels/TableCellValueSorter.h"
+#include "Insights/Table/ViewModels/TreeNodeGrouping.h"
 
 ////////////////////////////////////////////////////////////////////////////////////////////////////
-// Grouping
+// Sorters
+////////////////////////////////////////////////////////////////////////////////////////////////////
+
+class FTimerNodeSortingByTimerType: public Insights::FTableCellValueSorter
+{
+public:
+	FTimerNodeSortingByTimerType(TSharedRef<Insights::FTableColumn> InColumnRef);
+
+	virtual void Sort(TArray<Insights::FBaseTreeNodePtr>& NodesToSort, Insights::ESortMode SortMode) const override;
+};
+
+////////////////////////////////////////////////////////////////////////////////////////////////////
+
+class FTimerNodeSortingByInstanceCount : public Insights::FTableCellValueSorter
+{
+public:
+	FTimerNodeSortingByInstanceCount(TSharedRef<Insights::FTableColumn> InColumnRef);
+
+	virtual void Sort(TArray<Insights::FBaseTreeNodePtr>& NodesToSort, Insights::ESortMode SortMode) const override;
+};
+
+////////////////////////////////////////////////////////////////////////////////////////////////////
+
+class FTimerNodeSortingByTotalInclusiveTime : public Insights::FTableCellValueSorter
+{
+public:
+	FTimerNodeSortingByTotalInclusiveTime(TSharedRef<Insights::FTableColumn> InColumnRef);
+
+	virtual void Sort(TArray<Insights::FBaseTreeNodePtr>& NodesToSort, Insights::ESortMode SortMode) const override;
+};
+
+////////////////////////////////////////////////////////////////////////////////////////////////////
+
+class FTimerNodeSortingByTotalExclusiveTime : public Insights::FTableCellValueSorter
+{
+public:
+	FTimerNodeSortingByTotalExclusiveTime(TSharedRef<Insights::FTableColumn> InColumnRef);
+
+	virtual void Sort(TArray<Insights::FBaseTreeNodePtr>& NodesToSort, Insights::ESortMode SortMode) const override;
+};
+
+////////////////////////////////////////////////////////////////////////////////////////////////////
+// Organizers
+////////////////////////////////////////////////////////////////////////////////////////////////////
 
 /** Enumerates types of grouping or sorting for the timer nodes. */
 enum class ETimerGroupingMode
@@ -37,3 +85,15 @@ enum class ETimerGroupingMode
 typedef TSharedPtr<ETimerGroupingMode> ETimerGroupingModePtr;
 
 ////////////////////////////////////////////////////////////////////////////////////////////////////
+
+/** Creates a group for each timer type. */
+class FTimerNodeGroupingByTimerType : public Insights::FTreeNodeGrouping
+{
+public:
+	FTimerNodeGroupingByTimerType();
+	virtual ~FTimerNodeGroupingByTimerType() {}
+
+	virtual Insights::FTreeNodeGroupInfo GetGroupForNode(const Insights::FBaseTreeNodePtr InNode) const override;
+};
+
+////////////////////////////////////////////////////////////////////////////////////////////////////
diff --git a/Engine/Source/Developer/TraceInsights/Private/Insights/ViewModels/TimerNode.cpp b/Engine/Source/Developer/TraceInsights/Private/Insights/ViewModels/TimerNode.cpp
--- a/Engine/Source/Developer/TraceInsights/Private/Insights/ViewModels/TimerNode.cpp
+++ b/Engine/Source/Developer/TraceInsights/Private/Insights/ViewModels/TimerNode.cpp
@@ -4,6 +4,8 @@
 
 #define LOCTEXT_NAMESPACE "TimerNode"
 
+const FName FTimerNode::TypeName(TEXT("FTimerNode"));
+
 ////////////////////////////////////////////////////////////////////////////////////////////////////
 
 void FTimerNode::ResetAggregatedStats()
@@ -20,40 +22,4 @@ void FTimerNode::SetAggregatedStats(const Trace::FAggregatedTimingStats& InAggre
 
 ////////////////////////////////////////////////////////////////////////////////////////////////////
 
-const FText FTimerNode::GetNameEx() const
-{
-	FText Text = FText::GetEmpty();
-
-	if (IsGroup())
-	{
-		const int32 NumChildren = Children.Num();
-		const int32 NumFilteredChildren = FilteredChildren.Num();
-
-		if (NumFilteredChildren == NumChildren)
-		{
-			Text = FText::Format(LOCTEXT("TimerNodeGroupTextFmt1", "{0} ({1})"), FText::FromName(Name), FText::AsNumber(NumChildren));
-		}
-		else
-		{
-			Text = FText::Format(LOCTEXT("TimerNodeGroupTextFmt2", "{0} ({1} / {2})"), FText::FromName(Name), FText::AsNumber(NumFilteredChildren), FText::AsNumber(NumChildren));
-		}
-	}
-	else
-	{
-		//const bool bIsStatTracked = FProfilerManager::Get()->IsStatTracked(GroupOrStatNode->GetStatID());
-		//if (bIsStatTracked)
-		//{
-		//	Text = FText::Format(LOCTEXT("TimerNodeTrackedTextFmt", "{0}*"), FText::FromName(Name));
-		//}
-		//else
-		{
-			Text = FText::FromName(Name);
-		}
-	}
-
-	return Text;
-}
-
-////////////////////////////////////////////////////////////////////////////////////////////////////
-
 #undef LOCTEXT_NAMESPACE
diff --git a/Engine/Source/Developer/TraceInsights/Private/Insights/ViewModels/TimerNode.h b/Engine/Source/Developer/TraceInsights/Private/Insights/ViewModels/TimerNode.h
--- a/Engine/Source/Developer/TraceInsights/Private/Insights/ViewModels/TimerNode.h
+++ b/Engine/Source/Developer/TraceInsights/Private/Insights/ViewModels/TimerNode.h
@@ -3,21 +3,12 @@
 #pragma once
 
 #include "CoreMinimal.h"
-#include "TraceServices/AnalysisService.h"
+#include "TraceServices/Containers/Timelines.h"
 
-class FTimerNode;
-
-/** Type definition for shared pointers to instances of FTimerNode. */
-typedef TSharedPtr<class FTimerNode> FTimerNodePtr;
-
-/** Type definition for shared references to instances of FTimerNode. */
-typedef TSharedRef<class FTimerNode> FTimerNodeRef;
-
-/** Type definition for shared references to const instances of FTimerNode. */
-typedef TSharedRef<const class FTimerNode> FTimerNodeRefConst;
+// Insights
+#include "Insights/Table/ViewModels/BaseTreeNode.h"
 
-/** Type definition for weak references to instances of FTimerNode. */
-typedef TWeakPtr<class FTimerNode> FTimerNodeWeak;
+////////////////////////////////////////////////////////////////////////////////////////////////////
 
 enum class ETimerNodeType
 {
@@ -37,191 +28,87 @@ enum class ETimerNodeType
 	InvalidOrMax,
 };
 
+////////////////////////////////////////////////////////////////////////////////////////////////////
+
+class FTimerNode;
+
+/** Type definition for shared pointers to instances of FTimerNode. */
+typedef TSharedPtr<class FTimerNode> FTimerNodePtr;
+
+/** Type definition for shared references to instances of FTimerNode. */
+typedef TSharedRef<class FTimerNode> FTimerNodeRef;
+
+/** Type definition for shared references to const instances of FTimerNode. */
+typedef TSharedRef<const class FTimerNode> FTimerNodeRefConst;
+
+/** Type definition for weak references to instances of FTimerNode. */
+typedef TWeakPtr<class FTimerNode> FTimerNodeWeak;
+
+////////////////////////////////////////////////////////////////////////////////////////////////////
 /**
- * Class used to store information about a timer node (used in timers' tree view).
+ * Class used to store information about a timer node (used in the STimersView).
  */
-class FTimerNode : public TSharedFromThis<FTimerNode>
+class FTimerNode : public Insights::FBaseTreeNode
 {
 public:
-	static const uint64 InvalidId = -1;
+	static const FName TypeName;
+	//static constexpr uint64 InvalidTimerId = -1;
 
 public:
 	/** Initialization constructor for the timer node. */
 	FTimerNode(uint64 InId, const FName InName, const FName InMetaGroupName, ETimerNodeType InType)
-		: Id(InId)
-		, Name(InName)
+		: FBaseTreeNode(InId, InName, InType == ETimerNodeType::Group)
 		, MetaGroupName(InMetaGroupName)
 		, Type(InType)
-		, bForceExpandGroupNode(false)
+		, bIsHotPath(false)
 	{
 		ResetAggregatedStats();
 	}
 
 	/** Initialization constructor for the group node. */
 	FTimerNode(const FName InGroupName)
-		: Id(0)
-		, Name(InGroupName)
+		: FBaseTreeNode(0, InGroupName, true)
 		, Type(ETimerNodeType::Group)
-		, bForceExpandGroupNode(false)
+		, bIsHotPath(false)
 	{
 		ResetAggregatedStats();
 	}
 
-	/**
-	 * @return an Id of this timer, valid only for timer nodes.
-	 */
-	const uint64 GetId() const
-	{
-		return Id;
-	}
-
-	/**
-	 * @return a name of this node, group or timer.
-	 */
-	const FName& GetName() const
-	{
-		return Name;
-	}
-
-	/**
-	 * @return a name of this node, group or timer + addditional info, to display in Timers tree view.
-	*/
-	const FText GetNameEx() const;
-
-	/**
-	 * @return a name of the group that this timer node belongs to, taken from the metadata.
-	 */
-	const FName& GetMetaGroupName() const
-	{
-		return MetaGroupName;
-	}
+	virtual const FName& GetTypeName() const override { return TypeName; }
 
 	/**
-	 * @return a type of this timer or ETimerNodeType::Group for group nodes.
+	 * @return a name of the meta group that this timer node belongs to, taken from the metadata.
 	 */
-	const ETimerNodeType& GetType() const
-	{
-		return Type;
-	}
+	const FName& GetMetaGroupName() const { return MetaGroupName; }
 
 	/**
-	 * @return true, if this node is a group node.
+	 * @return a type of this timer node or ETimerNodeType::Group for group nodes.
 	 */
-	bool IsGroup() const
-	{
-		return Type == ETimerNodeType::Group;
-	}
+	const ETimerNodeType& GetType() const { return Type; }
 
 	/**
 	 * @return the aggregated stats for this timer.
 	 */
-	const Trace::FAggregatedTimingStats& GetAggregatedStats() const
-	{
-		return AggregatedStats;
-	}
+	const Trace::FAggregatedTimingStats& GetAggregatedStats() const { return AggregatedStats; }
 
 	void ResetAggregatedStats();
-
 	void SetAggregatedStats(const Trace::FAggregatedTimingStats& AggregatedStats);
 
-	/**
-	 * @return a const reference to the child nodes of this group.
-	 */
-	FORCEINLINE_DEBUGGABLE const TArray<FTimerNodePtr>& GetChildren() const
-	{
-		return Children;
-	}
-
-	/**
-	 * @return a const reference to the child nodes that should be visible to the UI based on filtering.
-	 */
-	FORCEINLINE_DEBUGGABLE const TArray<FTimerNodePtr>& GetFilteredChildren() const
-	{
-		return FilteredChildren;
-	}
-
-	/**
-	 * @return a weak reference to the group of this timer node, may be invalid.
-	 */
-	FTimerNodeWeak GetGroupPtr() const
-	{
-		return GroupPtr;
-	}
-
-	/**
-	 * @return a name of the fake group that this timer node belongs to.
-	 */
-	const FName& GetGroupName() const
-	{
-		return GroupPtr.Pin()->Name;
-	}
-
-	/**
-	 * @return a name of the fake group that this timer node belongs to.
-	 */
-	//const FName& GetGroupNameSafe() const
-	//{
-	//	return GroupPtr.IsValid() ? GroupPtr.Pin()->Name : NAME_None;
-	//}
+	bool IsHotPath() const { return bIsHotPath; }
+	void SetIsHotPath(bool bOnOff) { bIsHotPath = bOnOff; }
 
-	bool IsFiltered() const
-	{
-		return false; // TODO
-	}
-
-public:
-	/** Sorts children using the specified class instance. */
-	template<typename TSortingClass>
-	void SortChildren(const TSortingClass& Instance)
-	{
-		Children.Sort(Instance);
-	}
-
-	/** Adds specified child to the children and sets group for it. */
-	FORCEINLINE_DEBUGGABLE void AddChildAndSetGroupPtr(const FTimerNodePtr& ChildPtr)
-	{
-		ChildPtr->GroupPtr = AsShared();
-		Children.Add(ChildPtr);
-	}
-
-	/** Adds specified child to the filtered children. */
-	FORCEINLINE_DEBUGGABLE void AddFilteredChild(const FTimerNodePtr& ChildPtr)
-	{
-		FilteredChildren.Add(ChildPtr);
-	}
-
-	/** Clears filtered children. */
-	void ClearFilteredChildren()
-	{
-		FilteredChildren.Reset();
-	}
-
-protected:
-	/** The Id of this timer or group. */
-	const uint64 Id;
-
-	/** The name of this timer or group. */
-	const FName Name;
-
-	/** The name of the group that this timer belongs to, based on the timer's metadata; only valid for timer nodes. */
+private:
+	/** The name of the meta group that this timer belongs to, based on the timer's metadata; only valid for timer nodes. */
 	const FName MetaGroupName;
 
-	/** Holds the type of this timer; for the group, this is Group. */
+	/** Holds the type of this timer. */
 	const ETimerNodeType Type;
 
 	/** Aggregated stats. */
 	Trace::FAggregatedTimingStats AggregatedStats;
 
-	/** Children of this node. */
-	TArray<FTimerNodePtr> Children;
-
-	/** Filtered children of this node. */
-	TArray<FTimerNodePtr> FilteredChildren;
-
-	/** A weak pointer to the group/parent of this node. */
-	FTimerNodeWeak GroupPtr;
-
-public:
-	/** Whether this group node should be expanded when the text filtering is enabled. */
-	bool bForceExpandGroupNode;
+	/** True if this tree node is on the hot path. */
+	bool bIsHotPath;
 };
+
+////////////////////////////////////////////////////////////////////////////////////////////////////
diff --git a/Engine/Source/Developer/TraceInsights/Private/Insights/ViewModels/TimerNodeHelper.cpp b/Engine/Source/Developer/TraceInsights/Private/Insights/ViewModels/TimerNodeHelper.cpp
--- a/Engine/Source/Developer/TraceInsights/Private/Insights/ViewModels/TimerNodeHelper.cpp
+++ b/Engine/Source/Developer/TraceInsights/Private/Insights/ViewModels/TimerNodeHelper.cpp
@@ -9,7 +9,7 @@
 // TimerNode Type Helper
 ////////////////////////////////////////////////////////////////////////////////////////////////////
 
-FText TimerNodeTypeHelper::ToName(const ETimerNodeType NodeType)
+FText TimerNodeTypeHelper::ToText(const ETimerNodeType NodeType)
 {
 	static_assert(static_cast<int>(ETimerNodeType::InvalidOrMax) == 4, "Not all cases are handled in switch below!?");
 	switch (NodeType)
@@ -78,7 +78,7 @@ const FSlateBrush* TimerNodeTypeHelper::GetIconForTimerNodeType(const ETimerNode
 // TimerNode Grouping Helper
 ////////////////////////////////////////////////////////////////////////////////////////////////////
 
-FText TimerNodeGroupingHelper::ToName(const ETimerGroupingMode GroupingMode)
+FText TimerNodeGroupingHelper::ToText(const ETimerGroupingMode GroupingMode)
 {
 	static_assert(static_cast<int>(ETimerGroupingMode::InvalidOrMax) == 7, "Not all cases are handled in switch below!?");
 	switch (GroupingMode)
diff --git a/Engine/Source/Developer/TraceInsights/Private/Insights/ViewModels/TimersViewColumn.h b/Engine/Source/Developer/TraceInsights/Private/Insights/ViewModels/TimersViewColumn.h
--- a/Engine/Source/Developer/TraceInsights/Private/Insights/ViewModels/TimersViewColumn.h
+++ b/Engine/Source/Developer/TraceInsights/Private/Insights/ViewModels/TimersViewColumn.h
@@ -1,118 +1,3 @@
 // Copyright 1998-2019 Epic Games, Inc. All Rights Reserved.
 
 #pragma once
-
-#include "CoreMinimal.h"
-#include "Types/SlateEnums.h"
-#include "Misc/EnumClassFlags.h"
-
-class FTimerNode; // TODO: IInsightsTreeNode; generic formatter; see also StatsViewColumn.h
-
-enum class ETimersViewColumnFlags : uint32
-{
-	None = 0,
-
-	CanBeHidden = (1 << 0),
-	CanBeSorted = (1 << 1),
-	CanBeFiltered = (1 << 2),
-};
-ENUM_CLASS_FLAGS(ETimersViewColumnFlags);
-
-/** Holds information about a column in the timer view widget. */
-class FTimersViewColumn
-{
-	friend struct FTimersViewColumnFactory;
-
-public:
-	typedef TFunction<FText(const FTimersViewColumn& Column, const FTimerNode& TimerNode)> FGetFormattedValueFn;
-
-public:
-	/** Whether this column can be hidden. */
-	bool bCanBeHidden() const { return EnumHasAnyFlags(Flags, ETimersViewColumnFlags::CanBeHidden); }
-
-	/** Whether this column cab be used for sorting. */
-	bool bCanBeSorted() const { return EnumHasAnyFlags(Flags, ETimersViewColumnFlags::CanBeSorted); }
-
-	/** Where this column can be used to filtering displayed results. */
-	bool bCanBeFiltered() const { return EnumHasAnyFlags(Flags, ETimersViewColumnFlags::CanBeFiltered); }
-
-	/** If MinColumnWidth == MaxColumnWidth, this column has fixed width and cannot be resized. */
-	bool bIsFixedColumnWidth() const { return MinColumnWidth == MaxColumnWidth; }
-
-	FText GetFormattedValue(const FTimerNode& TimerNode) const
-	{
-		return GetFormattedValueFn(*this, TimerNode);
-	}
-
-protected:
-	/** No default constructor. */
-	FTimersViewColumn() = delete;
-
-	/** Initialization constructor, only used in FStatsViewColumnFactory. */
-	FTimersViewColumn
-	(
-		int32 InOrder,
-		const FName InId,
-		const FName InSearchId,
-		FText InShortName,
-		FText InTitleName,
-		FText InDescription,
-		bool bInIsVisible,
-		const ETimersViewColumnFlags InFlags,
-		const EHorizontalAlignment InHorizontalAlignment,
-		const float InInitialColumnWidth,
-		const float InMinColumnWidth,
-		const float InMaxColumnWidth,
-		const FGetFormattedValueFn InGetFormattedValueFn
-	)
-		: Order(InOrder)
-		, Id(InId)
-		, SearchId(InSearchId)
-		, ShortName(MoveTemp(InShortName))
-		, TitleName(MoveTemp(InTitleName))
-		, Description(MoveTemp(InDescription))
-		, bIsVisible(bInIsVisible)
-		, Flags(InFlags)
-		, HorizontalAlignment(InHorizontalAlignment)
-		, InitialColumnWidth(InInitialColumnWidth)
-		, MinColumnWidth(InMinColumnWidth)
-		, MaxColumnWidth(InMaxColumnWidth)
-		, GetFormattedValueFn(InGetFormattedValueFn)
-	{
-	}
-
-public:
-	/** Order value, to sort columns in the tree view. */
-	int32 Order;
-
-	/** Name of the column, name of the property. */
-	FName Id;
-
-	/** Name of the column used by the searching system. */
-	FName SearchId;
-
-	/** Short name of the column, displayed in the column header. */
-	FText ShortName;
-
-	/** Title name of the column, displayed as title in the column tooltip. */
-	FText TitleName;
-
-	/** Long name of the column, displayed in the column tooltip. */
-	FText Description;
-
-	/** Is this column visible? */
-	bool bIsVisible;
-
-	/** On/off switches. */
-	ETimersViewColumnFlags Flags;
-
-	/** Horizontal alignment of the content in this column. */
-	EHorizontalAlignment HorizontalAlignment;
-
-	float InitialColumnWidth; /**< Initial column width. */
-	float MinColumnWidth; /**< Minimum column width. */
-	float MaxColumnWidth; /**< Maximum column width. */
-
-	/** Custom function used to format (as an FText) the value of the timer node to be displayed by this column. */
-	FGetFormattedValueFn GetFormattedValueFn;
-};
