Copying //UE4/Dev-Core to Main (//UE4/Main)

diff --git a/Engine/Source/Developer/TraceInsights/Private/Insights/ViewModels/TimersViewColumnFactory.cpp b/Engine/Source/Developer/TraceInsights/Private/Insights/ViewModels/TimersViewColumnFactory.cpp
--- a/Engine/Source/Developer/TraceInsights/Private/Insights/ViewModels/TimersViewColumnFactory.cpp
+++ b/Engine/Source/Developer/TraceInsights/Private/Insights/ViewModels/TimersViewColumnFactory.cpp
@@ -4,6 +4,10 @@
 
 // Insights
 #include "Insights/Common/TimeUtils.h"
+#include "Insights/Table/ViewModels/TableCellValueFormatter.h"
+#include "Insights/Table/ViewModels/TableCellValueGetter.h"
+#include "Insights/Table/ViewModels/TableCellValueSorter.h"
+#include "Insights/ViewModels/TimerGroupingAndSorting.h"
 #include "Insights/ViewModels/TimerNodeHelper.h"
 
 #define LOCTEXT_NAMESPACE "STimerView"
@@ -11,7 +15,7 @@
 ////////////////////////////////////////////////////////////////////////////////////////////////////
 // Column identifiers
 
-const FName FTimersViewColumns::NameColumnID(TEXT("Name"));
+const FName FTimersViewColumns::NameColumnID(TEXT("Name")); // TEXT("_Hierarchy")
 const FName FTimersViewColumns::MetaGroupNameColumnID(TEXT("MetaGroupName"));
 const FName FTimersViewColumns::TypeColumnID(TEXT("Type"));
 const FName FTimersViewColumns::InstanceCountColumnID(TEXT("Count"));
@@ -36,319 +40,669 @@ const FName FTimersViewColumns::MinExclusiveTimeColumnID(TEXT("MinExclTime"));
 
 ////////////////////////////////////////////////////////////////////////////////////////////////////
 
-FTimersViewColumnFactory::FTimersViewColumnFactory()
+void FTimersViewColumnFactory::CreateTimersViewColumns(TArray<TSharedPtr<Insights::FTableColumn>>& Columns)
 {
-	uint32 Index = 0;
-
-	// FTimerViewColumn
-	// (
-	//     Order,
-	//     Id,
-	//     SearchId,
-	//     ShortName,
-	//     TitleName,
-	//     Description,
-	//     bIsVisible,
-	//     Flags,
-	//     HorizontalAlignment,
-	//     InitialWidth, MinWidth, MaxWidth,
-	//     GetFormattedValueFn
-	// )
-
-	//////////////////////////////////////////////////
-	// Name
-
-	const ETimersViewColumnFlags NameColumnFlags = ETimersViewColumnFlags::CanBeSorted |
-												   ETimersViewColumnFlags::CanBeFiltered;
-
-	Collection.Add(new FTimersViewColumn
-	(
-		Index++,
-		FTimersViewColumns::NameColumnID,
-		TEXT("name"),
-		LOCTEXT("TimerNameColumnName", "Name"),
-		LOCTEXT("TimerNameColumnTitle", "Timer or Group Name"),
-		LOCTEXT("TimerNameColumnDesc", "Name of the timer or group"),
-		true,
-		NameColumnFlags,
-		HAlign_Left,
-		246.0f, 10.0f, FLT_MAX,
-		[](const FTimersViewColumn& Column, const FTimerNode& TimerNode) -> FText
-		{ return FText::FromName(TimerNode.GetName()); }
-	));
-
-	//////////////////////////////////////////////////
-	// Meta Group Name
-
-	const ETimersViewColumnFlags MetaGroupNameColumnFlags = ETimersViewColumnFlags::CanBeHidden |
-															ETimersViewColumnFlags::CanBeSorted |
-															ETimersViewColumnFlags::CanBeFiltered;
-
-	Collection.Add(new FTimersViewColumn
-	(
-		Index++,
-		FTimersViewColumns::MetaGroupNameColumnID,
-		TEXT("metagroupname"),
-		LOCTEXT("TimerMetaGroupNameColumnName", "Meta Group"),
-		LOCTEXT("TimerMetaGroupNameColumnTitle", "Meta Group Name"),
-		LOCTEXT("TimerMetaGroupNameColumnDesc", "Name of the meta group"),
-		false,
-		MetaGroupNameColumnFlags,
-		HAlign_Left,
-		100.0f, 10.0f, FLT_MAX,
-		[](const FTimersViewColumn& Column, const FTimerNode& TimerNode) -> FText
-		{ return FText::FromName(TimerNode.GetMetaGroupName()); }
-	));
-
-	//////////////////////////////////////////////////
-	// Type
-
-	const ETimersViewColumnFlags TypeColumnFlags = ETimersViewColumnFlags::CanBeHidden |
-												   ETimersViewColumnFlags::CanBeSorted |
-												   ETimersViewColumnFlags::CanBeFiltered;
-
-	Collection.Add(new FTimersViewColumn
-	(
-		Index++,
-		FTimersViewColumns::TypeColumnID,
-		TEXT("type"),
-		LOCTEXT("TimerTypeColumnName", "Type"),
-		LOCTEXT("TimerTypeColumnTitle", "Type"),
-		LOCTEXT("TimerTypeColumnDesc", "Type of timer or group"),
-		false,
-		TypeColumnFlags,
-		HAlign_Left,
-		60.0f, 10.0f, FLT_MAX,
-		[](const FTimersViewColumn& Column, const FTimerNode& TimerNode) -> FText
-		{ return TimerNodeTypeHelper::ToName(TimerNode.GetType()); }
-	));
-		
-	//////////////////////////////////////////////////
-
-	const ETimersViewColumnFlags AggregatedStatsColumnFlags = ETimersViewColumnFlags::CanBeHidden |
-															  ETimersViewColumnFlags::CanBeSorted |
-															  ETimersViewColumnFlags::CanBeFiltered;
-
-	//////////////////////////////////////////////////
-	// Instance Count
-
-	Collection.Add(new FTimersViewColumn
-	(
-		Index++,
-		FTimersViewColumns::InstanceCountColumnID,
-		TEXT("count"),
-		LOCTEXT("InstanceCountName", "Count"),
-		LOCTEXT("InstanceCountTitle", "Instance Count"),
-		LOCTEXT("InstanceCountDesc", "Number of timer's instances"),
-		true,
-		AggregatedStatsColumnFlags,
-		HAlign_Right,
-		60.0f, 0.0f, FLT_MAX,
-		[](const FTimersViewColumn& Column, const FTimerNode& TimerNode) -> FText
-		{ return FText::AsNumber(TimerNode.GetAggregatedStats().InstanceCount); }
-	));
-
-	//////////////////////////////////////////////////
-
-	constexpr float TotalTimeColumnInitialWidth = 60.0f;
-	constexpr float TotalTimeColumnMinWidth = 4.0f;
-	constexpr float TotalTimeColumnMaxWidth = FLT_MAX;
-
-	constexpr float TimeMsColumnInitialWidth = 50.0f;
-	constexpr float TimeMsColumnMinWidth = 4.0f;
-	constexpr float TimeMsColumnMaxWidth = FLT_MAX;
-
-	//////////////////////////////////////////////////
-	// Inclusive Time stats
-
-	Collection.Add(new FTimersViewColumn
-	(
-		Index++,
-		FTimersViewColumns::TotalInclusiveTimeColumnID,
-		TEXT("inc"),
-		LOCTEXT("TotalInclusiveTimeName", "Incl"),
-		LOCTEXT("TotalInclusiveTimeTitle", "Total Inclusive Time"),
-		LOCTEXT("TotalInclusiveTimeDesc", "Total inclusive duration of selected timer's instances"),
-		true,
-		AggregatedStatsColumnFlags,
-		HAlign_Right,
-		TotalTimeColumnInitialWidth, TotalTimeColumnMinWidth, TotalTimeColumnMaxWidth,
-		[](const FTimersViewColumn& Column, const FTimerNode& TimerNode) -> FText
-		{ return FText::FromString(TimeUtils::FormatTimeAuto(TimerNode.GetAggregatedStats().TotalInclusiveTime)); }
-	));
-
-	Collection.Add(new FTimersViewColumn
-	(
-		Index++,
-		FTimersViewColumns::MaxInclusiveTimeColumnID,
-		TEXT("maxinc"),
-		LOCTEXT("MaxInclusiveTimeName", "I.Max"),
-		LOCTEXT("MaxInclusiveTimeTitle", "Max Inclusive Time (ms)"),
-		LOCTEXT("MaxInclusiveTimeDesc", "Maximum inclusive duration of selected timer's instances, in milliseconds"),
-		false,
-		AggregatedStatsColumnFlags,
-		HAlign_Right,
-		TimeMsColumnInitialWidth, TimeMsColumnMinWidth, TimeMsColumnMaxWidth,
-		[](const FTimersViewColumn& Column, const FTimerNode& TimerNode) -> FText
-		{ return FText::FromString(TimeUtils::FormatTimeMs(TimerNode.GetAggregatedStats().MaxInclusiveTime)); }
-	));
-
-	Collection.Add(new FTimersViewColumn
-	(
-		Index++,
-		FTimersViewColumns::AverageInclusiveTimeColumnID,
-		TEXT("avginc"),
-		LOCTEXT("AvgInclusiveTimeName", "I.Avg"),
-		LOCTEXT("AvgInclusiveTimeTitle", "Average Inclusive Time (ms)"),
-		LOCTEXT("AvgInclusiveTimeDesc", "Average inclusive duration of selected timer's instances, in milliseconds"),
-		false,
-		AggregatedStatsColumnFlags,
-		HAlign_Right,
-		TimeMsColumnInitialWidth, TimeMsColumnMinWidth, TimeMsColumnMaxWidth,
-		[](const FTimersViewColumn& Column, const FTimerNode& TimerNode) -> FText
-		{ return FText::FromString(TimeUtils::FormatTimeMs(TimerNode.GetAggregatedStats().AverageInclusiveTime)); }
-	));
-
-	Collection.Add(new FTimersViewColumn
-	(
-		Index++,
-		FTimersViewColumns::MedianInclusiveTimeColumnID,
-		TEXT("medinc"),
-		LOCTEXT("MedInclusiveTimeName", "I.Med"),
-		LOCTEXT("MedInclusiveTimeTitle", "Median Inclusive Time (ms)"),
-		LOCTEXT("MedInclusiveTimeDesc", "Median inclusive duration of selected timer's instances, in milliseconds"),
-		false,
-		AggregatedStatsColumnFlags,
-		HAlign_Right,
-		TimeMsColumnInitialWidth, TimeMsColumnMinWidth, TimeMsColumnMaxWidth,
-		[](const FTimersViewColumn& Column, const FTimerNode& TimerNode) -> FText
-		{ return FText::FromString(TimeUtils::FormatTimeMs(TimerNode.GetAggregatedStats().MedianInclusiveTime)); }
-	));
-
-	Collection.Add(new FTimersViewColumn
-	(
-		Index++,
-		FTimersViewColumns::MinInclusiveTimeColumnID,
-		TEXT("mininc"),
-		LOCTEXT("MinInclusiveTimeName", "I.Min"),
-		LOCTEXT("MinInclusiveTimeTitle", "Min Inclusive Time (ms)"),
-		LOCTEXT("MinInclusiveTimeDesc", "Minimum inclusive duration of selected timer's instances, in milliseconds"),
-		false,
-		AggregatedStatsColumnFlags,
-		HAlign_Right,
-		TimeMsColumnInitialWidth, TimeMsColumnMinWidth, TimeMsColumnMaxWidth,
-		[](const FTimersViewColumn& Column, const FTimerNode& TimerNode) -> FText
-		{ return FText::FromString(TimeUtils::FormatTimeMs(TimerNode.GetAggregatedStats().MinInclusiveTime)); }
-	));
-
-	//////////////////////////////////////////////////
-	// Exclusive Time stats
-
-	Collection.Add(new FTimersViewColumn
-	(
-		Index++,
-		FTimersViewColumns::TotalExclusiveTimeColumnID,
-		TEXT("exc"),
-		LOCTEXT("TotalExclusiveTimeName", "Excl"),
-		LOCTEXT("TotalExclusiveTimeTitle", "Total Exclusive Time"),
-		LOCTEXT("TotalExclusiveTimeDesc", "Total exclusive duration of selected timer's instances"),
-		true,
-		AggregatedStatsColumnFlags,
-		HAlign_Right,
-		TotalTimeColumnInitialWidth, TotalTimeColumnMinWidth, TotalTimeColumnMaxWidth,
-		[](const FTimersViewColumn& Column, const FTimerNode& TimerNode) -> FText
-		{ return FText::FromString(TimeUtils::FormatTimeAuto(TimerNode.GetAggregatedStats().TotalExclusiveTime)); }
-	));
-
-	Collection.Add(new FTimersViewColumn
-	(
-		Index++,
-		FTimersViewColumns::MaxExclusiveTimeColumnID,
-		TEXT("maxexc"),
-		LOCTEXT("MaxExclusiveTimeName", "E.Max"),
-		LOCTEXT("MaxExclusiveTimeTitle", "Max Exclusive Time (ms)"),
-		LOCTEXT("MaxExclusiveTimeDesc", "Maximum exclusive duration of selected timer's instances, in milliseconds"),
-		false,
-		AggregatedStatsColumnFlags,
-		HAlign_Right,
-		TimeMsColumnInitialWidth, TimeMsColumnMinWidth, TimeMsColumnMaxWidth,
-		[](const FTimersViewColumn& Column, const FTimerNode& TimerNode) -> FText
-		{ return FText::FromString(TimeUtils::FormatTimeMs(TimerNode.GetAggregatedStats().MaxExclusiveTime)); }
-	));
-
-	Collection.Add(new FTimersViewColumn
-	(
-		Index++,
-		FTimersViewColumns::AverageExclusiveTimeColumnID,
-		TEXT("avgexc"),
-		LOCTEXT("AvgExclusiveTimeName", "E.Avg"),
-		LOCTEXT("AvgExclusiveTimeTitle", "Average Exclusive Time (ms)"),
-		LOCTEXT("AvgExclusiveTimeDesc", "Average exclusive duration of selected timer's instances, in milliseconds"),
-		false,
-		AggregatedStatsColumnFlags,
-		HAlign_Right,
-		TimeMsColumnInitialWidth, TimeMsColumnMinWidth, TimeMsColumnMaxWidth,
-		[](const FTimersViewColumn& Column, const FTimerNode& TimerNode) -> FText
-		{ return FText::FromString(TimeUtils::FormatTimeMs(TimerNode.GetAggregatedStats().AverageExclusiveTime)); }
-	));
-
-	Collection.Add(new FTimersViewColumn
-	(
-		Index++,
-		FTimersViewColumns::MedianExclusiveTimeColumnID,
-		TEXT("medexc"),
-		LOCTEXT("MedExclusiveTimeName", "E.Med"),
-		LOCTEXT("MedExclusiveTimeTitle", "Median Exclusive Time (ms)"),
-		LOCTEXT("MedExclusiveTimeDesc", "Median exclusive duration of timer's instances, in milliseconds"),
-		false,
-		AggregatedStatsColumnFlags,
-		HAlign_Right,
-		TimeMsColumnInitialWidth, TimeMsColumnMinWidth, TimeMsColumnMaxWidth,
-		[](const FTimersViewColumn& Column, const FTimerNode& TimerNode) -> FText
-		{ return FText::FromString(TimeUtils::FormatTimeMs(TimerNode.GetAggregatedStats().MedianExclusiveTime)); }
-	));
-
-	Collection.Add(new FTimersViewColumn
-	(
-		Index++,
-		FTimersViewColumns::MinExclusiveTimeColumnID,
-		TEXT("minexc"),
-		LOCTEXT("MinExclusiveTimeName", "E.Min"),
-		LOCTEXT("MinExclusiveTimeTitle", "Min Exclusive Time (ms)"),
-		LOCTEXT("MinExclusiveTimeDesc", "Minimum exclusive duration of selected timer's instances, in milliseconds"),
-		false,
-		AggregatedStatsColumnFlags,
-		HAlign_Right,
-		TimeMsColumnInitialWidth, TimeMsColumnMinWidth, TimeMsColumnMaxWidth,
-		[](const FTimersViewColumn& Column, const FTimerNode& TimerNode) -> FText
-		{ return FText::FromString(TimeUtils::FormatTimeMs(TimerNode.GetAggregatedStats().MinExclusiveTime)); }
-	));
-
-	//////////////////////////////////////////////////
-
-	for (int32 MapIndex = 0; MapIndex < Collection.Num(); MapIndex++)
+	Columns.Reset();
+
+	Columns.Add(CreateNameColumn());
+	Columns.Add(CreateMetaGroupNameColumn());
+	Columns.Add(CreateTypeColumn());
+	Columns.Add(CreateInstanceCountColumn());
+
+	Columns.Add(CreateTotalInclusiveTimeColumn());
+	Columns.Add(CreateMaxInclusiveTimeColumn());
+	Columns.Add(CreateAverageInclusiveTimeColumn());
+	Columns.Add(CreateMedianInclusiveTimeColumn());
+	Columns.Add(CreateMinInclusiveTimeColumn());
+
+	Columns.Add(CreateTotalExclusiveTimeColumn());
+	Columns.Add(CreateMaxExclusiveTimeColumn());
+	Columns.Add(CreateAverageExclusiveTimeColumn());
+	Columns.Add(CreateMedianExclusiveTimeColumn());
+	Columns.Add(CreateMinExclusiveTimeColumn());
+}
+
+////////////////////////////////////////////////////////////////////////////////////////////////////
+
+void FTimersViewColumnFactory::CreateTimerTreeViewColumns(TArray<TSharedPtr<Insights::FTableColumn>>& Columns)
+{
+	Columns.Reset();
+
+	Columns.Add(CreateNameColumn());
+	Columns.Add(CreateTypeColumn());
+	Columns.Add(CreateInstanceCountColumn());
+	Columns.Add(CreateTotalInclusiveTimeColumn());
+	Columns.Add(CreateTotalExclusiveTimeColumn());
+}
+
+////////////////////////////////////////////////////////////////////////////////////////////////////
+
+TSharedRef<Insights::FTableColumn> FTimersViewColumnFactory::CreateNameColumn()
+{
+	using namespace Insights;
+
+	TSharedPtr<FTableColumn> ColumnPtr = MakeShareable(new FTableColumn(FTimersViewColumns::NameColumnID));
+	FTableColumn& Column = *ColumnPtr;
+
+	Column.SetShortName(LOCTEXT("TimerNameColumnName", "Name"));
+	Column.SetTitleName(LOCTEXT("TimerNameColumnTitle", "Timer or Group Name"));
+	Column.SetDescription(LOCTEXT("TimerNameColumnDesc", "Name of the timer or group"));
+
+	Column.SetFlags(ETableColumnFlags::ShouldBeVisible |
+					ETableColumnFlags::CanBeFiltered |
+					ETableColumnFlags::IsHierarchy);
+
+	Column.SetHorizontalAlignment(HAlign_Left);
+	Column.SetInitialWidth(246.0f);
+	Column.SetMinWidth(42.0f);
+
+	Column.SetDataType(ETableCellDataType::Text);
+
+	TSharedPtr<ITableCellValueGetter> GetterPtr = MakeShareable(new FDisplayNameValueGetter());
+	Column.SetValueGetter(GetterPtr.ToSharedRef());
+
+	TSharedPtr<ITableCellValueFormatter> FormatterPtr = MakeShareable(new FTextValueFormatter());
+	Column.SetValueFormatter(FormatterPtr.ToSharedRef());
+
+	TSharedPtr<ITableCellValueSorter> SorterPtr = MakeShareable(new FSorterByName(ColumnPtr.ToSharedRef()));
+	Column.SetValueSorter(SorterPtr);
+
+	return ColumnPtr.ToSharedRef();
+}
+
+////////////////////////////////////////////////////////////////////////////////////////////////////
+
+TSharedRef<Insights::FTableColumn> FTimersViewColumnFactory::CreateMetaGroupNameColumn()
+{
+	using namespace Insights;
+
+	TSharedPtr<FTableColumn> ColumnPtr = MakeShareable(new FTableColumn(FTimersViewColumns::MetaGroupNameColumnID));
+	FTableColumn& Column = *ColumnPtr;
+
+	Column.SetShortName(LOCTEXT("TimerMetaGroupNameColumnName", "Meta Group"));
+	Column.SetTitleName(LOCTEXT("TimerMetaGroupNameColumnTitle", "Meta Group Name"));
+	Column.SetDescription(LOCTEXT("TimerMetaGroupNameColumnDesc", "Name of the meta group"));
+
+	Column.SetFlags(ETableColumnFlags::CanBeHidden |
+					//ETableColumnFlags::ShouldBeVisible |
+					ETableColumnFlags::CanBeFiltered);
+
+	Column.SetHorizontalAlignment(HAlign_Left);
+	Column.SetInitialWidth(100.0f);
+
+	Column.SetDataType(ETableCellDataType::Text);
+
+	class FMetaGroupNameValueGetter : public FTableCellValueGetter
+	{
+	public:
+		virtual const TOptional<FTableCellValue> GetValue(const FTableColumn& Column, const FBaseTreeNode& Node) const
+		{
+			ensure(Column.GetId() == FTimersViewColumns::MetaGroupNameColumnID);
+			const FTimerNode& TimerNode = static_cast<const FTimerNode&>(Node);
+			return TOptional<FTableCellValue>(FTableCellValue(FText::FromName(TimerNode.GetMetaGroupName())));
+		}
+	};
+
+	TSharedPtr<FTableCellValueGetter> GetterPtr = MakeShareable(new FMetaGroupNameValueGetter());
+	Column.SetValueGetter(GetterPtr.ToSharedRef());
+
+	TSharedPtr<FTableCellValueFormatter> FormatterPtr = MakeShareable(new FTextValueFormatter());
+	Column.SetValueFormatter(FormatterPtr.ToSharedRef());
+
+	TSharedPtr<ITableCellValueSorter> SorterPtr = MakeShareable(new FSorterByTextValue(ColumnPtr.ToSharedRef()));
+	Column.SetValueSorter(SorterPtr);
+
+	return ColumnPtr.ToSharedRef();
+}
+
+////////////////////////////////////////////////////////////////////////////////////////////////////
+
+TSharedRef<Insights::FTableColumn> FTimersViewColumnFactory::CreateTypeColumn()
+{
+	using namespace Insights;
+
+	TSharedPtr<FTableColumn> ColumnPtr = MakeShareable(new FTableColumn(FTimersViewColumns::TypeColumnID));
+	FTableColumn& Column = *ColumnPtr;
+
+	Column.SetShortName(LOCTEXT("TimerTypeColumnName", "Type"));
+	Column.SetTitleName(LOCTEXT("TimerTypeColumnTitle", "Type"));
+	Column.SetDescription(LOCTEXT("TimerTypeColumnDesc", "Type of timer or group"));
+
+	Column.SetFlags(ETableColumnFlags::CanBeHidden |
+					//ETableColumnFlags::ShouldBeVisible |
+					ETableColumnFlags::CanBeFiltered);
+
+	Column.SetHorizontalAlignment(HAlign_Left);
+	Column.SetInitialWidth(60.0f);
+
+	Column.SetDataType(ETableCellDataType::Text);
+
+	class FTimerTypeValueGetter : public FTableCellValueGetter
+	{
+	public:
+		virtual const TOptional<FTableCellValue> GetValue(const FTableColumn& Column, const FBaseTreeNode& Node) const
+		{
+			ensure(Column.GetId() == FTimersViewColumns::TypeColumnID);
+			const FTimerNode& TimerNode = static_cast<const FTimerNode&>(Node);
+			return TOptional<FTableCellValue>(FTableCellValue(TimerNodeTypeHelper::ToText(TimerNode.GetType())));
+		}
+	};
+
+	TSharedPtr<FTableCellValueGetter> GetterPtr = MakeShareable(new FTimerTypeValueGetter());
+	Column.SetValueGetter(GetterPtr.ToSharedRef());
+
+	TSharedPtr<FTableCellValueFormatter> FormatterPtr = MakeShareable(new FTextValueFormatter());
+	Column.SetValueFormatter(FormatterPtr.ToSharedRef());
+
+	//TSharedPtr<ITableCellValueSorter> SorterPtr = MakeShareable(new FSorterByTextValue(ColumnPtr.ToSharedRef()));
+	TSharedPtr<ITableCellValueSorter> SorterPtr = MakeShareable(new FTimerNodeSortingByTimerType(ColumnPtr.ToSharedRef()));
+	Column.SetValueSorter(SorterPtr);
+
+	return ColumnPtr.ToSharedRef();
+}
+
+////////////////////////////////////////////////////////////////////////////////////////////////////
+
+TSharedRef<Insights::FTableColumn> FTimersViewColumnFactory::CreateInstanceCountColumn()
+{
+	using namespace Insights;
+
+	TSharedPtr<FTableColumn> ColumnPtr = MakeShareable(new FTableColumn(FTimersViewColumns::InstanceCountColumnID));
+	FTableColumn& Column = *ColumnPtr;
+
+	Column.SetShortName(LOCTEXT("InstanceCountName", "Count"));
+	Column.SetTitleName(LOCTEXT("InstanceCountTitle", "Instance Count"));
+	Column.SetDescription(LOCTEXT("InstanceCountDesc", "Number of timer's instances"));
+
+	Column.SetFlags(ETableColumnFlags::CanBeHidden |
+					ETableColumnFlags::ShouldBeVisible |
+					ETableColumnFlags::CanBeFiltered);
+
+	Column.SetHorizontalAlignment(HAlign_Right);
+	Column.SetInitialWidth(60.0f);
+
+	Column.SetDataType(ETableCellDataType::Int64);
+
+	class FInstanceCountValueGetter : public FTableCellValueGetter
+	{
+	public:
+		virtual const TOptional<FTableCellValue> GetValue(const FTableColumn& Column, const FBaseTreeNode& Node) const
+		{
+			ensure(Column.GetId() == FTimersViewColumns::InstanceCountColumnID);
+			const FTimerNode& TimerNode = static_cast<const FTimerNode&>(Node);
+			return TOptional<FTableCellValue>(FTableCellValue(static_cast<int64>(TimerNode.GetAggregatedStats().InstanceCount)));
+		}
+	};
+
+	TSharedPtr<FTableCellValueGetter> GetterPtr = MakeShareable(new FInstanceCountValueGetter());
+	Column.SetValueGetter(GetterPtr.ToSharedRef());
+
+	TSharedPtr<FTableCellValueFormatter> FormatterPtr = MakeShareable(new FInt64ValueFormatterAsNumber());
+	Column.SetValueFormatter(FormatterPtr.ToSharedRef());
+
+	//TSharedPtr<ITableCellValueSorter> SorterPtr = MakeShareable(new FSorterByInt64Value(ColumnPtr.ToSharedRef()));
+	TSharedPtr<ITableCellValueSorter> SorterPtr = MakeShareable(new FTimerNodeSortingByInstanceCount(ColumnPtr.ToSharedRef()));
+	Column.SetValueSorter(SorterPtr);
+
+	return ColumnPtr.ToSharedRef();
+}
+
+////////////////////////////////////////////////////////////////////////////////////////////////////
+// Inclusive Time Columns
+////////////////////////////////////////////////////////////////////////////////////////////////////
+
+TSharedRef<Insights::FTableColumn> FTimersViewColumnFactory::CreateTotalInclusiveTimeColumn()
+{
+	using namespace Insights;
+
+	TSharedPtr<FTableColumn> ColumnPtr = MakeShareable(new FTableColumn(FTimersViewColumns::TotalInclusiveTimeColumnID));
+	FTableColumn& Column = *ColumnPtr;
+
+	Column.SetShortName(LOCTEXT("TotalInclusiveTimeName", "Incl"));
+	Column.SetTitleName(LOCTEXT("TotalInclusiveTimeTitle", "Total Inclusive Time"));
+	Column.SetDescription(LOCTEXT("TotalInclusiveTimeDesc", "Total inclusive duration of selected timer's instances"));
+
+	Column.SetFlags(ETableColumnFlags::CanBeHidden |
+					ETableColumnFlags::ShouldBeVisible |
+					ETableColumnFlags::CanBeFiltered);
+
+	Column.SetHorizontalAlignment(HAlign_Right);
+	Column.SetInitialWidth(TotalTimeColumnInitialWidth);
+
+	Column.SetDataType(ETableCellDataType::Double);
+	Column.SetAggregation(ETableColumnAggregation::Sum);
+
+	class FTotalInclusiveTimeValueGetter : public FTableCellValueGetter
+	{
+	public:
+		virtual const TOptional<FTableCellValue> GetValue(const FTableColumn& Column, const FBaseTreeNode& Node) const
+		{
+			ensure(Column.GetId() == FTimersViewColumns::TotalInclusiveTimeColumnID);
+			const FTimerNode& TimerNode = static_cast<const FTimerNode&>(Node);
+			return TOptional<FTableCellValue>(FTableCellValue(TimerNode.GetAggregatedStats().TotalInclusiveTime));
+		}
+	};
+
+	TSharedPtr<FTableCellValueGetter> GetterPtr = MakeShareable(new FTotalInclusiveTimeValueGetter());
+	Column.SetValueGetter(GetterPtr.ToSharedRef());
+
+	TSharedPtr<FTableCellValueFormatter> FormatterPtr = MakeShareable(new FDoubleValueFormatterAsTimeAuto());
+	Column.SetValueFormatter(FormatterPtr.ToSharedRef());
+
+	//TSharedPtr<ITableCellValueSorter> SorterPtr = MakeShareable(new FSorterByDoubleValue(ColumnPtr.ToSharedRef()));
+	TSharedPtr<ITableCellValueSorter> SorterPtr = MakeShareable(new FTimerNodeSortingByTotalInclusiveTime(ColumnPtr.ToSharedRef()));
+	Column.SetValueSorter(SorterPtr);
+
+	return ColumnPtr.ToSharedRef();
+}
+
+////////////////////////////////////////////////////////////////////////////////////////////////////
+
+TSharedRef<Insights::FTableColumn> FTimersViewColumnFactory::CreateMaxInclusiveTimeColumn()
+{
+	using namespace Insights;
+
+	TSharedPtr<FTableColumn> ColumnPtr = MakeShareable(new FTableColumn(FTimersViewColumns::MaxInclusiveTimeColumnID));
+	FTableColumn& Column = *ColumnPtr;
+
+	Column.SetShortName(LOCTEXT("MaxInclusiveTimeName", "I.Max"));
+	Column.SetTitleName(LOCTEXT("MaxInclusiveTimeTitle", "Max Inclusive Time (ms)"));
+	Column.SetDescription(LOCTEXT("MaxInclusiveTimeDesc", "Maximum inclusive duration of selected timer's instances, in milliseconds"));
+
+	Column.SetFlags(ETableColumnFlags::CanBeHidden |
+					//ETableColumnFlags::ShouldBeVisible |
+					ETableColumnFlags::CanBeFiltered);
+
+	Column.SetHorizontalAlignment(HAlign_Right);
+	Column.SetInitialWidth(TimeMsColumnInitialWidth);
+
+	Column.SetDataType(ETableCellDataType::Double);
+
+	class FMaxInclusiveTimeValueGetter : public FTableCellValueGetter
+	{
+	public:
+		virtual const TOptional<FTableCellValue> GetValue(const FTableColumn& Column, const FBaseTreeNode& Node) const
+		{
+			ensure(Column.GetId() == FTimersViewColumns::MaxInclusiveTimeColumnID);
+			const FTimerNode& TimerNode = static_cast<const FTimerNode&>(Node);
+			return TOptional<FTableCellValue>(FTableCellValue(TimerNode.GetAggregatedStats().MaxInclusiveTime));
+		}
+	};
+
+	TSharedPtr<FTableCellValueGetter> GetterPtr = MakeShareable(new FMaxInclusiveTimeValueGetter());
+	Column.SetValueGetter(GetterPtr.ToSharedRef());
+
+	TSharedPtr<FTableCellValueFormatter> FormatterPtr = MakeShareable(new FDoubleValueFormatterAsTimeMs());
+	Column.SetValueFormatter(FormatterPtr.ToSharedRef());
+
+	TSharedPtr<ITableCellValueSorter> SorterPtr = MakeShareable(new FSorterByDoubleValue(ColumnPtr.ToSharedRef()));
+	Column.SetValueSorter(SorterPtr);
+
+	return ColumnPtr.ToSharedRef();
+}
+
+////////////////////////////////////////////////////////////////////////////////////////////////////
+
+TSharedRef<Insights::FTableColumn> FTimersViewColumnFactory::CreateAverageInclusiveTimeColumn()
+{
+	using namespace Insights;
+
+	TSharedPtr<FTableColumn> ColumnPtr = MakeShareable(new FTableColumn(FTimersViewColumns::AverageInclusiveTimeColumnID));
+	FTableColumn& Column = *ColumnPtr;
+
+	Column.SetShortName(LOCTEXT("AvgInclusiveTimeName", "I.Avg"));
+	Column.SetTitleName(LOCTEXT("AvgInclusiveTimeTitle", "Average Inclusive Time (ms)"));
+	Column.SetDescription(LOCTEXT("AvgInclusiveTimeDesc", "Average inclusive duration of selected timer's instances, in milliseconds"));
+
+	Column.SetFlags(ETableColumnFlags::CanBeHidden |
+					//ETableColumnFlags::ShouldBeVisible |
+					ETableColumnFlags::CanBeFiltered);
+
+	Column.SetHorizontalAlignment(HAlign_Right);
+	Column.SetInitialWidth(TimeMsColumnInitialWidth);
+
+	Column.SetDataType(ETableCellDataType::Double);
+
+	class FAverageInclusiveTimeValueGetter : public FTableCellValueGetter
+	{
+	public:
+		virtual const TOptional<FTableCellValue> GetValue(const FTableColumn& Column, const FBaseTreeNode& Node) const
+		{
+			ensure(Column.GetId() == FTimersViewColumns::AverageInclusiveTimeColumnID);
+			const FTimerNode& TimerNode = static_cast<const FTimerNode&>(Node);
+			return TOptional<FTableCellValue>(FTableCellValue(TimerNode.GetAggregatedStats().AverageInclusiveTime));
+		}
+	};
+
+	TSharedPtr<FTableCellValueGetter> GetterPtr = MakeShareable(new FAverageInclusiveTimeValueGetter());
+	Column.SetValueGetter(GetterPtr.ToSharedRef());
+
+	TSharedPtr<FTableCellValueFormatter> FormatterPtr = MakeShareable(new FDoubleValueFormatterAsTimeMs());
+	Column.SetValueFormatter(FormatterPtr.ToSharedRef());
+
+	TSharedPtr<ITableCellValueSorter> SorterPtr = MakeShareable(new FSorterByDoubleValue(ColumnPtr.ToSharedRef()));
+	Column.SetValueSorter(SorterPtr);
+
+	return ColumnPtr.ToSharedRef();
+}
+
+////////////////////////////////////////////////////////////////////////////////////////////////////
+
+TSharedRef<Insights::FTableColumn> FTimersViewColumnFactory::CreateMedianInclusiveTimeColumn()
+{
+	using namespace Insights;
+
+	TSharedPtr<FTableColumn> ColumnPtr = MakeShareable(new FTableColumn(FTimersViewColumns::MedianInclusiveTimeColumnID));
+	FTableColumn& Column = *ColumnPtr;
+
+	Column.SetShortName(LOCTEXT("MedInclusiveTimeName", "I.Med"));
+	Column.SetTitleName(LOCTEXT("MedInclusiveTimeTitle", "Median Inclusive Time (ms)"));
+	Column.SetDescription(LOCTEXT("MedInclusiveTimeDesc", "Median inclusive duration of selected timer's instances, in milliseconds"));
+
+	Column.SetFlags(ETableColumnFlags::CanBeHidden |
+					//ETableColumnFlags::ShouldBeVisible |
+					ETableColumnFlags::CanBeFiltered);
+
+	Column.SetHorizontalAlignment(HAlign_Right);
+	Column.SetInitialWidth(TimeMsColumnInitialWidth);
+
+	Column.SetDataType(ETableCellDataType::Double);
+
+	class FMedianInclusiveTimeValueGetter : public FTableCellValueGetter
 	{
-		ColumnIdToPtrMapping.Add(Collection[MapIndex]->Id, Collection[MapIndex]);
-	}
+	public:
+		virtual const TOptional<FTableCellValue> GetValue(const FTableColumn& Column, const FBaseTreeNode& Node) const
+		{
+			ensure(Column.GetId() == FTimersViewColumns::MedianInclusiveTimeColumnID);
+			const FTimerNode& TimerNode = static_cast<const FTimerNode&>(Node);
+			return TOptional<FTableCellValue>(FTableCellValue(TimerNode.GetAggregatedStats().MedianInclusiveTime));
+		}
+	};
+
+	TSharedPtr<FTableCellValueGetter> GetterPtr = MakeShareable(new FMedianInclusiveTimeValueGetter());
+	Column.SetValueGetter(GetterPtr.ToSharedRef());
+
+	TSharedPtr<FTableCellValueFormatter> FormatterPtr = MakeShareable(new FDoubleValueFormatterAsTimeMs());
+	Column.SetValueFormatter(FormatterPtr.ToSharedRef());
+
+	TSharedPtr<ITableCellValueSorter> SorterPtr = MakeShareable(new FSorterByDoubleValue(ColumnPtr.ToSharedRef()));
+	Column.SetValueSorter(SorterPtr);
+
+	return ColumnPtr.ToSharedRef();
 }
 
 ////////////////////////////////////////////////////////////////////////////////////////////////////
 
-FTimersViewColumnFactory::~FTimersViewColumnFactory()
+TSharedRef<Insights::FTableColumn> FTimersViewColumnFactory::CreateMinInclusiveTimeColumn()
 {
-	for (const FTimersViewColumn* Column : Collection)
+	using namespace Insights;
+
+	TSharedPtr<FTableColumn> ColumnPtr = MakeShareable(new FTableColumn(FTimersViewColumns::MinInclusiveTimeColumnID));
+	FTableColumn& Column = *ColumnPtr;
+
+	Column.SetShortName(LOCTEXT("MinInclusiveTimeName", "I.Min"));
+	Column.SetTitleName(LOCTEXT("MinInclusiveTimeTitle", "Min Inclusive Time (ms)"));
+	Column.SetDescription(LOCTEXT("MinInclusiveTimeDesc", "Minimum inclusive duration of selected timer's instances, in milliseconds"));
+
+	Column.SetFlags(ETableColumnFlags::CanBeHidden |
+					//ETableColumnFlags::ShouldBeVisible |
+					ETableColumnFlags::CanBeFiltered);
+
+	Column.SetHorizontalAlignment(HAlign_Right);
+	Column.SetInitialWidth(TimeMsColumnInitialWidth);
+
+	Column.SetDataType(ETableCellDataType::Double);
+
+	class FMinInclusiveTimeValueGetter : public FTableCellValueGetter
 	{
-		delete Column;
-	}
-	Collection.Reset();
+	public:
+		virtual const TOptional<FTableCellValue> GetValue(const FTableColumn& Column, const FBaseTreeNode& Node) const
+		{
+			ensure(Column.GetId() == FTimersViewColumns::MinInclusiveTimeColumnID);
+			const FTimerNode& TimerNode = static_cast<const FTimerNode&>(Node);
+			return TOptional<FTableCellValue>(FTableCellValue(TimerNode.GetAggregatedStats().MinInclusiveTime));
+		}
+	};
+
+	TSharedPtr<FTableCellValueGetter> GetterPtr = MakeShareable(new FMinInclusiveTimeValueGetter());
+	Column.SetValueGetter(GetterPtr.ToSharedRef());
+
+	TSharedPtr<FTableCellValueFormatter> FormatterPtr = MakeShareable(new FDoubleValueFormatterAsTimeMs());
+	Column.SetValueFormatter(FormatterPtr.ToSharedRef());
+
+	TSharedPtr<ITableCellValueSorter> SorterPtr = MakeShareable(new FSorterByDoubleValue(ColumnPtr.ToSharedRef()));
+	Column.SetValueSorter(SorterPtr);
+
+	return ColumnPtr.ToSharedRef();
 }
 
+////////////////////////////////////////////////////////////////////////////////////////////////////
+// Exclusive Time Columns
 ////////////////////////////////////////////////////////////////////////////////////////////////////
 
-const FTimersViewColumnFactory& FTimersViewColumnFactory::Get()
+TSharedRef<Insights::FTableColumn> FTimersViewColumnFactory::CreateTotalExclusiveTimeColumn()
 {
-	static FTimersViewColumnFactory Instance;
-	return Instance;
+	using namespace Insights;
+
+	TSharedPtr<FTableColumn> ColumnPtr = MakeShareable(new FTableColumn(FTimersViewColumns::TotalExclusiveTimeColumnID));
+	FTableColumn& Column = *ColumnPtr;
+
+	Column.SetShortName(LOCTEXT("TotalExclusiveTimeName", "Excl"));
+	Column.SetTitleName(LOCTEXT("TotalExclusiveTimeTitle", "Total Exclusive Time"));
+	Column.SetDescription(LOCTEXT("TotalExclusiveTimeDesc", "Total exclusive duration of selected timer's instances"));
+
+	Column.SetFlags(ETableColumnFlags::CanBeHidden |
+					ETableColumnFlags::ShouldBeVisible |
+					ETableColumnFlags::CanBeFiltered);
+
+	Column.SetHorizontalAlignment(HAlign_Right);
+	Column.SetInitialWidth(TotalTimeColumnInitialWidth);
+
+	Column.SetDataType(ETableCellDataType::Double);
+	Column.SetAggregation(ETableColumnAggregation::Sum);
+
+	class FTotalExclusiveTimeValueGetter : public FTableCellValueGetter
+	{
+	public:
+		virtual const TOptional<FTableCellValue> GetValue(const FTableColumn& Column, const FBaseTreeNode& Node) const
+		{
+			ensure(Column.GetId() == FTimersViewColumns::TotalExclusiveTimeColumnID);
+			const FTimerNode& TimerNode = static_cast<const FTimerNode&>(Node);
+			return TOptional<FTableCellValue>(FTableCellValue(TimerNode.GetAggregatedStats().TotalExclusiveTime));
+		}
+	};
+
+	TSharedPtr<FTableCellValueGetter> GetterPtr = MakeShareable(new FTotalExclusiveTimeValueGetter());
+	Column.SetValueGetter(GetterPtr.ToSharedRef());
+
+	TSharedPtr<FTableCellValueFormatter> FormatterPtr = MakeShareable(new FDoubleValueFormatterAsTimeAuto());
+	Column.SetValueFormatter(FormatterPtr.ToSharedRef());
+
+	//TSharedPtr<ITableCellValueSorter> SorterPtr = MakeShareable(new FSorterByDoubleValue(ColumnPtr.ToSharedRef()));
+	TSharedPtr<ITableCellValueSorter> SorterPtr = MakeShareable(new FTimerNodeSortingByTotalExclusiveTime(ColumnPtr.ToSharedRef()));
+	Column.SetValueSorter(SorterPtr);
+
+	return ColumnPtr.ToSharedRef();
+}
+
+////////////////////////////////////////////////////////////////////////////////////////////////////
+
+TSharedRef<Insights::FTableColumn> FTimersViewColumnFactory::CreateMaxExclusiveTimeColumn()
+{
+	using namespace Insights;
+
+	TSharedPtr<FTableColumn> ColumnPtr = MakeShareable(new FTableColumn(FTimersViewColumns::MaxExclusiveTimeColumnID));
+	FTableColumn& Column = *ColumnPtr;
+
+	Column.SetShortName(LOCTEXT("MaxExclusiveTimeName", "E.Max"));
+	Column.SetTitleName(LOCTEXT("MaxExclusiveTimeTitle", "Max Exclusive Time (ms)"));
+	Column.SetDescription(LOCTEXT("MaxExclusiveTimeDesc", "Maximum exclusive duration of selected timer's instances, in milliseconds"));
+
+	Column.SetFlags(ETableColumnFlags::CanBeHidden |
+					//ETableColumnFlags::ShouldBeVisible |
+					ETableColumnFlags::CanBeFiltered);
+
+	Column.SetHorizontalAlignment(HAlign_Right);
+	Column.SetInitialWidth(TimeMsColumnInitialWidth);
+
+	Column.SetDataType(ETableCellDataType::Double);
+
+	class FMaxExclusiveTimeValueGetter : public FTableCellValueGetter
+	{
+	public:
+		virtual const TOptional<FTableCellValue> GetValue(const FTableColumn& Column, const FBaseTreeNode& Node) const
+		{
+			ensure(Column.GetId() == FTimersViewColumns::MaxExclusiveTimeColumnID);
+			const FTimerNode& TimerNode = static_cast<const FTimerNode&>(Node);
+			return TOptional<FTableCellValue>(FTableCellValue(TimerNode.GetAggregatedStats().MaxExclusiveTime));
+		}
+	};
+
+	TSharedPtr<FTableCellValueGetter> GetterPtr = MakeShareable(new FMaxExclusiveTimeValueGetter());
+	Column.SetValueGetter(GetterPtr.ToSharedRef());
+
+	TSharedPtr<FTableCellValueFormatter> FormatterPtr = MakeShareable(new FDoubleValueFormatterAsTimeMs());
+	Column.SetValueFormatter(FormatterPtr.ToSharedRef());
+
+	TSharedPtr<ITableCellValueSorter> SorterPtr = MakeShareable(new FSorterByDoubleValue(ColumnPtr.ToSharedRef()));
+	Column.SetValueSorter(SorterPtr);
+
+	return ColumnPtr.ToSharedRef();
+}
+
+////////////////////////////////////////////////////////////////////////////////////////////////////
+
+TSharedRef<Insights::FTableColumn> FTimersViewColumnFactory::CreateAverageExclusiveTimeColumn()
+{
+	using namespace Insights;
+
+	TSharedPtr<FTableColumn> ColumnPtr = MakeShareable(new FTableColumn(FTimersViewColumns::AverageExclusiveTimeColumnID));
+	FTableColumn& Column = *ColumnPtr;
+
+	Column.SetShortName(LOCTEXT("AvgExclusiveTimeName", "E.Avg"));
+	Column.SetTitleName(LOCTEXT("AvgExclusiveTimeTitle", "Average Exclusive Time (ms)"));
+	Column.SetDescription(LOCTEXT("AvgExclusiveTimeDesc", "Average exclusive duration of selected timer's instances, in milliseconds"));
+
+	Column.SetFlags(ETableColumnFlags::CanBeHidden |
+					//ETableColumnFlags::ShouldBeVisible |
+					ETableColumnFlags::CanBeFiltered);
+
+	Column.SetHorizontalAlignment(HAlign_Right);
+	Column.SetInitialWidth(TimeMsColumnInitialWidth);
+
+	Column.SetDataType(ETableCellDataType::Double);
+
+	class FAverageExclusiveTimeValueGetter : public FTableCellValueGetter
+	{
+	public:
+		virtual const TOptional<FTableCellValue> GetValue(const FTableColumn& Column, const FBaseTreeNode& Node) const
+		{
+			ensure(Column.GetId() == FTimersViewColumns::AverageExclusiveTimeColumnID);
+			const FTimerNode& TimerNode = static_cast<const FTimerNode&>(Node);
+			return TOptional<FTableCellValue>(FTableCellValue(TimerNode.GetAggregatedStats().AverageExclusiveTime));
+		}
+	};
+
+	TSharedPtr<FTableCellValueGetter> GetterPtr = MakeShareable(new FAverageExclusiveTimeValueGetter());
+	Column.SetValueGetter(GetterPtr.ToSharedRef());
+
+	TSharedPtr<FTableCellValueFormatter> FormatterPtr = MakeShareable(new FDoubleValueFormatterAsTimeMs());
+	Column.SetValueFormatter(FormatterPtr.ToSharedRef());
+
+	TSharedPtr<ITableCellValueSorter> SorterPtr = MakeShareable(new FSorterByDoubleValue(ColumnPtr.ToSharedRef()));
+	Column.SetValueSorter(SorterPtr);
+
+	return ColumnPtr.ToSharedRef();
+}
+
+////////////////////////////////////////////////////////////////////////////////////////////////////
+
+TSharedRef<Insights::FTableColumn> FTimersViewColumnFactory::CreateMedianExclusiveTimeColumn()
+{
+	using namespace Insights;
+
+	TSharedPtr<FTableColumn> ColumnPtr = MakeShareable(new FTableColumn(FTimersViewColumns::MedianExclusiveTimeColumnID));
+	FTableColumn& Column = *ColumnPtr;
+
+	Column.SetShortName(LOCTEXT("MedExclusiveTimeName", "E.Med"));
+	Column.SetTitleName(LOCTEXT("MedExclusiveTimeTitle", "Median Exclusive Time (ms)"));
+	Column.SetDescription(LOCTEXT("MedExclusiveTimeDesc", "Median exclusive duration of selected timer's instances, in milliseconds"));
+
+	Column.SetFlags(ETableColumnFlags::CanBeHidden |
+					//ETableColumnFlags::ShouldBeVisible |
+					ETableColumnFlags::CanBeFiltered);
+
+	Column.SetHorizontalAlignment(HAlign_Right);
+	Column.SetInitialWidth(TimeMsColumnInitialWidth);
+
+	Column.SetDataType(ETableCellDataType::Double);
+
+	class FMedianExclusiveTimeValueGetter : public FTableCellValueGetter
+	{
+	public:
+		virtual const TOptional<FTableCellValue> GetValue(const FTableColumn& Column, const FBaseTreeNode& Node) const
+		{
+			ensure(Column.GetId() == FTimersViewColumns::MedianExclusiveTimeColumnID);
+			const FTimerNode& TimerNode = static_cast<const FTimerNode&>(Node);
+			return TOptional<FTableCellValue>(FTableCellValue(TimerNode.GetAggregatedStats().MedianExclusiveTime));
+		}
+	};
+
+	TSharedPtr<FTableCellValueGetter> GetterPtr = MakeShareable(new FMedianExclusiveTimeValueGetter());
+	Column.SetValueGetter(GetterPtr.ToSharedRef());
+
+	TSharedPtr<FTableCellValueFormatter> FormatterPtr = MakeShareable(new FDoubleValueFormatterAsTimeMs());
+	Column.SetValueFormatter(FormatterPtr.ToSharedRef());
+
+	TSharedPtr<ITableCellValueSorter> SorterPtr = MakeShareable(new FSorterByDoubleValue(ColumnPtr.ToSharedRef()));
+	Column.SetValueSorter(SorterPtr);
+
+	return ColumnPtr.ToSharedRef();
+}
+
+////////////////////////////////////////////////////////////////////////////////////////////////////
+
+TSharedRef<Insights::FTableColumn> FTimersViewColumnFactory::CreateMinExclusiveTimeColumn()
+{
+	using namespace Insights;
+
+	TSharedPtr<FTableColumn> ColumnPtr = MakeShareable(new FTableColumn(FTimersViewColumns::MinExclusiveTimeColumnID));
+	FTableColumn& Column = *ColumnPtr;
+
+	Column.SetShortName(LOCTEXT("MinExclusiveTimeName", "E.Min"));
+	Column.SetTitleName(LOCTEXT("MinExclusiveTimeTitle", "Min Exclusive Time (ms)"));
+	Column.SetDescription(LOCTEXT("MinExclusiveTimeDesc", "Minimum exclusive duration of selected timer's instances, in milliseconds"));
+
+	Column.SetFlags(ETableColumnFlags::CanBeHidden |
+					//ETableColumnFlags::ShouldBeVisible |
+					ETableColumnFlags::CanBeFiltered);
+
+	Column.SetHorizontalAlignment(HAlign_Right);
+	Column.SetInitialWidth(TimeMsColumnInitialWidth);
+
+	Column.SetDataType(ETableCellDataType::Double);
+
+	class FMinExclusiveTimeValueGetter : public FTableCellValueGetter
+	{
+	public:
+		virtual const TOptional<FTableCellValue> GetValue(const FTableColumn& Column, const FBaseTreeNode& Node) const
+		{
+			ensure(Column.GetId() == FTimersViewColumns::MinExclusiveTimeColumnID);
+			const FTimerNode& TimerNode = static_cast<const FTimerNode&>(Node);
+			return TOptional<FTableCellValue>(FTableCellValue(TimerNode.GetAggregatedStats().MinExclusiveTime));
+		}
+	};
+
+	TSharedPtr<FTableCellValueGetter> GetterPtr = MakeShareable(new FMinExclusiveTimeValueGetter());
+	Column.SetValueGetter(GetterPtr.ToSharedRef());
+
+	TSharedPtr<FTableCellValueFormatter> FormatterPtr = MakeShareable(new FDoubleValueFormatterAsTimeMs());
+	Column.SetValueFormatter(FormatterPtr.ToSharedRef());
+
+	TSharedPtr<ITableCellValueSorter> SorterPtr = MakeShareable(new FSorterByDoubleValue(ColumnPtr.ToSharedRef()));
+	Column.SetValueSorter(SorterPtr);
+
+	return ColumnPtr.ToSharedRef();
 }
 
 ////////////////////////////////////////////////////////////////////////////////////////////////////
diff --git a/Engine/Source/Developer/TraceInsights/Private/Insights/ViewModels/TimersViewColumnFactory.h b/Engine/Source/Developer/TraceInsights/Private/Insights/ViewModels/TimersViewColumnFactory.h
--- a/Engine/Source/Developer/TraceInsights/Private/Insights/ViewModels/TimersViewColumnFactory.h
+++ b/Engine/Source/Developer/TraceInsights/Private/Insights/ViewModels/TimersViewColumnFactory.h
@@ -4,8 +4,6 @@
 
 #include "CoreMinimal.h"
 
-#include "Insights/ViewModels/TimersViewColumn.h"
-
 ////////////////////////////////////////////////////////////////////////////////////////////////////
 
 struct FTimersViewColumns
@@ -39,24 +37,39 @@ struct FTimersViewColumns
 
 ////////////////////////////////////////////////////////////////////////////////////////////////////
 
-struct FTimersViewColumnFactory
+namespace Insights
 {
-private:
-	/** Default constructor. */
-	FTimersViewColumnFactory();
+	class FTableColumn;
+}
 
-	/** Destructor. */
-	~FTimersViewColumnFactory();
+////////////////////////////////////////////////////////////////////////////////////////////////////
 
+struct FTimersViewColumnFactory
+{
 public:
-	/** Contains basic information about columns used in the Timers view widget. Names should be localized. */
-	TArray<FTimersViewColumn*> Collection;
+	static void CreateTimersViewColumns(TArray<TSharedPtr<Insights::FTableColumn>>& Columns);
+	static void CreateTimerTreeViewColumns(TArray<TSharedPtr<Insights::FTableColumn>>& Columns);
+
+	static TSharedRef<Insights::FTableColumn> CreateNameColumn();
+	static TSharedRef<Insights::FTableColumn> CreateMetaGroupNameColumn();
+	static TSharedRef<Insights::FTableColumn> CreateTypeColumn();
+	static TSharedRef<Insights::FTableColumn> CreateInstanceCountColumn();
 
-	/** Mapping between column IDs and FTimersViewColumn pointers. */
-	TMap<FName, const FTimersViewColumn*> ColumnIdToPtrMapping;
+	static TSharedRef<Insights::FTableColumn> CreateTotalInclusiveTimeColumn();
+	static TSharedRef<Insights::FTableColumn> CreateMaxInclusiveTimeColumn();
+	static TSharedRef<Insights::FTableColumn> CreateAverageInclusiveTimeColumn();
+	static TSharedRef<Insights::FTableColumn> CreateMedianInclusiveTimeColumn();
+	static TSharedRef<Insights::FTableColumn> CreateMinInclusiveTimeColumn();
 
-	/** Singleton */
-	static const FTimersViewColumnFactory& Get();
+	static TSharedRef<Insights::FTableColumn> CreateTotalExclusiveTimeColumn();
+	static TSharedRef<Insights::FTableColumn> CreateMaxExclusiveTimeColumn();
+	static TSharedRef<Insights::FTableColumn> CreateAverageExclusiveTimeColumn();
+	static TSharedRef<Insights::FTableColumn> CreateMedianExclusiveTimeColumn();
+	static TSharedRef<Insights::FTableColumn> CreateMinExclusiveTimeColumn();
+
+private:
+	static constexpr float TotalTimeColumnInitialWidth = 60.0f;
+	static constexpr float TimeMsColumnInitialWidth = 50.0f;
 };
 
 ////////////////////////////////////////////////////////////////////////////////////////////////////
diff --git a/Engine/Source/Developer/TraceInsights/Private/Insights/ViewModels/TimingTrackViewport.cpp b/Engine/Source/Developer/TraceInsights/Private/Insights/ViewModels/TimingTrackViewport.cpp
--- a/Engine/Source/Developer/TraceInsights/Private/Insights/ViewModels/TimingTrackViewport.cpp
+++ b/Engine/Source/Developer/TraceInsights/Private/Insights/ViewModels/TimingTrackViewport.cpp
@@ -2,10 +2,26 @@
 
 #include "TimingTrackViewport.h"
 
+#include "Widgets/Layout/SScrollBar.h"
+
 #include <limits>
 
 ////////////////////////////////////////////////////////////////////////////////////////////////////
 
+bool FTimingTrackViewport::UpdateSize(const float InWidth, const float InHeight)
+{
+	if (Width != InWidth || Height != InHeight)
+	{
+		Width = InWidth;
+		Height = InHeight;
+		EndTime = SlateUnitsToTime(Width);
+		return true;
+	}
+	return false;
+}
+
+////////////////////////////////////////////////////////////////////////////////////////////////////
+
 bool FTimingTrackViewport::ScrollAtTime(const double Time)
 {
 	const double NewStartTime = AlignTimeToPixel(Time);
@@ -55,6 +71,25 @@ bool FTimingTrackViewport::ZoomOnTimeInterval(const double Time, const double Du
 
 ////////////////////////////////////////////////////////////////////////////////////////////////////
 
+bool FTimingTrackViewport::RelativeZoomWithFixedX(const float Delta, const float X)
+{
+	constexpr double ZoomStep = 0.25; // as percent
+	double NewScaleX;
+
+	if (Delta > 0)
+	{
+		NewScaleX = ScaleX * FMath::Pow(1.0 + ZoomStep, Delta);
+	}
+	else
+	{
+		NewScaleX = ScaleX * FMath::Pow(1.0 / (1.0 + ZoomStep), -Delta);
+	}
+
+	return ZoomWithFixedX(NewScaleX, X);
+}
+
+////////////////////////////////////////////////////////////////////////////////////////////////////
+
 bool FTimingTrackViewport::ZoomWithFixedX(const double NewScaleX, const float X)
 {
 	const double LocalNewScaleX = FMath::Clamp(NewScaleX, MinScaleX, MaxScaleX);
@@ -129,18 +164,92 @@ bool FTimingTrackViewport::EnforceHorizontalScrollLimits(const double U)
 	{
 		NewStartTime = (1.0 - U) * NewStartTime + U * MinT;
 		if (FMath::IsNearlyEqual(NewStartTime, MinT, 1.0 / ScaleX))
+		{
 			NewStartTime = MinT;
+		}
 	}
 	else if (NewStartTime > MaxT)
 	{
 		NewStartTime = (1.0 - U) * NewStartTime + U * MaxT;
 		if (FMath::IsNearlyEqual(NewStartTime, MaxT, 1.0 / ScaleX))
+		{
 			NewStartTime = MaxT;
+		}
 		if (NewStartTime < MinT)
+		{
 			NewStartTime = MinT;
+		}
 	}
 
 	return ScrollAtTime(NewStartTime);
 }
 
 ////////////////////////////////////////////////////////////////////////////////////////////////////
+
+bool FTimingTrackViewport::OnUserScrolled(TSharedPtr<SScrollBar> ScrollBar, float ScrollOffset)
+{
+	const double S = 1.0 / (MaxValidTime - MinValidTime);
+	const double D = EndTime - StartTime;
+	const float ThumbSizeFraction = FMath::Clamp<float>(D * S, 0.0f, 1.0f);
+	const float OffsetFraction = FMath::Clamp<float>(ScrollOffset, 0.0f, 1.0f - ThumbSizeFraction);
+
+	ScrollBar->SetState(OffsetFraction, ThumbSizeFraction);
+
+	const double NewStartTime = MinValidTime + static_cast<double>(OffsetFraction) * (MaxValidTime - MinValidTime);
+	if (NewStartTime != StartTime)
+	{
+		ScrollAtTime(NewStartTime);
+		return true;
+	}
+
+	return false;
+}
+
+////////////////////////////////////////////////////////////////////////////////////////////////////
+
+void FTimingTrackViewport::UpdateScrollBar(TSharedPtr<SScrollBar> ScrollBar) const
+{
+	const double S = 1.0 / (MaxValidTime - MinValidTime);
+	const double D = EndTime - StartTime;
+	const float ThumbSizeFraction = FMath::Clamp<float>(D * S, 0.0f, 1.0f);
+	const float ScrollOffset = static_cast<float>((StartTime - MinValidTime) * S);
+	const float OffsetFraction = FMath::Clamp<float>(ScrollOffset, 0.0f, 1.0f - ThumbSizeFraction);
+
+	ScrollBar->SetState(OffsetFraction, ThumbSizeFraction);
+}
+
+////////////////////////////////////////////////////////////////////////////////////////////////////
+
+bool FTimingTrackViewport::OnUserScrolledY(TSharedPtr<SScrollBar> ScrollBar, float ScrollOffset)
+{
+	const float S = 1.0 / ScrollHeight;
+	const float D = Height - TopOffset;
+	const float ThumbSizeFraction = FMath::Clamp<float>(D * S, 0.0f, 1.0f);
+	const float OffsetFraction = FMath::Clamp<float>(ScrollOffset, 0.0f, 1.0f - ThumbSizeFraction);
+
+	ScrollBar->SetState(OffsetFraction, ThumbSizeFraction);
+
+	const float NewScrollPosY = OffsetFraction * ScrollHeight;
+	if (NewScrollPosY != ScrollPosY)
+	{
+		ScrollPosY = NewScrollPosY;
+		return true;
+	}
+
+	return false;
+}
+
+////////////////////////////////////////////////////////////////////////////////////////////////////
+
+void FTimingTrackViewport::UpdateScrollBarY(TSharedPtr<SScrollBar> ScrollBar) const
+{
+	const float S = 1.0 / ScrollHeight;
+	const float D = Height - TopOffset;
+	const float ThumbSizeFraction = FMath::Clamp<float>(D * S, 0.0f, 1.0f);
+	const float ScrollOffset = ScrollPosY * S;
+	const float OffsetFraction = FMath::Clamp<float>(ScrollOffset, 0.0f, 1.0f - ThumbSizeFraction);
+
+	ScrollBar->SetState(OffsetFraction, ThumbSizeFraction);
+}
+
+////////////////////////////////////////////////////////////////////////////////////////////////////
diff --git a/Engine/Source/Developer/TraceInsights/Private/Insights/ViewModels/TimingTrackViewport.h b/Engine/Source/Developer/TraceInsights/Private/Insights/ViewModels/TimingTrackViewport.h
--- a/Engine/Source/Developer/TraceInsights/Private/Insights/ViewModels/TimingTrackViewport.h
+++ b/Engine/Source/Developer/TraceInsights/Private/Insights/ViewModels/TimingTrackViewport.h
@@ -4,24 +4,12 @@
 
 #include "CoreMinimal.h"
 
-class FTimingTrackViewport
-{
-public:
-	float Width; // width of viewport, in Slate units
-	float Height; // height of viewport, in Slate units
+class SScrollBar;
 
-	double MinValidTime; // min session time, in seconds
-	double MaxValidTime; // max session time, in seconds
-	double StartTime; // time of viewport's left side, in seconds
-	double EndTime; // time of viewport's right side, in seconds; computed when StartTime or Width changes
-	double MinScaleX;
-	double MaxScaleX;
-	double ScaleX; // scale factor between seconds and pixels (Slate units), in pixels (Slate units) per second
-
-	float TopOffset; // top offset (to allow the time ruller to be visible), in pixels (Slate units)
-	float ScrollHeight; // height of the vertical scrollable area, in pixels (Slate units)
-	float ScrollPosY; // current vertical scroll position, in pixels (Slate units)
+////////////////////////////////////////////////////////////////////////////////////////////////////
 
+class FTimingTrackViewport
+{
 public:
 	FTimingTrackViewport()
 	{
@@ -35,8 +23,10 @@ public:
 
 		MinValidTime = 0.0;
 		MaxValidTime = 0.0;
+
 		StartTime = 0.0;
 		EndTime = 0.0;
+
 		MinScaleX = (5 * 20) / 3600.0; // 1h between major tick marks
 		MaxScaleX = 1.0E10; // 10ns between major tick marks
 		ScaleX = (5 * 20) / 5.0; // 5s between major tick marks
@@ -46,6 +36,40 @@ public:
 		ScrollPosY = 0.0f;
 	}
 
+	/** Width of viewport, in pixels. [px] */
+	float GetWidth() const { return Width; }
+
+	/** Height of viewport, in pixels. [px] */
+	float GetHeight() const { return Height; }
+
+	bool UpdateSize(const float InWidth, const float InHeight);
+
+	/** Minimum session time, in seconds. [s] */
+	double GetMinValidTime() const { return MinValidTime; }
+	/** Maximum session time, in seconds. [s] */
+	double GetMaxValidTime() const { return MaxValidTime; }
+	void SetMaxValidTime(const double InMaxValidTime) { MaxValidTime = InMaxValidTime; }
+
+	/**
+	  * Time of viewport's left side, in seconds. [s]
+	  */
+	double GetStartTime() const { return StartTime; }
+	/**
+	  * Time of viewport's right side, in seconds. [s]
+	  * Computed when StartTime, ScaleX or Width changes.
+	  */
+	double GetEndTime() const { return EndTime; }
+	double GetDuration() const { return EndTime - StartTime; }
+
+	double GetMinScaleX() const { return MinScaleX; }
+	double GetMaxScaleX() const { return MaxScaleX; }
+
+	/** Current scale factor between seconds and pixels (Slate units). [px/s] */
+	double GetScaleX() const { return ScaleX; }
+
+	//////////////////////////////////////////////////
+	// Conversions
+
 	float TimeToSlateUnits(const double Time) const
 	{
 		return static_cast<float>((Time - StartTime) * ScaleX);
@@ -71,18 +95,6 @@ public:
 		return StartTime + static_cast<double>(X) / ScaleX;
 	}
 
-	bool UpdateSize(const float InWidth, const float InHeight)
-	{
-		if (Width != InWidth || Height != InHeight)
-		{
-			Width = InWidth;
-			Height = InHeight;
-			EndTime = SlateUnitsToTime(Width);
-			return true;
-		}
-		return false;
-	}
-
 	double AlignTimeToPixel(const double InTime, const double InScaleX) const
 	{
 		return FMath::RoundToDouble(InTime * InScaleX) / InScaleX;
@@ -93,17 +105,70 @@ public:
 		return AlignTimeToPixel(Time, ScaleX);
 	}
 
+	double RestrictEndTime(const double InEndTime) const;
+	double RestrictDuration(const double InStartTime, const double InEndTime) const;
+
+	//////////////////////////////////////////////////
+	// Scrolling
+
 	bool ScrollAtTime(const double Time);
 	bool CenterOnTimeInterval(const double Time, const double Duration);
 
+	void GetHorizontalScrollLimits(double& OutMinT, double& OutMaxT);
+	bool EnforceHorizontalScrollLimits(const double U);
+
+	//////////////////////////////////////////////////
+	// Zooming
+
 	bool ZoomOnTimeInterval(const double Time, const double Duration);
+	bool RelativeZoomWithFixedX(const float Delta, const float X);
 	bool ZoomWithFixedX(const double NewScaleX, const float X);
 
-	double RestrictEndTime(const double InEndTime) const;
-	double RestrictDuration(const double InStartTime, const double InEndTime) const;
+	//////////////////////////////////////////////////
+	// Vertical axis
 
-	void GetHorizontalScrollLimits(double& OutMinT, double& OutMaxT);
-	bool EnforceHorizontalScrollLimits(const double U);
+	/** Top offset (to allow the time ruller to be visible), in pixels. [px] */
+	float GetTopOffset() const { return TopOffset; }
+
+	void SetTopOffset(float InTopOffset) { TopOffset = InTopOffset; }
+
+	/** Height of the vertical scrollable area, in pixels. [px] */
+	float GetScrollHeight() const { return ScrollHeight; }
+	void SetScrollHeight(const float InScrollHeight) { ScrollHeight = InScrollHeight; }
+
+	/** Current vertical scroll position, in pixels. [px] */
+	float GetScrollPosY() const { return ScrollPosY; }
+	void SetScrollPosY(const float InScrollPosY) { ScrollPosY = InScrollPosY; }
 
 	float GetViewportY(const float Y) const { return TopOffset + Y - ScrollPosY; }
+
+	//////////////////////////////////////////////////
+
+	bool OnUserScrolled(TSharedPtr<SScrollBar> ScrollBar, float ScrollOffset);
+	void UpdateScrollBar(TSharedPtr<SScrollBar> ScrollBar) const;
+
+	bool OnUserScrolledY(TSharedPtr<SScrollBar> ScrollBar, float ScrollOffset);
+	void UpdateScrollBarY(TSharedPtr<SScrollBar> ScrollBar) const;
+
+	//////////////////////////////////////////////////
+
+private:
+	float Width; // width of viewport, in pixels; [px]
+	float Height; // height of viewport, in pixels; [px]
+
+	double MinValidTime; // min session time, in seconds; [s]
+	double MaxValidTime; // max session time, in seconds; [s]
+
+	double StartTime; // time of viewport's left side, in seconds; [s]
+	double EndTime; // time of viewport's right side, in seconds; [s]; computed when StartTime, Scale or Width changes
+
+	double MinScaleX; // min scale factor; [px/s]
+	double MaxScaleX; // max scale factor; [px/s]
+	double ScaleX; // scale factor between seconds and pixels; [px/s]
+
+	float TopOffset; // top offset (to allow the time ruller to be visible), in pixels; [px]
+	float ScrollHeight; // height of the vertical scrollable area, in pixels; [px]
+	float ScrollPosY; // current vertical scroll position, in pixels; [px]
 };
+
+////////////////////////////////////////////////////////////////////////////////////////////////////
diff --git a/Engine/Source/Developer/TraceInsights/Private/Insights/ViewModels/TimingViewDrawHelper.cpp b/Engine/Source/Developer/TraceInsights/Private/Insights/ViewModels/TimingViewDrawHelper.cpp
--- a/Engine/Source/Developer/TraceInsights/Private/Insights/ViewModels/TimingViewDrawHelper.cpp
+++ b/Engine/Source/Developer/TraceInsights/Private/Insights/ViewModels/TimingViewDrawHelper.cpp
@@ -2,9 +2,6 @@
 
 #include "TimingViewDrawHelper.h"
 
-#include "Brushes/SlateBorderBrush.h"
-#include "Brushes/SlateColorBrush.h"
-#include "EditorStyleSet.h"
 #include "Fonts/FontMeasure.h"
 #include "Fonts/SlateFontInfo.h"
 #include "Framework/Application/SlateApplication.h"
@@ -15,6 +12,8 @@
 // Insights
 #include "Insights/Common/PaintUtils.h"
 #include "Insights/Common/TimeUtils.h"
+#include "Insights/InsightsStyle.h"
+#include "Insights/ViewModels/DrawHelpers.h"
 #include "Insights/ViewModels/MarkersTimingTrack.h"
 #include "Insights/ViewModels/TimingEventsTrack.h"
 #include "Insights/ViewModels/TimingTrackViewport.h"
@@ -25,13 +24,23 @@ FTimingViewDrawHelper::FTimingViewDrawHelper(const FDrawContext& InDC, const FTi
 	: DrawContext(InDC)
 	, Viewport(InViewport)
 	, Layout(InLayout)
-	, WhiteBrush(FCoreStyle::Get().GetBrush("WhiteBrush"))
-	, BorderBrush(FEditorStyle::GetBrush("PlainBorder"))
-	, EventsBorderBrush(new FSlateBorderBrush(NAME_None, FMargin(1.0f)))
+	, WhiteBrush(FInsightsStyle::Get().GetBrush("WhiteBrush"))
+	, EventBorderBrush(FInsightsStyle::Get().GetBrush("EventBorder"))
+	, HoveredEventBorderBrush(FInsightsStyle::Get().GetBrush("HoveredEventBorder"))
+	, SelectedEventBorderBrush(FInsightsStyle::Get().GetBrush("SelectedEventBorder"))
 	, BackgroundAreaBrush(WhiteBrush)
 	, ValidAreaColor(0.07f, 0.07f, 0.07f, 1.0f)
 	, InvalidAreaColor(0.1f, 0.07f, 0.07f, 1.0f)
+	, EdgeColor(0.05f, 0.05f, 0.05f, 1.0f)
 	, EventFont(FCoreStyle::GetDefaultFontStyle("Regular", 8))
+	, ValidAreaX(0.0f)
+	, ValidAreaW(0.0f)
+	, TimelineTopY(0.0f)
+	, TimelineY(0.0f)
+	, MaxDepth(0)
+	, TimelineIndex(0)
+	, LastEventX2()
+	, LastBox()
 	, Stats()
 {
 }
@@ -40,51 +49,15 @@ FTimingViewDrawHelper::FTimingViewDrawHelper(const FDrawContext& InDC, const FTi
 
 FTimingViewDrawHelper::~FTimingViewDrawHelper()
 {
-	delete EventsBorderBrush; //im: is it safe to delete brush (i.e in OnPaint, imediately after use)?
 }
 
 ////////////////////////////////////////////////////////////////////////////////////////////////////
 
 void FTimingViewDrawHelper::DrawBackground() const
 {
-	const float X0 = Viewport.TimeToSlateUnitsRounded(0.0);
-	const float X1 = Viewport.TimeToSlateUnitsRounded(Viewport.MaxValidTime);
-	const float W = FMath::CeilToFloat(Viewport.Width);
-	const float H = FMath::CeilToFloat(Viewport.Height);
-
-	if (X0 >= W || X1 <= 0.0f)
-	{
-		ValidX0 = 0.0f;
-		ValidX1 = W;
-
-		// Draw invalid area (entire view).
-		DrawContext.DrawBox(0.0f, 0.0f, W, H, BackgroundAreaBrush, InvalidAreaColor);
-	}
-	else // X0 < W && X1 > 0
-	{
-		if (X0 > 0.0f)
-		{
-			// Draw invalid area (left).
-			DrawContext.DrawBox(0.0f, 0.0f, X0, H, BackgroundAreaBrush, InvalidAreaColor);
-		}
-
-		if (X1 < W)
-		{
-			// Draw invalid area (right).
-			DrawContext.DrawBox(X1, 0.0f, W - X1, H, BackgroundAreaBrush, InvalidAreaColor);
-		}
-
-		ValidX0 = FMath::Max(X0, 0.0f);
-		ValidX1 = FMath::Min(X1, W);
-
-		if (ValidX1 > ValidX0)
-		{
-			// Draw valid area.
-			DrawContext.DrawBox(ValidX0, 0.0f, ValidX1 - ValidX0, H, BackgroundAreaBrush, ValidAreaColor);
-		}
-	}
-
-	DrawContext.LayerId++;
+	const float Y = 0.0f;
+	const float H = FMath::CeilToFloat(Viewport.GetHeight());
+	FDrawHelpers::DrawBackground(DrawContext, BackgroundAreaBrush, Viewport, Y, H, ValidAreaX, ValidAreaW);
 }
 
 ////////////////////////////////////////////////////////////////////////////////////////////////////
@@ -92,16 +65,16 @@ void FTimingViewDrawHelper::DrawBackground() const
 void FTimingViewDrawHelper::BeginTimelines()
 {
 	TimelineIndex = -1;
-	TimelineTopY = -Viewport.ScrollPosY;
+	TimelineTopY = -Viewport.GetScrollPosY();
 
-	float Y = TimelineTopY + Viewport.TopOffset;
+	float Y = TimelineTopY + Viewport.GetTopOffset();
 
-	if (Y > 0.0f && ValidX1 > ValidX0)
+	if (Y > 0.0f && ValidAreaW > 0.0f)
 	{
-		Y = FMath::Min(Y, Viewport.Height);
+		Y = FMath::Min(Y, Viewport.GetHeight());
 
 		// Draw invalid area (top).
-		DrawContext.DrawBox(DrawContext.LayerId + ToInt32(EDrawLayer::TimelineHeader), ValidX0, 0.0f, ValidX1 - ValidX0, Y, BackgroundAreaBrush, InvalidAreaColor);
+		DrawContext.DrawBox(DrawContext.LayerId + ToInt32(EDrawLayer::TimelineHeader), ValidAreaX, 0.0f, ValidAreaW, Y, BackgroundAreaBrush, InvalidAreaColor);
 	}
 }
 
@@ -111,10 +84,12 @@ bool FTimingViewDrawHelper::BeginTimeline(FTimingEventsTrack& Track)
 {
 	TimelineIndex++;
 
-	Track.SetPosY(TimelineTopY + Viewport.ScrollPosY);
+	Track.SetPosY(TimelineTopY + Viewport.GetScrollPosY());
 
 	if (Track.GetHeight() < Layout.MinTimelineH)
+	{
 		Track.SetHeight(Layout.MinTimelineH);
+	}
 
 	if (TimelineTopY + Track.GetHeight() < -1.0f)
 	{
@@ -122,7 +97,7 @@ bool FTimingViewDrawHelper::BeginTimeline(FTimingEventsTrack& Track)
 		return false;
 	}
 
-	if (TimelineTopY > Viewport.Height - Viewport.TopOffset)
+	if (TimelineTopY > Viewport.GetHeight() - Viewport.GetTopOffset())
 	{
 		return false;
 	}
@@ -130,7 +105,7 @@ bool FTimingViewDrawHelper::BeginTimeline(FTimingEventsTrack& Track)
 	MaxDepth = -1;
 
 	// +1.0f is for horizontal line between timelines
-	TimelineY = Viewport.TopOffset + TimelineTopY + 1.0f + Layout.TimelineDY;
+	TimelineY = Viewport.GetTopOffset() + TimelineTopY + 1.0f + Layout.TimelineDY;
 
 	// Reset "last event X2" for all depths.
 	LastEventX2.Reset();
@@ -156,7 +131,7 @@ void FTimingViewDrawHelper::EndTimeline(FTimingEventsTrack& Track)
 		}
 	}
 
-	Track.Depth = MaxDepth;
+	Track.SetDepth(MaxDepth);
 
 	float NewH;
 	if (MaxDepth < 0)
@@ -203,9 +178,8 @@ void FTimingViewDrawHelper::EndTimeline(FTimingEventsTrack& Track)
 
 	if (Track.GetHeight() > 0)
 	{
-		const float Y = Viewport.TopOffset + TimelineTopY;
+		const float Y = Viewport.GetTopOffset() + TimelineTopY;
 
-		FLinearColor Color(0.05f, 0.05f, 0.05f, 1.0f);
 		FLinearColor TextColor(1.0f, 1.0f, 1.0f, 1.0f);
 
 		if (Track.IsSelected())
@@ -225,13 +199,13 @@ void FTimingViewDrawHelper::EndTimeline(FTimingEventsTrack& Track)
 		}
 
 		// Draw a horizontal line between timelines.
-		DrawContext.DrawBox(DrawContext.LayerId + ToInt32(EDrawLayer::TimelineHeader), 0, Y, Viewport.Width, 1.0f, WhiteBrush, Color);
+		DrawContext.DrawBox(DrawContext.LayerId + ToInt32(EDrawLayer::TimelineHeader), 0, Y, Viewport.GetWidth(), 1.0f, WhiteBrush, EdgeColor);
 
 		// Draw name of timeline.
 		//const FString Name = FString::Printf(TEXT("%d Y: %g, H: %g, VO: %g, VY: %g"), TimelineIndex, Track.Y, Track.H, Viewport.TopOffset, Viewport.ScrollPosY); //debug
 		const TSharedRef<FSlateFontMeasure> FontMeasureService = FSlateApplication::Get().GetRenderer()->GetFontMeasureService();
 		float NameWidth = FontMeasureService->Measure(Track.GetName(), EventFont).X;
-		DrawContext.DrawBox(DrawContext.LayerId + ToInt32(EDrawLayer::TimelineHeader), 0.0f, Y + 1.0f, NameWidth + 4.0f, 12.0f, WhiteBrush, Color);
+		DrawContext.DrawBox(DrawContext.LayerId + ToInt32(EDrawLayer::TimelineHeader), 0.0f, Y + 1.0f, NameWidth + 4.0f, 12.0f, WhiteBrush, EdgeColor);
 		DrawContext.DrawText(DrawContext.LayerId + ToInt32(EDrawLayer::TimelineText), 2.0f, Y, Track.GetName(), EventFont, TextColor);
 	}
 
@@ -245,12 +219,12 @@ void FTimingViewDrawHelper::AddEvent(double EventStartTime, double EventEndTime,
 	Stats.NumEvents++;
 
 	float EventX1 = Viewport.TimeToSlateUnitsRounded(EventStartTime);
-	if (EventX1 > Viewport.Width)
+	if (EventX1 > Viewport.GetWidth())
 	{
 		return;
 	}
 
-	double RestrictedEndTime = Viewport.RestrictEndTime(EventEndTime);
+	const double RestrictedEndTime = Viewport.RestrictEndTime(EventEndTime);
 	float EventX2 = Viewport.TimeToSlateUnitsRounded(RestrictedEndTime);
 	if (EventX2 < 0)
 	{
@@ -276,7 +250,7 @@ void FTimingViewDrawHelper::AddEvent(double EventStartTime, double EventEndTime,
 	}
 
 	float EventY = TimelineY + (Layout.EventH + Layout.EventDY) * Depth;
-	if (EventY < -Layout.EventH || EventY > Viewport.Height)
+	if (EventY < -Layout.EventH || EventY > Viewport.GetHeight())
 	{
 		return;
 	}
@@ -296,13 +270,13 @@ void FTimingViewDrawHelper::AddEvent(double EventStartTime, double EventEndTime,
 	}
 
 	// Limit event width on the viewport's right side.
-	const float MaxX = Viewport.Width + 2.0f; // +2 allows event border to remain outside screen
+	const float MaxX = Viewport.GetWidth() + 2.0f; // +2 allows event border to remain outside screen
 	if (EventX2 > MaxX)
 	{
 		EventX2 = MaxX;
 	}
 
-	float EventW = EventX2 - EventX1;
+	const float EventW = EventX2 - EventX1;
 
 	// Optimization...
 	if (EventW == 1.0f && EventX1 == LastEventX2[Depth] - 1.0f)
@@ -415,7 +389,7 @@ void FTimingViewDrawHelper::AddEvent(double EventStartTime, double EventEndTime,
 		}
 
 		const FLinearColor EventColorBorder(EventColorFill.R * BorderColorFactor, EventColorFill.G * BorderColorFactor, EventColorFill.B * BorderColorFactor, 1.0f);
-		DrawContext.DrawBox(DrawContext.LayerId + ToInt32(EDrawLayer::EventBorder), EventX1, EventY, EventW, Layout.EventH, EventsBorderBrush, EventColorBorder);
+		DrawContext.DrawBox(DrawContext.LayerId + ToInt32(EDrawLayer::EventBorder), EventX1, EventY, EventW, Layout.EventH, EventBorderBrush, EventColorBorder);
 		Stats.NumDrawBorders++;
 	}
 	else // 1px or 2px boxes
@@ -487,18 +461,18 @@ void FTimingViewDrawHelper::DrawBox(const FBoxData& Box, const float EventY, con
 
 void FTimingViewDrawHelper::EndTimelines()
 {
-	float Y = Viewport.TopOffset + TimelineTopY;
+	float Y = Viewport.GetTopOffset() + TimelineTopY;
 
 	// Draw a last horizontal line.
-	DrawContext.DrawBox(DrawContext.LayerId + ToInt32(EDrawLayer::TimelineHeader), 0, Y, Viewport.Width, 1.0f, WhiteBrush, FLinearColor(0.05f, 0.05f, 0.05f, 1.0f));
+	DrawContext.DrawBox(DrawContext.LayerId + ToInt32(EDrawLayer::TimelineHeader), 0, Y, Viewport.GetWidth(), 1.0f, WhiteBrush, EdgeColor);
 	Y += 1.0f;
 
-	if (Y < Viewport.Height && ValidX1 > ValidX0)
+	if (Y < Viewport.GetHeight() && ValidAreaW > 0.0f)
 	{
 		Y = FMath::Max(Y, 0.0f);
 
 		// Draw invalid area (bottom).
-		DrawContext.DrawBox(DrawContext.LayerId + ToInt32(EDrawLayer::TimelineHeader), ValidX0, Y, ValidX1 - ValidX0, Viewport.Height - Y, BackgroundAreaBrush, InvalidAreaColor);
+		DrawContext.DrawBox(DrawContext.LayerId + ToInt32(EDrawLayer::TimelineHeader), ValidAreaX, Y, ValidAreaW, Viewport.GetHeight() - Y, BackgroundAreaBrush, InvalidAreaColor);
 	}
 
 	DrawContext.LayerId += ToInt32(EDrawLayer::Count);
@@ -509,7 +483,7 @@ void FTimingViewDrawHelper::EndTimelines()
 void FTimingViewDrawHelper::DrawTimingEventHighlight(double StartTime, double EndTime, float Y, EHighlightMode Mode)
 {
 	float EventX1 = Viewport.TimeToSlateUnitsRounded(StartTime);
-	if (EventX1 > Viewport.Width)
+	if (EventX1 > Viewport.GetWidth())
 	{
 		return;
 	}
@@ -521,6 +495,7 @@ void FTimingViewDrawHelper::DrawTimingEventHighlight(double StartTime, double En
 		return;
 	}
 
+	// Timing events are displayed with minimum 1px (including empty ones).
 	if (EventX1 == EventX2)
 	{
 		EventX2 = EventX1 + 1.0f;
@@ -534,27 +509,23 @@ void FTimingViewDrawHelper::DrawTimingEventHighlight(double StartTime, double En
 	}
 
 	// Limit event width on the viewport's right side.
-	const float MaxX = Viewport.Width + 2.0f; // +2 allows event border to remain outside screen
+	const float MaxX = Viewport.GetWidth() + 2.0f; // +2 allows event border to remain outside screen
 	if (EventX2 > MaxX)
 	{
 		EventX2 = MaxX;
 	}
 
-	float EventW = EventX2 - EventX1;
+	const float EventW = EventX2 - EventX1;
 
 	if (Mode == EHighlightMode::Hovered)
 	{
-		const FSlateBrush* HighlightBrush = BorderBrush;
-
 		const FLinearColor Color(1.0f, 1.0f, 0.0f, 1.0f); // yellow
 
 		// Draw border around the timing event box.
-		DrawContext.DrawBox(EventX1 - 2.0f, Y - 2.0f, EventW + 4.0f, Layout.EventH + 4.0f, HighlightBrush, Color);
+		DrawContext.DrawBox(EventX1 - 2.0f, Y - 2.0f, EventW + 4.0f, Layout.EventH + 4.0f, HoveredEventBorderBrush, Color);
 	}
 	else // EHighlightMode::Selected or EHighlightMode::SelectedAndHovered
 	{
-		const FSlateBrush* SelectedBrush = BorderBrush;
-
 		// Animate color from white (if selected and hovered) or yellow (if only selected) to black, using a squared sine function.
 		const double Time = static_cast<double>(FPlatformTime::Cycles64()) * FPlatformTime::GetSecondsPerCycle64();
 		float S = FMath::Sin(2.0 * Time);
@@ -563,7 +534,7 @@ void FTimingViewDrawHelper::DrawTimingEventHighlight(double StartTime, double En
 		const FLinearColor Color(S, S, Blue, 1.0f);
 
 		// Draw border around the timing event box.
-		DrawContext.DrawBox(EventX1 - 2.0f, Y - 2.0f, EventW + 4.0f, Layout.EventH + 4.0f, SelectedBrush, Color);
+		DrawContext.DrawBox(EventX1 - 2.0f, Y - 2.0f, EventW + 4.0f, Layout.EventH + 4.0f, SelectedEventBorderBrush, Color);
 	}
 	DrawContext.LayerId++;
 }
diff --git a/Engine/Source/Developer/TraceInsights/Private/Insights/ViewModels/TimingViewDrawHelper.h b/Engine/Source/Developer/TraceInsights/Private/Insights/ViewModels/TimingViewDrawHelper.h
--- a/Engine/Source/Developer/TraceInsights/Private/Insights/ViewModels/TimingViewDrawHelper.h
+++ b/Engine/Source/Developer/TraceInsights/Private/Insights/ViewModels/TimingViewDrawHelper.h
@@ -20,11 +20,11 @@ class FTimingTrackViewport;
 class FTimingViewDrawHelper
 {
 public:
-	enum class EHighlightMode
+	enum class EHighlightMode : uint32
 	{
-		Hovered,
-		Selected,
-		SelectedAndHovered
+		Hovered = 1,
+		Selected = 2,
+		SelectedAndHovered = 3
 	};
 
 private:
@@ -74,8 +74,6 @@ private:
 		int32 NumMergedBoxes;
 		int32 NumDrawBorders;
 		int32 NumDrawTexts;
-		int32 NumDrawTimeMarkerBoxes;
-		int32 NumDrawTimeMarkerTexts;
 
 		FStats()
 			: NumEvents(0)
@@ -83,13 +81,11 @@ private:
 			, NumMergedBoxes(0)
 			, NumDrawBorders(0)
 			, NumDrawTexts(0)
-			, NumDrawTimeMarkerBoxes(0)
-			, NumDrawTimeMarkerTexts(0)
 		{}
 	};
 
 public:
-	FTimingViewDrawHelper(const FDrawContext& InDrawContext, const FTimingTrackViewport& InViewport, const FTimingEventsTrackLayout& InLayout);
+	explicit FTimingViewDrawHelper(const FDrawContext& InDrawContext, const FTimingTrackViewport& InViewport, const FTimingEventsTrackLayout& InLayout);
 	~FTimingViewDrawHelper();
 
 	/**
@@ -102,13 +98,14 @@ public:
 	const FTimingTrackViewport& GetViewport() const { return Viewport; }
 	const FTimingEventsTrackLayout& GetLayout() const { return Layout; }
 
-	int32 GetNumEvents() const              { return Stats.NumEvents; }
-	int32 GetNumDrawBoxes() const           { return Stats.NumDrawBoxes; }
-	int32 GetNumMergedBoxes() const         { return Stats.NumMergedBoxes; }
-	int32 GetNumDrawBorders() const         { return Stats.NumDrawBorders; }
-	int32 GetNumDrawTexts() const           { return Stats.NumDrawTexts; }
-	int32 GetNumDrawTimeMarkerBoxes() const { return Stats.NumDrawTimeMarkerBoxes; }
-	int32 GetNumDrawTimeMarkerTexts() const { return Stats.NumDrawTimeMarkerTexts; }
+	const FSlateBrush* GetWhiteBrush() const { return WhiteBrush; }
+	const FSlateFontInfo& GetEventFont() const { return EventFont; }
+
+	int32 GetNumEvents() const      { return Stats.NumEvents; }
+	int32 GetNumDrawBoxes() const   { return Stats.NumDrawBoxes; }
+	int32 GetNumMergedBoxes() const { return Stats.NumMergedBoxes; }
+	int32 GetNumDrawBorders() const { return Stats.NumDrawBorders; }
+	int32 GetNumDrawTexts() const   { return Stats.NumDrawTexts; }
 
 	void DrawBackground() const;
 	void DrawTimingEventHighlight(double StartTime, double EndTime, float Y, EHighlightMode Mode);
@@ -129,18 +126,20 @@ private:
 	const FTimingEventsTrackLayout& Layout;
 
 	const FSlateBrush* WhiteBrush;
-	const FSlateBrush* BorderBrush;
-	const FSlateBrush* EventsBorderBrush;
+	const FSlateBrush* EventBorderBrush;
+	const FSlateBrush* HoveredEventBorderBrush;
+	const FSlateBrush* SelectedEventBorderBrush;
 	const FSlateBrush* BackgroundAreaBrush;
 	const FLinearColor ValidAreaColor;
 	const FLinearColor InvalidAreaColor;
+	const FLinearColor EdgeColor;
 	const FSlateFontInfo EventFont;
 
-	//static const FSlateColorBrush SolidWhiteBrush = FSlateColorBrush(FColorList::White);
-	//static const FSlateBrush BorderBrush = FSlateBorderBrush(NAME_None, FMargin(1.0f));
+	mutable float ValidAreaX;
+	mutable float ValidAreaW;
 
-	mutable float ValidX0;
-	mutable float ValidX1;
+	//////////////////////////////////////////////////
+	// Builder state
 
 	float TimelineTopY;
 	float TimelineY;
@@ -150,6 +149,8 @@ private:
 	TArray<float> LastEventX2; // X2 value for last event on each depth, for current timeline
 	TArray<FBoxData> LastBox;
 
+	//////////////////////////////////////////////////
+
 	/** Debug stats */
 	mutable FStats Stats;
 };
diff --git a/Engine/Source/Developer/TraceInsights/Private/Insights/ViewModels/TooltipDrawState.cpp b/Engine/Source/Developer/TraceInsights/Private/Insights/ViewModels/TooltipDrawState.cpp
new file mode 100644
--- /dev/null
+++ b/Engine/Source/Developer/TraceInsights/Private/Insights/ViewModels/TooltipDrawState.cpp
@@ -0,0 +1,259 @@
+// Copyright 1998-2019 Epic Games, Inc. All Rights Reserved.
+
+#include "TooltipDrawState.h"
+
+#include "Fonts/FontMeasure.h"
+#include "Fonts/SlateFontInfo.h"
+#include "Styling/CoreStyle.h"
+
+// Insights
+#include "Insights/Common/PaintUtils.h"
+#include "Insights/InsightsStyle.h"
+
+////////////////////////////////////////////////////////////////////////////////////////////////////
+
+FLinearColor FTooltipDrawState::DefaultTitleColor(0.9f, 0.9f, 0.5f, 1.0f);
+FLinearColor FTooltipDrawState::DefaultNameColor(0.6f, 0.6f, 0.6f, 1.0f);
+FLinearColor FTooltipDrawState::DefaultValueColor(1.0f, 1.0f, 1.0f, 1.0f);
+
+////////////////////////////////////////////////////////////////////////////////////////////////////
+
+FTooltipDrawState::FTooltipDrawState()
+	: WhiteBrush(FInsightsStyle::Get().GetBrush("WhiteBrush"))
+	, Font(FCoreStyle::GetDefaultFontStyle("Regular", 8))
+	, BackgroundColor(0.05f, 0.05f, 0.05f, 1.0f)
+	, Size(0.0f, 0.0f)
+	, DesiredSize(0.0f, 0.0f)
+	, Position(0.0f, 0.0f)
+	, ValueOffsetX(0.0f)
+	, NewLineY(0.0f)
+	, Opacity(0.0f)
+	, DesiredOpacity(0.0f)
+	, Texts()
+{
+}
+
+////////////////////////////////////////////////////////////////////////////////////////////////////
+
+FTooltipDrawState::~FTooltipDrawState()
+{
+}
+
+////////////////////////////////////////////////////////////////////////////////////////////////////
+
+void FTooltipDrawState::Reset()
+{
+	Opacity = 0.0f;
+	DesiredOpacity = 0.0f;
+
+	ResetContent();
+}
+
+////////////////////////////////////////////////////////////////////////////////////////////////////
+
+void FTooltipDrawState::ResetContent()
+{
+	Texts.Reset();
+
+	ValueOffsetX = 0.0f;
+	NewLineY = BorderY;
+}
+
+////////////////////////////////////////////////////////////////////////////////////////////////////
+
+void FTooltipDrawState::AddTitle(const FString& Title)
+{
+	const TSharedRef<FSlateFontMeasure> FontMeasureService = FSlateApplication::Get().GetRenderer()->GetFontMeasureService();
+
+	const FVector2D TextSize = FontMeasureService->Measure(Title, Font);
+	Texts.Add({ BorderX, NewLineY, TextSize, Title, DefaultTitleColor, FDrawTextType::Title });
+
+	NewLineY += DefaultTitleHeight;
+}
+
+////////////////////////////////////////////////////////////////////////////////////////////////////
+
+void FTooltipDrawState::AddTitle(const FString& Title, const FLinearColor& Color)
+{
+	const TSharedRef<FSlateFontMeasure> FontMeasureService = FSlateApplication::Get().GetRenderer()->GetFontMeasureService();
+
+	const FVector2D TextSize = FontMeasureService->Measure(Title, Font);
+	Texts.Add({ BorderX, NewLineY, TextSize, Title, Color, FDrawTextType::Misc });
+
+	NewLineY += DefaultTitleHeight;
+}
+
+////////////////////////////////////////////////////////////////////////////////////////////////////
+
+void FTooltipDrawState::AddNameValueTextLine(const FString& Name, const FString& Value)
+{
+	const TSharedRef<FSlateFontMeasure> FontMeasureService = FSlateApplication::Get().GetRenderer()->GetFontMeasureService();
+
+	const FVector2D NameTextSize = FontMeasureService->Measure(Name, Font);
+	Texts.Add({ 0.0f, NewLineY, NameTextSize, Name, DefaultNameColor, FDrawTextType::Name });
+
+	const FVector2D ValueTextSize = FontMeasureService->Measure(Value, Font);
+	Texts.Add({ 0.0f, NewLineY, ValueTextSize, Value, DefaultValueColor, FDrawTextType::Value });
+
+	NewLineY += DefaultLineHeight;
+}
+
+////////////////////////////////////////////////////////////////////////////////////////////////////
+
+void FTooltipDrawState::AddTextLine(const FString& Text, const FLinearColor& Color)
+{
+	const TSharedRef<FSlateFontMeasure> FontMeasureService = FSlateApplication::Get().GetRenderer()->GetFontMeasureService();
+
+	const FVector2D TextSize = FontMeasureService->Measure(Text, Font);
+	Texts.Add({ BorderX, NewLineY, TextSize, Text, Color, FDrawTextType::Misc });
+
+	NewLineY += DefaultLineHeight;
+}
+
+////////////////////////////////////////////////////////////////////////////////////////////////////
+
+void FTooltipDrawState::AddTextLine(const float X, const float Y, const FString& Text, const FLinearColor& Color)
+{
+	const TSharedRef<FSlateFontMeasure> FontMeasureService = FSlateApplication::Get().GetRenderer()->GetFontMeasureService();
+
+	const FVector2D TextSize = FontMeasureService->Measure(Text, Font);
+	Texts.Add({ X, Y, TextSize, Text, Color, FDrawTextType::Misc });
+}
+
+////////////////////////////////////////////////////////////////////////////////////////////////////
+
+void FTooltipDrawState::UpdateLayout()
+{
+	ValueOffsetX = 0.0f;
+	for (const FDrawTextInfo& TextInfo : Texts)
+	{
+		if (TextInfo.Type == FDrawTextType::Name)
+		{
+			if (ValueOffsetX < TextInfo.TextSize.X)
+			{
+				ValueOffsetX = TextInfo.TextSize.X;
+			}
+		}
+	}
+	ValueOffsetX += BorderX;
+
+	DesiredSize.Set(BorderX, NewLineY + BorderY);
+	for (const FDrawTextInfo& TextInfo : Texts)
+	{
+		float RightX;
+		switch (TextInfo.Type)
+		{
+		case FDrawTextType::Name:
+			RightX = ValueOffsetX;
+			break;
+
+		case FDrawTextType::Value:
+			RightX = ValueOffsetX + NameValueDX + TextInfo.TextSize.X;
+			break;
+
+		default:
+			RightX = TextInfo.X + TextInfo.TextSize.X;
+		}
+		if (DesiredSize.X < RightX)
+		{
+			DesiredSize.X = RightX;
+		}
+	}
+	DesiredSize.X += BorderX;
+}
+
+////////////////////////////////////////////////////////////////////////////////////////////////////
+
+void FTooltipDrawState::Update()
+{
+	if (Size.X != DesiredSize.X)
+	{
+		Size.X = Size.X * 0.75f + DesiredSize.X * 0.25f;
+
+		if (FMath::IsNearlyEqual(Size.X, DesiredSize.X))
+		{
+			Size.X = DesiredSize.X;
+		}
+	}
+
+	if (Size.Y != DesiredSize.Y)
+	{
+		Size.Y = Size.Y * 0.5f + DesiredSize.Y * 0.5f;
+
+		if (FMath::IsNearlyEqual(Size.Y, DesiredSize.Y))
+		{
+			Size.Y = DesiredSize.Y;
+		}
+	}
+
+	const float RealDesiredOpacity = FMath::Min(DesiredOpacity, 1.0f - FMath::Abs(Size.X - DesiredSize.X) / DesiredSize.X);
+	if (Opacity != RealDesiredOpacity)
+	{
+		if (Opacity < RealDesiredOpacity)
+		{
+			// slow fade in
+			Opacity = Opacity * 0.9f + RealDesiredOpacity * 0.1f;
+		}
+		else
+		{
+			// fast fade out
+			Opacity = Opacity * 0.75f + RealDesiredOpacity * 0.25f;
+		}
+
+		if (FMath::IsNearlyEqual(Opacity, RealDesiredOpacity))
+		{
+			Opacity = RealDesiredOpacity;
+		}
+	}
+}
+
+////////////////////////////////////////////////////////////////////////////////////////////////////
+
+void FTooltipDrawState::SetPosition(const FVector2D& MousePosition, const float MinX, const float MaxX, const float MinY, const float MaxY)
+{
+	Position.X = FMath::Max(MinX, FMath::Min(MousePosition.X + 12.0f, MaxX - Size.X));
+	Position.Y = FMath::Max(MinY, FMath::Min(MousePosition.Y + 15.0f, MaxY - Size.Y));
+}
+
+////////////////////////////////////////////////////////////////////////////////////////////////////
+
+void FTooltipDrawState::Draw(const FDrawContext& DrawContext) const
+{
+	if (Opacity > 0.0f && Size.X > 0.0f && Size.Y > 0.0f)
+	{
+		// Draw background.
+		DrawContext.DrawBox(Position.X, Position.Y, Size.X, Size.Y, WhiteBrush, BackgroundColor.CopyWithNewOpacity(Opacity));
+		if (Size.X < DesiredSize.X)
+		{
+			DrawContext.DrawBox(Position.X + Size.X, Position.Y, DesiredSize.X - Size.X, Size.Y, WhiteBrush, BackgroundColor.CopyWithNewOpacity(Opacity * 0.5f));
+		}
+		DrawContext.LayerId++;
+
+		// Draw border.
+		//DrawContext.DrawBox(Position.X, Position.Y, Size.X, Size.Y, BorderBrush, BorderColor.CopyWithNewOpacity(Opacity));
+		//DrawContext.LayerId++;
+
+		// Draw cached texts.
+		for (const FDrawTextInfo& TextInfo : Texts)
+		{
+			float X = Position.X;
+			switch (TextInfo.Type)
+			{
+				case FDrawTextType::Name:
+					X += ValueOffsetX - TextInfo.TextSize.X;
+					break;
+
+				case FDrawTextType::Value:
+					X += ValueOffsetX + NameValueDX;
+					break;
+
+				default:
+					X += TextInfo.X;
+			}
+			DrawContext.DrawText(X, Position.Y + TextInfo.Y, TextInfo.Text, Font, TextInfo.Color.CopyWithNewOpacity(Opacity));
+		}
+		DrawContext.LayerId++;
+	}
+}
+
+////////////////////////////////////////////////////////////////////////////////////////////////////
diff --git a/Engine/Source/Developer/TraceInsights/Private/Insights/Widgets/SFrameTrack.h b/Engine/Source/Developer/TraceInsights/Private/Insights/Widgets/SFrameTrack.h
--- a/Engine/Source/Developer/TraceInsights/Private/Insights/Widgets/SFrameTrack.h
+++ b/Engine/Source/Developer/TraceInsights/Private/Insights/Widgets/SFrameTrack.h
@@ -19,33 +19,67 @@ class SScrollBar;
 
 ////////////////////////////////////////////////////////////////////////////////////////////////////
 
-struct FSampleRef
+struct FFrameTrackSampleRef
 {
-	const FFrameTrackTimeline* Timeline;
-	const FFrameTrackSample* Sample;
+	TSharedPtr<FFrameTrackSeries> Series;
+	TSharedPtr<FFrameTrackSample> Sample;
 
-	FSampleRef(const FFrameTrackTimeline* InTimeline, const FFrameTrackSample* InSample)
-		: Timeline(InTimeline), Sample(InSample)
+	FFrameTrackSampleRef()
+		: Series(), Sample()
 	{
 	}
 
-	void Reset() { Timeline = nullptr; Sample = nullptr; }
-	bool IsValid() const { return Timeline != nullptr && Sample != nullptr; }
+	FFrameTrackSampleRef(TSharedPtr<FFrameTrackSeries> InSeries, TSharedPtr<FFrameTrackSample> InSample)
+		: Series(InSeries), Sample(InSample)
+	{
+	}
+
+	FFrameTrackSampleRef(const FFrameTrackSampleRef& Other)
+		: Series(Other.Series), Sample(Other.Sample)
+	{
+	}
+
+	FFrameTrackSampleRef& operator=(const FFrameTrackSampleRef& Other)
+	{
+		Series = Other.Series;
+		Sample = Other.Sample;
+		return *this;
+	}
+
+	void Reset()
+	{
+		Series.Reset();
+		Sample.Reset();
+	}
+
+	bool IsValid() const
+	{
+		return Series.IsValid() && Sample.IsValid();
+	}
+
+	bool Equals(const FFrameTrackSampleRef& Other) const
+	{
+		return Series == Other.Series
+			&& ((Sample == Other.Sample) || (Sample.IsValid() && Other.Sample.IsValid() && Sample->Equals(*Other.Sample)));
+	}
+
+	static bool AreEquals(const FFrameTrackSampleRef& A, const FFrameTrackSampleRef& B)
+	{
+		return A.Equals(B);
+	}
 };
 
 ////////////////////////////////////////////////////////////////////////////////////////////////////
 /**
- * Widget used to present frames data in a track.
+ * Widget used to present frames data in a bar track.
  */
 class SFrameTrack : public SCompoundWidget
 {
-	enum
-	{
-		/** Number of pixels. */
-		MOUSE_SNAP_DISTANCE = 4,
-	};
+public:
+	/** Number of pixels. */
+	static constexpr float MOUSE_SNAP_DISTANCE = 2.0f;
 
-	enum class EFrameTrackCursor
+	enum class ECursorType
 	{
 		Default,
 		Arrow,
@@ -95,14 +129,29 @@ public:
 	virtual FCursorReply OnCursorQuery(const FGeometry& MyGeometry, const FPointerEvent& CursorEvent) const override;
 
 protected:
-	bool IsReady() { return true; }
-
+	TSharedPtr<FFrameTrackSeries> FindOrAddSeries(int32 FrameType);
+	TSharedPtr<FFrameTrackSeries> FindSeries(int32 FrameType) const;
 	void UpdateState();
 
-	FSampleRef GetSampleAtMousePosition(float X, float Y);
+	void DrawHorizontalAxisGrid(FDrawContext& DrawContext, const FSlateBrush* Brush, const FSlateFontInfo& Font) const;
+	void DrawVerticalAxisGrid(FDrawContext& DrawContext, const FSlateBrush* Brush, const FSlateFontInfo& Font) const;
+
+	FFrameTrackSampleRef GetSampleAtMousePosition(float X, float Y);
 	void SelectFrameAtMousePosition(float X, float Y);
 
-	void ShowContextMenu(const FVector2D& ScreenSpacePosition);
+	void ShowContextMenu(const FPointerEvent& MouseEvent);
+
+	void ContextMenu_ShowGameFrames_Execute();
+	bool ContextMenu_ShowGameFrames_CanExecute();
+	bool ContextMenu_ShowGameFrames_IsChecked();
+
+	void ContextMenu_ShowRenderingFrames_Execute();
+	bool ContextMenu_ShowRenderingFrames_CanExecute();
+	bool ContextMenu_ShowRenderingFrames_IsChecked();
+
+	void ContextMenu_AutoZoom_Execute();
+	bool ContextMenu_AutoZoom_CanExecute();
+	bool ContextMenu_AutoZoom_IsChecked();
 
 	/** Binds our UI commands to delegates. */
 	void BindCommands();
@@ -116,30 +165,21 @@ protected:
 
 	void ZoomHorizontally(const float Delta, const float X);
 
-	//////////////////////////////////////////////////
-	// SelectionBoxChanged Event
-
-public:
-	/** The event to execute when the selection box has been changed. */
-	DECLARE_EVENT_TwoParams(SFrameTrack, FSelectionBoxChangedEvent, int32 /*FrameStart*/, int32 /*FrameEnd*/);
-	FSelectionBoxChangedEvent& OnSelectionBoxChanged() { return SelectionBoxChangedEvent; }
-protected:
-	/** The event to execute when the selection box has been changed. */
-	FSelectionBoxChangedEvent SelectionBoxChangedEvent;
-
-	//////////////////////////////////////////////////
-
 protected:
 	/** The track's viewport. Encapsulates info about position and scale. */
 	FFrameTrackViewport Viewport;
 	bool bIsViewportDirty;
 
-	/** Cached info for timelines. */
-	TMap<uint64, FFrameTrackTimeline> CachedTimelines;
-	TArray<int32> TimelinesOrder;
+	/** Cached info for all frame series. */
+	TMap<int32, TSharedPtr<FFrameTrackSeries>> SeriesMap;
+	TArray<int32> SeriesOrder;
 
 	bool bIsStateDirty;
 
+	bool bShowGameFrames;
+	bool bShowRenderingFrames;
+	bool bIsAutoZoomEnabled;
+
 	uint64 AnalysisSyncNextTimestamp;
 
 	//////////////////////////////////////////////////
@@ -147,7 +187,7 @@ protected:
 	TSharedPtr<SScrollBar> HorizontalScrollBar;
 
 	//////////////////////////////////////////////////
-	// Panning, Zooming and Selection behaviors
+	// Panning and Zooming behaviors
 
 	/** The current mouse position. */
 	FVector2D MousePosition;
@@ -155,7 +195,6 @@ protected:
 	/** Mouse position during the call on mouse button down. */
 	FVector2D MousePositionOnButtonDown;
 	float ViewportPosXOnButtonDown;
-	float ViewportPosYOnButtonDown;
 
 	/** Mouse position during the call on mouse button up. */
 	FVector2D MousePositionOnButtonUp;
@@ -163,15 +202,16 @@ protected:
 	bool bIsLMB_Pressed;
 	bool bIsRMB_Pressed;
 
-	/** True, if the user is currently interactively scrolling the view by holding the right mouse button and dragging. */
+	/** True, if the user is currently interactively scrolling the view (ex.: by holding the left mouse button and dragging). */
 	bool bIsScrolling;
 
 	//////////////////////////////////////////////////
+	// Selection
 
-	int32 SelectionStartFrameIndex;
-	int32 SelectionEndFrameIndex;
+	FFrameTrackSampleRef HoveredSample;
 
-	FSampleRef HoveredSample;
+	float TooltipDesiredOpacity;
+	mutable float TooltipOpacity;
 
 	//////////////////////////////////////////////////
 	// Misc
@@ -179,7 +219,7 @@ protected:
 	FGeometry ThisGeometry;
 
 	/** Cursor type. */
-	EFrameTrackCursor CursorType;
+	ECursorType CursorType;
 
 	// Debug stats
 	int32 NumUpdatedFrames;
diff --git a/Engine/Source/Developer/TraceInsights/Private/Insights/Widgets/SStatsTableCell.cpp b/Engine/Source/Developer/TraceInsights/Private/Insights/Widgets/SStatsTableCell.cpp
--- a/Engine/Source/Developer/TraceInsights/Private/Insights/Widgets/SStatsTableCell.cpp
+++ b/Engine/Source/Developer/TraceInsights/Private/Insights/Widgets/SStatsTableCell.cpp
@@ -6,11 +6,14 @@
 #include "SlateOptMacros.h"
 #include "Widgets/Images/SImage.h"
 #include "Widgets/Input/SButton.h"
+#include "Widgets/Layout/SBox.h"
 #include "Widgets/SBoxPanel.h"
+#include "Widgets/SToolTip.h"
 #include "Widgets/Text/STextBlock.h"
 #include "Widgets/Views/SExpanderArrow.h"
 
 // Insights
+#include "Insights/InsightsStyle.h"
 #include "Insights/ViewModels/StatsViewColumn.h"
 #include "Insights/ViewModels/StatsViewColumnFactory.h"
 #include "Insights/Widgets/SStatsViewTooltip.h"
@@ -38,7 +41,7 @@ void SStatsTableCell::Construct(const FArguments& InArgs, const TSharedRef<class
 
 TSharedRef<SWidget> SStatsTableCell::GenerateWidgetForColumn(const FArguments& InArgs, const TSharedRef<class ITableRow>& TableRow)
 {
-	if (InArgs._IsStatsNameColumn)
+	if (InArgs._IsNameColumn)
 	{
 		return GenerateWidgetForNameColumn(InArgs, TableRow);
 	}
@@ -65,7 +68,7 @@ TSharedRef<SWidget> SStatsTableCell::GenerateWidgetForNameColumn(const FArgument
 			SNew(SExpanderArrow, TableRow)
 		]
 
-		// Event info icon + tooltip.
+		// Info icon + tooltip
 		+ SHorizontalBox::Slot()
 		.AutoWidth()
 		.HAlign(HAlign_Center)
@@ -77,6 +80,25 @@ TSharedRef<SWidget> SStatsTableCell::GenerateWidgetForNameColumn(const FArgument
 			.ToolTip(SStatsViewTooltip::GetTableCellTooltip(StatsNodePtr))
 		]
 
+		// Color box
+		+ SHorizontalBox::Slot()
+		.AutoWidth()
+		.HAlign(HAlign_Center)
+		.VAlign(VAlign_Center)
+		[
+			SNew(SBox)
+			.Visibility(this, &SStatsTableCell::GetColorBoxVisibility)
+			.WidthOverride(14.0f)
+			.HeightOverride(14.0f)
+			[
+				SNew(SBorder)
+				.BorderImage(FInsightsStyle::Get().GetBrush("WhiteBrush"))
+				.BorderBackgroundColor(this, &SStatsTableCell::GetStatsBoxColorAndOpacity)
+				.HAlign(HAlign_Fill)
+				.VAlign(VAlign_Fill)
+			]
+		]
+
 		// Name
 		+ SHorizontalBox::Slot()
 		.AutoWidth()
@@ -85,34 +107,13 @@ TSharedRef<SWidget> SStatsTableCell::GenerateWidgetForNameColumn(const FArgument
 		.Padding(FMargin(2.0f, 0.0f))
 		[
 			SNew(STextBlock)
-			//.Text(StatsNodePtr->GetNameEx())
-			.Text(this, &SStatsTableCell::GetNameEx)
+			.Text(this, &SStatsTableCell::GetDisplayName)
 			.HighlightText(InArgs._HighlightText)
 			.TextStyle(FEditorStyle::Get(), TEXT("Profiler.Tooltip"))
 			.ColorAndOpacity(this, &SStatsTableCell::GetColorAndOpacity)
 			.ShadowColorAndOpacity(this, &SStatsTableCell::GetShadowColorAndOpacity)
 		]
-
-		/*
-		// Culled children warning icon + tooltip.
-		+ SHorizontalBox::Slot()
-		.AutoWidth()
-		.HAlign(HAlign_Center)
-		.VAlign(VAlign_Center)
-		[
-			SNew(SButton)
-			.ButtonStyle(FEditorStyle::Get(), TEXT("HoverHintOnly"))
-			.ContentPadding(0.0f)
-			.IsFocusable(false)
-			//.OnClicked(this, &SStatsTableCell::ExpandCulledEvents_OnClicked)
-			[
-				SNew(SImage)
-				.Visibility(this, &SStatsTableCell::GetCulledEventsIconVisibility)
-				.Image(FEditorStyle::GetBrush("Profiler.EventGraph.HasCulledEventsSmall"))
-				.ToolTipText(LOCTEXT("HasCulledEvents_TT", "This event contains culled children, if you want to see all children, please disable culling or use function details, or press this icon"))
-			]
-		]*/
-		;
+	;
 }
 
 ////////////////////////////////////////////////////////////////////////////////////////////////////
@@ -120,11 +121,13 @@ TSharedRef<SWidget> SStatsTableCell::GenerateWidgetForNameColumn(const FArgument
 TSharedRef<SWidget> SStatsTableCell::GenerateWidgetForStatsColumn(const FArguments& InArgs, const TSharedRef<class ITableRow>& TableRow)
 {
 	const FStatsViewColumn& Column = *FStatsViewColumnFactory::Get().ColumnIdToPtrMapping.FindChecked(ColumnId);
-	const FText FormattedValue = Column.GetFormattedValue(*StatsNodePtr);
+	const FText CellText = Column.GetFormattedValue(*StatsNodePtr);
+	//const FText CellText = ColumnPtr->GetValueAsText(*StatsNodePtr);
 
 	return
 		SNew(SHorizontalBox)
 
+		// Value
 		+ SHorizontalBox::Slot()
 		.FillWidth(1.0f)
 		.VAlign(VAlign_Center)
@@ -132,11 +135,12 @@ TSharedRef<SWidget> SStatsTableCell::GenerateWidgetForStatsColumn(const FArgumen
 		.Padding(FMargin(2.0f, 0.0f))
 		[
 			SNew(STextBlock)
-			.Text(FormattedValue)
+			.Text(CellText)
 			.TextStyle(FEditorStyle::Get(), TEXT("Profiler.Tooltip"))
 			.ColorAndOpacity(this, &SStatsTableCell::GetStatsColorAndOpacity)
 			.ShadowColorAndOpacity(this, &SStatsTableCell::GetShadowColorAndOpacity)
-		];
+		]
+	;
 }
 
 END_SLATE_FUNCTION_BUILD_OPTIMIZATION
diff --git a/Engine/Source/Developer/TraceInsights/Private/Insights/Widgets/SStatsTableCell.h b/Engine/Source/Developer/TraceInsights/Private/Insights/Widgets/SStatsTableCell.h
--- a/Engine/Source/Developer/TraceInsights/Private/Insights/Widgets/SStatsTableCell.h
+++ b/Engine/Source/Developer/TraceInsights/Private/Insights/Widgets/SStatsTableCell.h
@@ -10,7 +10,6 @@
 #include "Insights/ViewModels/StatsNode.h"
 
 DECLARE_DELEGATE_TwoParams(FSetHoveredStatsTableCell, const FName /*ColumnId*/, const FStatsNodePtr /*SamplePtr*/);
-DECLARE_DELEGATE_RetVal_OneParam(bool, FIsColumnVisibleDelegate, const FName /*ColumnId*/);
 DECLARE_DELEGATE_RetVal_OneParam(EHorizontalAlignment, FGetColumnOutlineHAlignmentDelegate, const FName /*ColumnId*/);
 
 class SStatsTableCell : public SCompoundWidget
@@ -21,7 +20,7 @@ public:
 		SLATE_ATTRIBUTE(FText, HighlightText)
 		SLATE_ARGUMENT(FStatsNodePtr, StatsNodePtr)
 		SLATE_ARGUMENT(FName, ColumnId)
-		SLATE_ARGUMENT(bool, IsStatsNameColumn)
+		SLATE_ARGUMENT(bool, IsNameColumn)
 	SLATE_END_ARGS()
 
 	void Construct(const FArguments& InArgs, const TSharedRef<class ITableRow>& TableRow);
@@ -88,20 +87,19 @@ protected:
 		SetHoveredTableCellDelegate.ExecuteIfBound(NAME_None, nullptr);
 	}
 
-	EVisibility GetHintIconVisibility() const
+	EVisibility GetColorBoxVisibility() const
 	{
-		return IsHovered() ? EVisibility::Visible : EVisibility::Hidden;
+		return StatsNodePtr->IsAddedToGraph() ? EVisibility::Visible : EVisibility::Collapsed;
 	}
 
-	EVisibility GetCulledEventsIconVisibility() const
+	EVisibility GetHintIconVisibility() const
 	{
-		//return StatsNodePtr->HasCulledChildren() ? EVisibility::Visible : EVisibility::Collapsed;
-		return EVisibility::Collapsed;
+		return IsHovered() ? EVisibility::Visible : EVisibility::Hidden;
 	}
 
-	FText GetNameEx() const
+	FText GetDisplayName() const
 	{
-		return StatsNodePtr->GetNameEx();
+		return StatsNodePtr->GetDisplayName();
 	}
 
 	FSlateColor GetColorAndOpacity() const
@@ -121,6 +119,11 @@ protected:
 		return TextColor;
 	}
 
+	FSlateColor GetStatsBoxColorAndOpacity() const
+	{
+		return StatsNodePtr->GetColor();
+	}
+
 	FLinearColor GetShadowColorAndOpacity() const
 	{
 		const FLinearColor ShadowColor =
