Copying //UE4/Dev-Core to Main (//UE4/Main)

diff --git a/Engine/Source/Developer/TraceInsights/Private/Insights/Widgets/SStatsTableRow.cpp b/Engine/Source/Developer/TraceInsights/Private/Insights/Widgets/SStatsTableRow.cpp
--- a/Engine/Source/Developer/TraceInsights/Private/Insights/Widgets/SStatsTableRow.cpp
+++ b/Engine/Source/Developer/TraceInsights/Private/Insights/Widgets/SStatsTableRow.cpp
@@ -2,6 +2,8 @@
 
 #include "SStatsTableRow.h"
 
+#include "EditorStyleSet.h"
+#include "SlateOptMacros.h"
 #include "Widgets/SOverlay.h"
 #include "Widgets/Images/SImage.h"
 
@@ -13,6 +15,8 @@
 
 #define LOCTEXT_NAMESPACE "SStatsView"
 
+////////////////////////////////////////////////////////////////////////////////////////////////////
+
 BEGIN_SLATE_FUNCTION_BUILD_OPTIMIZATION
 
 void SStatsTableRow::Construct(const FArguments& InArgs, const TSharedRef<STableViewBase>& InOwnerTableView)
@@ -23,109 +27,55 @@ void SStatsTableRow::Construct(const FArguments& InArgs, const TSharedRef<STable
 	GetColumnOutlineHAlignmentDelegate = InArgs._OnGetColumnOutlineHAlignmentDelegate;
 
 	HighlightText = InArgs._HighlightText;
-	HighlightedStatsName = InArgs._HighlightedStatsName;
+	HighlightedNodeName = InArgs._HighlightedNodeName;
 
 	StatsNodePtr = InArgs._StatsNodePtr;
 
 	SetEnabled(TAttribute<bool>(this, &SStatsTableRow::HandleShouldBeEnabled));
 
 	SMultiColumnTableRow<FStatsNodePtr>::Construct(SMultiColumnTableRow<FStatsNodePtr>::FArguments(), InOwnerTableView);
-
-	/*
-	const FSlateBrush* const NodeIcon = StatsNodeTypeHelper::GetIconForStatsNodeType(InStatsNode->GetType());
-	const TSharedRef<SToolTip> Tooltip = InStatsNode->IsGroup() ? SNew(SToolTip) : SStatsViewTooltip(InStatsNode->GetId()).GetTooltip();
-
-	ChildSlot
-		[
-			SNew(SHorizontalBox)
-
-			// Expander arrow.
-			+ SHorizontalBox::Slot()
-			.AutoWidth()
-			.HAlign(HAlign_Right)
-			.VAlign(VAlign_Center)
-			[
-				SNew(SExpanderArrow, SharedThis(this))
-			]
-
-			// Icon to visualize group or timer type.
-			+ SHorizontalBox::Slot()
-			.AutoWidth()
-			.Padding(0.0f, 0.0f, 8.0f, 0.0f)
-			[
-				SNew(SImage)
-				.Image(NodeIcon)
-				.ToolTip(Tooltip)
-			]
-
-			// Description text.
-			+ SHorizontalBox::Slot()
-			.AutoWidth()
-			.VAlign(VAlign_Center)
-			.HAlign(HAlign_Left)
-			.Padding(FMargin(2.0f, 0.0f))
-			[
-				SNew(STextBlock)
-				.Text(this, &SStatsTableRow::GetText)
-				.HighlightText(InArgs._HighlightText)
-				.TextStyle(FEditorStyle::Get(), TEXT("Profiler.Tooltip"))
-				.ColorAndOpacity(this, &SStatsTableRow::GetColorAndOpacity)
-			]
-
-			+ SHorizontalBox::Slot()
-			.AutoWidth()
-			.HAlign(HAlign_Center)
-			.VAlign(VAlign_Top)
-			.Padding(0.0f, 1.0f, 0.0f, 0.0f)
-			[
-				SNew(SImage)
-				.Visibility(!InStatsNode->IsGroup() ? EVisibility::Visible : EVisibility::Collapsed)
-				.Image(FEditorStyle::GetBrush("Profiler.Tooltip.HintIcon10"))
-				.ToolTip(Tooltip)
-			]
-		];
-
-	STableRow<FStatsNodePtr>::ConstructInternal(STableRow::FArguments().ShowSelection(true), InOwnerTableView);
-	*/
 }
 
+////////////////////////////////////////////////////////////////////////////////////////////////////
+
 TSharedRef<SWidget> SStatsTableRow::GenerateWidgetForColumn(const FName& ColumnId)
 {
 	return
+		SNew(SOverlay)
+		.Visibility(EVisibility::SelfHitTestInvisible)
 
-	SNew(SOverlay)
-	.Visibility(EVisibility::SelfHitTestInvisible)
-
-	+SOverlay::Slot()
-	.Padding(0.0f)
-	[
-		SNew(SImage)
-		.Image(FEditorStyle::GetBrush("Profiler.LineGraphArea"))
-		.ColorAndOpacity(this, &SStatsTableRow::GetBackgroundColorAndOpacity)
-	]
-
-	+SOverlay::Slot()
-	.Padding(0.0f)
-	[
-		SNew(SImage)
-		.Image(this, &SStatsTableRow::GetOutlineBrush, ColumnId)
-		.ColorAndOpacity(this, &SStatsTableRow::GetOutlineColorAndOpacity)
-	]
-
-	+SOverlay::Slot()
-	[
-		SNew(SStatsTableCell, SharedThis(this))
-		.Visibility(this, &SStatsTableRow::IsColumnVisible, ColumnId)
-		.StatsNodePtr(StatsNodePtr)
-		.ColumnId(ColumnId)
-		.HighlightText(HighlightText)
-		.IsStatsNameColumn(ColumnId == FStatsViewColumnFactory::Get().Collection[0]->Id) // name column
-		.OnSetHoveredTableCell(this, &SStatsTableRow::OnSetHoveredTableCell)
-	];
+		+SOverlay::Slot()
+		.Padding(0.0f)
+		[
+			SNew(SImage)
+			.Image(FEditorStyle::GetBrush("Profiler.LineGraphArea"))
+			.ColorAndOpacity(this, &SStatsTableRow::GetBackgroundColorAndOpacity)
+		]
+
+		+SOverlay::Slot()
+		.Padding(0.0f)
+		[
+			SNew(SImage)
+			.Image(this, &SStatsTableRow::GetOutlineBrush, ColumnId)
+			.ColorAndOpacity(this, &SStatsTableRow::GetOutlineColorAndOpacity)
+		]
+
+		+SOverlay::Slot()
+		[
+			SNew(SStatsTableCell, SharedThis(this))
+			.Visibility(this, &SStatsTableRow::IsColumnVisible, ColumnId)
+			.StatsNodePtr(StatsNodePtr)
+			.ColumnId(ColumnId)
+			.HighlightText(HighlightText)
+			.IsNameColumn(ColumnId == FStatsViewColumnFactory::Get().Collection[0]->Id) // name column
+			.OnSetHoveredTableCell(this, &SStatsTableRow::OnSetHoveredTableCell)
+		];
 }
 
 END_SLATE_FUNCTION_BUILD_OPTIMIZATION
 
+////////////////////////////////////////////////////////////////////////////////////////////////////
+
 FReply SStatsTableRow::OnDragDetected(const FGeometry& MyGeometry, const FPointerEvent& MouseEvent)
 {
 	//im:TODO
@@ -155,37 +105,7 @@ FReply SStatsTableRow::OnDragDetected(const FGeometry& MyGeometry, const FPointe
 	return SMultiColumnTableRow<FStatsNodePtr>::OnDragDetected(MyGeometry, MouseEvent);
 }
 
-FText SStatsTableRow::GetText() const
-{
-	FText Text = FText::GetEmpty();
-
-	if (StatsNodePtr->IsGroup())
-	{
-		Text = FText::Format(LOCTEXT("StatsNode_GroupNodeTextFmt", "{0} ({1})"), FText::FromName(StatsNodePtr->GetName()), FText::AsNumber(StatsNodePtr->GetChildren().Num()));
-	}
-	else
-	{
-		Text = FText::FromName(StatsNodePtr->GetName());
-	}
-
-	return Text;
-}
-
-FSlateFontInfo SStatsTableRow::GetFont() const
-{
-	const bool bIsStatTracked = false;//im:TODO: FTimingProfilerManager::Get()->IsStatTracked(StatsNodePtr->GetStatsId());
-	const FSlateFontInfo FontInfo = bIsStatTracked ? FEditorStyle::GetFontStyle("BoldFont") : FEditorStyle::GetFontStyle("NormalFont");
-	return FontInfo;
-}
-
-FSlateColor SStatsTableRow::GetColorAndOpacity() const
-{
-	//const bool bIsStatTracked = FTimingProfilerManager::Get()->IsStatTracked(StatsNodePtr->GetId());
-	//const FSlateColor Color = bIsStatTracked ? FTimingProfilerManager::Get()->GetColorForStatsId(StatsNodePtr->GetId()) : FLinearColor::White;
-	//return Color;
-	//FTimingProfilerManager::GetSettings().GetColorForStat(StatName)
-	return FLinearColor::White;
-}
+////////////////////////////////////////////////////////////////////////////////////////////////////
 
 FSlateColor SStatsTableRow::GetBackgroundColorAndOpacity() const
 {
@@ -193,6 +113,8 @@ FSlateColor SStatsTableRow::GetBackgroundColorAndOpacity() const
 	return FLinearColor(0.0f, 0.0f, 0.0f, 1.0f);
 }
 
+////////////////////////////////////////////////////////////////////////////////////////////////////
+
 FSlateColor SStatsTableRow::GetBackgroundColorAndOpacity(double Time) const
 {
 	const FLinearColor Color =	Time > TimeUtils::Second      ? FLinearColor(0.3f, 0.0f, 0.0f, 1.0f) :
@@ -202,14 +124,18 @@ FSlateColor SStatsTableRow::GetBackgroundColorAndOpacity(double Time) const
 	return Color;
 }
 
+////////////////////////////////////////////////////////////////////////////////////////////////////
+
 FSlateColor SStatsTableRow::GetOutlineColorAndOpacity() const
 {
 	const FLinearColor NoColor(0.0f, 0.0f, 0.0f, 0.0f);
-	const bool bShouldBeHighlighted = StatsNodePtr->GetName() == HighlightedStatsName.Get();
+	const bool bShouldBeHighlighted = StatsNodePtr->GetName() == HighlightedNodeName.Get();
 	const FLinearColor OutlineColorAndOpacity = bShouldBeHighlighted ? FLinearColor(FColorList::SlateBlue) : NoColor;
 	return OutlineColorAndOpacity;
 }
 
+////////////////////////////////////////////////////////////////////////////////////////////////////
+
 const FSlateBrush* SStatsTableRow::GetOutlineBrush(const FName ColumnId) const
 {
 	EHorizontalAlignment Result = HAlign_Center;
@@ -234,6 +160,8 @@ const FSlateBrush* SStatsTableRow::GetOutlineBrush(const FName ColumnId) const
 	return Brush;
 }
 
+////////////////////////////////////////////////////////////////////////////////////////////////////
+
 bool SStatsTableRow::HandleShouldBeEnabled() const
 {
 	bool bResult = false;
@@ -253,6 +181,8 @@ bool SStatsTableRow::HandleShouldBeEnabled() const
 	return bResult;
 }
 
+////////////////////////////////////////////////////////////////////////////////////////////////////
+
 EVisibility SStatsTableRow::IsColumnVisible(const FName ColumnId) const
 {
 	EVisibility Result = EVisibility::Collapsed;
@@ -265,9 +195,13 @@ EVisibility SStatsTableRow::IsColumnVisible(const FName ColumnId) const
 	return Result;
 }
 
+////////////////////////////////////////////////////////////////////////////////////////////////////
+
 void SStatsTableRow::OnSetHoveredTableCell(const FName InColumnId, const FStatsNodePtr InSamplePtr)
 {
 	SetHoveredTableCellDelegate.ExecuteIfBound(InColumnId, InSamplePtr);
 }
 
+////////////////////////////////////////////////////////////////////////////////////////////////////
+
 #undef LOCTEXT_NAMESPACE
diff --git a/Engine/Source/Developer/TraceInsights/Private/Insights/Widgets/SStatsTableRow.h b/Engine/Source/Developer/TraceInsights/Private/Insights/Widgets/SStatsTableRow.h
--- a/Engine/Source/Developer/TraceInsights/Private/Insights/Widgets/SStatsTableRow.h
+++ b/Engine/Source/Developer/TraceInsights/Private/Insights/Widgets/SStatsTableRow.h
@@ -27,7 +27,7 @@ public:
 		SLATE_EVENT(FSetHoveredStatsTableCell, OnSetHoveredTableCell)
 		SLATE_EVENT(FGetColumnOutlineHAlignmentDelegate, OnGetColumnOutlineHAlignmentDelegate)
 		SLATE_ATTRIBUTE(FText, HighlightText)
-		SLATE_ATTRIBUTE(FName, HighlightedStatsName)
+		SLATE_ATTRIBUTE(FName, HighlightedNodeName)
 		SLATE_ARGUMENT(FStatsNodePtr, StatsNodePtr)
 	SLATE_END_ARGS()
 
@@ -53,21 +53,6 @@ public:
 	virtual FReply OnDragDetected(const FGeometry& MyGeometry, const FPointerEvent& MouseEvent) override;
 
 protected:
-	/**
-	 * @return a text which describes this table row, refers to both groups and stats
-	 */
-	FText GetText() const;
-
-	/**
-	 * @return a font style which is used to draw this table row, refers to both groups and stats
-	 */
-	FSlateFontInfo GetFont() const;
-
-	/**
-	 * @return a color and opacity value used to draw this table row, refers to both groups and stats
-	 */
-	FSlateColor GetColorAndOpacity() const;
-
 	FSlateColor GetBackgroundColorAndOpacity() const;
 	FSlateColor GetBackgroundColorAndOpacity(double Time) const;
 	FSlateColor GetOutlineColorAndOpacity() const;
@@ -88,6 +73,6 @@ protected:
 	/** Text to be highlighted on stats name. */
 	TAttribute<FText> HighlightText;
 
-	/** Name of the stats that should be drawn as highlighted. */
-	TAttribute<FName> HighlightedStatsName;
+	/** Name of the stats node that should be drawn as highlighted. */
+	TAttribute<FName> HighlightedNodeName;
 };
diff --git a/Engine/Source/Developer/TraceInsights/Private/Insights/Widgets/SStatsView.cpp b/Engine/Source/Developer/TraceInsights/Private/Insights/Widgets/SStatsView.cpp
--- a/Engine/Source/Developer/TraceInsights/Private/Insights/Widgets/SStatsView.cpp
+++ b/Engine/Source/Developer/TraceInsights/Private/Insights/Widgets/SStatsView.cpp
@@ -18,10 +18,13 @@
 // Insights
 #include "Insights/TimingProfilerCommon.h"
 #include "Insights/TimingProfilerManager.h"
+#include "Insights/ViewModels/TimingGraphTrack.h"
 #include "Insights/ViewModels/StatsNodeHelper.h"
 #include "Insights/ViewModels/StatsViewColumnFactory.h"
 #include "Insights/Widgets/SStatsViewTooltip.h"
 #include "Insights/Widgets/SStatsTableRow.h"
+#include "Insights/Widgets/STimingProfilerWindow.h"
+#include "Insights/Widgets/STimingView.h"
 
 #define LOCTEXT_NAMESPACE "SStatsView"
 
@@ -43,6 +46,7 @@ const FName SStatsView::GetDefaultColumnBeingSorted()
 
 SStatsView::SStatsView()
 	: bExpansionSaved(false)
+	, bFilterOutZeroCountStats(false)
 	, GroupingMode(EStatsGroupingMode::Flat)
 	, ColumnSortMode(GetDefaultColumnSortMode())
 	, ColumnBeingSorted(GetDefaultColumnBeingSorted())
@@ -73,7 +77,7 @@ void SStatsView::Construct(const FArguments& InArgs)
 	[
 		SNew(SVerticalBox)
 
-		+SVerticalBox::Slot()
+		+ SVerticalBox::Slot()
 		.VAlign(VAlign_Center)
 		.AutoHeight()
 		[
@@ -83,28 +87,57 @@ void SStatsView::Construct(const FArguments& InArgs)
 			[
 				SNew(SVerticalBox)
 
-				// Search box
-				+SVerticalBox::Slot()
+				+ SVerticalBox::Slot()
 				.VAlign(VAlign_Center)
 				.Padding(2.0f)
 				.AutoHeight()
 				[
-					SAssignNew(SearchBox, SSearchBox)
-					.HintText(LOCTEXT("SearchBoxHint", "Search timers or groups"))
-					.OnTextChanged(this, &SStatsView::SearchBox_OnTextChanged)
-					.IsEnabled(this, &SStatsView::SearchBox_IsEnabled)
-					.ToolTipText(LOCTEXT("FilterSearchHint", "Type here to search timer or group"))
+					SNew(SHorizontalBox)
+
+					// Search box
+					+ SHorizontalBox::Slot()
+					.VAlign(VAlign_Center)
+					.Padding(2.0f)
+					.FillWidth(1.0f)
+					[
+						SAssignNew(SearchBox, SSearchBox)
+						.HintText(LOCTEXT("SearchBoxHint", "Search stats counters or groups"))
+						.OnTextChanged(this, &SStatsView::SearchBox_OnTextChanged)
+						.IsEnabled(this, &SStatsView::SearchBox_IsEnabled)
+						.ToolTipText(LOCTEXT("FilterSearchHint", "Type here to search stats counter or group"))
+					]
+
+					// Filter out timers with zero instance count
+					+ SHorizontalBox::Slot()
+					.VAlign(VAlign_Center)
+					.Padding(2.0f)
+					.AutoWidth()
+					[
+						SNew(SCheckBox)
+						.Style(FEditorStyle::Get(), "ToggleButtonCheckbox")
+						.HAlign(HAlign_Center)
+						.Padding(2.0f)
+						.OnCheckStateChanged(this, &SStatsView::FilterOutZeroCountStats_OnCheckStateChanged)
+						.IsChecked(this, &SStatsView::FilterOutZeroCountStats_IsChecked)
+						.ToolTipText(LOCTEXT("FilterOutZeroCountStats_Tooltip", "Filter out the stats counters having zero total instance count (aggregated stats)."))
+						[
+							//TODO: SNew(SImage)
+							SNew(STextBlock)
+							.Text(LOCTEXT("FilterOutZeroCountStats_Button", " !0 "))
+							.TextStyle(FEditorStyle::Get(), TEXT("Profiler.Caption"))
+						]
+					]
 				]
 
 				// Group by
-				+SVerticalBox::Slot()
+				+ SVerticalBox::Slot()
 				.VAlign(VAlign_Center)
 				.Padding(2.0f)
 				.AutoHeight()
 				[
 					SNew(SHorizontalBox)
 
-					+SHorizontalBox::Slot()
+					+ SHorizontalBox::Slot()
 					.FillWidth(1.0f)
 					.VAlign(VAlign_Center)
 					[
@@ -112,7 +145,7 @@ void SStatsView::Construct(const FArguments& InArgs)
 						.Text(LOCTEXT("GroupByText", "Group by"))
 					]
 
-					+SHorizontalBox::Slot()
+					+ SHorizontalBox::Slot()
 					.FillWidth(2.0f)
 					.VAlign(VAlign_Center)
 					[
@@ -128,7 +161,7 @@ void SStatsView::Construct(const FArguments& InArgs)
 					]
 				]
 
-				// Check boxes for: GpuScope, ComputeScope, CpuScope
+				// Check boxes for: Int64, Float
 				+ SVerticalBox::Slot()
 				.VAlign(VAlign_Center)
 				.Padding(2.0f)
@@ -136,14 +169,14 @@ void SStatsView::Construct(const FArguments& InArgs)
 				[
 					SNew(SHorizontalBox)
 
-					+SHorizontalBox::Slot()
+					+ SHorizontalBox::Slot()
 					.Padding(FMargin(0.0f,0.0f,1.0f,0.0f))
 					.FillWidth(1.0f)
 					[
 						GetToggleButtonForStatsType(EStatsNodeType::Int64)
 					]
 
-					+SHorizontalBox::Slot()
+					+ SHorizontalBox::Slot()
 					.Padding(FMargin(1.0f,0.0f,1.0f,0.0f))
 					.FillWidth(1.0f)
 					[
@@ -154,7 +187,7 @@ void SStatsView::Construct(const FArguments& InArgs)
 		]
 
 		// Tree view
-		+SVerticalBox::Slot()
+		+ SVerticalBox::Slot()
 		.FillHeight(1.0f)
 		.Padding(0.0f, 6.0f, 0.0f, 0.0f)
 		[
@@ -482,10 +515,9 @@ void SStatsView::TreeViewHeaderRow_CreateColumnArgs(const int32 ColumnIndex)
 	ColumnArgs
 		.ColumnId(Column.Id)
 		.DefaultLabel(Column.ShortName)
-		.SortMode(EColumnSortMode::None)
 		.HAlignHeader(HAlign_Fill)
 		.VAlignHeader(VAlign_Fill)
-		.HeaderContentPadding(TOptional<FMargin>(2.0f))
+		.HeaderContentPadding(FMargin(2.0f))
 		.HAlignCell(HAlign_Fill)
 		.VAlignCell(VAlign_Fill)
 		.SortMode(this, &SStatsView::GetSortModeForColumn, Column.Id)
@@ -494,10 +526,8 @@ void SStatsView::TreeViewHeaderRow_CreateColumnArgs(const int32 ColumnIndex)
 		.FixedWidth(Column.bIsFixedColumnWidth() ? Column.InitialColumnWidth : TOptional<float>())
 		.HeaderContent()
 		[
-			SNew(SHorizontalBox)
+			SNew(SBox)
 			.ToolTip(SStatsViewTooltip::GetColumnTooltip(Column))
-
-			+ SHorizontalBox::Slot()
 			.HAlign(Column.HorizontalAlignment)
 			.VAlign(VAlign_Center)
 			[
@@ -628,7 +658,7 @@ void SStatsView::InsightsManager_OnSessionChanged()
 	if (NewSession != Session)
 	{
 		Session = NewSession;
-		RebuildTree();
+		Reset();
 	}
 	else
 	{
@@ -640,7 +670,6 @@ void SStatsView::InsightsManager_OnSessionChanged()
 
 void SStatsView::UpdateTree()
 {
-	// Create groups, sort timers within the group and apply filtering.
 	CreateGroups();
 	SortStats();
 	ApplyFiltering();
@@ -660,14 +689,16 @@ void SStatsView::ApplyFiltering()
 		GroupPtr->ClearFilteredChildren();
 		const bool bIsGroupVisible = Filters->PassesAllFilters(GroupPtr);
 
-		const TArray<FStatsNodePtr>& GroupChildren = GroupPtr->GetChildren();
+		const TArray<Insights::FBaseTreeNodePtr>& GroupChildren = GroupPtr->GetChildren();
 		const int32 NumChildren = GroupChildren.Num();
 		int32 NumVisibleChildren = 0;
 		for (int32 Cx = 0; Cx < NumChildren; ++Cx)
 		{
 			// Add a child.
-			const FStatsNodePtr& NodePtr = GroupChildren[Cx];
-			const bool bIsChildVisible = Filters->PassesAllFilters(NodePtr) && bStatsNodeIsVisible[static_cast<int>(NodePtr->GetType())];
+			const FStatsNodePtr& NodePtr = StaticCastSharedPtr<FStatsNode, Insights::FBaseTreeNode>(GroupChildren[Cx]);
+			const bool bIsChildVisible = (!bFilterOutZeroCountStats || NodePtr->GetAggregatedStats().Count > 0)
+									  && bStatsNodeIsVisible[static_cast<int>(NodePtr->GetType())]
+									  && Filters->PassesAllFilters(NodePtr);
 			if (bIsChildVisible)
 			{
 				GroupPtr->AddFilteredChild(NodePtr);
@@ -679,15 +710,15 @@ void SStatsView::ApplyFiltering()
 		{
 			// Add a group.
 			FilteredGroupNodes.Add(GroupPtr);
-			GroupPtr->bForceExpandGroupNode = true;
+			GroupPtr->SetExpansion(true);
 		}
 		else
 		{
-			GroupPtr->bForceExpandGroupNode = false;
+			GroupPtr->SetExpansion(false);
 		}
 	}
 
-	// Only expand timer nodes if we have a text filter.
+	// Only expand stats nodes if we have a text filter.
 	const bool bNonEmptyTextFilter = !TextFilter->GetRawFilterText().IsEmpty();
 	if (bNonEmptyTextFilter)
 	{
@@ -701,7 +732,7 @@ void SStatsView::ApplyFiltering()
 		for (int32 Fx = 0; Fx < FilteredGroupNodes.Num(); Fx++)
 		{
 			const FStatsNodePtr& GroupPtr = FilteredGroupNodes[Fx];
-			TreeView->SetItemExpansion(GroupPtr, GroupPtr->bForceExpandGroupNode);
+			TreeView->SetItemExpansion(GroupPtr, GroupPtr->IsExpanded());
 		}
 	}
 	else
@@ -751,7 +782,7 @@ TSharedRef<SWidget> SStatsView::GetToggleButtonForStatsType(const EStatsNodeType
 						.Image(StatsNodeTypeHelper::GetIconForStatsNodeType(NodeType))
 				]
 
-			+SHorizontalBox::Slot()
+			+ SHorizontalBox::Slot()
 				.Padding(2.0f, 0.0f, 0.0f, 0.0f)
 				.VAlign(VAlign_Center)
 				[
@@ -764,6 +795,21 @@ TSharedRef<SWidget> SStatsView::GetToggleButtonForStatsType(const EStatsNodeType
 
 ////////////////////////////////////////////////////////////////////////////////////////////////////
 
+void SStatsView::FilterOutZeroCountStats_OnCheckStateChanged(ECheckBoxState NewRadioState)
+{
+	bFilterOutZeroCountStats = (NewRadioState == ECheckBoxState::Checked);
+	ApplyFiltering();
+}
+
+////////////////////////////////////////////////////////////////////////////////////////////////////
+
+ECheckBoxState SStatsView::FilterOutZeroCountStats_IsChecked() const
+{
+	return bFilterOutZeroCountStats ? ECheckBoxState::Checked : ECheckBoxState::Unchecked;
+}
+
+////////////////////////////////////////////////////////////////////////////////////////////////////
+
 void SStatsView::FilterByStatsType_OnCheckStateChanged(ECheckBoxState NewRadioState, const EStatsNodeType InStatType)
 {
 	bStatsNodeIsVisible[static_cast<int>(InStatType)] = (NewRadioState == ECheckBoxState::Checked);
@@ -798,11 +844,11 @@ void SStatsView::TreeView_OnSelectionChanged(FStatsNodePtr SelectedItem, ESelect
 		TArray<FStatsNodePtr> SelectedItems = TreeView->GetSelectedItems();
 		if (SelectedItems.Num() == 1)
 		{
-			//HighlightedStatsName = SelectedItems[0]->GetName();
+			//HighlightedNodeName = SelectedItems[0]->GetName();
 		}
 		else
 		{
-			//HighlightedStatsName = NAME_None;
+			//HighlightedNodeName = NAME_None;
 		}
 	}
 }
@@ -811,7 +857,40 @@ void SStatsView::TreeView_OnSelectionChanged(FStatsNodePtr SelectedItem, ESelect
 
 void SStatsView::TreeView_OnGetChildren(FStatsNodePtr InParent, TArray<FStatsNodePtr>& OutChildren)
 {
-	OutChildren = InParent->GetFilteredChildren();
+	const TArray<Insights::FBaseTreeNodePtr>& Children = InParent->GetFilteredChildren();
+	OutChildren.Reset(Children.Num());
+	for (const Insights::FBaseTreeNodePtr& Child : Children)
+	{
+		OutChildren.Add(StaticCastSharedPtr<FStatsNode, Insights::FBaseTreeNode>(Child));
+	}
+}
+
+////////////////////////////////////////////////////////////////////////////////////////////////////
+
+void SStatsView::UpdateStatsNode(FStatsNodePtr StatsNode)
+{
+	bool bAddedToGraphFlag = false;
+
+	if (!StatsNode->IsGroup())
+	{
+		TSharedPtr<STimingProfilerWindow> Wnd = FTimingProfilerManager::Get()->GetProfilerWindow();
+		if (Wnd.IsValid())
+		{
+			TSharedPtr<STimingView> TimingView = Wnd->GetTimingView();
+			if (TimingView.IsValid())
+			{
+				TSharedPtr<FTimingGraphTrack> GraphTrack = TimingView->GetMainTimingGraphTrack();
+				if (GraphTrack.IsValid())
+				{
+					uint32 StatsCounterId = static_cast<uint32>(StatsNode->GetId());
+					TSharedPtr<FTimingGraphSeries> Series = GraphTrack->GetStatsCounterSeries(StatsCounterId);
+					bAddedToGraphFlag = Series.IsValid();
+				}
+			}
+		}
+	}
+
+	StatsNode->SetAddedToGraphFlag(bAddedToGraphFlag);
 }
 
 ////////////////////////////////////////////////////////////////////////////////////////////////////
@@ -820,16 +899,43 @@ void SStatsView::TreeView_OnMouseButtonDoubleClick(FStatsNodePtr StatsNode)
 {
 	if (!StatsNode->IsGroup())
 	{
-		//im:TODO: const bool bIsStatTracked = FTimingProfilerManager::Get()->IsStatTracked(StatsNode->GetId());
-	//	if (!bIsStatTracked)
+		TSharedPtr<STimingProfilerWindow> Wnd = FTimingProfilerManager::Get()->GetProfilerWindow();
+		if (Wnd.IsValid())
+		{
+			TSharedPtr<STimingView> TimingView = Wnd->GetTimingView();
+			if (TimingView.IsValid())
+			{
+				TSharedPtr<FTimingGraphTrack> GraphTrack = TimingView->GetMainTimingGraphTrack();
+				if (GraphTrack.IsValid())
+				{
+					uint32 StatsCounterId = static_cast<uint32>(StatsNode->GetId());
+					TSharedPtr<FTimingGraphSeries> Series = GraphTrack->GetStatsCounterSeries(StatsCounterId);
+					if (Series.IsValid())
+					{
+						GraphTrack->RemoveStatsCounterSeries(StatsCounterId);
+						GraphTrack->SetDirtyFlag();
+						StatsNode->SetAddedToGraphFlag(false);
+					}
+					else
+					{
+						GraphTrack->AddStatsCounterSeries(StatsCounterId, StatsNode->GetColor());
+						GraphTrack->SetDirtyFlag();
+						StatsNode->SetAddedToGraphFlag(true);
+					}
+				}
+			}
+		}
+
+		//im:TODO: const bool bIsTracked = FTimingProfilerManager::Get()->IsStatsCounterTracked(StatsNode->GetId());
+	//	if (!bIsTracked)
 	//	{
-	//		// Add a new graph.
-	//		FTimingProfilerManager::Get()->TrackStat(StatsNode->GetId());
+	//		// Add a new graph series.
+	//		FTimingProfilerManager::Get()->TrackStatsCounter(StatsNode->GetId());
 	//	}
 	//	else
 	//	{
-	//		// Remove a graph
-	//		FTimingProfilerManager::Get()->UntrackStat(StatsNode->GetId());
+	//		// Remove the corresponding graph series.
+	//		FTimingProfilerManager::Get()->UntrackStatsCounter(StatsNode->GetId());
 	//	}
 	}
 	else
@@ -852,7 +958,7 @@ TSharedRef<ITableRow> SStatsView::TreeView_OnGenerateRow(FStatsNodePtr StatsNode
 		.OnSetHoveredTableCell(this, &SStatsView::TableRow_SetHoveredTableCell)
 		.OnGetColumnOutlineHAlignmentDelegate(this, &SStatsView::TableRow_GetColumnOutlineHAlignment)
 		.HighlightText(this, &SStatsView::TableRow_GetHighlightText)
-		.HighlightedStatsName(this, &SStatsView::TableRow_GetHighlightedStatsName)
+		.HighlightedNodeName(this, &SStatsView::TableRow_GetHighlightedNodeName)
 		.StatsNodePtr(StatsNodePtr);
 
 	return TableRow;
@@ -914,9 +1020,9 @@ FText SStatsView::TableRow_GetHighlightText() const
 
 ////////////////////////////////////////////////////////////////////////////////////////////////////
 
-FName SStatsView::TableRow_GetHighlightedStatsName() const
+FName SStatsView::TableRow_GetHighlightedNodeName() const
 {
-	return HighlightedStatsName;
+	return HighlightedNodeName;
 }
 
 ////////////////////////////////////////////////////////////////////////////////////////////////////
@@ -1054,7 +1160,6 @@ void SStatsView::GroupBy_OnSelectionChanged(TSharedPtr<EStatsGroupingMode> NewGr
 	{
 		GroupingMode = *NewGroupingMode;
 
-		// Create groups, sort timers within the group and apply filtering.
 		CreateGroups();
 		SortStats();
 		ApplyFiltering();
@@ -1145,7 +1250,7 @@ void SStatsView::SetSortModeForColumn(const FName& ColumnId, const EColumnSortMo
 	ColumnBeingSorted = ColumnId;
 	ColumnSortMode = SortMode;
 
-	// Sort timers and apply filtering.
+	// Sort stats and apply filtering.
 	SortStats();
 	ApplyFiltering();
 }
@@ -1396,43 +1501,60 @@ void SStatsView::ContextMenu_ResetColumns_Execute()
 
 ////////////////////////////////////////////////////////////////////////////////////////////////////
 
+void SStatsView::Reset()
+{
+	StatsStartTime = 0.0;
+	StatsEndTime = 0.0;
+
+	RebuildTree(true);
+}
+
+////////////////////////////////////////////////////////////////////////////////////////////////////
+
 void SStatsView::RebuildTree(bool bResync)
 {
+	TArray<FStatsNodePtr> SelectedItems;
 	bool bListHasChanged = false;
 
 	if (bResync)
 	{
-		StatsNodes.Empty(StatsNodes.Num());
-		//StatsNodesMap.Empty(StatsNodesMap.Num());
-		StatsNodesIdMap.Empty(StatsNodesIdMap.Num());
+		const int32 PreviousNodeCount = StatsNodes.Num();
+		StatsNodes.Empty(PreviousNodeCount);
+		//StatsNodesMap.Empty(PreviousNodeCount);
+		StatsNodesIdMap.Empty(PreviousNodeCount);
 		bListHasChanged = true;
 	}
 
-	if (Session.IsValid() && Trace::ReadCounterProvider(*Session.Get()))
+	if (Session.IsValid())
 	{
 		Trace::FAnalysisSessionReadScope SessionReadScope(*Session.Get());
 
-		const Trace::ICounterProvider& CountersProvider = *Trace::ReadCounterProvider(*Session.Get());
+		const Trace::ICounterProvider& CountersProvider = Trace::ReadCounterProvider(*Session.Get());
 
-		if (!bResync)
+		if (CountersProvider.GetCounterCount() != StatsNodes.Num())
 		{
-			bResync = (CountersProvider.GetCounterCount() != StatsNodes.Num());
+			bResync = true;
 		}
 
 		if (bResync)
 		{
-			StatsNodes.Empty(StatsNodes.Num());
-			//StatsNodesMap.Empty(StatsNodesMap.Num());
-			StatsNodesIdMap.Empty(StatsNodesIdMap.Num());
+			// Save selection.
+			TreeView->GetSelectedItems(SelectedItems);
+
+			const int32 PreviousNodeCount = StatsNodes.Num();
+			StatsNodes.Empty(PreviousNodeCount);
+			//StatsNodesMap.Empty(PreviousNodeCount);
+			StatsNodesIdMap.Empty(PreviousNodeCount);
 			bListHasChanged = true;
 
 			CountersProvider.EnumerateCounters([this](const Trace::ICounter& Counter)
 			{
 				FName Name(Counter.GetName());
 				FName Group(Counter.GetDisplayHint() == Trace::CounterDisplayHint_Memory ? TEXT("Memory") :
-							Counter.GetDisplayHint() == Trace::CounterDisplayHint_FloatingPoint ? TEXT("float") : TEXT("int64"));
-				EStatsNodeType Type = Counter.GetDisplayHint() == Trace::CounterDisplayHint_FloatingPoint ? EStatsNodeType::Float : EStatsNodeType::Int64;
+							Counter.IsFloatingPoint() ? TEXT("float") : TEXT("int64"));
+				EStatsNodeType Type = Counter.IsFloatingPoint() ? EStatsNodeType::Float : EStatsNodeType::Int64;
 				FStatsNodePtr StatsNodePtr = MakeShareable(new FStatsNode(Counter.GetId(), Name, Group, Type));
+				UpdateStatsNode(StatsNodePtr);
 				StatsNodes.Add(StatsNodePtr);
 				//StatsNodesMap.Add(Name, StatsNodePtr);
 				StatsNodesIdMap.Add(Counter.GetId(), StatsNodePtr);
@@ -1444,6 +1566,28 @@ void SStatsView::RebuildTree(bool bResync)
 	{
 		UpdateTree();
 		UpdateStats(StatsStartTime, StatsEndTime);
+
+		TreeView->RebuildList();
+
+		// Restore selection.
+		if (SelectedItems.Num() > 0)
+		{
+			TreeView->ClearSelection();
+			TArray<FStatsNodePtr> NewSelectedItems;
+			for (const FStatsNodePtr& StatsNode : SelectedItems)
+			{
+				FStatsNodePtr* StatsNodePtrPtr = StatsNodesIdMap.Find(StatsNode->GetId());
+				if (StatsNodePtrPtr != nullptr)
+				{
+					NewSelectedItems.Add(*StatsNodePtrPtr);
+				}
+			}
+			if (NewSelectedItems.Num() > 0)
+			{
+				TreeView->SetItemSelection(NewSelectedItems, true);
+				TreeView->RequestScrollIntoView(NewSelectedItems[0]);
+			}
+		}
 	}
 }
 
@@ -1717,12 +1861,21 @@ void TTimeCalculationHelper<int64>::PostProcess(TMap<uint64, FStatsNodePtr>& Sta
 
 void SStatsView::UpdateStats(double StartTime, double EndTime)
 {
+	StatsStartTime = StartTime;
+	StatsEndTime = EndTime;
+
+	if (StartTime >= EndTime)
+	{
+		// keep previous aggregated stats
+		return;
+	}
+
 	for (const FStatsNodePtr& StatsNodePtr : StatsNodes)
 	{
 		StatsNodePtr->ResetAggregatedStats();
 	}
 
-	if (Session.IsValid() && StartTime < EndTime && Trace::ReadCounterProvider(*Session.Get()))
+	if (Session.IsValid())
 	{
 		const bool bComputeMedian = true;
 
@@ -1732,13 +1885,13 @@ void SStatsView::UpdateStats(double StartTime, double EndTime)
 		{
 			Trace::FAnalysisSessionReadScope SessionReadScope(*Session.Get());
 
-			const Trace::ICounterProvider& CountersProvider = *Trace::ReadCounterProvider(*Session.Get());
+			const Trace::ICounterProvider& CountersProvider = Trace::ReadCounterProvider(*Session.Get());
 
 			// Compute instance count and total/min/max inclusive/exclusive times for each counter.
 			// Iterate through all counters.
 			CountersProvider.EnumerateCounters([&CalculationHelperDbl, &CalculationHelperInt](const Trace::ICounter& Counter)
 			{
-				if (Counter.GetDisplayHint() == Trace::CounterDisplayHint_FloatingPoint)
+				if (Counter.IsFloatingPoint())
 				{
 					CalculationHelperDbl.Update(Counter);
 				}
@@ -1759,7 +1912,7 @@ void SStatsView::UpdateStats(double StartTime, double EndTime)
 				// Iterate again through all counters.
 				CountersProvider.EnumerateCounters([&CalculationHelperDbl, &CalculationHelperInt](const Trace::ICounter& Counter)
 				{
-					if (Counter.GetDisplayHint() == Trace::CounterDisplayHint_FloatingPoint)
+					if (Counter.IsFloatingPoint())
 					{
 						CalculationHelperDbl.UpdateHistograms(Counter);
 					}
@@ -1774,13 +1927,9 @@ void SStatsView::UpdateStats(double StartTime, double EndTime)
 		// Compute average and median inclusive/exclusive times.
 		CalculationHelperDbl.PostProcess(StatsNodesIdMap, bComputeMedian);
 		CalculationHelperInt.PostProcess(StatsNodesIdMap, bComputeMedian);
-
-		StatsStartTime = StartTime;
-		StatsEndTime = EndTime;
-
-		UpdateTree();
-		TreeView->RebuildList();
 	}
+
+	UpdateTree();
 }
 
 ////////////////////////////////////////////////////////////////////////////////////////////////////
diff --git a/Engine/Source/Developer/TraceInsights/Private/Insights/Widgets/SStatsView.h b/Engine/Source/Developer/TraceInsights/Private/Insights/Widgets/SStatsView.h
--- a/Engine/Source/Developer/TraceInsights/Private/Insights/Widgets/SStatsView.h
+++ b/Engine/Source/Developer/TraceInsights/Private/Insights/Widgets/SStatsView.h
@@ -80,11 +80,13 @@ public:
 	 */
 	void Construct(const FArguments& InArgs);
 
+	void Reset();
+
 	/**
 	 * Rebuilds the tree (if necessary).
 	 * @param bResync - If true, it forces a resync with list of stats counters from Analysis, even if the list did not changed since last sync.
 	 */
-	void RebuildTree(bool bResync = true);
+	void RebuildTree(bool bResync);
 	void UpdateStats(double StartTime, double EndTime);
 
 	void SelectStatsNode(uint64 Id);
@@ -153,7 +155,7 @@ protected:
 	EHorizontalAlignment TableRow_GetColumnOutlineHAlignment(const FName ColumnId) const;
 
 	FText TableRow_GetHighlightText() const;
-	FName TableRow_GetHighlightedStatsName() const;
+	FName TableRow_GetHighlightedNodeName() const;
 
 	bool TableRow_ShouldBeEnabled(const uint32 StatsId) const;
 
@@ -163,6 +165,9 @@ protected:
 	/** Populates the group and stat tree with items based on the current data. */
 	void ApplyFiltering();
 
+	void FilterOutZeroCountStats_OnCheckStateChanged(ECheckBoxState NewRadioState);
+	ECheckBoxState FilterOutZeroCountStats_IsChecked() const;
+
 	TSharedRef<SWidget> GetToggleButtonForStatsType(const EStatsNodeType StatType);
 	void FilterByStatsType_OnCheckStateChanged(ECheckBoxState NewRadioState, const EStatsNodeType InStatType);
 	ECheckBoxState FilterByStatsType_IsChecked(const EStatsNodeType InStatType) const;
@@ -240,6 +245,10 @@ protected:
 	bool ContextMenu_ResetColumns_CanExecute() const;
 	void ContextMenu_ResetColumns_Execute();
 
+	////////////////////////////////////////////////////////////////////////////////////////////////////
+
+	void UpdateStatsNode(FStatsNodePtr StatsNode);
+
 protected:
 	/** A weak pointer to the profiler session used to populate this widget. */
 	TSharedPtr<const Trace::IAnalysisSession>/*Weak*/ Session;
@@ -272,7 +281,7 @@ protected:
 	FStatsNodePtr HoveredStatsNodePtr;
 
 	/** Name of the stats that should be drawn as highlighted. */
-	FName HighlightedStatsName;
+	FName HighlightedNodeName;
 
 	//////////////////////////////////////////////////
 	// Stats Nodes
@@ -313,6 +322,9 @@ protected:
 	/** Holds the visibility of each stats type. */
 	bool bStatsNodeIsVisible[static_cast<int>(EStatsNodeType::InvalidOrMax)];
 
+	/** Filter out the stats counters having zero total instance count (aggregated stats). */
+	bool bFilterOutZeroCountStats;
+
 	//////////////////////////////////////////////////
 	// Grouping
 
diff --git a/Engine/Source/Developer/TraceInsights/Private/Insights/Widgets/SStatsViewTooltip.cpp b/Engine/Source/Developer/TraceInsights/Private/Insights/Widgets/SStatsViewTooltip.cpp
--- a/Engine/Source/Developer/TraceInsights/Private/Insights/Widgets/SStatsViewTooltip.cpp
+++ b/Engine/Source/Developer/TraceInsights/Private/Insights/Widgets/SStatsViewTooltip.cpp
@@ -4,7 +4,7 @@
 
 #include "EditorStyleSet.h"
 #include "SlateOptMacros.h"
-#include "Widgets/Images/SImage.h"
+#include "TraceServices/Model/AnalysisSession.h"
 #include "Widgets/Layout/SGridPanel.h"
 #include "Widgets/Layout/SSeparator.h"
 #include "Widgets/SBoxPanel.h"
@@ -115,6 +115,7 @@ TSharedPtr<SToolTip> SStatsViewTooltip::GetTableCellTooltip(const TSharedPtr<FSt
 					// Id: [Id]
 					+SGridPanel::Slot(0, 0)
 					.Padding(2.0f)
+					//.HAlign(HAlign_Right)
 					[
 						SNew(STextBlock)
 						.Text(LOCTEXT("TT_Id", "Id:"))
@@ -122,6 +123,7 @@ TSharedPtr<SToolTip> SStatsViewTooltip::GetTableCellTooltip(const TSharedPtr<FSt
 					]
 					+SGridPanel::Slot(1, 0)
 					.Padding(2.0f)
+					.HAlign(HAlign_Left)
 					[
 						SNew(STextBlock)
 						.Text(FText::AsNumber(StatsNodePtr->GetId()))
@@ -131,6 +133,7 @@ TSharedPtr<SToolTip> SStatsViewTooltip::GetTableCellTooltip(const TSharedPtr<FSt
 					// Name: [Name]
 					+SGridPanel::Slot(0, 1)
 					.Padding(2.0f)
+					//.HAlign(HAlign_Right)
 					[
 						SNew(STextBlock)
 						.Text(LOCTEXT("TT_Name", "Name:"))
@@ -138,6 +141,7 @@ TSharedPtr<SToolTip> SStatsViewTooltip::GetTableCellTooltip(const TSharedPtr<FSt
 					]
 					+SGridPanel::Slot(1, 1)
 					.Padding(2.0f)
+					.HAlign(HAlign_Left)
 					[
 						SNew(STextBlock)
 						.Text(FText::FromName(StatsNodePtr->GetName()))
@@ -147,6 +151,7 @@ TSharedPtr<SToolTip> SStatsViewTooltip::GetTableCellTooltip(const TSharedPtr<FSt
 					//// Group: [MetaGroupName]
 					//+SGridPanel::Slot(0, 2)
 					//.Padding(2.0f)
+					////.HAlign(HAlign_Right)
 					//[
 					//	SNew(STextBlock)
 					//	.Text(LOCTEXT("TT_Group", "Group:"))
@@ -154,6 +159,7 @@ TSharedPtr<SToolTip> SStatsViewTooltip::GetTableCellTooltip(const TSharedPtr<FSt
 					//]
 					//+SGridPanel::Slot(1, 2)
 					//.Padding(2.0f)
+					//.HAlign(HAlign_Left)
 					//[
 					//	SNew(STextBlock)
 					//	.Text(FText::FromName(StatsNodePtr->GetMetaGroupName()))
@@ -163,6 +169,7 @@ TSharedPtr<SToolTip> SStatsViewTooltip::GetTableCellTooltip(const TSharedPtr<FSt
 					// Type: [Type]
 					+ SGridPanel::Slot(0, 3)
 					.Padding(2.0f)
+					//.HAlign(HAlign_Right)
 					[
 						SNew(STextBlock)
 						.Text(LOCTEXT("TT_Type", "Type:"))
@@ -170,6 +177,7 @@ TSharedPtr<SToolTip> SStatsViewTooltip::GetTableCellTooltip(const TSharedPtr<FSt
 					]
 					+ SGridPanel::Slot(1, 3)
 					.Padding(2.0f)
+					.HAlign(HAlign_Left)
 					[
 						SNew(STextBlock)
 						.Text(StatsNodeTypeHelper::ToName(StatsNodePtr->GetType()))
@@ -193,6 +201,7 @@ TSharedPtr<SToolTip> SStatsViewTooltip::GetTableCellTooltip(const TSharedPtr<FSt
 
 					+ SGridPanel::Slot(0, 0)
 					.Padding(2.0f)
+					//.HAlign(HAlign_Right)
 					[
 						SNew(STextBlock)
 						.Text(LOCTEXT("TT_NumInstances", "Num Instances:"))
@@ -200,6 +209,7 @@ TSharedPtr<SToolTip> SStatsViewTooltip::GetTableCellTooltip(const TSharedPtr<FSt
 					]
 					+ SGridPanel::Slot(1, 0)
 					.Padding(2.0f)
+					.HAlign(HAlign_Left)
 					[
 						SNew(STextBlock)
 						.Text(InstanceCountText)
@@ -501,6 +511,7 @@ void SStatsViewTooltip::AddAggregatedStatsRow(TSharedPtr<SGridPanel> Grid, int32
 {
 	Grid->AddSlot(0, Row)
 		.Padding(2.0f)
+		//.HAlign(HAlign_Right)
 		[
 			SNew(STextBlock)
 			.Text(Name)
diff --git a/Engine/Source/Developer/TraceInsights/Private/Insights/Widgets/SStatsViewTooltip.h b/Engine/Source/Developer/TraceInsights/Private/Insights/Widgets/SStatsViewTooltip.h
--- a/Engine/Source/Developer/TraceInsights/Private/Insights/Widgets/SStatsViewTooltip.h
+++ b/Engine/Source/Developer/TraceInsights/Private/Insights/Widgets/SStatsViewTooltip.h
@@ -3,18 +3,15 @@
 #pragma once
 
 #include "CoreMinimal.h"
-#include "EditorStyleSet.h"
-#include "SlateOptMacros.h"
-#include "TraceServices/AnalysisService.h"
-#include "Widgets/Images/SImage.h"
-#include "Widgets/Layout/SGridPanel.h"
-#include "Widgets/Layout/SSeparator.h"
-#include "Widgets/SToolTip.h"
-#include "Widgets/Text/STextBlock.h"
 
 // Insights
 #include "Insights/InsightsManager.h"
 
+class SToolTip;
+class SGridPanel;
+
+namespace Trace { class IAnalysisSession; }
+
 // Insights
 class FStatsNode;
 class FStatsViewColumn;
diff --git a/Engine/Source/Developer/TraceInsights/Private/Insights/Widgets/STimerTableCell.cpp b/Engine/Source/Developer/TraceInsights/Private/Insights/Widgets/STimerTableCell.cpp
--- a/Engine/Source/Developer/TraceInsights/Private/Insights/Widgets/STimerTableCell.cpp
+++ b/Engine/Source/Developer/TraceInsights/Private/Insights/Widgets/STimerTableCell.cpp
@@ -7,12 +7,14 @@
 #include "Widgets/Images/SImage.h"
 #include "Widgets/Input/SButton.h"
 #include "Widgets/SBoxPanel.h"
+#include "Widgets/SOverlay.h"
+#include "Widgets/SToolTip.h"
 #include "Widgets/Text/STextBlock.h"
 #include "Widgets/Views/SExpanderArrow.h"
 
 // Insights
-#include "Insights/ViewModels/TimersViewColumn.h"
-#include "Insights/ViewModels/TimersViewColumnFactory.h"
+#include "Insights/Table/ViewModels/Table.h"
+#include "Insights/Table/ViewModels/TableColumn.h"
 #include "Insights/Widgets/STimersViewTooltip.h"
 
 #define LOCTEXT_NAMESPACE "STimersView"
@@ -23,10 +25,15 @@ BEGIN_SLATE_FUNCTION_BUILD_OPTIMIZATION
 
 void STimerTableCell::Construct(const FArguments& InArgs, const TSharedRef<class ITableRow>& TableRow)
 {
-	SetHoveredTableCellDelegate = InArgs._OnSetHoveredTableCell;
+	TablePtr = InArgs._TablePtr;
+	ColumnPtr = InArgs._ColumnPtr;
 	TimerNodePtr = InArgs._TimerNodePtr;
-	ColumnId = InArgs._ColumnId;
-	//HighlightText = InArgs._HighlightText;
+
+	ensure(TablePtr.IsValid());
+	ensure(ColumnPtr.IsValid());
+	ensure(TimerNodePtr.IsValid());
+
+	SetHoveredCellDelegate = InArgs._OnSetHoveredCell;
 
 	ChildSlot
 		[
@@ -38,7 +45,7 @@ void STimerTableCell::Construct(const FArguments& InArgs, const TSharedRef<class
 
 TSharedRef<SWidget> STimerTableCell::GenerateWidgetForColumn(const FArguments& InArgs, const TSharedRef<class ITableRow>& TableRow)
 {
-	if (InArgs._IsTimerNameColumn)
+	if (InArgs._IsNameColumn)
 	{
 		return GenerateWidgetForNameColumn(InArgs, TableRow);
 	}
@@ -52,8 +59,6 @@ TSharedRef<SWidget> STimerTableCell::GenerateWidgetForColumn(const FArguments& I
 
 TSharedRef<SWidget> STimerTableCell::GenerateWidgetForNameColumn(const FArguments& InArgs, const TSharedRef<class ITableRow>& TableRow)
 {
-	const FTimersViewColumn& Column = *FTimersViewColumnFactory::Get().ColumnIdToPtrMapping.FindChecked(ColumnId);
-
 	return
 		SNew(SHorizontalBox)
 
@@ -65,78 +70,71 @@ TSharedRef<SWidget> STimerTableCell::GenerateWidgetForNameColumn(const FArgument
 			SNew(SExpanderArrow, TableRow)
 		]
 
-		// Event info icon + tooltip.
 		+ SHorizontalBox::Slot()
 		.AutoWidth()
 		.HAlign(HAlign_Center)
 		.VAlign(VAlign_Center)
 		[
-			SNew(SImage)
-			.Visibility(this, &STimerTableCell::GetHintIconVisibility)
-			.Image(FEditorStyle::GetBrush("Profiler.Tooltip.HintIcon10"))
-			.ToolTip(STimersViewTooltip::GetTableCellTooltip(TimerNodePtr))
+			SNew(SOverlay)
+
+			// Hot path icon
+			+ SOverlay::Slot()
+			[
+				SNew(SImage)
+				.Visibility(this, &STimerTableCell::GetHotPathIconVisibility)
+				.Image(FEditorStyle::GetBrush(TEXT("Profiler.EventGraph.HotPathSmall")))
+			]
+
+			// Info icon + tooltip
+			+ SOverlay::Slot()
+			[
+				SNew(SImage)
+				.Visibility(this, &STimerTableCell::GetHintIconVisibility)
+				.Image(FEditorStyle::GetBrush("Profiler.Tooltip.HintIcon10"))
+				.ToolTip(STimersViewTooltip::GetCellTooltip(TimerNodePtr, ColumnPtr))
+			]
 		]
 
 		// Name
 		+ SHorizontalBox::Slot()
 		.AutoWidth()
 		.VAlign(VAlign_Center)
-		.HAlign(Column.HorizontalAlignment)
+		.HAlign(ColumnPtr->GetHorizontalAlignment())
 		.Padding(FMargin(2.0f, 0.0f))
 		[
 			SNew(STextBlock)
-			//.Text(TimerNodePtr->GetNameEx())
-			.Text(this, &STimerTableCell::GetNameEx)
+			.Text(this, &STimerTableCell::GetDisplayName)
 			.HighlightText(InArgs._HighlightText)
 			.TextStyle(FEditorStyle::Get(), TEXT("Profiler.Tooltip"))
 			.ColorAndOpacity(this, &STimerTableCell::GetColorAndOpacity)
 			.ShadowColorAndOpacity(this, &STimerTableCell::GetShadowColorAndOpacity)
 		]
-
-		/*
-		// Culled children warning icon + tooltip.
-		+ SHorizontalBox::Slot()
-		.AutoWidth()
-		.HAlign(HAlign_Center)
-		.VAlign(VAlign_Center)
-		[
-			SNew(SButton)
-			.ButtonStyle(FEditorStyle::Get(), TEXT("HoverHintOnly"))
-			.ContentPadding(0.0f)
-			.IsFocusable(false)
-			//.OnClicked(this, &STimerTableCell::ExpandCulledEvents_OnClicked)
-			[
-				SNew(SImage)
-				.Visibility(this, &STimerTableCell::GetCulledEventsIconVisibility)
-				.Image(FEditorStyle::GetBrush("Profiler.EventGraph.HasCulledEventsSmall"))
-				.ToolTipText(LOCTEXT("HasCulledEvents_TT", "This event contains culled children, if you want to see all children, please disable culling or use function details, or press this icon"))
-			]
-		]*/
-		;
+	;
 }
 
 ////////////////////////////////////////////////////////////////////////////////////////////////////
 
 TSharedRef<SWidget> STimerTableCell::GenerateWidgetForStatsColumn(const FArguments& InArgs, const TSharedRef<class ITableRow>& TableRow)
 {
-	const FTimersViewColumn& Column = *FTimersViewColumnFactory::Get().ColumnIdToPtrMapping.FindChecked(ColumnId);
-	const FText FormattedValue = Column.GetFormattedValue(*TimerNodePtr);
+	const FText CellText = ColumnPtr->GetValueAsText(*TimerNodePtr);
 
 	return
 		SNew(SHorizontalBox)
 
+		// Value
 		+ SHorizontalBox::Slot()
 		.FillWidth(1.0f)
 		.VAlign(VAlign_Center)
-		.HAlign(Column.HorizontalAlignment)
+		.HAlign(ColumnPtr->GetHorizontalAlignment())
 		.Padding(FMargin(2.0f, 0.0f))
 		[
 			SNew(STextBlock)
-			.Text(FormattedValue)
+			.Text(CellText)
 			.TextStyle(FEditorStyle::Get(), TEXT("Profiler.Tooltip"))
 			.ColorAndOpacity(this, &STimerTableCell::GetStatsColorAndOpacity)
 			.ShadowColorAndOpacity(this, &STimerTableCell::GetShadowColorAndOpacity)
-		];
+		]
+	;
 }
 
 END_SLATE_FUNCTION_BUILD_OPTIMIZATION
diff --git a/Engine/Source/Developer/TraceInsights/Private/Insights/Widgets/STimerTableCell.h b/Engine/Source/Developer/TraceInsights/Private/Insights/Widgets/STimerTableCell.h
--- a/Engine/Source/Developer/TraceInsights/Private/Insights/Widgets/STimerTableCell.h
+++ b/Engine/Source/Developer/TraceInsights/Private/Insights/Widgets/STimerTableCell.h
@@ -9,19 +9,24 @@
 // Insights
 #include "Insights/ViewModels/TimerNode.h"
 
-DECLARE_DELEGATE_TwoParams(FSetHoveredTimerTableCell, const FName /*ColumnId*/, const FTimerNodePtr /*SamplePtr*/);
-DECLARE_DELEGATE_RetVal_OneParam(bool, FIsColumnVisibleDelegate, const FName /*ColumnId*/);
-DECLARE_DELEGATE_RetVal_OneParam(EHorizontalAlignment, FGetColumnOutlineHAlignmentDelegate, const FName /*ColumnId*/);
+namespace Insights
+{
+	class FTable;
+	class FTableColumn;
+}
+
+DECLARE_DELEGATE_ThreeParams(FSetHoveredTimerTableCell, TSharedPtr<Insights::FTable> /*TablePtr*/, TSharedPtr<Insights::FTableColumn> /*ColumnPtr*/, const FTimerNodePtr /*TimerNodePtr*/);
 
 class STimerTableCell : public SCompoundWidget
 {
 public:
 	SLATE_BEGIN_ARGS(STimerTableCell) {}
-		SLATE_EVENT(FSetHoveredTimerTableCell, OnSetHoveredTableCell)
+		SLATE_EVENT(FSetHoveredTimerTableCell, OnSetHoveredCell)
 		SLATE_ATTRIBUTE(FText, HighlightText)
+		SLATE_ARGUMENT(TSharedPtr<Insights::FTable>, TablePtr)
+		SLATE_ARGUMENT(TSharedPtr<Insights::FTableColumn>, ColumnPtr)
 		SLATE_ARGUMENT(FTimerNodePtr, TimerNodePtr)
-		SLATE_ARGUMENT(FName, ColumnId)
-		SLATE_ARGUMENT(bool, IsTimerNameColumn)
+		SLATE_ARGUMENT(bool, IsNameColumn)
 	SLATE_END_ARGS()
 
 	void Construct(const FArguments& InArgs, const TSharedRef<class ITableRow>& TableRow);
@@ -40,7 +45,7 @@ protected:
 	virtual void OnMouseEnter(const FGeometry& MyGeometry, const FPointerEvent& MouseEvent) override
 	{
 		SCompoundWidget::OnMouseEnter(MyGeometry, MouseEvent);
-		SetHoveredTableCellDelegate.ExecuteIfBound(ColumnId, TimerNodePtr);
+		SetHoveredCellDelegate.ExecuteIfBound(TablePtr, ColumnPtr, TimerNodePtr);
 	}
 
 	/**
@@ -51,7 +56,7 @@ protected:
 	virtual void OnMouseLeave(const FPointerEvent& MouseEvent) override
 	{
 		SCompoundWidget::OnMouseLeave(MouseEvent);
-		SetHoveredTableCellDelegate.ExecuteIfBound(NAME_None, nullptr);
+		SetHoveredCellDelegate.ExecuteIfBound(nullptr, nullptr, nullptr);
 	}
 
 	/**
@@ -74,7 +79,7 @@ protected:
 	virtual void OnDragEnter(const FGeometry& MyGeometry, const FDragDropEvent& DragDropEvent) override
 	{
 		SCompoundWidget::OnDragEnter(MyGeometry, DragDropEvent);
-		SetHoveredTableCellDelegate.ExecuteIfBound(ColumnId, TimerNodePtr);
+		SetHoveredCellDelegate.ExecuteIfBound(TablePtr, ColumnPtr, TimerNodePtr);
 	}
 
 	/**
@@ -85,56 +90,62 @@ protected:
 	virtual void OnDragLeave(const FDragDropEvent& DragDropEvent)  override
 	{
 		SCompoundWidget::OnDragLeave(DragDropEvent);
-		SetHoveredTableCellDelegate.ExecuteIfBound(NAME_None, nullptr);
+		SetHoveredCellDelegate.ExecuteIfBound(nullptr, nullptr, nullptr);
 	}
 
-	EVisibility GetHintIconVisibility() const
+	EVisibility GetHotPathIconVisibility() const
 	{
-		return IsHovered() ? EVisibility::Visible : EVisibility::Hidden;
+		return TimerNodePtr->IsHotPath() && !IsHovered() ? EVisibility::Visible : EVisibility::Collapsed;
 	}
 
-	EVisibility GetCulledEventsIconVisibility() const
+	EVisibility GetHintIconVisibility() const
 	{
-		//return TimerNodePtr->HasCulledChildren() ? EVisibility::Visible : EVisibility::Collapsed;
-		return EVisibility::Collapsed;
+		return IsHovered() ? EVisibility::Visible : EVisibility::Hidden;
 	}
 
-	FText GetNameEx() const
+	FText GetDisplayName() const
 	{
-		return TimerNodePtr->GetNameEx();
+		return TimerNodePtr->GetDisplayName();
 	}
 
 	FSlateColor GetColorAndOpacity() const
 	{
 		const FLinearColor TextColor =
-			TimerNodePtr->IsFiltered() ? FLinearColor(1.0f, 1.0f, 1.0f, 0.5f) :
-			FLinearColor::White;
+			TimerNodePtr->IsFiltered() ?
+				FLinearColor(1.0f, 1.0f, 1.0f, 0.5f) :
+				FLinearColor(1.0f, 1.0f, 1.0f, 1.0f);
 		return TextColor;
 	}
 
 	FSlateColor GetStatsColorAndOpacity() const
 	{
 		const FLinearColor TextColor =
-			TimerNodePtr->IsFiltered() ? FLinearColor(1.0f, 1.0f, 1.0f, 0.5f) :
-			TimerNodePtr->GetAggregatedStats().InstanceCount == 0 ? FLinearColor(1.0f, 1.0f, 1.0f, 0.6f) :
-			FLinearColor::White;
+			TimerNodePtr->IsFiltered() ?
+				FLinearColor(1.0f, 1.0f, 1.0f, 0.5f) :
+			TimerNodePtr->GetAggregatedStats().InstanceCount == 0 ?
+				FLinearColor(1.0f, 1.0f, 1.0f, 0.6f) :
+				FLinearColor(1.0f, 1.0f, 1.0f, 1.0f);
 		return TextColor;
 	}
 
 	FLinearColor GetShadowColorAndOpacity() const
 	{
 		const FLinearColor ShadowColor =
-			TimerNodePtr->IsFiltered() ? FLinearColor(0.f, 0.f, 0.f, 0.25f) :
-			FLinearColor(0.0f, 0.0f, 0.0f, 0.5f);
+			TimerNodePtr->IsFiltered() ?
+				FLinearColor(0.f, 0.f, 0.f, 0.25f) :
+				FLinearColor(0.0f, 0.0f, 0.0f, 0.5f);
 		return ShadowColor;
 	}
 
 protected:
-	/** A shared pointer to the timer node. */
-	FTimerNodePtr TimerNodePtr;
+	/** A shared pointer to the table view model. */
+	TSharedPtr<Insights::FTable> TablePtr; // TODO: TSharedRef
 
-	/** The Id of the column where this timer belongs. */
-	FName ColumnId;
+	/** A shared pointer to the table column view model. */
+	TSharedPtr<Insights::FTableColumn> ColumnPtr; // TODO: TSharedRef
+
+	/** A shared pointer to the timer node. */
+	FTimerNodePtr TimerNodePtr; // TODO: TSharedRef
 
-	FSetHoveredTimerTableCell SetHoveredTableCellDelegate;
+	FSetHoveredTimerTableCell SetHoveredCellDelegate;
 };
diff --git a/Engine/Source/Developer/TraceInsights/Private/Insights/Widgets/STimerTableRow.cpp b/Engine/Source/Developer/TraceInsights/Private/Insights/Widgets/STimerTableRow.cpp
--- a/Engine/Source/Developer/TraceInsights/Private/Insights/Widgets/STimerTableRow.cpp
+++ b/Engine/Source/Developer/TraceInsights/Private/Insights/Widgets/STimerTableRow.cpp
@@ -2,130 +2,84 @@
 
 #include "STimerTableRow.h"
 
+#include "EditorStyleSet.h"
+#include "SlateOptMacros.h"
 #include "Widgets/SOverlay.h"
 #include "Widgets/Images/SImage.h"
 
 // Insights
 #include "Insights/Common/TimeUtils.h"
-#include "Insights/ViewModels/TimersViewColumnFactory.h"
-#include "Insights/Widgets/STimersViewTooltip.h"
+#include "Insights/Table/ViewModels/Table.h"
+#include "Insights/Table/ViewModels/TableColumn.h"
 #include "Insights/Widgets/STimerTableCell.h"
 
 #define LOCTEXT_NAMESPACE "STimersView"
 
+////////////////////////////////////////////////////////////////////////////////////////////////////
+
 BEGIN_SLATE_FUNCTION_BUILD_OPTIMIZATION
 
 void STimerTableRow::Construct(const FArguments& InArgs, const TSharedRef<STableViewBase>& InOwnerTableView)
 {
 	OnShouldBeEnabled = InArgs._OnShouldBeEnabled;
 	IsColumnVisibleDelegate = InArgs._OnIsColumnVisible;
-	SetHoveredTableCellDelegate = InArgs._OnSetHoveredTableCell;
+	SetHoveredCellDelegate = InArgs._OnSetHoveredCell;
 	GetColumnOutlineHAlignmentDelegate = InArgs._OnGetColumnOutlineHAlignmentDelegate;
 
 	HighlightText = InArgs._HighlightText;
-	HighlightedTimerName = InArgs._HighlightedTimerName;
+	HighlightedNodeName = InArgs._HighlightedNodeName;
 
+	TablePtr = InArgs._TablePtr;
 	TimerNodePtr = InArgs._TimerNodePtr;
 
 	SetEnabled(TAttribute<bool>(this, &STimerTableRow::HandleShouldBeEnabled));
 
 	SMultiColumnTableRow<FTimerNodePtr>::Construct(SMultiColumnTableRow<FTimerNodePtr>::FArguments(), InOwnerTableView);
-
-	/*
-	const FSlateBrush* const NodeIcon = TimerNodeTypeHelper::GetIconForTimerNodeType(InTimerNode->GetType());
-	const TSharedRef<SToolTip> Tooltip = InTimerNode->IsGroup() ? SNew(SToolTip) : STimersViewTooltip(InTimerNode->GetId()).GetTooltip();
-
-	ChildSlot
-		[
-			SNew(SHorizontalBox)
-
-			// Expander arrow.
-			+ SHorizontalBox::Slot()
-			.AutoWidth()
-			.HAlign(HAlign_Right)
-			.VAlign(VAlign_Center)
-			[
-				SNew(SExpanderArrow, SharedThis(this))
-			]
-
-			// Icon to visualize group or timer type.
-			+ SHorizontalBox::Slot()
-			.AutoWidth()
-			.Padding(0.0f, 0.0f, 8.0f, 0.0f)
-			[
-				SNew(SImage)
-				.Image(NodeIcon)
-				.ToolTip(Tooltip)
-			]
-
-			// Description text.
-			+ SHorizontalBox::Slot()
-			.AutoWidth()
-			.VAlign(VAlign_Center)
-			.HAlign(HAlign_Left)
-			.Padding(FMargin(2.0f, 0.0f))
-			[
-				SNew(STextBlock)
-				.Text(this, &STimerTableRow::GetText)
-				.HighlightText(InArgs._HighlightText)
-				.TextStyle(FEditorStyle::Get(), TEXT("Profiler.Tooltip"))
-				.ColorAndOpacity(this, &STimerTableRow::GetColorAndOpacity)
-			]
-
-			+ SHorizontalBox::Slot()
-			.AutoWidth()
-			.HAlign(HAlign_Center)
-			.VAlign(VAlign_Top)
-			.Padding(0.0f, 1.0f, 0.0f, 0.0f)
-			[
-				SNew(SImage)
-				.Visibility(!InTimerNode->IsGroup() ? EVisibility::Visible : EVisibility::Collapsed)
-				.Image(FEditorStyle::GetBrush("Profiler.Tooltip.HintIcon10"))
-				.ToolTip(Tooltip)
-			]
-		];
-
-	STableRow<FTimerNodePtr>::ConstructInternal(STableRow::FArguments().ShowSelection(true), InOwnerTableView);
-	*/
 }
 
+////////////////////////////////////////////////////////////////////////////////////////////////////
+
 TSharedRef<SWidget> STimerTableRow::GenerateWidgetForColumn(const FName& ColumnId)
 {
+	TSharedPtr<Insights::FTableColumn> ColumnPtr = TablePtr->FindColumnChecked(ColumnId);
+
 	return
+		SNew(SOverlay)
+		.Visibility(EVisibility::SelfHitTestInvisible)
+
+		+SOverlay::Slot()
+		.Padding(0.0f)
+		[
+			SNew(SImage)
+			.Image(FEditorStyle::GetBrush("Profiler.LineGraphArea"))
+			.ColorAndOpacity(this, &STimerTableRow::GetBackgroundColorAndOpacity)
+		]
 
-	SNew(SOverlay)
-	.Visibility(EVisibility::SelfHitTestInvisible)
-
-	+SOverlay::Slot()
-	.Padding(0.0f)
-	[
-		SNew(SImage)
-		.Image(FEditorStyle::GetBrush("Profiler.LineGraphArea"))
-		.ColorAndOpacity(this, &STimerTableRow::GetBackgroundColorAndOpacity)
-	]
-
-	+SOverlay::Slot()
-	.Padding(0.0f)
-	[
-		SNew(SImage)
-		.Image(this, &STimerTableRow::GetOutlineBrush, ColumnId)
-		.ColorAndOpacity(this, &STimerTableRow::GetOutlineColorAndOpacity)
-	]
-
-	+SOverlay::Slot()
-	[
-		SNew(STimerTableCell, SharedThis(this))
-		.Visibility(this, &STimerTableRow::IsColumnVisible, ColumnId)
-		.TimerNodePtr(TimerNodePtr)
-		.ColumnId(ColumnId)
-		.HighlightText(HighlightText)
-		.IsTimerNameColumn(ColumnId == FTimersViewColumnFactory::Get().Collection[0]->Id) // name column
-		.OnSetHoveredTableCell(this, &STimerTableRow::OnSetHoveredTableCell)
-	];
+		+SOverlay::Slot()
+		.Padding(0.0f)
+		[
+			SNew(SImage)
+			.Image(this, &STimerTableRow::GetOutlineBrush, ColumnId)
+			.ColorAndOpacity(this, &STimerTableRow::GetOutlineColorAndOpacity)
+		]
+
+		+SOverlay::Slot()
+		[
+			SNew(STimerTableCell, SharedThis(this))
+			.Visibility(this, &STimerTableRow::IsColumnVisible, ColumnId)
+			.TablePtr(TablePtr)
+			.ColumnPtr(ColumnPtr)
+			.TimerNodePtr(TimerNodePtr)
+			.HighlightText(HighlightText)
+			.IsNameColumn(ColumnPtr->IsHierarchy())
+			.OnSetHoveredCell(this, &STimerTableRow::OnSetHoveredCell)
+		];
 }
 
 END_SLATE_FUNCTION_BUILD_OPTIMIZATION
 
+////////////////////////////////////////////////////////////////////////////////////////////////////
+
 FReply STimerTableRow::OnDragDetected(const FGeometry& MyGeometry, const FPointerEvent& MouseEvent)
 {
 	//im:TODO
@@ -155,37 +109,7 @@ FReply STimerTableRow::OnDragDetected(const FGeometry& MyGeometry, const FPointe
 	return SMultiColumnTableRow<FTimerNodePtr>::OnDragDetected(MyGeometry, MouseEvent);
 }
 
-FText STimerTableRow::GetText() const
-{
-	FText Text = FText::GetEmpty();
-
-	if (TimerNodePtr->IsGroup())
-	{
-		Text = FText::Format(LOCTEXT("TimerNode_GroupNodeTextFmt", "{0} ({1})"), FText::FromName(TimerNodePtr->GetName()), FText::AsNumber(TimerNodePtr->GetChildren().Num()));
-	}
-	else
-	{
-		Text = FText::FromName(TimerNodePtr->GetName());
-	}
-
-	return Text;
-}
-
-FSlateFontInfo STimerTableRow::GetFont() const
-{
-	const bool bIsStatTracked = false;//im:TODO: FTimingProfilerManager::Get()->IsStatTracked(TimerNodePtr->GetTimerId());
-	const FSlateFontInfo FontInfo = bIsStatTracked ? FEditorStyle::GetFontStyle("BoldFont") : FEditorStyle::GetFontStyle("NormalFont");
-	return FontInfo;
-}
-
-FSlateColor STimerTableRow::GetColorAndOpacity() const
-{
-	//const bool bIsStatTracked = FTimingProfilerManager::Get()->IsStatTracked(TimerNodePtr->GetId());
-	//const FSlateColor Color = bIsStatTracked ? FTimingProfilerManager::Get()->GetColorForTimerId(TimerNodePtr->GetId()) : FLinearColor::White;
-	//return Color;
-	//FTimingProfilerManager::GetSettings().GetColorForStat(StatName)
-	return FLinearColor::White;
-}
+////////////////////////////////////////////////////////////////////////////////////////////////////
 
 FSlateColor STimerTableRow::GetBackgroundColorAndOpacity() const
 {
@@ -193,6 +117,8 @@ FSlateColor STimerTableRow::GetBackgroundColorAndOpacity() const
 	//return FLinearColor(0.0f, 0.0f, 0.0f, 1.0f)
 }
 
+////////////////////////////////////////////////////////////////////////////////////////////////////
+
 FSlateColor STimerTableRow::GetBackgroundColorAndOpacity(double Time) const
 {
 	const FLinearColor Color =	Time > TimeUtils::Second      ? FLinearColor(0.3f, 0.0f, 0.0f, 1.0f) :
@@ -202,14 +128,18 @@ FSlateColor STimerTableRow::GetBackgroundColorAndOpacity(double Time) const
 	return Color;
 }
 
+////////////////////////////////////////////////////////////////////////////////////////////////////
+
 FSlateColor STimerTableRow::GetOutlineColorAndOpacity() const
 {
 	const FLinearColor NoColor(0.0f, 0.0f, 0.0f, 0.0f);
-	const bool bShouldBeHighlighted = TimerNodePtr->GetName() == HighlightedTimerName.Get();
+	const bool bShouldBeHighlighted = TimerNodePtr->GetName() == HighlightedNodeName.Get();
 	const FLinearColor OutlineColorAndOpacity = bShouldBeHighlighted ? FLinearColor(FColorList::SlateBlue) : NoColor;
 	return OutlineColorAndOpacity;
 }
 
+////////////////////////////////////////////////////////////////////////////////////////////////////
+
 const FSlateBrush* STimerTableRow::GetOutlineBrush(const FName ColumnId) const
 {
 	EHorizontalAlignment Result = HAlign_Center;
@@ -234,6 +164,8 @@ const FSlateBrush* STimerTableRow::GetOutlineBrush(const FName ColumnId) const
 	return Brush;
 }
 
+////////////////////////////////////////////////////////////////////////////////////////////////////
+
 bool STimerTableRow::HandleShouldBeEnabled() const
 {
 	bool bResult = false;
@@ -253,6 +185,8 @@ bool STimerTableRow::HandleShouldBeEnabled() const
 	return bResult;
 }
 
+////////////////////////////////////////////////////////////////////////////////////////////////////
+
 EVisibility STimerTableRow::IsColumnVisible(const FName ColumnId) const
 {
 	EVisibility Result = EVisibility::Collapsed;
@@ -265,9 +199,13 @@ EVisibility STimerTableRow::IsColumnVisible(const FName ColumnId) const
 	return Result;
 }
 
-void STimerTableRow::OnSetHoveredTableCell(const FName InColumnId, const FTimerNodePtr InSamplePtr)
+////////////////////////////////////////////////////////////////////////////////////////////////////
+
+void STimerTableRow::OnSetHoveredCell(TSharedPtr<Insights::FTable> InTablePtr, TSharedPtr<Insights::FTableColumn> InColumnPtr, const FTimerNodePtr InTimerNodePtr)
 {
-	SetHoveredTableCellDelegate.ExecuteIfBound(InColumnId, InSamplePtr);
+	SetHoveredCellDelegate.ExecuteIfBound(InTablePtr, InColumnPtr, InTimerNodePtr);
 }
 
+////////////////////////////////////////////////////////////////////////////////////////////////////
+
 #undef LOCTEXT_NAMESPACE
diff --git a/Engine/Source/Developer/TraceInsights/Private/Insights/Widgets/STimerTableRow.h b/Engine/Source/Developer/TraceInsights/Private/Insights/Widgets/STimerTableRow.h
--- a/Engine/Source/Developer/TraceInsights/Private/Insights/Widgets/STimerTableRow.h
+++ b/Engine/Source/Developer/TraceInsights/Private/Insights/Widgets/STimerTableRow.h
@@ -12,22 +12,29 @@
 // Insights
 #include "Insights/ViewModels/TimerNodeHelper.h"
 
+namespace Insights
+{
+	class FTable;
+	class FTableColumn;
+}
+
 DECLARE_DELEGATE_RetVal_OneParam(bool, FShouldBeEnabledDelegate, const uint32 /*TimerId*/);
 DECLARE_DELEGATE_RetVal_OneParam(bool, FIsColumnVisibleDelegate, const FName /*ColumnId*/);
-DECLARE_DELEGATE_TwoParams(FSetHoveredTimerTableCell, const FName /*ColumnId*/, const FTimerNodePtr /*TimerNodePtr*/);
 DECLARE_DELEGATE_RetVal_OneParam(EHorizontalAlignment, FGetColumnOutlineHAlignmentDelegate, const FName /*ColumnId*/);
+DECLARE_DELEGATE_ThreeParams(FSetHoveredTimerTableCell, TSharedPtr<Insights::FTable> /*TablePtr*/, TSharedPtr<Insights::FTableColumn> /*ColumnPtr*/, const FTimerNodePtr /*TimerNodePtr*/);
 
-/** Widget that represents a table row in the timers' tree control. Generates widgets for each column on demand. */
+/** Widget that represents a table row in a timer tree control. Generates widgets for each column on demand. */
 class STimerTableRow : public SMultiColumnTableRow<FTimerNodePtr>
 {
 public:
 	SLATE_BEGIN_ARGS(STimerTableRow) {}
 		SLATE_EVENT(FShouldBeEnabledDelegate, OnShouldBeEnabled)
 		SLATE_EVENT(FIsColumnVisibleDelegate, OnIsColumnVisible)
-		SLATE_EVENT(FSetHoveredTimerTableCell, OnSetHoveredTableCell)
 		SLATE_EVENT(FGetColumnOutlineHAlignmentDelegate, OnGetColumnOutlineHAlignmentDelegate)
+		SLATE_EVENT(FSetHoveredTimerTableCell, OnSetHoveredCell)
 		SLATE_ATTRIBUTE(FText, HighlightText)
-		SLATE_ATTRIBUTE(FName, HighlightedTimerName)
+		SLATE_ATTRIBUTE(FName, HighlightedNodeName)
+		SLATE_ARGUMENT(TSharedPtr<Insights::FTable>, TablePtr)
 		SLATE_ARGUMENT(FTimerNodePtr, TimerNodePtr)
 	SLATE_END_ARGS()
 
@@ -53,41 +60,29 @@ public:
 	virtual FReply OnDragDetected(const FGeometry& MyGeometry, const FPointerEvent& MouseEvent) override;
 
 protected:
-	/**
-	 * @return a text which describes this table row, refers to both groups and timers
-	 */
-	FText GetText() const;
-
-	/**
-	 * @return a font style which is used to draw this table row, refers to both groups and timers
-	 */
-	FSlateFontInfo GetFont() const;
-
-	/**
-	 * @return a color and opacity value used to draw this table row, refers to both groups and timers
-	 */
-	FSlateColor GetColorAndOpacity() const;
-
 	FSlateColor GetBackgroundColorAndOpacity() const;
 	FSlateColor GetBackgroundColorAndOpacity(double Time) const;
 	FSlateColor GetOutlineColorAndOpacity() const;
 	const FSlateBrush* GetOutlineBrush(const FName ColumnId) const;
 	bool HandleShouldBeEnabled() const;
 	EVisibility IsColumnVisible(const FName ColumnId) const;
-	void OnSetHoveredTableCell(const FName ColumnId, const FTimerNodePtr SamplePtr);
+	void OnSetHoveredCell(TSharedPtr<Insights::FTable> InTablePtr, TSharedPtr<Insights::FTableColumn> InColumnPtr, const FTimerNodePtr InTimerNodePtr);
 
 protected:
+	/** A shared pointer to the table view model. */
+	TSharedPtr<Insights::FTable> TablePtr;
+
 	/** Data context for this table row. */
 	FTimerNodePtr TimerNodePtr;
 
 	FShouldBeEnabledDelegate OnShouldBeEnabled;
 	FIsColumnVisibleDelegate IsColumnVisibleDelegate;
-	FSetHoveredTimerTableCell SetHoveredTableCellDelegate;
+	FSetHoveredTimerTableCell SetHoveredCellDelegate;
 	FGetColumnOutlineHAlignmentDelegate GetColumnOutlineHAlignmentDelegate;
 
 	/** Text to be highlighted on timer name. */
 	TAttribute<FText> HighlightText;
 
-	/** Name of the timer that should be drawn as highlighted. */
-	TAttribute<FName> HighlightedTimerName;
+	/** Name of the timer node that should be drawn as highlighted. */
+	TAttribute<FName> HighlightedNodeName;
 };
diff --git a/Engine/Source/Developer/TraceInsights/Private/Insights/Widgets/STimerTreeView.cpp b/Engine/Source/Developer/TraceInsights/Private/Insights/Widgets/STimerTreeView.cpp
new file mode 100644
--- /dev/null
+++ b/Engine/Source/Developer/TraceInsights/Private/Insights/Widgets/STimerTreeView.cpp
@@ -0,0 +1,1009 @@
+// Copyright 1998-2019 Epic Games, Inc. All Rights Reserved.
+
+#include "STimerTreeView.h"
+
+#include "EditorStyleSet.h"
+#include "Framework/Application/SlateApplication.h"
+#include "Framework/MultiBox/MultiBoxBuilder.h"
+#include "SlateOptMacros.h"
+#include "Widgets/Layout/SScrollBox.h"
+#include "Widgets/SToolTip.h"
+
+// Insights
+#include "Insights/Table/ViewModels/Table.h"
+#include "Insights/Table/ViewModels/TableColumn.h"
+#include "Insights/Table/ViewModels/TreeNodeSorting.h"
+#include "Insights/ViewModels/TimersViewColumnFactory.h"
+#include "Insights/Widgets/STimersViewTooltip.h"
+#include "Insights/Widgets/STimerTableRow.h"
+
+#define LOCTEXT_NAMESPACE "STimerTreeView"
+
+////////////////////////////////////////////////////////////////////////////////////////////////////
+
+STimerTreeView::STimerTreeView()
+	: Table(MakeShareable(new Insights::FTable()))
+	, ViewName()
+	, TreeView(nullptr)
+	, TreeViewHeaderRow(nullptr)
+	, ExternalScrollbar(nullptr)
+	, HoveredColumnId()
+	, HoveredNodePtr(nullptr)
+	, HighlightedNodeName()
+	, TreeNodes()
+	, AvailableSorters()
+	, CurrentSorter(nullptr)
+	, ColumnBeingSorted(GetDefaultColumnBeingSorted())
+	, ColumnSortMode(GetDefaultColumnSortMode())
+	, StatsStartTime(0.0)
+	, StatsEndTime(0.0)
+	, StatsTimerId(0)
+{
+}
+
+////////////////////////////////////////////////////////////////////////////////////////////////////
+
+STimerTreeView::~STimerTreeView()
+{
+}
+
+////////////////////////////////////////////////////////////////////////////////////////////////////
+
+BEGIN_SLATE_FUNCTION_BUILD_OPTIMIZATION
+void STimerTreeView::Construct(const FArguments& InArgs, const FText& InViewName)
+{
+	ViewName = InViewName;
+
+	SAssignNew(ExternalScrollbar, SScrollBar)
+	.AlwaysShowScrollbar(true);
+
+	ChildSlot
+	[
+		SNew(SHorizontalBox)
+
+		+ SHorizontalBox::Slot()
+		.FillWidth(1.0f)
+		.Padding(0.0f)
+		[
+			SNew(SScrollBox)
+			.Orientation(Orient_Horizontal)
+
+			+ SScrollBox::Slot()
+			[
+				SNew(SBorder)
+				.BorderImage(FEditorStyle::GetBrush("ToolPanel.GroupBorder"))
+				.Padding(0.0f)
+				[
+					SAssignNew(TreeView, STreeView<FTimerNodePtr>)
+					.ExternalScrollbar(ExternalScrollbar)
+					.SelectionMode(ESelectionMode::Multi)
+					.TreeItemsSource(&TreeNodes)
+					.OnGetChildren(this, &STimerTreeView::TreeView_OnGetChildren)
+					.OnGenerateRow(this, &STimerTreeView::TreeView_OnGenerateRow)
+					//.OnSelectionChanged(this, &STimerTreeView::TreeView_OnSelectionChanged)
+					//.OnMouseButtonDoubleClick(this, &STimerTreeView::TreeView_OnMouseButtonDoubleClick)
+					.OnContextMenuOpening(FOnContextMenuOpening::CreateSP(this, &STimerTreeView::TreeView_GetMenuContent))
+					.ItemHeight(12.0f)
+					.HeaderRow
+					(
+						SAssignNew(TreeViewHeaderRow, SHeaderRow)
+						.Visibility(EVisibility::Visible)
+					)
+				]
+			]
+		]
+
+		+ SHorizontalBox::Slot()
+		.AutoWidth()
+		.Padding(0.0f)
+		[
+			SNew(SBox)
+			.WidthOverride(FOptionalSize(13.0f))
+			[
+				ExternalScrollbar.ToSharedRef()
+			]
+		]
+	];
+
+	InitializeAndShowHeaderColumns();
+
+	CreateSortings();
+}
+END_SLATE_FUNCTION_BUILD_OPTIMIZATION
+
+////////////////////////////////////////////////////////////////////////////////////////////////////
+
+TSharedPtr<SWidget> STimerTreeView::TreeView_GetMenuContent()
+{
+	const TArray<FTimerNodePtr> SelectedNodes = TreeView->GetSelectedItems();
+	const int32 NumSelectedNodes = SelectedNodes.Num();
+	FTimerNodePtr SelectedNode = NumSelectedNodes ? SelectedNodes[0] : nullptr;
+
+	const TSharedPtr<Insights::FTableColumn> HoveredColumnPtr = Table->FindColumn(HoveredColumnId);
+
+	FText SelectionStr;
+	FText PropertyName;
+	FText PropertyValue;
+
+	if (NumSelectedNodes == 0)
+	{
+		SelectionStr = LOCTEXT("NothingSelected", "Nothing selected");
+	}
+	else if (NumSelectedNodes == 1)
+	{
+		if (HoveredColumnPtr != nullptr)
+		{
+			PropertyName = HoveredColumnPtr->GetShortName();
+			PropertyValue = HoveredColumnPtr->GetValueAsTooltipText(*SelectedNode);
+		}
+		SelectionStr = FText::FromName(SelectedNode->GetName());
+	}
+	else
+	{
+		SelectionStr = LOCTEXT("MultipleSelection", "Multiple selection");
+	}
+
+	const bool bShouldCloseWindowAfterMenuSelection = true;
+	FMenuBuilder MenuBuilder(bShouldCloseWindowAfterMenuSelection, NULL);
+
+	// Selection menu
+	MenuBuilder.BeginSection("Selection", LOCTEXT("ContextMenu_Header_Selection", "Selection"));
+	{
+		struct FLocal
+		{
+			static bool ReturnFalse()
+			{
+				return false;
+			}
+		};
+
+		FUIAction DummyUIAction;
+		DummyUIAction.CanExecuteAction = FCanExecuteAction::CreateStatic(&FLocal::ReturnFalse);
+		MenuBuilder.AddMenuEntry
+		(
+			SelectionStr,
+			LOCTEXT("ContextMenu_Selection", "Currently selected items"),
+			FSlateIcon(FEditorStyle::GetStyleSetName(), "@missing.icon"), DummyUIAction, NAME_None, EUserInterfaceActionType::Button
+		);
+	}
+	MenuBuilder.EndSection();
+
+	MenuBuilder.BeginSection("Misc", LOCTEXT("ContextMenu_Header_Misc", "Miscellaneous"));
+	{
+		MenuBuilder.AddSubMenu
+		(
+			LOCTEXT("ContextMenu_Header_Misc_Sort", "Sort By"),
+			LOCTEXT("ContextMenu_Header_Misc_Sort_Desc", "Sort by column"),
+			FNewMenuDelegate::CreateSP(this, &STimerTreeView::TreeView_BuildSortByMenu),
+			false,
+			FSlateIcon(FEditorStyle::GetStyleSetName(), "Profiler.Misc.SortBy")
+		);
+	}
+	MenuBuilder.EndSection();
+
+	MenuBuilder.BeginSection("Columns", LOCTEXT("ContextMenu_Header_Columns", "Columns"));
+	{
+		MenuBuilder.AddSubMenu
+		(
+			LOCTEXT("ContextMenu_Header_Columns_View", "View Column"),
+			LOCTEXT("ContextMenu_Header_Columns_View_Desc", "Hides or shows columns"),
+			FNewMenuDelegate::CreateSP(this, &STimerTreeView::TreeView_BuildViewColumnMenu),
+			false,
+			FSlateIcon(FEditorStyle::GetStyleSetName(), "Profiler.EventGraph.ViewColumn")
+		);
+
+		FUIAction Action_ShowAllColumns
+		(
+			FExecuteAction::CreateSP(this, &STimerTreeView::ContextMenu_ShowAllColumns_Execute),
+			FCanExecuteAction::CreateSP(this, &STimerTreeView::ContextMenu_ShowAllColumns_CanExecute)
+		);
+		MenuBuilder.AddMenuEntry
+		(
+			LOCTEXT("ContextMenu_Header_Columns_ShowAllColumns", "Show All Columns"),
+			LOCTEXT("ContextMenu_Header_Columns_ShowAllColumns_Desc", "Resets tree view to show all columns"),
+			FSlateIcon(FEditorStyle::GetStyleSetName(), "Profiler.EventGraph.ResetColumn"), Action_ShowAllColumns, NAME_None, EUserInterfaceActionType::Button
+		);
+
+		FUIAction Action_ResetColumns
+		(
+			FExecuteAction::CreateSP(this, &STimerTreeView::ContextMenu_ResetColumns_Execute),
+			FCanExecuteAction::CreateSP(this, &STimerTreeView::ContextMenu_ResetColumns_CanExecute)
+		);
+		MenuBuilder.AddMenuEntry
+		(
+			LOCTEXT("ContextMenu_Header_Columns_ResetColumns", "Reset Columns to Default"),
+			LOCTEXT("ContextMenu_Header_Columns_ResetColumns_Desc", "Resets columns to default"),
+			FSlateIcon(FEditorStyle::GetStyleSetName(), "Profiler.EventGraph.ResetColumn"), Action_ResetColumns, NAME_None, EUserInterfaceActionType::Button
+		);
+	}
+	MenuBuilder.EndSection();
+
+	return MenuBuilder.MakeWidget();
+}
+
+////////////////////////////////////////////////////////////////////////////////////////////////////
+
+void STimerTreeView::TreeView_BuildSortByMenu(FMenuBuilder& MenuBuilder)
+{
+	MenuBuilder.BeginSection("ColumnName", LOCTEXT("ContextMenu_Header_Misc_ColumnName", "Column Name"));
+
+	for (const TSharedPtr<Insights::FTableColumn>& ColumnPtr : Table->GetColumns())
+	{
+		const Insights::FTableColumn& Column = *ColumnPtr;
+
+		if (Column.IsVisible() && Column.CanBeSorted())
+		{
+			FUIAction Action_SortByColumn
+			(
+				FExecuteAction::CreateSP(this, &STimerTreeView::ContextMenu_SortByColumn_Execute, Column.GetId()),
+				FCanExecuteAction::CreateSP(this, &STimerTreeView::ContextMenu_SortByColumn_CanExecute, Column.GetId()),
+				FIsActionChecked::CreateSP(this, &STimerTreeView::ContextMenu_SortByColumn_IsChecked, Column.GetId())
+			);
+			MenuBuilder.AddMenuEntry
+			(
+				Column.GetTitleName(),
+				Column.GetDescription(),
+				FSlateIcon(), Action_SortByColumn, NAME_None, EUserInterfaceActionType::RadioButton
+			);
+		}
+	}
+
+	MenuBuilder.EndSection();
+
+	//-----------------------------------------------------------------------------
+
+	MenuBuilder.BeginSection("SortMode", LOCTEXT("ContextMenu_Header_Misc_Sort_SortMode", "Sort Mode"));
+	{
+		FUIAction Action_SortAscending
+		(
+			FExecuteAction::CreateSP(this, &STimerTreeView::ContextMenu_SortMode_Execute, EColumnSortMode::Ascending),
+			FCanExecuteAction::CreateSP(this, &STimerTreeView::ContextMenu_SortMode_CanExecute, EColumnSortMode::Ascending),
+			FIsActionChecked::CreateSP(this, &STimerTreeView::ContextMenu_SortMode_IsChecked, EColumnSortMode::Ascending)
+		);
+		MenuBuilder.AddMenuEntry
+		(
+			LOCTEXT("ContextMenu_Header_Misc_Sort_SortAscending", "Sort Ascending"),
+			LOCTEXT("ContextMenu_Header_Misc_Sort_SortAscending_Desc", "Sorts ascending"),
+			FSlateIcon(FEditorStyle::GetStyleSetName(), "Profiler.Misc.SortAscending"), Action_SortAscending, NAME_None, EUserInterfaceActionType::RadioButton
+		);
+
+		FUIAction Action_SortDescending
+		(
+			FExecuteAction::CreateSP(this, &STimerTreeView::ContextMenu_SortMode_Execute, EColumnSortMode::Descending),
+			FCanExecuteAction::CreateSP(this, &STimerTreeView::ContextMenu_SortMode_CanExecute, EColumnSortMode::Descending),
+			FIsActionChecked::CreateSP(this, &STimerTreeView::ContextMenu_SortMode_IsChecked, EColumnSortMode::Descending)
+		);
+		MenuBuilder.AddMenuEntry
+		(
+			LOCTEXT("ContextMenu_Header_Misc_Sort_SortDescending", "Sort Descending"),
+			LOCTEXT("ContextMenu_Header_Misc_Sort_SortDescending_Desc", "Sorts descending"),
+			FSlateIcon(FEditorStyle::GetStyleSetName(), "Profiler.Misc.SortDescending"), Action_SortDescending, NAME_None, EUserInterfaceActionType::RadioButton
+		);
+	}
+	MenuBuilder.EndSection();
+}
+
+////////////////////////////////////////////////////////////////////////////////////////////////////
+
+void STimerTreeView::TreeView_BuildViewColumnMenu(FMenuBuilder& MenuBuilder)
+{
+	MenuBuilder.BeginSection("ViewColumn", LOCTEXT("ContextMenu_Header_Columns_View", "View Column"));
+
+	for (const TSharedPtr<Insights::FTableColumn>& ColumnPtr : Table->GetColumns())
+	{
+		const Insights::FTableColumn& Column = *ColumnPtr;
+
+		FUIAction Action_ToggleColumn
+		(
+			FExecuteAction::CreateSP(this, &STimerTreeView::ToggleColumnVisibility, Column.GetId()),
+			FCanExecuteAction::CreateSP(this, &STimerTreeView::CanToggleColumnVisibility, Column.GetId()),
+			FIsActionChecked::CreateSP(this, &STimerTreeView::IsColumnVisible, Column.GetId())
+		);
+		MenuBuilder.AddMenuEntry
+		(
+			Column.GetTitleName(),
+			Column.GetDescription(),
+			FSlateIcon(), Action_ToggleColumn, NAME_None, EUserInterfaceActionType::ToggleButton
+		);
+	}
+
+	MenuBuilder.EndSection();
+}
+
+////////////////////////////////////////////////////////////////////////////////////////////////////
+
+void STimerTreeView::InitializeAndShowHeaderColumns()
+{
+	// Create columns.
+	TArray<TSharedPtr<Insights::FTableColumn>> Columns;
+	FTimersViewColumnFactory::CreateTimerTreeViewColumns(Columns);
+	ensure(Columns.Num() > 0 && Columns[0]->IsHierarchy());
+	Columns[0]->SetShortName(ViewName);
+	Columns[0]->SetTitleName(ViewName);
+	Table->SetColumns(Columns);
+
+	// Show columns.
+	for (const TSharedPtr<Insights::FTableColumn>& ColumnPtr : Table->GetColumns())
+	{
+		if (ColumnPtr->ShouldBeVisible())
+		{
+			ShowColumn(ColumnPtr->GetId());
+		}
+	}
+}
+
+////////////////////////////////////////////////////////////////////////////////////////////////////
+
+FText STimerTreeView::GetColumnHeaderText(const FName ColumnId) const
+{
+	const Insights::FTableColumn& Column = *Table->FindColumnChecked(ColumnId);
+	return Column.GetShortName();
+}
+
+////////////////////////////////////////////////////////////////////////////////////////////////////
+
+TSharedRef<SWidget> STimerTreeView::TreeViewHeaderRow_GenerateColumnMenu(const Insights::FTableColumn& Column)
+{
+	bool bIsMenuVisible = false;
+
+	const bool bShouldCloseWindowAfterMenuSelection = true;
+	FMenuBuilder MenuBuilder(bShouldCloseWindowAfterMenuSelection, NULL);
+	{
+		if (Column.CanBeHidden())
+		{
+			MenuBuilder.BeginSection("Column", LOCTEXT("TreeViewHeaderRow_Header_Column", "Column"));
+
+			FUIAction Action_HideColumn
+			(
+				FExecuteAction::CreateSP(this, &STimerTreeView::HideColumn, Column.GetId()),
+				FCanExecuteAction::CreateSP(this, &STimerTreeView::CanHideColumn, Column.GetId())
+			);
+
+			MenuBuilder.AddMenuEntry
+			(
+				LOCTEXT("TreeViewHeaderRow_HideColumn", "Hide"),
+				LOCTEXT("TreeViewHeaderRow_HideColumn_Desc", "Hides the selected column"),
+				FSlateIcon(), Action_HideColumn, NAME_None, EUserInterfaceActionType::Button
+			);
+			bIsMenuVisible = true;
+
+			MenuBuilder.EndSection();
+		}
+
+		if (Column.CanBeSorted())
+		{
+			MenuBuilder.BeginSection("SortMode", LOCTEXT("ContextMenu_Header_Misc_Sort_SortMode", "Sort Mode"));
+
+			FUIAction Action_SortAscending
+			(
+				FExecuteAction::CreateSP(this, &STimerTreeView::HeaderMenu_SortMode_Execute, Column.GetId(), EColumnSortMode::Ascending),
+				FCanExecuteAction::CreateSP(this, &STimerTreeView::HeaderMenu_SortMode_CanExecute, Column.GetId(), EColumnSortMode::Ascending),
+				FIsActionChecked::CreateSP(this, &STimerTreeView::HeaderMenu_SortMode_IsChecked, Column.GetId(), EColumnSortMode::Ascending)
+			);
+			MenuBuilder.AddMenuEntry
+			(
+				LOCTEXT("ContextMenu_Header_Misc_Sort_SortAscending", "Sort Ascending"),
+				LOCTEXT("ContextMenu_Header_Misc_Sort_SortAscending_Desc", "Sorts ascending"),
+				FSlateIcon(FEditorStyle::GetStyleSetName(), "Profiler.Misc.SortAscending"), Action_SortAscending, NAME_None, EUserInterfaceActionType::RadioButton
+			);
+
+			FUIAction Action_SortDescending
+			(
+				FExecuteAction::CreateSP(this, &STimerTreeView::HeaderMenu_SortMode_Execute, Column.GetId(), EColumnSortMode::Descending),
+				FCanExecuteAction::CreateSP(this, &STimerTreeView::HeaderMenu_SortMode_CanExecute, Column.GetId(), EColumnSortMode::Descending),
+				FIsActionChecked::CreateSP(this, &STimerTreeView::HeaderMenu_SortMode_IsChecked, Column.GetId(), EColumnSortMode::Descending)
+			);
+			MenuBuilder.AddMenuEntry
+			(
+				LOCTEXT("ContextMenu_Header_Misc_Sort_SortDescending", "Sort Descending"),
+				LOCTEXT("ContextMenu_Header_Misc_Sort_SortDescending_Desc", "Sorts descending"),
+				FSlateIcon(FEditorStyle::GetStyleSetName(), "Profiler.Misc.SortDescending"), Action_SortDescending, NAME_None, EUserInterfaceActionType::RadioButton
+			);
+			bIsMenuVisible = true;
+
+			MenuBuilder.EndSection();
+		}
+	}
+
+	return bIsMenuVisible ? MenuBuilder.MakeWidget() : (TSharedRef<SWidget>)SNullWidget::NullWidget;
+}
+
+////////////////////////////////////////////////////////////////////////////////////////////////////
+// TreeView
+////////////////////////////////////////////////////////////////////////////////////////////////////
+
+void STimerTreeView::TreeView_Refresh()
+{
+	if (TreeView.IsValid())
+	{
+		TreeView->RequestTreeRefresh();
+	}
+}
+
+////////////////////////////////////////////////////////////////////////////////////////////////////
+
+void STimerTreeView::TreeView_OnSelectionChanged(FTimerNodePtr SelectedItem, ESelectInfo::Type SelectInfo)
+{
+}
+
+////////////////////////////////////////////////////////////////////////////////////////////////////
+
+void STimerTreeView::TreeView_OnGetChildren(FTimerNodePtr InParent, TArray<FTimerNodePtr>& OutChildren)
+{
+	constexpr bool bUseFiltering = false;
+	if (bUseFiltering)
+	{
+		const TArray<Insights::FBaseTreeNodePtr>& Children = InParent->GetFilteredChildren();
+		OutChildren.Reset(Children.Num());
+		for (const Insights::FBaseTreeNodePtr& Child : Children)
+		{
+			OutChildren.Add(StaticCastSharedPtr<FTimerNode, Insights::FBaseTreeNode>(Child));
+		}
+	}
+	else
+	{
+		const TArray<Insights::FBaseTreeNodePtr>& Children = InParent->GetChildren();
+		OutChildren.Reset(Children.Num());
+		for (const Insights::FBaseTreeNodePtr& Child : Children)
+		{
+			OutChildren.Add(StaticCastSharedPtr<FTimerNode, Insights::FBaseTreeNode>(Child));
+		}
+	}
+}
+
+////////////////////////////////////////////////////////////////////////////////////////////////////
+
+void STimerTreeView::TreeView_OnMouseButtonDoubleClick(FTimerNodePtr TimerNodePtr)
+{
+	//if (TimerNodePtr->IsGroup())
+	if (TimerNodePtr->GetChildren().Num() > 0)
+	{
+		const bool bIsGroupExpanded = TreeView->IsItemExpanded(TimerNodePtr);
+		TreeView->SetItemExpansion(TimerNodePtr, !bIsGroupExpanded);
+	}
+}
+
+////////////////////////////////////////////////////////////////////////////////////////////////////
+// Tree View's Table Row
+////////////////////////////////////////////////////////////////////////////////////////////////////
+
+TSharedRef<ITableRow> STimerTreeView::TreeView_OnGenerateRow(FTimerNodePtr TimerNodePtr, const TSharedRef<STableViewBase>& OwnerTable)
+{
+	TSharedRef<ITableRow> TableRow =
+		SNew(STimerTableRow, OwnerTable)
+		.OnShouldBeEnabled(this, &STimerTreeView::TableRow_ShouldBeEnabled)
+		.OnIsColumnVisible(this, &STimerTreeView::IsColumnVisible)
+		.OnSetHoveredCell(this, &STimerTreeView::TableRow_SetHoveredCell)
+		.OnGetColumnOutlineHAlignmentDelegate(this, &STimerTreeView::TableRow_GetColumnOutlineHAlignment)
+		.HighlightText(this, &STimerTreeView::TableRow_GetHighlightText)
+		.HighlightedNodeName(this, &STimerTreeView::TableRow_GetHighlightedNodeName)
+		.TablePtr(Table)
+		.TimerNodePtr(TimerNodePtr);
+
+	return TableRow;
+}
+
+////////////////////////////////////////////////////////////////////////////////////////////////////
+
+bool STimerTreeView::TableRow_ShouldBeEnabled(const uint32 TimerId) const
+{
+	return true;
+}
+
+////////////////////////////////////////////////////////////////////////////////////////////////////
+
+void STimerTreeView::TableRow_SetHoveredCell(TSharedPtr<Insights::FTable> InTablePtr, TSharedPtr<Insights::FTableColumn> InColumnPtr, const FTimerNodePtr InNodePtr)
+{
+	HoveredColumnId = InColumnPtr ? InColumnPtr->GetId() : FName();
+
+	const bool bIsAnyMenusVisible = FSlateApplication::Get().AnyMenusVisible();
+	if (!HasMouseCapture() && !bIsAnyMenusVisible)
+	{
+		HoveredNodePtr = InNodePtr;
+	}
+}
+
+////////////////////////////////////////////////////////////////////////////////////////////////////
+
+EHorizontalAlignment STimerTreeView::TableRow_GetColumnOutlineHAlignment(const FName ColumnId) const
+{
+	const TIndirectArray<SHeaderRow::FColumn>& Columns = TreeViewHeaderRow->GetColumns();
+	const int32 LastColumnIdx = Columns.Num() - 1;
+
+	// First column
+	if (Columns[0].ColumnId == ColumnId)
+	{
+		return HAlign_Left;
+	}
+	// Last column
+	else if (Columns[LastColumnIdx].ColumnId == ColumnId)
+	{
+		return HAlign_Right;
+	}
+	// Middle columns
+	{
+		return HAlign_Center;
+	}
+}
+
+////////////////////////////////////////////////////////////////////////////////////////////////////
+
+FText STimerTreeView::TableRow_GetHighlightText() const
+{
+	return FText::GetEmpty();
+}
+
+////////////////////////////////////////////////////////////////////////////////////////////////////
+
+FName STimerTreeView::TableRow_GetHighlightedNodeName() const
+{
+	return HighlightedNodeName;
+}
+
+////////////////////////////////////////////////////////////////////////////////////////////////////
+// Sorting
+////////////////////////////////////////////////////////////////////////////////////////////////////
+
+const FName STimerTreeView::GetDefaultColumnBeingSorted()
+{
+	return FTimersViewColumns::TotalInclusiveTimeColumnID;
+}
+
+////////////////////////////////////////////////////////////////////////////////////////////////////
+
+const EColumnSortMode::Type STimerTreeView::GetDefaultColumnSortMode()
+{
+	return EColumnSortMode::Descending;
+}
+
+////////////////////////////////////////////////////////////////////////////////////////////////////
+
+void STimerTreeView::CreateSortings()
+{
+	AvailableSorters.Reset();
+	CurrentSorter = nullptr;
+
+	for (const TSharedPtr<Insights::FTableColumn> ColumnPtr : Table->GetColumns())
+	{
+		if (ColumnPtr->CanBeSorted())
+		{
+			TSharedPtr<Insights::ITableCellValueSorter> SorterPtr = ColumnPtr->GetValueSorter();
+			if (ensure(SorterPtr.IsValid()))
+			{
+				AvailableSorters.Add(SorterPtr);
+			}
+		}
+	}
+
+	UpdateCurrentSortingByColumn();
+}
+
+////////////////////////////////////////////////////////////////////////////////////////////////////
+
+void STimerTreeView::UpdateCurrentSortingByColumn()
+{
+	TSharedPtr<Insights::FTableColumn> ColumnPtr = Table->FindColumn(ColumnBeingSorted);
+	CurrentSorter = ColumnPtr.IsValid() ? ColumnPtr->GetValueSorter() : nullptr;
+}
+
+////////////////////////////////////////////////////////////////////////////////////////////////////
+
+void STimerTreeView::SortTreeNodes()
+{
+	if (CurrentSorter.IsValid())
+	{
+		for (FTimerNodePtr& Root : TreeNodes)
+		{
+			SortTreeNodesRec(*Root, *CurrentSorter);
+		}
+	}
+}
+
+////////////////////////////////////////////////////////////////////////////////////////////////////
+
+void STimerTreeView::SortTreeNodesRec(FTimerNode& Node, const Insights::ITableCellValueSorter& Sorter)
+{
+	if (ColumnSortMode == EColumnSortMode::Type::Descending)
+	{
+		Node.SortChildrenDescending(Sorter);
+	}
+	else // if (ColumnSortMode == EColumnSortMode::Type::Ascending)
+	{
+		Node.SortChildrenAscending(Sorter);
+	}
+
+	for (Insights::FBaseTreeNodePtr ChildPtr : Node.GetChildren())
+	{
+		//if (ChildPtr->IsGroup())
+		if (ChildPtr->GetChildren().Num() > 0)
+		{
+			SortTreeNodesRec(*StaticCastSharedPtr<FTimerNode>(ChildPtr), Sorter);
+		}
+	}
+}
+
+////////////////////////////////////////////////////////////////////////////////////////////////////
+
+EColumnSortMode::Type STimerTreeView::GetSortModeForColumn(const FName ColumnId) const
+{
+	if (ColumnBeingSorted != ColumnId)
+	{
+		return EColumnSortMode::None;
+	}
+
+	return ColumnSortMode;
+}
+
+////////////////////////////////////////////////////////////////////////////////////////////////////
+
+void STimerTreeView::SetSortModeForColumn(const FName& ColumnId, const EColumnSortMode::Type SortMode)
+{
+	ColumnBeingSorted = ColumnId;
+	ColumnSortMode = SortMode;
+	UpdateCurrentSortingByColumn();
+
+	SortTreeNodes();
+	//ApplyFiltering();
+}
+
+////////////////////////////////////////////////////////////////////////////////////////////////////
+
+void STimerTreeView::OnSortModeChanged(const EColumnSortPriority::Type SortPriority, const FName& ColumnId, const EColumnSortMode::Type SortMode)
+{
+	SetSortModeForColumn(ColumnId, SortMode);
+	TreeView_Refresh();
+}
+
+////////////////////////////////////////////////////////////////////////////////////////////////////
+// SortMode action (HeaderMenu)
+////////////////////////////////////////////////////////////////////////////////////////////////////
+
+bool STimerTreeView::HeaderMenu_SortMode_IsChecked(const FName ColumnId, const EColumnSortMode::Type InSortMode)
+{
+	return ColumnBeingSorted == ColumnId && ColumnSortMode == InSortMode;
+}
+
+////////////////////////////////////////////////////////////////////////////////////////////////////
+
+bool STimerTreeView::HeaderMenu_SortMode_CanExecute(const FName ColumnId, const EColumnSortMode::Type InSortMode) const
+{
+	const Insights::FTableColumn& Column = *Table->FindColumnChecked(ColumnId);
+	return Column.CanBeSorted();
+}
+
+////////////////////////////////////////////////////////////////////////////////////////////////////
+
+void STimerTreeView::HeaderMenu_SortMode_Execute(const FName ColumnId, const EColumnSortMode::Type InSortMode)
+{
+	SetSortModeForColumn(ColumnId, InSortMode);
+	TreeView_Refresh();
+}
+
+////////////////////////////////////////////////////////////////////////////////////////////////////
+// SortMode action (ContextMenu)
+////////////////////////////////////////////////////////////////////////////////////////////////////
+
+bool STimerTreeView::ContextMenu_SortMode_IsChecked(const EColumnSortMode::Type InSortMode)
+{
+	return ColumnSortMode == InSortMode;
+}
+
+////////////////////////////////////////////////////////////////////////////////////////////////////
+
+bool STimerTreeView::ContextMenu_SortMode_CanExecute(const EColumnSortMode::Type InSortMode) const
+{
+	return true; //ColumnSortMode != InSortMode;
+}
+
+////////////////////////////////////////////////////////////////////////////////////////////////////
+
+void STimerTreeView::ContextMenu_SortMode_Execute(const EColumnSortMode::Type InSortMode)
+{
+	SetSortModeForColumn(ColumnBeingSorted, InSortMode);
+	TreeView_Refresh();
+}
+
+////////////////////////////////////////////////////////////////////////////////////////////////////
+// SortByColumn action (ContextMenu)
+////////////////////////////////////////////////////////////////////////////////////////////////////
+
+bool STimerTreeView::ContextMenu_SortByColumn_IsChecked(const FName ColumnId)
+{
+	return ColumnId == ColumnBeingSorted;
+}
+
+////////////////////////////////////////////////////////////////////////////////////////////////////
+
+bool STimerTreeView::ContextMenu_SortByColumn_CanExecute(const FName ColumnId) const
+{
+	return true; //ColumnId != ColumnBeingSorted;
+}
+
+////////////////////////////////////////////////////////////////////////////////////////////////////
+
+void STimerTreeView::ContextMenu_SortByColumn_Execute(const FName ColumnId)
+{
+	SetSortModeForColumn(ColumnId, EColumnSortMode::Descending);
+	TreeView_Refresh();
+}
+
+////////////////////////////////////////////////////////////////////////////////////////////////////
+// ShowColumn action
+////////////////////////////////////////////////////////////////////////////////////////////////////
+
+bool STimerTreeView::CanShowColumn(const FName ColumnId) const
+{
+	return true;
+}
+
+////////////////////////////////////////////////////////////////////////////////////////////////////
+
+void STimerTreeView::ShowColumn(const FName ColumnId)
+{
+	Insights::FTableColumn& Column = *Table->FindColumnChecked(ColumnId);
+	Column.Show();
+
+	SHeaderRow::FColumn::FArguments ColumnArgs;
+	ColumnArgs
+		.ColumnId(Column.GetId())
+		.DefaultLabel(Column.GetShortName())
+		.HAlignHeader(HAlign_Fill)
+		.VAlignHeader(VAlign_Fill)
+		.HeaderContentPadding(FMargin(2.0f))
+		.HAlignCell(HAlign_Fill)
+		.VAlignCell(VAlign_Fill)
+		.SortMode(this, &STimerTreeView::GetSortModeForColumn, Column.GetId())
+		.OnSort(this, &STimerTreeView::OnSortModeChanged)
+		.ManualWidth(Column.GetInitialWidth())
+		.FixedWidth(Column.IsFixedWidth() ? Column.GetInitialWidth() : TOptional<float>())
+		.HeaderContent()
+		[
+			SNew(SBox)
+			.ToolTip(STimersViewTooltip::GetColumnTooltip(Column))
+			.HAlign(Column.GetHorizontalAlignment())
+			.VAlign(VAlign_Center)
+			[
+				SNew(STextBlock)
+				.Text(this, &STimerTreeView::GetColumnHeaderText, Column.GetId())
+			]
+		]
+		.MenuContent()
+		[
+			TreeViewHeaderRow_GenerateColumnMenu(Column)
+		];
+
+	int32 ColumnIndex = 0;
+	const int32 NewColumnPosition = Table->GetColumnPositionIndex(ColumnId);
+	const int32 NumColumns = TreeViewHeaderRow->GetColumns().Num();
+	for (; ColumnIndex < NumColumns; ColumnIndex++)
+	{
+		const SHeaderRow::FColumn& CurrentColumn = TreeViewHeaderRow->GetColumns()[ColumnIndex];
+		const int32 CurrentColumnPosition = Table->GetColumnPositionIndex(CurrentColumn.ColumnId);
+		if (NewColumnPosition < CurrentColumnPosition)
+		{
+			break;
+		}
+	}
+
+	TreeViewHeaderRow->InsertColumn(ColumnArgs, ColumnIndex);
+}
+
+////////////////////////////////////////////////////////////////////////////////////////////////////
+// HideColumn action
+////////////////////////////////////////////////////////////////////////////////////////////////////
+
+bool STimerTreeView::CanHideColumn(const FName ColumnId) const
+{
+	const Insights::FTableColumn& Column = *Table->FindColumnChecked(ColumnId);
+	return Column.CanBeHidden();
+}
+
+////////////////////////////////////////////////////////////////////////////////////////////////////
+
+void STimerTreeView::HideColumn(const FName ColumnId)
+{
+	Insights::FTableColumn& Column = *Table->FindColumnChecked(ColumnId);
+	Column.Hide();
+
+	TreeViewHeaderRow->RemoveColumn(ColumnId);
+}
+
+////////////////////////////////////////////////////////////////////////////////////////////////////
+// ToggleColumn action
+////////////////////////////////////////////////////////////////////////////////////////////////////
+
+bool STimerTreeView::IsColumnVisible(const FName ColumnId)
+{
+	const Insights::FTableColumn& Column = *Table->FindColumnChecked(ColumnId);
+	return Column.IsVisible();
+}
+
+////////////////////////////////////////////////////////////////////////////////////////////////////
+
+bool STimerTreeView::CanToggleColumnVisibility(const FName ColumnId) const
+{
+	const Insights::FTableColumn& Column = *Table->FindColumnChecked(ColumnId);
+	return !Column.IsVisible() || Column.CanBeHidden();
+}
+
+////////////////////////////////////////////////////////////////////////////////////////////////////
+
+void STimerTreeView::ToggleColumnVisibility(const FName ColumnId)
+{
+	const Insights::FTableColumn& Column = *Table->FindColumnChecked(ColumnId);
+	if (Column.IsVisible())
+	{
+		HideColumn(ColumnId);
+	}
+	else
+	{
+		ShowColumn(ColumnId);
+	}
+}
+
+////////////////////////////////////////////////////////////////////////////////////////////////////
+// "Show All Columns" action (ContextMenu)
+////////////////////////////////////////////////////////////////////////////////////////////////////
+
+bool STimerTreeView::ContextMenu_ShowAllColumns_CanExecute() const
+{
+	return true;
+}
+
+////////////////////////////////////////////////////////////////////////////////////////////////////
+
+void STimerTreeView::ContextMenu_ShowAllColumns_Execute()
+{
+	ColumnBeingSorted = GetDefaultColumnBeingSorted();
+	ColumnSortMode = GetDefaultColumnSortMode();
+	UpdateCurrentSortingByColumn();
+
+	for (const TSharedPtr<Insights::FTableColumn>& ColumnPtr : Table->GetColumns())
+	{
+		const Insights::FTableColumn& Column = *ColumnPtr;
+
+		if (!Column.IsVisible())
+		{
+			ShowColumn(Column.GetId());
+		}
+	}
+}
+
+////////////////////////////////////////////////////////////////////////////////////////////////////
+// ResetColumns action (ContextMenu)
+////////////////////////////////////////////////////////////////////////////////////////////////////
+
+bool STimerTreeView::ContextMenu_ResetColumns_CanExecute() const
+{
+	return true;
+}
+
+////////////////////////////////////////////////////////////////////////////////////////////////////
+
+void STimerTreeView::ContextMenu_ResetColumns_Execute()
+{
+	ColumnBeingSorted = GetDefaultColumnBeingSorted();
+	ColumnSortMode = GetDefaultColumnSortMode();
+	UpdateCurrentSortingByColumn();
+
+	for (const TSharedPtr<Insights::FTableColumn>& ColumnPtr : Table->GetColumns())
+	{
+		const Insights::FTableColumn& Column = *ColumnPtr;
+
+		if (Column.ShouldBeVisible() && !Column.IsVisible())
+		{
+			ShowColumn(Column.GetId());
+		}
+		else if (!Column.ShouldBeVisible() && Column.IsVisible())
+		{
+			HideColumn(Column.GetId());
+		}
+	}
+}
+
+////////////////////////////////////////////////////////////////////////////////////////////////////
+
+void STimerTreeView::Reset()
+{
+	TreeNodes.Reset();
+	TreeView_Refresh();
+}
+
+////////////////////////////////////////////////////////////////////////////////////////////////////
+
+void STimerTreeView::SetTree(const Trace::FTimingProfilerButterflyNode& Root)
+{
+	TreeNodes.Reset();
+
+	FTimerNodePtr RootTimerNodePtr = CreateTimerNodeRec(Root);
+	if (RootTimerNodePtr)
+	{
+		// Mark the hot path. The child nodes are already sorted by InclTime (descending), so we just follow the first child.
+		FTimerNodePtr TimerNodePtr = RootTimerNodePtr;
+		while (TimerNodePtr.IsValid())
+		{
+			TimerNodePtr->SetIsHotPath(true);
+			const TArray<Insights::FBaseTreeNodePtr>& Children = TimerNodePtr->GetChildren();
+			TimerNodePtr = Children.Num() > 0 ? StaticCastSharedPtr<FTimerNode, Insights::FBaseTreeNode>(Children[0]) : nullptr;
+		}
+
+		TreeNodes.Add(RootTimerNodePtr);
+	}
+
+	SortTreeNodes();
+
+	TreeView_Refresh();
+
+	if (RootTimerNodePtr)
+	{
+		ExpandNodesRec(RootTimerNodePtr, 0);
+	}
+}
+
+////////////////////////////////////////////////////////////////////////////////////////////////////
+
+FTimerNodePtr STimerTreeView::CreateTimerNodeRec(const Trace::FTimingProfilerButterflyNode& Node)
+{
+	if (Node.Timer == nullptr)
+	{
+		return nullptr;
+		//return MakeShareable(new FTimerNode(0, FName(TEXT("!!!!!")), FName(TEXT("[Group]")), ETimerNodeType::InvalidOrMax));
+	}
+
+	const uint64 Id = Node.Timer->Id;
+	const FName Name(Node.Timer->Name);
+	const FName Group(Node.Timer->IsGpuTimer ? TEXT("GPU") : TEXT("CPU"));
+	const ETimerNodeType Type = Node.Timer->IsGpuTimer ? ETimerNodeType::GpuScope : ETimerNodeType::CpuScope;
+
+	FTimerNodePtr TimerNodePtr = MakeShareable(new FTimerNode(Id, Name, Group, Type));
+
+	Trace::FAggregatedTimingStats AggregatedStats;
+	AggregatedStats.InstanceCount = Node.Count;
+	AggregatedStats.TotalInclusiveTime = Node.InclusiveTime;
+	AggregatedStats.TotalExclusiveTime = Node.ExclusiveTime;
+	TimerNodePtr->SetAggregatedStats(AggregatedStats);
+
+	for (const Trace::FTimingProfilerButterflyNode* ChildNodePtr : Node.Children)
+	{
+		if (ChildNodePtr != nullptr)
+		{
+			FTimerNodePtr ChildTimerNodePtr = CreateTimerNodeRec(*ChildNodePtr);
+			if (ChildTimerNodePtr)
+			{
+				TimerNodePtr->AddChildAndSetGroupPtr(ChildTimerNodePtr);
+			}
+		}
+	}
+
+	// Sort children by InclTime (descending).
+	TArray<Insights::FBaseTreeNodePtr>& Children = const_cast<TArray<Insights::FBaseTreeNodePtr>&>(TimerNodePtr->GetChildren());
+	Children.Sort([](const Insights::FBaseTreeNodePtr& A, const Insights::FBaseTreeNodePtr& B) -> bool
+	{
+		const double InclTimeA = StaticCastSharedPtr<FTimerNode>(A)->GetAggregatedStats().TotalInclusiveTime;
+		const double InclTimeB = StaticCastSharedPtr<FTimerNode>(B)->GetAggregatedStats().TotalInclusiveTime;
+		return InclTimeA >= InclTimeB;
+	});
+
+	return TimerNodePtr;
+}
+
+////////////////////////////////////////////////////////////////////////////////////////////////////
+
+void STimerTreeView::ExpandNodesRec(FTimerNodePtr NodePtr, int32 Depth)
+{
+	//constexpr int32 MaxDepth = 3;
+
+	TreeView->SetItemExpansion(NodePtr, NodePtr->IsHotPath()); // expand only the hot path
+
+	//if (Depth < MaxDepth)
+	{
+		for (const Insights::FBaseTreeNodePtr& ChildPtr : NodePtr->GetChildren())
+		{
+			ExpandNodesRec(StaticCastSharedPtr<FTimerNode>(ChildPtr), Depth + 1);
+		}
+	}
+}
+
+////////////////////////////////////////////////////////////////////////////////////////////////////
+
+#undef LOCTEXT_NAMESPACE
