Copying //UE4/Dev-Core to Main (//UE4/Main)

diff --git a/Engine/Source/Developer/TraceInsights/Private/Insights/Widgets/STimersView.cpp b/Engine/Source/Developer/TraceInsights/Private/Insights/Widgets/STimersView.cpp
--- a/Engine/Source/Developer/TraceInsights/Private/Insights/Widgets/STimersView.cpp
+++ b/Engine/Source/Developer/TraceInsights/Private/Insights/Widgets/STimersView.cpp
@@ -5,6 +5,7 @@
 #include "EditorStyleSet.h"
 #include "Framework/MultiBox/MultiBoxBuilder.h"
 #include "SlateOptMacros.h"
+#include "Templates/UniquePtr.h"
 #include "TraceServices/AnalysisService.h"
 #include "Widgets/Input/SCheckBox.h"
 #include "Widgets/Layout/SScrollBox.h"
@@ -16,6 +17,8 @@
 #include "Widgets/Views/STableViewBase.h"
 
 // Insights
+#include "Insights/Table/ViewModels/Table.h"
+#include "Insights/Table/ViewModels/TableColumn.h"
 #include "Insights/TimingProfilerCommon.h"
 #include "Insights/TimingProfilerManager.h"
 #include "Insights/ViewModels/TimerNodeHelper.h"
@@ -29,27 +32,19 @@
 
 ////////////////////////////////////////////////////////////////////////////////////////////////////
 
-const EColumnSortMode::Type STimersView::GetDefaultColumnSortMode()
-{
-	return EColumnSortMode::Descending;
-}
-
-////////////////////////////////////////////////////////////////////////////////////////////////////
-
-const FName STimersView::GetDefaultColumnBeingSorted()
-{
-	return FTimersViewColumns::TotalInclusiveTimeColumnID;
-}
-
-////////////////////////////////////////////////////////////////////////////////////////////////////
-
 STimersView::STimersView()
-	: bExpansionSaved(false)
+	: Table(MakeShareable(new Insights::FTable()))
+	, bExpansionSaved(false)
+	, bFilterOutZeroCountTimers(false)
 	, GroupingMode(ETimerGroupingMode::ByType)
-	, ColumnSortMode(GetDefaultColumnSortMode())
+	, AvailableSorters()
+	, CurrentSorter(nullptr)
 	, ColumnBeingSorted(GetDefaultColumnBeingSorted())
+	, ColumnSortMode(GetDefaultColumnSortMode())
+	, StatsStartTime(0.0)
+	, StatsEndTime(0.0)
 {
-	FMemory::Memset(bTimerNodeIsVisible, 1);
+	FMemory::Memset(bTimerTypeIsVisible, 1);
 }
 
 ////////////////////////////////////////////////////////////////////////////////////////////////////
@@ -75,7 +70,7 @@ void STimersView::Construct(const FArguments& InArgs)
 	[
 		SNew(SVerticalBox)
 
-		+SVerticalBox::Slot()
+		+ SVerticalBox::Slot()
 		.VAlign(VAlign_Center)
 		.AutoHeight()
 		[
@@ -85,28 +80,57 @@ void STimersView::Construct(const FArguments& InArgs)
 			[
 				SNew(SVerticalBox)
 
-				// Search box
-				+SVerticalBox::Slot()
+				+ SVerticalBox::Slot()
 				.VAlign(VAlign_Center)
 				.Padding(2.0f)
 				.AutoHeight()
 				[
-					SAssignNew(SearchBox, SSearchBox)
-					.HintText(LOCTEXT("SearchBoxHint", "Search timers or groups"))
-					.OnTextChanged(this, &STimersView::SearchBox_OnTextChanged)
-					.IsEnabled(this, &STimersView::SearchBox_IsEnabled)
-					.ToolTipText(LOCTEXT("FilterSearchHint", "Type here to search timer or group"))
+					SNew(SHorizontalBox)
+
+					// Search box
+					+ SHorizontalBox::Slot()
+					.VAlign(VAlign_Center)
+					.Padding(2.0f)
+					.FillWidth(1.0f)
+					[
+						SAssignNew(SearchBox, SSearchBox)
+						.HintText(LOCTEXT("SearchBoxHint", "Search timers or groups"))
+						.OnTextChanged(this, &STimersView::SearchBox_OnTextChanged)
+						.IsEnabled(this, &STimersView::SearchBox_IsEnabled)
+						.ToolTipText(LOCTEXT("FilterSearchHint", "Type here to search timer or group"))
+					]
+
+					// Filter out timers with zero instance count
+					+ SHorizontalBox::Slot()
+					.VAlign(VAlign_Center)
+					.Padding(2.0f)
+					.AutoWidth()
+					[
+						SNew(SCheckBox)
+						.Style(FEditorStyle::Get(), "ToggleButtonCheckbox")
+						.HAlign(HAlign_Center)
+						.Padding(2.0f)
+						.OnCheckStateChanged(this, &STimersView::FilterOutZeroCountTimers_OnCheckStateChanged)
+						.IsChecked(this, &STimersView::FilterOutZeroCountTimers_IsChecked)
+						.ToolTipText(LOCTEXT("FilterOutZeroCountTimers_Tooltip", "Filter out the timers having zero total instance count (aggregated stats)."))
+						[
+							//TODO: SNew(SImage)
+							SNew(STextBlock)
+							.Text(LOCTEXT("FilterOutZeroCountTimers_Button", " !0 "))
+							.TextStyle(FEditorStyle::Get(), TEXT("Profiler.Caption"))
+						]
+					]
 				]
 
 				// Group by
-				+SVerticalBox::Slot()
+				+ SVerticalBox::Slot()
 				.VAlign(VAlign_Center)
 				.Padding(2.0f)
 				.AutoHeight()
 				[
 					SNew(SHorizontalBox)
 
-					+SHorizontalBox::Slot()
+					+ SHorizontalBox::Slot()
 					.FillWidth(1.0f)
 					.VAlign(VAlign_Center)
 					[
@@ -114,7 +138,7 @@ void STimersView::Construct(const FArguments& InArgs)
 						.Text(LOCTEXT("GroupByText", "Group by"))
 					]
 
-					+SHorizontalBox::Slot()
+					+ SHorizontalBox::Slot()
 					.FillWidth(2.0f)
 					.VAlign(VAlign_Center)
 					[
@@ -138,21 +162,21 @@ void STimersView::Construct(const FArguments& InArgs)
 				[
 					SNew(SHorizontalBox)
 
-					+SHorizontalBox::Slot()
+					+ SHorizontalBox::Slot()
 					.Padding(FMargin(0.0f,0.0f,1.0f,0.0f))
 					.FillWidth(1.0f)
 					[
 						GetToggleButtonForTimerType(ETimerNodeType::GpuScope)
 					]
 
-					//+SHorizontalBox::Slot()
+					//+ SHorizontalBox::Slot()
 					//.Padding(FMargin(1.0f,0.0f,1.0f,0.0f))
 					//.FillWidth(1.0f)
 					//[
 					//	GetToggleButtonForTimerType(ETimerNodeType::ComputeScope)
 					//]
 
-					+SHorizontalBox::Slot()
+					+ SHorizontalBox::Slot()
 					.Padding(FMargin(1.0f,0.0f,1.0f,0.0f))
 					.FillWidth(1.0f)
 					[
@@ -163,7 +187,7 @@ void STimersView::Construct(const FArguments& InArgs)
 		]
 
 		// Tree view
-		+SVerticalBox::Slot()
+		+ SVerticalBox::Slot()
 		.FillHeight(1.0f)
 		.Padding(0.0f, 6.0f, 0.0f, 0.0f)
 		[
@@ -223,6 +247,7 @@ void STimersView::Construct(const FArguments& InArgs)
 	Filters->Add(TextFilter);
 
 	CreateGroupByOptionsSources();
+	CreateSortings();
 
 	// Register ourselves with the Insights manager.
 	FInsightsManager::Get()->GetSessionChangedEvent().AddSP(this, &STimersView::InsightsManager_OnSessionChanged);
@@ -236,29 +261,28 @@ END_SLATE_FUNCTION_BUILD_OPTIMIZATION
 
 TSharedPtr<SWidget> STimersView::TreeView_GetMenuContent()
 {
-	const TArray<FTimerNodePtr> SelectedTimerNodes = TreeView->GetSelectedItems();
-	const int32 NumSelectedTimerNodes = SelectedTimerNodes.Num();
-	FTimerNodePtr SelectedTimerNode = NumSelectedTimerNodes ? SelectedTimerNodes[0] : nullptr;
+	const TArray<FTimerNodePtr> SelectedNodes = TreeView->GetSelectedItems();
+	const int32 NumSelectedNodes = SelectedNodes.Num();
+	FTimerNodePtr SelectedNode = NumSelectedNodes ? SelectedNodes[0] : nullptr;
 
-	const FTimersViewColumn* const * ColumnPtrPtr = FTimersViewColumnFactory::Get().ColumnIdToPtrMapping.Find(HoveredColumnId);
-	const FTimersViewColumn* const ColumnPtr = (ColumnPtrPtr != nullptr) ? *ColumnPtrPtr : nullptr;
+	const TSharedPtr<Insights::FTableColumn> HoveredColumnPtr = Table->FindColumn(HoveredColumnId);
 
 	FText SelectionStr;
 	FText PropertyName;
 	FText PropertyValue;
 
-	if (NumSelectedTimerNodes == 0)
+	if (NumSelectedNodes == 0)
 	{
 		SelectionStr = LOCTEXT("NothingSelected", "Nothing selected");
 	}
-	else if (NumSelectedTimerNodes == 1)
+	else if (NumSelectedNodes == 1)
 	{
-		if (ColumnPtr != nullptr)
+		if (HoveredColumnPtr != nullptr)
 		{
-			PropertyName = ColumnPtr->ShortName;
-			PropertyValue = ColumnPtr->GetFormattedValue(*SelectedTimerNode);
+			PropertyName = HoveredColumnPtr->GetShortName();
+			PropertyValue = HoveredColumnPtr->GetValueAsTooltipText(*SelectedNode);
 		}
-		SelectionStr = FText::FromName(SelectedTimerNode->GetName());
+		SelectionStr = FText::FromName(SelectedNode->GetName());
 	}
 	else
 	{
@@ -377,22 +401,22 @@ void STimersView::TreeView_BuildSortByMenu(FMenuBuilder& MenuBuilder)
 
 	MenuBuilder.BeginSection("ColumnName", LOCTEXT("ContextMenu_Header_Misc_ColumnName", "Column Name"));
 
-	for (auto It = TreeViewHeaderColumns.CreateConstIterator(); It; ++It)
+	for (const TSharedPtr<Insights::FTableColumn>& ColumnPtr : Table->GetColumns())
 	{
-		const FTimersViewColumn& Column = It.Value();
+		const Insights::FTableColumn& Column = *ColumnPtr;
 
-		if (Column.bIsVisible && Column.bCanBeSorted())
+		if (Column.IsVisible() && Column.CanBeSorted())
 		{
 			FUIAction Action_SortByColumn
 			(
-				FExecuteAction::CreateSP(this, &STimersView::ContextMenu_SortByColumn_Execute, Column.Id),
-				FCanExecuteAction::CreateSP(this, &STimersView::ContextMenu_SortByColumn_CanExecute, Column.Id),
-				FIsActionChecked::CreateSP(this, &STimersView::ContextMenu_SortByColumn_IsChecked, Column.Id)
+				FExecuteAction::CreateSP(this, &STimersView::ContextMenu_SortByColumn_Execute, Column.GetId()),
+				FCanExecuteAction::CreateSP(this, &STimersView::ContextMenu_SortByColumn_CanExecute, Column.GetId()),
+				FIsActionChecked::CreateSP(this, &STimersView::ContextMenu_SortByColumn_IsChecked, Column.GetId())
 			);
 			MenuBuilder.AddMenuEntry
 			(
-				Column.TitleName,
-				Column.Description,
+				Column.GetTitleName(),
+				Column.GetDescription(),
 				FSlateIcon(), Action_SortByColumn, NAME_None, EUserInterfaceActionType::RadioButton
 			);
 		}
@@ -439,20 +463,20 @@ void STimersView::TreeView_BuildViewColumnMenu(FMenuBuilder& MenuBuilder)
 {
 	MenuBuilder.BeginSection("ViewColumn", LOCTEXT("ContextMenu_Header_Columns_View", "View Column"));
 
-	for (auto It = TreeViewHeaderColumns.CreateConstIterator(); It; ++It)
+	for (const TSharedPtr<Insights::FTableColumn>& ColumnPtr : Table->GetColumns())
 	{
-		const FTimersViewColumn& Column = It.Value();
+		const Insights::FTableColumn& Column = *ColumnPtr;
 
 		FUIAction Action_ToggleColumn
 		(
-			FExecuteAction::CreateSP(this, &STimersView::ContextMenu_ToggleColumn_Execute, Column.Id),
-			FCanExecuteAction::CreateSP(this, &STimersView::ContextMenu_ToggleColumn_CanExecute, Column.Id),
-			FIsActionChecked::CreateSP(this, &STimersView::ContextMenu_ToggleColumn_IsChecked, Column.Id)
+			FExecuteAction::CreateSP(this, &STimersView::ToggleColumnVisibility, Column.GetId()),
+			FCanExecuteAction::CreateSP(this, &STimersView::CanToggleColumnVisibility, Column.GetId()),
+			FIsActionChecked::CreateSP(this, &STimersView::IsColumnVisible, Column.GetId())
 		);
 		MenuBuilder.AddMenuEntry
 		(
-			Column.TitleName,
-			Column.Description,
+			Column.GetTitleName(),
+			Column.GetDescription(),
 			FSlateIcon(), Action_ToggleColumn, NAME_None, EUserInterfaceActionType::ToggleButton
 		);
 	}
@@ -464,103 +488,46 @@ void STimersView::TreeView_BuildViewColumnMenu(FMenuBuilder& MenuBuilder)
 
 void STimersView::InitializeAndShowHeaderColumns()
 {
-	const int32 NumColumns = FTimersViewColumnFactory::Get().Collection.Num();
-	for (int32 ColumnIndex = 0; ColumnIndex < NumColumns; ColumnIndex++)
-	{
-		TreeViewHeaderRow_CreateColumnArgs(ColumnIndex);
-	}
+	// Create columns.
+	TArray<TSharedPtr<Insights::FTableColumn>> Columns;
+	FTimersViewColumnFactory::CreateTimersViewColumns(Columns);
+	Table->SetColumns(Columns);
 
-	for (auto It = TreeViewHeaderColumns.CreateConstIterator(); It; ++It)
+	// Show columns.
+	for (const TSharedPtr<Insights::FTableColumn>& ColumnPtr : Table->GetColumns())
 	{
-		const FTimersViewColumn& Column = It.Value();
-
-		if (Column.bIsVisible)
+		if (ColumnPtr->ShouldBeVisible())
 		{
-			TreeViewHeaderRow_ShowColumn(Column.Id);
+			ShowColumn(ColumnPtr->GetId());
 		}
 	}
 }
 
 ////////////////////////////////////////////////////////////////////////////////////////////////////
 
-void STimersView::TreeViewHeaderRow_CreateColumnArgs(const int32 ColumnIndex)
-{
-	const FTimersViewColumn& Column = *FTimersViewColumnFactory::Get().Collection[ColumnIndex];
-	SHeaderRow::FColumn::FArguments ColumnArgs;
-
-	ColumnArgs
-		.ColumnId(Column.Id)
-		.DefaultLabel(Column.ShortName)
-		.SortMode(EColumnSortMode::None)
-		.HAlignHeader(HAlign_Fill)
-		.VAlignHeader(VAlign_Fill)
-		.HeaderContentPadding(TOptional<FMargin>(2.0f))
-		.HAlignCell(HAlign_Fill)
-		.VAlignCell(VAlign_Fill)
-		.SortMode(this, &STimersView::GetSortModeForColumn, Column.Id)
-		.OnSort(this, &STimersView::OnSortModeChanged)
-		.ManualWidth(Column.InitialColumnWidth)
-		.FixedWidth(Column.bIsFixedColumnWidth() ? Column.InitialColumnWidth : TOptional<float>())
-		.HeaderContent()
-		[
-			SNew(SHorizontalBox)
-			.ToolTip(STimersViewTooltip::GetColumnTooltip(Column))
-
-			+ SHorizontalBox::Slot()
-			.HAlign(Column.HorizontalAlignment)
-			.VAlign(VAlign_Center)
-			[
-				SNew(STextBlock)
-				.Text(Column.ShortName)
-			]
-		]
-		.MenuContent()
-		[
-			TreeViewHeaderRow_GenerateColumnMenu(Column)
-		];
-
-	TreeViewHeaderColumnArgs.Add(Column.Id, ColumnArgs);
-	TreeViewHeaderColumns.Add(Column.Id, Column);
-}
-
-////////////////////////////////////////////////////////////////////////////////////////////////////
-
-void STimersView::TreeViewHeaderRow_ShowColumn(const FName ColumnId)
+FText STimersView::GetColumnHeaderText(const FName ColumnId) const
 {
-	FTimersViewColumn& Column = TreeViewHeaderColumns.FindChecked(ColumnId);
-	Column.bIsVisible = true;
-	SHeaderRow::FColumn::FArguments& ColumnArgs = TreeViewHeaderColumnArgs.FindChecked(ColumnId);
-
-	const int32 NumColumns = TreeViewHeaderRow->GetColumns().Num();
-	int32 ColumnIndex = 0;
-	for (; ColumnIndex < NumColumns; ColumnIndex++)
-	{
-		const SHeaderRow::FColumn& CurrentColumn = TreeViewHeaderRow->GetColumns()[ColumnIndex];
-		const FTimersViewColumn& CurrentTimersViewColumn = TreeViewHeaderColumns.FindChecked(CurrentColumn.ColumnId);
-		if (Column.Order < CurrentTimersViewColumn.Order)
-			break;
-	}
-
-	TreeViewHeaderRow->InsertColumn(ColumnArgs, ColumnIndex);
+	const Insights::FTableColumn& Column = *Table->FindColumnChecked(ColumnId);
+	return Column.GetShortName();
 }
 
 ////////////////////////////////////////////////////////////////////////////////////////////////////
 
-TSharedRef<SWidget> STimersView::TreeViewHeaderRow_GenerateColumnMenu(const FTimersViewColumn& Column)
+TSharedRef<SWidget> STimersView::TreeViewHeaderRow_GenerateColumnMenu(const Insights::FTableColumn& Column)
 {
 	bool bIsMenuVisible = false;
 
 	const bool bShouldCloseWindowAfterMenuSelection = true;
 	FMenuBuilder MenuBuilder(bShouldCloseWindowAfterMenuSelection, NULL);
 	{
-		if (Column.bCanBeHidden())
+		if (Column.CanBeHidden())
 		{
 			MenuBuilder.BeginSection("Column", LOCTEXT("TreeViewHeaderRow_Header_Column", "Column"));
 
 			FUIAction Action_HideColumn
 			(
-				FExecuteAction::CreateSP(this, &STimersView::HeaderMenu_HideColumn_Execute, Column.Id),
-				FCanExecuteAction::CreateSP(this, &STimersView::HeaderMenu_HideColumn_CanExecute, Column.Id)
+				FExecuteAction::CreateSP(this, &STimersView::HideColumn, Column.GetId()),
+				FCanExecuteAction::CreateSP(this, &STimersView::CanHideColumn, Column.GetId())
 			);
 
 			MenuBuilder.AddMenuEntry
@@ -574,15 +541,15 @@ TSharedRef<SWidget> STimersView::TreeViewHeaderRow_GenerateColumnMenu(const FTim
 			MenuBuilder.EndSection();
 		}
 
-		if (Column.bCanBeSorted())
+		if (Column.CanBeSorted())
 		{
 			MenuBuilder.BeginSection("SortMode", LOCTEXT("ContextMenu_Header_Misc_Sort_SortMode", "Sort Mode"));
 
 			FUIAction Action_SortAscending
 			(
-				FExecuteAction::CreateSP(this, &STimersView::HeaderMenu_SortMode_Execute, Column.Id, EColumnSortMode::Ascending),
-				FCanExecuteAction::CreateSP(this, &STimersView::HeaderMenu_SortMode_CanExecute, Column.Id, EColumnSortMode::Ascending),
-				FIsActionChecked::CreateSP(this, &STimersView::HeaderMenu_SortMode_IsChecked, Column.Id, EColumnSortMode::Ascending)
+				FExecuteAction::CreateSP(this, &STimersView::HeaderMenu_SortMode_Execute, Column.GetId(), EColumnSortMode::Ascending),
+				FCanExecuteAction::CreateSP(this, &STimersView::HeaderMenu_SortMode_CanExecute, Column.GetId(), EColumnSortMode::Ascending),
+				FIsActionChecked::CreateSP(this, &STimersView::HeaderMenu_SortMode_IsChecked, Column.GetId(), EColumnSortMode::Ascending)
 			);
 			MenuBuilder.AddMenuEntry
 			(
@@ -593,9 +560,9 @@ TSharedRef<SWidget> STimersView::TreeViewHeaderRow_GenerateColumnMenu(const FTim
 
 			FUIAction Action_SortDescending
 			(
-				FExecuteAction::CreateSP(this, &STimersView::HeaderMenu_SortMode_Execute, Column.Id, EColumnSortMode::Descending),
-				FCanExecuteAction::CreateSP(this, &STimersView::HeaderMenu_SortMode_CanExecute, Column.Id, EColumnSortMode::Descending),
-				FIsActionChecked::CreateSP(this, &STimersView::HeaderMenu_SortMode_IsChecked, Column.Id, EColumnSortMode::Descending)
+				FExecuteAction::CreateSP(this, &STimersView::HeaderMenu_SortMode_Execute, Column.GetId(), EColumnSortMode::Descending),
+				FCanExecuteAction::CreateSP(this, &STimersView::HeaderMenu_SortMode_CanExecute, Column.GetId(), EColumnSortMode::Descending),
+				FIsActionChecked::CreateSP(this, &STimersView::HeaderMenu_SortMode_IsChecked, Column.GetId(), EColumnSortMode::Descending)
 			);
 			MenuBuilder.AddMenuEntry
 			(
@@ -608,12 +575,12 @@ TSharedRef<SWidget> STimersView::TreeViewHeaderRow_GenerateColumnMenu(const FTim
 			MenuBuilder.EndSection();
 		}
 
-		if (Column.bCanBeFiltered())
-		{
-			MenuBuilder.BeginSection("FilterMode", LOCTEXT("ContextMenu_Header_Misc_Filter_FilterMode", "Filter Mode"));
-			bIsMenuVisible = true;
-			MenuBuilder.EndSection();
-		}
+		//if (Column.CanBeFiltered())
+		//{
+		//	MenuBuilder.BeginSection("FilterMode", LOCTEXT("ContextMenu_Header_Misc_Filter_FilterMode", "Filter Mode"));
+		//	bIsMenuVisible = true;
+		//	MenuBuilder.EndSection();
+		//}
 	}
 
 	/*
@@ -637,7 +604,7 @@ void STimersView::InsightsManager_OnSessionChanged()
 	if (NewSession != Session)
 	{
 		Session = NewSession;
-		RebuildTree();
+		Reset();
 	}
 	else
 	{
@@ -649,9 +616,8 @@ void STimersView::InsightsManager_OnSessionChanged()
 
 void STimersView::UpdateTree()
 {
-	// Create groups, sort timers within the group and apply filtering.
 	CreateGroups();
-	SortTimers();
+	SortTreeNodes();
 	ApplyFiltering();
 }
 
@@ -669,14 +635,16 @@ void STimersView::ApplyFiltering()
 		GroupPtr->ClearFilteredChildren();
 		const bool bIsGroupVisible = Filters->PassesAllFilters(GroupPtr);
 
-		const TArray<FTimerNodePtr>& GroupChildren = GroupPtr->GetChildren();
+		const TArray<Insights::FBaseTreeNodePtr>& GroupChildren = GroupPtr->GetChildren();
 		const int32 NumChildren = GroupChildren.Num();
 		int32 NumVisibleChildren = 0;
 		for (int32 Cx = 0; Cx < NumChildren; ++Cx)
 		{
 			// Add a child.
-			const FTimerNodePtr& NodePtr = GroupChildren[Cx];
-			const bool bIsChildVisible = Filters->PassesAllFilters(NodePtr) && bTimerNodeIsVisible[static_cast<int>(NodePtr->GetType())];
+			const FTimerNodePtr& NodePtr = StaticCastSharedPtr<FTimerNode, Insights::FBaseTreeNode>(GroupChildren[Cx]);
+			const bool bIsChildVisible = (!bFilterOutZeroCountTimers || NodePtr->GetAggregatedStats().InstanceCount > 0)
+									  && bTimerTypeIsVisible[static_cast<int>(NodePtr->GetType())]
+									  && Filters->PassesAllFilters(NodePtr);
 			if (bIsChildVisible)
 			{
 				GroupPtr->AddFilteredChild(NodePtr);
@@ -688,11 +656,11 @@ void STimersView::ApplyFiltering()
 		{
 			// Add a group.
 			FilteredGroupNodes.Add(GroupPtr);
-			GroupPtr->bForceExpandGroupNode = true;
+			GroupPtr->SetExpansion(true);
 		}
 		else
 		{
-			GroupPtr->bForceExpandGroupNode = false;
+			GroupPtr->SetExpansion(false);
 		}
 	}
 
@@ -710,7 +678,7 @@ void STimersView::ApplyFiltering()
 		for (int32 Fx = 0; Fx < FilteredGroupNodes.Num(); Fx++)
 		{
 			const FTimerNodePtr& GroupPtr = FilteredGroupNodes[Fx];
-			TreeView->SetItemExpansion(GroupPtr, GroupPtr->bForceExpandGroupNode);
+			TreeView->SetItemExpansion(GroupPtr, GroupPtr->IsExpanded());
 		}
 	}
 	else
@@ -760,12 +728,12 @@ TSharedRef<SWidget> STimersView::GetToggleButtonForTimerType(const ETimerNodeTyp
 						.Image(TimerNodeTypeHelper::GetIconForTimerNodeType(NodeType))
 				]
 
-			+SHorizontalBox::Slot()
+			+ SHorizontalBox::Slot()
 				.Padding(2.0f, 0.0f, 0.0f, 0.0f)
 				.VAlign(VAlign_Center)
 				[
 					SNew(STextBlock)
-						.Text(TimerNodeTypeHelper::ToName(NodeType))
+						.Text(TimerNodeTypeHelper::ToText(NodeType))
 						.TextStyle(FEditorStyle::Get(), TEXT("Profiler.Caption"))
 				]
 		];
@@ -773,9 +741,24 @@ TSharedRef<SWidget> STimersView::GetToggleButtonForTimerType(const ETimerNodeTyp
 
 ////////////////////////////////////////////////////////////////////////////////////////////////////
 
+void STimersView::FilterOutZeroCountTimers_OnCheckStateChanged(ECheckBoxState NewRadioState)
+{
+	bFilterOutZeroCountTimers = (NewRadioState == ECheckBoxState::Checked);
+	ApplyFiltering();
+}
+
+////////////////////////////////////////////////////////////////////////////////////////////////////
+
+ECheckBoxState STimersView::FilterOutZeroCountTimers_IsChecked() const
+{
+	return bFilterOutZeroCountTimers ? ECheckBoxState::Checked : ECheckBoxState::Unchecked;
+}
+
+////////////////////////////////////////////////////////////////////////////////////////////////////
+
 void STimersView::FilterByTimerType_OnCheckStateChanged(ECheckBoxState NewRadioState, const ETimerNodeType InStatType)
 {
-	bTimerNodeIsVisible[static_cast<int>(InStatType)] = (NewRadioState == ECheckBoxState::Checked);
+	bTimerTypeIsVisible[static_cast<int>(InStatType)] = (NewRadioState == ECheckBoxState::Checked);
 	ApplyFiltering();
 }
 
@@ -783,7 +766,7 @@ void STimersView::FilterByTimerType_OnCheckStateChanged(ECheckBoxState NewRadioS
 
 ECheckBoxState STimersView::FilterByTimerType_IsChecked(const ETimerNodeType InStatType) const
 {
-	return bTimerNodeIsVisible[static_cast<int>(InStatType)] ? ECheckBoxState::Checked : ECheckBoxState::Unchecked;
+	return bTimerTypeIsVisible[static_cast<int>(InStatType)] ? ECheckBoxState::Checked : ECheckBoxState::Unchecked;
 }
 
 ////////////////////////////////////////////////////////////////////////////////////////////////////
@@ -805,13 +788,9 @@ void STimersView::TreeView_OnSelectionChanged(FTimerNodePtr SelectedItem, ESelec
 	if (SelectInfo != ESelectInfo::Direct)
 	{
 		TArray<FTimerNodePtr> SelectedItems = TreeView->GetSelectedItems();
-		if (SelectedItems.Num() == 1)
-		{
-			//HighlightedTimerName = SelectedItems[0]->GetName();
-		}
-		else
+		if (SelectedItems.Num() == 1 && !SelectedItems[0]->IsGroup())
 		{
-			//HighlightedTimerName = NAME_None;
+			FTimingProfilerManager::Get()->SetSelectedTimer(SelectedItems[0]->GetId());
 		}
 	}
 }
@@ -820,32 +799,37 @@ void STimersView::TreeView_OnSelectionChanged(FTimerNodePtr SelectedItem, ESelec
 
 void STimersView::TreeView_OnGetChildren(FTimerNodePtr InParent, TArray<FTimerNodePtr>& OutChildren)
 {
-	OutChildren = InParent->GetFilteredChildren();
+	const TArray<Insights::FBaseTreeNodePtr>& Children = InParent->GetFilteredChildren();
+	OutChildren.Reset(Children.Num());
+	for (const Insights::FBaseTreeNodePtr& Child : Children)
+	{
+		OutChildren.Add(StaticCastSharedPtr<FTimerNode, Insights::FBaseTreeNode>(Child));
+	}
 }
 
 ////////////////////////////////////////////////////////////////////////////////////////////////////
 
-void STimersView::TreeView_OnMouseButtonDoubleClick(FTimerNodePtr TimerNode)
+void STimersView::TreeView_OnMouseButtonDoubleClick(FTimerNodePtr TimerNodePtr)
 {
-	if (!TimerNode->IsGroup())
+	if (TimerNodePtr->IsGroup())
 	{
-		//im:TODO: const bool bIsStatTracked = FTimingProfilerManager::Get()->IsStatTracked(TimerNode->GetId());
-	//	if (!bIsStatTracked)
+		const bool bIsGroupExpanded = TreeView->IsItemExpanded(TimerNodePtr);
+		TreeView->SetItemExpansion(TimerNodePtr, !bIsGroupExpanded);
+	}
+	else
+	{
+		//im:TODO: const bool bIsTracked = FTimingProfilerManager::Get()->IsTimerTracked(TimerNodePtr->GetId());
+	//	if (!bIsTracked)
 	//	{
-	//		// Add a new graph.
-	//		FTimingProfilerManager::Get()->TrackStat(TimerNode->GetId());
+	//		// Add a new graph series.
+	//		FTimingProfilerManager::Get()->TrackTimer(TimerNodePtr->GetId());
 	//	}
 	//	else
 	//	{
-	//		// Remove a graph
-	//		FTimingProfilerManager::Get()->UntrackStat(TimerNode->GetId());
+	//		// Remove the corresponding graph series.
+	//		FTimingProfilerManager::Get()->UntrackTimer(TimerNodePtr->GetId());
 	//	}
 	}
-	else
-	{
-		const bool bIsGroupExpanded = TreeView->IsItemExpanded(TimerNode);
-		TreeView->SetItemExpansion(TimerNode, !bIsGroupExpanded);
-	}
 }
 
 ////////////////////////////////////////////////////////////////////////////////////////////////////
@@ -857,11 +841,12 @@ TSharedRef<ITableRow> STimersView::TreeView_OnGenerateRow(FTimerNodePtr TimerNod
 	TSharedRef<ITableRow> TableRow =
 		SNew(STimerTableRow, OwnerTable)
 		.OnShouldBeEnabled(this, &STimersView::TableRow_ShouldBeEnabled)
-		.OnIsColumnVisible(this, &STimersView::TableRow_IsColumnVisible)
-		.OnSetHoveredTableCell(this, &STimersView::TableRow_SetHoveredTableCell)
+		.OnIsColumnVisible(this, &STimersView::IsColumnVisible)
+		.OnSetHoveredCell(this, &STimersView::TableRow_SetHoveredCell)
 		.OnGetColumnOutlineHAlignmentDelegate(this, &STimersView::TableRow_GetColumnOutlineHAlignment)
 		.HighlightText(this, &STimersView::TableRow_GetHighlightText)
-		.HighlightedTimerName(this, &STimersView::TableRow_GetHighlightedTimerName)
+		.HighlightedNodeName(this, &STimersView::TableRow_GetHighlightedNodeName)
+		.TablePtr(Table)
 		.TimerNodePtr(TimerNodePtr);
 
 	return TableRow;
@@ -869,26 +854,22 @@ TSharedRef<ITableRow> STimersView::TreeView_OnGenerateRow(FTimerNodePtr TimerNod
 
 ////////////////////////////////////////////////////////////////////////////////////////////////////
 
-bool STimersView::TableRow_IsColumnVisible(const FName ColumnId) const
+bool STimersView::TableRow_ShouldBeEnabled(const uint32 TimerId) const
 {
-	bool bResult = false;
-	const FTimersViewColumn& ColumnPtr = TreeViewHeaderColumns.FindChecked(ColumnId);
-	return ColumnPtr.bIsVisible;
+	return true;//im:TODO: Session->GetAggregatedStat(TimerId) != nullptr;
 }
 
 ////////////////////////////////////////////////////////////////////////////////////////////////////
 
-void STimersView::TableRow_SetHoveredTableCell(const FName ColumnId, const FTimerNodePtr TimerNodePtr)
+void STimersView::TableRow_SetHoveredCell(TSharedPtr<Insights::FTable> InTablePtr, TSharedPtr<Insights::FTableColumn> InColumnPtr, const FTimerNodePtr InNodePtr)
 {
-	HoveredColumnId = ColumnId;
+	HoveredColumnId = InColumnPtr ? InColumnPtr->GetId() : FName();
 
 	const bool bIsAnyMenusVisible = FSlateApplication::Get().AnyMenusVisible();
 	if (!HasMouseCapture() && !bIsAnyMenusVisible)
 	{
-		HoveredTimerNodePtr = TimerNodePtr;
+		HoveredNodePtr = InNodePtr;
 	}
-
-	//UE_LOG(TimingProfiler, Log, TEXT("%s -> %s"), *HoveredColumnId.GetPlainNameString(), TimerNodePtr.IsValid() ? *TimerNodePtr->GetName().GetPlainNameString() : TEXT("nullptr"));
 }
 
 ////////////////////////////////////////////////////////////////////////////////////////////////////
@@ -923,16 +904,9 @@ FText STimersView::TableRow_GetHighlightText() const
 
 ////////////////////////////////////////////////////////////////////////////////////////////////////
 
-FName STimersView::TableRow_GetHighlightedTimerName() const
-{
-	return HighlightedTimerName;
-}
-
-////////////////////////////////////////////////////////////////////////////////////////////////////
-
-bool STimersView::TableRow_ShouldBeEnabled(const uint32 TimerId) const
+FName STimersView::TableRow_GetHighlightedNodeName() const
 {
-	return true;//im:TODO: Session->GetAggregatedStat(TimerId) != nullptr;
+	return HighlightedNodeName;
 }
 
 ////////////////////////////////////////////////////////////////////////////////////////////////////
@@ -954,7 +928,7 @@ bool STimersView::SearchBox_IsEnabled() const
 }
 
 ////////////////////////////////////////////////////////////////////////////////////////////////////
-// GroupBy
+// Grouping
 ////////////////////////////////////////////////////////////////////////////////////////////////////
 
 void STimersView::CreateGroups()
@@ -999,7 +973,7 @@ void STimersView::CreateGroups()
 	{
 		for (const FTimerNodePtr& TimerNodePtr : TimerNodes)
 		{
-			const FName GroupName = *TimerNodeTypeHelper::ToName(TimerNodePtr->GetType()).ToString();
+			const FName GroupName = *TimerNodeTypeHelper::ToText(TimerNodePtr->GetType()).ToString();
 
 			FTimerNodePtr* GroupPtr = GroupNodeSet.Find(GroupName);
 			if (!GroupPtr)
@@ -1044,9 +1018,6 @@ void STimersView::CreateGroups()
 	}
 
 	GroupNodeSet.GenerateValueArray(GroupNodes);
-
-	// Sort by a fake group name.
-	GroupNodes.Sort(TimerNodeSortingHelper::ByNameAscending());
 }
 
 ////////////////////////////////////////////////////////////////////////////////////////////////////
@@ -1078,9 +1049,8 @@ void STimersView::GroupBy_OnSelectionChanged(TSharedPtr<ETimerGroupingMode> NewG
 	{
 		GroupingMode = *NewGroupingMode;
 
-		// Create groups, sort timers within the group and apply filtering.
 		CreateGroups();
-		SortTimers();
+		SortTreeNodes();
 		ApplyFiltering();
 	}
 }
@@ -1090,7 +1060,7 @@ void STimersView::GroupBy_OnSelectionChanged(TSharedPtr<ETimerGroupingMode> NewG
 TSharedRef<SWidget> STimersView::GroupBy_OnGenerateWidget(TSharedPtr<ETimerGroupingMode> InGroupingMode) const
 {
 	return SNew(STextBlock)
-		.Text(TimerNodeGroupingHelper::ToName(*InGroupingMode))
+		.Text(TimerNodeGroupingHelper::ToText(*InGroupingMode))
 		.ToolTipText(TimerNodeGroupingHelper::ToDescription(*InGroupingMode));
 }
 
@@ -1098,7 +1068,7 @@ TSharedRef<SWidget> STimersView::GroupBy_OnGenerateWidget(TSharedPtr<ETimerGroup
 
 FText STimersView::GroupBy_GetSelectedText() const
 {
-	return TimerNodeGroupingHelper::ToName(GroupingMode);
+	return TimerNodeGroupingHelper::ToText(GroupingMode);
 }
 
 ////////////////////////////////////////////////////////////////////////////////////////////////////
@@ -1109,50 +1079,85 @@ FText STimersView::GroupBy_GetSelectedTooltipText() const
 }
 
 ////////////////////////////////////////////////////////////////////////////////////////////////////
-// SortBy
+// Sorting
 ////////////////////////////////////////////////////////////////////////////////////////////////////
 
-void STimersView::SortTimers()
+const FName STimersView::GetDefaultColumnBeingSorted()
 {
-	const int32 NumGroups = GroupNodes.Num();
+	return FTimersViewColumns::TotalInclusiveTimeColumnID;
+}
+
+////////////////////////////////////////////////////////////////////////////////////////////////////
 
-	#define CHECK_AND_SORT_COLUMN(ColumnId, SortTypeName) \
-		if (ColumnBeingSorted == ColumnId) \
-		{ \
-			if (ColumnSortMode == EColumnSortMode::Type::Descending) \
-			{ \
-				for (int32 ID = 0; ID < NumGroups; ++ID) \
-				{ \
-					GroupNodes[ID]->SortChildren(TimerNodeSortingHelper::SortTypeName##Descending()); \
-				} \
-			} \
-			else /*if (ColumnSortMode == EColumnSortMode::Type::Ascending)*/ \
-			{ \
-				for (int32 ID = 0; ID < NumGroups; ++ID) \
-				{ \
-					GroupNodes[ID]->SortChildren(TimerNodeSortingHelper::SortTypeName##Ascending()); \
-				} \
-			} \
+const EColumnSortMode::Type STimersView::GetDefaultColumnSortMode()
+{
+	return EColumnSortMode::Descending;
+}
+
+////////////////////////////////////////////////////////////////////////////////////////////////////
+
+void STimersView::CreateSortings()
+{
+	AvailableSorters.Reset();
+	CurrentSorter = nullptr;
+
+	for (const TSharedPtr<Insights::FTableColumn> ColumnPtr : Table->GetColumns())
+	{
+		if (ColumnPtr->CanBeSorted())
+		{
+			TSharedPtr<Insights::ITableCellValueSorter> SorterPtr = ColumnPtr->GetValueSorter();
+			if (ensure(SorterPtr.IsValid()))
+			{
+				AvailableSorters.Add(SorterPtr);
+			}
+		}
+	}
+
+	UpdateCurrentSortingByColumn();
+}
+
+////////////////////////////////////////////////////////////////////////////////////////////////////
+
+void STimersView::UpdateCurrentSortingByColumn()
+{
+	TSharedPtr<Insights::FTableColumn> ColumnPtr = Table->FindColumn(ColumnBeingSorted);
+	CurrentSorter = ColumnPtr.IsValid() ? ColumnPtr->GetValueSorter() : nullptr;
+}
+
+////////////////////////////////////////////////////////////////////////////////////////////////////
+
+void STimersView::SortTreeNodes()
+{
+	if (CurrentSorter.IsValid())
+	{
+		for (FTimerNodePtr& Root : GroupNodes)
+		{
+			SortTreeNodesRec(*Root, *CurrentSorter);
 		}
+	}
+}
 
-		 CHECK_AND_SORT_COLUMN(FTimersViewColumns::NameColumnID,          ByName)
-	else CHECK_AND_SORT_COLUMN(FTimersViewColumns::MetaGroupNameColumnID, ByMetaGroupName)
-	else CHECK_AND_SORT_COLUMN(FTimersViewColumns::TypeColumnID,          ByType)
-	else CHECK_AND_SORT_COLUMN(FTimersViewColumns::InstanceCountColumnID, ByInstanceCount)
-	// Inclusive Time columns
-	else CHECK_AND_SORT_COLUMN(FTimersViewColumns::TotalInclusiveTimeColumnID, ByTotalInclusiveTime)
-	else CHECK_AND_SORT_COLUMN(FTimersViewColumns::MaxInclusiveTimeColumnID,     ByMaxInclusiveTime)
-	else CHECK_AND_SORT_COLUMN(FTimersViewColumns::AverageInclusiveTimeColumnID, ByAverageInclusiveTime)
-	else CHECK_AND_SORT_COLUMN(FTimersViewColumns::MedianInclusiveTimeColumnID,  ByMedianInclusiveTime)
-	else CHECK_AND_SORT_COLUMN(FTimersViewColumns::MinInclusiveTimeColumnID,     ByMinInclusiveTime)
-	// Exclusive Time columns
-	else CHECK_AND_SORT_COLUMN(FTimersViewColumns::TotalExclusiveTimeColumnID,   ByTotalExclusiveTime)
-	else CHECK_AND_SORT_COLUMN(FTimersViewColumns::MaxExclusiveTimeColumnID,     ByMaxExclusiveTime)
-	else CHECK_AND_SORT_COLUMN(FTimersViewColumns::AverageExclusiveTimeColumnID, ByAverageExclusiveTime)
-	else CHECK_AND_SORT_COLUMN(FTimersViewColumns::MedianExclusiveTimeColumnID,  ByMedianExclusiveTime)
-	else CHECK_AND_SORT_COLUMN(FTimersViewColumns::MinExclusiveTimeColumnID,     ByMinExclusiveTime)
+////////////////////////////////////////////////////////////////////////////////////////////////////
+
+void STimersView::SortTreeNodesRec(FTimerNode& Node, const Insights::ITableCellValueSorter& Sorter)
+{
+	if (ColumnSortMode == EColumnSortMode::Type::Descending)
+	{
+		Node.SortChildrenDescending(Sorter);
+	}
+	else // if (ColumnSortMode == EColumnSortMode::Type::Ascending)
+	{
+		Node.SortChildrenAscending(Sorter);
+	}
 
-	#undef CHECK_AND_SORT_COLUMN
+	for (Insights::FBaseTreeNodePtr ChildPtr : Node.GetChildren())
+	{
+		//if (ChildPtr->IsGroup())
+		if (ChildPtr->GetChildren().Num() > 0)
+		{
+			SortTreeNodesRec(*StaticCastSharedPtr<FTimerNode>(ChildPtr), Sorter);
+		}
+	}
 }
 
 ////////////////////////////////////////////////////////////////////////////////////////////////////
@@ -1173,9 +1178,9 @@ void STimersView::SetSortModeForColumn(const FName& ColumnId, const EColumnSortM
 {
 	ColumnBeingSorted = ColumnId;
 	ColumnSortMode = SortMode;
+	UpdateCurrentSortingByColumn();
 
-	// Sort timers and apply filtering.
-	SortTimers();
+	SortTreeNodes();
 	ApplyFiltering();
 }
 
@@ -1200,8 +1205,8 @@ bool STimersView::HeaderMenu_SortMode_IsChecked(const FName ColumnId, const ECol
 
 bool STimersView::HeaderMenu_SortMode_CanExecute(const FName ColumnId, const EColumnSortMode::Type InSortMode) const
 {
-	const FTimersViewColumn& Column = TreeViewHeaderColumns.FindChecked(ColumnId);
-	return Column.bCanBeSorted();
+	const Insights::FTableColumn& Column = *Table->FindColumnChecked(ColumnId);
+	return Column.CanBeSorted();
 }
 
 ////////////////////////////////////////////////////////////////////////////////////////////////////
@@ -1261,54 +1266,116 @@ void STimersView::ContextMenu_SortByColumn_Execute(const FName ColumnId)
 }
 
 ////////////////////////////////////////////////////////////////////////////////////////////////////
-// HideColumn action (HeaderMenu)
+// ShowColumn action
 ////////////////////////////////////////////////////////////////////////////////////////////////////
 
-bool STimersView::HeaderMenu_HideColumn_CanExecute(const FName ColumnId) const
+bool STimersView::CanShowColumn(const FName ColumnId) const
 {
-	const FTimersViewColumn& Column = TreeViewHeaderColumns.FindChecked(ColumnId);
-	return Column.bCanBeHidden();
+	return true;
 }
 
 ////////////////////////////////////////////////////////////////////////////////////////////////////
 
-void STimersView::HeaderMenu_HideColumn_Execute(const FName ColumnId)
+void STimersView::ShowColumn(const FName ColumnId)
 {
-	FTimersViewColumn& Column = TreeViewHeaderColumns.FindChecked(ColumnId);
-	Column.bIsVisible = false;
+	Insights::FTableColumn& Column = *Table->FindColumnChecked(ColumnId);
+	Column.Show();
+
+	SHeaderRow::FColumn::FArguments ColumnArgs;
+	ColumnArgs
+		.ColumnId(Column.GetId())
+		.DefaultLabel(Column.GetShortName())
+		.HAlignHeader(HAlign_Fill)
+		.VAlignHeader(VAlign_Fill)
+		.HeaderContentPadding(FMargin(2.0f))
+		.HAlignCell(HAlign_Fill)
+		.VAlignCell(VAlign_Fill)
+		.SortMode(this, &STimersView::GetSortModeForColumn, Column.GetId())
+		.OnSort(this, &STimersView::OnSortModeChanged)
+		.ManualWidth(Column.GetInitialWidth())
+		.FixedWidth(Column.IsFixedWidth() ? Column.GetInitialWidth() : TOptional<float>())
+		.HeaderContent()
+		[
+			SNew(SBox)
+			.ToolTip(STimersViewTooltip::GetColumnTooltip(Column))
+			.HAlign(Column.GetHorizontalAlignment())
+			.VAlign(VAlign_Center)
+			[
+				SNew(STextBlock)
+				.Text(this, &STimersView::GetColumnHeaderText, Column.GetId())
+			]
+		]
+		.MenuContent()
+		[
+			TreeViewHeaderRow_GenerateColumnMenu(Column)
+		];
+
+	int32 ColumnIndex = 0;
+	const int32 NewColumnPosition = Table->GetColumnPositionIndex(ColumnId);
+	const int32 NumColumns = TreeViewHeaderRow->GetColumns().Num();
+	for (; ColumnIndex < NumColumns; ColumnIndex++)
+	{
+		const SHeaderRow::FColumn& CurrentColumn = TreeViewHeaderRow->GetColumns()[ColumnIndex];
+		const int32 CurrentColumnPosition = Table->GetColumnPositionIndex(CurrentColumn.ColumnId);
+		if (NewColumnPosition < CurrentColumnPosition)
+		{
+			break;
+		}
+	}
+
+	TreeViewHeaderRow->InsertColumn(ColumnArgs, ColumnIndex);
+}
+
+////////////////////////////////////////////////////////////////////////////////////////////////////
+// HideColumn action
+////////////////////////////////////////////////////////////////////////////////////////////////////
+
+bool STimersView::CanHideColumn(const FName ColumnId) const
+{
+	const Insights::FTableColumn& Column = *Table->FindColumnChecked(ColumnId);
+	return Column.CanBeHidden();
+}
+
+////////////////////////////////////////////////////////////////////////////////////////////////////
+
+void STimersView::HideColumn(const FName ColumnId)
+{
+	Insights::FTableColumn& Column = *Table->FindColumnChecked(ColumnId);
+	Column.Hide();
+
 	TreeViewHeaderRow->RemoveColumn(ColumnId);
 }
 
 ////////////////////////////////////////////////////////////////////////////////////////////////////
-// ToggleColumn action (ContextMenu)
+// ToggleColumn action
 ////////////////////////////////////////////////////////////////////////////////////////////////////
 
-bool STimersView::ContextMenu_ToggleColumn_IsChecked(const FName ColumnId)
+bool STimersView::IsColumnVisible(const FName ColumnId)
 {
-	const FTimersViewColumn& Column = TreeViewHeaderColumns.FindChecked(ColumnId);
-	return Column.bIsVisible;
+	const Insights::FTableColumn& Column = *Table->FindColumnChecked(ColumnId);
+	return Column.IsVisible();
 }
 
 ////////////////////////////////////////////////////////////////////////////////////////////////////
 
-bool STimersView::ContextMenu_ToggleColumn_CanExecute(const FName ColumnId) const
+bool STimersView::CanToggleColumnVisibility(const FName ColumnId) const
 {
-	const FTimersViewColumn& Column = TreeViewHeaderColumns.FindChecked(ColumnId);
-	return Column.bCanBeHidden();
+	const Insights::FTableColumn& Column = *Table->FindColumnChecked(ColumnId);
+	return !Column.IsVisible() || Column.CanBeHidden();
 }
 
 ////////////////////////////////////////////////////////////////////////////////////////////////////
 
-void STimersView::ContextMenu_ToggleColumn_Execute(const FName ColumnId)
+void STimersView::ToggleColumnVisibility(const FName ColumnId)
 {
-	FTimersViewColumn& Column = TreeViewHeaderColumns.FindChecked(ColumnId);
-	if (Column.bIsVisible)
+	const Insights::FTableColumn& Column = *Table->FindColumnChecked(ColumnId);
+	if (Column.IsVisible())
 	{
-		HeaderMenu_HideColumn_Execute(ColumnId);
+		HideColumn(ColumnId);
 	}
 	else
 	{
-		TreeViewHeaderRow_ShowColumn(ColumnId);
+		ShowColumn(ColumnId);
 	}
 }
 
@@ -1325,18 +1392,17 @@ bool STimersView::ContextMenu_ShowAllColumns_CanExecute() const
 
 void STimersView::ContextMenu_ShowAllColumns_Execute()
 {
-	ColumnSortMode = GetDefaultColumnSortMode();
 	ColumnBeingSorted = GetDefaultColumnBeingSorted();
+	ColumnSortMode = GetDefaultColumnSortMode();
+	UpdateCurrentSortingByColumn();
 
-	const int32 NumColumns = FTimersViewColumnFactory::Get().Collection.Num();
-	for (int32 ColumnIndex = 0; ColumnIndex < NumColumns; ColumnIndex++)
+	for (const TSharedPtr<Insights::FTableColumn>& ColumnPtr : Table->GetColumns())
 	{
-		const FTimersViewColumn& DefaultColumn = *FTimersViewColumnFactory::Get().Collection[ColumnIndex];
-		const FTimersViewColumn& CurrentColumn = TreeViewHeaderColumns.FindChecked(DefaultColumn.Id);
+		const Insights::FTableColumn& Column = *ColumnPtr;
 
-		if (!CurrentColumn.bIsVisible)
+		if (!Column.IsVisible())
 		{
-			TreeViewHeaderRow_ShowColumn(DefaultColumn.Id);
+			ShowColumn(Column.GetId());
 		}
 	}
 }
@@ -1372,23 +1438,23 @@ void STimersView::ContextMenu_ShowMinMaxMedColumns_Execute()
 		FTimersViewColumns::MinExclusiveTimeColumnID,
 	};
 
-	ColumnSortMode = EColumnSortMode::Descending;
 	ColumnBeingSorted = FTimersViewColumns::TotalInclusiveTimeColumnID;
+	ColumnSortMode = EColumnSortMode::Descending;
+	UpdateCurrentSortingByColumn();
 
-	const int32 NumColumns = FTimersViewColumnFactory::Get().Collection.Num();
-	for (int32 ColumnIndex = 0; ColumnIndex < NumColumns; ColumnIndex++)
+	for (const TSharedPtr<Insights::FTableColumn>& ColumnPtr : Table->GetColumns())
 	{
-		const FTimersViewColumn& DefaultColumn = *FTimersViewColumnFactory::Get().Collection[ColumnIndex];
-		const FTimersViewColumn& CurrentColumn = TreeViewHeaderColumns.FindChecked(DefaultColumn.Id);
+		const Insights::FTableColumn& Column = *ColumnPtr;
 
-		bool bIsVisible = Preset.Contains(DefaultColumn.Id);
-		if (bIsVisible && !CurrentColumn.bIsVisible)
+		const bool bShouldBeVisible = Preset.Contains(Column.GetId());
+
+		if (bShouldBeVisible && !Column.IsVisible())
 		{
-			TreeViewHeaderRow_ShowColumn(DefaultColumn.Id);
+			ShowColumn(Column.GetId());
 		}
-		else if (!bIsVisible && CurrentColumn.bIsVisible)
+		else if (!bShouldBeVisible && Column.IsVisible())
 		{
-			HeaderMenu_HideColumn_Execute(DefaultColumn.Id);
+			HideColumn(Column.GetId());
 		}
 	}
 }
@@ -1406,37 +1472,47 @@ bool STimersView::ContextMenu_ResetColumns_CanExecute() const
 
 void STimersView::ContextMenu_ResetColumns_Execute()
 {
-	ColumnSortMode = GetDefaultColumnSortMode();
 	ColumnBeingSorted = GetDefaultColumnBeingSorted();
+	ColumnSortMode = GetDefaultColumnSortMode();
+	UpdateCurrentSortingByColumn();
 
-	const int32 NumColumns = FTimersViewColumnFactory::Get().Collection.Num();
-	for (int32 ColumnIndex = 0; ColumnIndex < NumColumns; ColumnIndex++)
+	for (const TSharedPtr<Insights::FTableColumn>& ColumnPtr : Table->GetColumns())
 	{
-		const FTimersViewColumn& DefaultColumn = *FTimersViewColumnFactory::Get().Collection[ColumnIndex];
-		const FTimersViewColumn& CurrentColumn = TreeViewHeaderColumns.FindChecked(DefaultColumn.Id);
+		const Insights::FTableColumn& Column = *ColumnPtr;
 
-		if (DefaultColumn.bIsVisible && !CurrentColumn.bIsVisible)
+		if (Column.ShouldBeVisible() && !Column.IsVisible())
 		{
-			TreeViewHeaderRow_ShowColumn(DefaultColumn.Id);
+			ShowColumn(Column.GetId());
 		}
-		else if (!DefaultColumn.bIsVisible && CurrentColumn.bIsVisible)
+		else if (!Column.ShouldBeVisible() && Column.IsVisible())
 		{
-			HeaderMenu_HideColumn_Execute(DefaultColumn.Id);
+			HideColumn(Column.GetId());
 		}
 	}
 }
 
 ////////////////////////////////////////////////////////////////////////////////////////////////////
 
+void STimersView::Reset()
+{
+	StatsStartTime = 0.0;
+	StatsEndTime = 0.0;
+
+	RebuildTree(true);
+}
+
+////////////////////////////////////////////////////////////////////////////////////////////////////
+
 void STimersView::RebuildTree(bool bResync)
 {
+	TArray<FTimerNodePtr> SelectedItems;
 	bool bListHasChanged = false;
 
 	if (bResync)
 	{
-		TimerNodes.Empty(TimerNodes.Num());
-		//TimerNodesMap.Empty(TimerNodesMap.Num());
-		TimerNodesIdMap.Empty(TimerNodesIdMap.Num());
+		const int32 PreviousNodeCount = TimerNodes.Num();
+		TimerNodes.Empty(PreviousNodeCount);
+		TimerNodesIdMap.Empty(PreviousNodeCount);
 		bListHasChanged = true;
 	}
 
@@ -1446,18 +1522,21 @@ void STimersView::RebuildTree(bool bResync)
 
 		const Trace::ITimingProfilerProvider& TimingProfilerProvider = *Trace::ReadTimingProfilerProvider(*Session.Get());
 
-		TimingProfilerProvider.ReadTimers([this, &bResync, &bListHasChanged](const Trace::FTimingProfilerTimer* Timers, uint64 TimersCount)
+		TimingProfilerProvider.ReadTimers([this, &bResync, &SelectedItems, &bListHasChanged](const Trace::FTimingProfilerTimer* Timers, uint64 TimersCount)
 		{
-			if (!bResync)
+			if (TimersCount != TimerNodes.Num())
 			{
-				bResync = (TimersCount != TimerNodes.Num());
+				bResync = true;
 			}
 
 			if (bResync)
 			{
-				TimerNodes.Empty(TimerNodes.Num());
-				//TimerNodesMap.Empty(TimerNodesMap.Num());
-				TimerNodesIdMap.Empty(TimerNodesIdMap.Num());
+				// Save selection.
+				TreeView->GetSelectedItems(SelectedItems);
+
+				const int32 PreviousNodeCount = TimerNodes.Num();
+				TimerNodes.Empty(PreviousNodeCount);
+				TimerNodesIdMap.Empty(PreviousNodeCount);
 				bListHasChanged = true;
 
 				for (uint64 TimerIndex = 0; TimerIndex < TimersCount; ++TimerIndex)
@@ -1466,10 +1545,8 @@ void STimersView::RebuildTree(bool bResync)
 					FName Name(Timer.Name);// +TEXT(" [GPU]")));
 					FName Group(Timer.IsGpuTimer ? TEXT("GPU") : TEXT("CPU"));
 					ETimerNodeType Type = Timer.IsGpuTimer ? ETimerNodeType::GpuScope : ETimerNodeType::CpuScope;
-					FTimerNode* TimerPtr = new FTimerNode(Timer.Id, Name, Group, Type);
-					FTimerNodePtr TimerNodePtr = MakeShareable(TimerPtr);
+					FTimerNodePtr TimerNodePtr = MakeShareable(new FTimerNode(Timer.Id, Name, Group, Type));
 					TimerNodes.Add(TimerNodePtr);
-					//TimerNodesMap.Add(Name, TimerNodePtr);
 					TimerNodesIdMap.Add(Timer.Id, TimerNodePtr);
 				}
 			}
@@ -1480,6 +1557,28 @@ void STimersView::RebuildTree(bool bResync)
 	{
 		UpdateTree();
 		UpdateStats(StatsStartTime, StatsEndTime);
+
+		TreeView->RebuildList();
+
+		// Restore selection.
+		if (SelectedItems.Num() > 0)
+		{
+			TreeView->ClearSelection();
+			TArray<FTimerNodePtr> NewSelectedItems;
+			for (const FTimerNodePtr& TimerNode : SelectedItems)
+			{
+				FTimerNodePtr* TimerNodePtrPtr = TimerNodesIdMap.Find(TimerNode->GetId());
+				if (TimerNodePtrPtr != nullptr)
+				{
+					NewSelectedItems.Add(*TimerNodePtrPtr);
+				}
+			}
+			if (NewSelectedItems.Num() > 0)
+			{
+				TreeView->SetItemSelection(NewSelectedItems, true);
+				TreeView->RequestScrollIntoView(NewSelectedItems[0]);
+			}
+		}
 	}
 }
 
@@ -1487,15 +1586,24 @@ void STimersView::RebuildTree(bool bResync)
 
 void STimersView::UpdateStats(double StartTime, double EndTime)
 {
+	StatsStartTime = StartTime;
+	StatsEndTime = EndTime;
+
+	if (StartTime >= EndTime)
+	{
+		// keep previous aggregated stats
+		return;
+	}
+
 	for (const FTimerNodePtr& TimerNodePtr : TimerNodes)
 	{
 		TimerNodePtr->ResetAggregatedStats();
 	}
 
-	if (Session.IsValid() && StartTime < EndTime && Trace::ReadTimingProfilerProvider(*Session.Get()))
+	if (Session.IsValid() && Trace::ReadTimingProfilerProvider(*Session.Get()))
 	{
 		TSharedPtr<STimingProfilerWindow> Wnd = FTimingProfilerManager::Get()->GetProfilerWindow();
-		TSharedPtr<STimingView> TimingView = Wnd.IsValid() ? Wnd->TimingView : nullptr;
+		TSharedPtr<STimingView> TimingView = Wnd.IsValid() ? Wnd->GetTimingView() : nullptr;
 
 		auto ThreadFilter = [&TimingView](uint32 ThreadId)
 		{
@@ -1522,25 +1630,22 @@ void STimersView::UpdateStats(double StartTime, double EndTime)
 			}
 			AggregationResultTableReader->NextRow();
 		}
-
-		StatsStartTime = StartTime;
-		StatsEndTime = EndTime;
-
-		UpdateTree();
-		TreeView->RebuildList();
 	}
+
+	UpdateTree();
 }
 
 ////////////////////////////////////////////////////////////////////////////////////////////////////
 
 void STimersView::SelectTimerNode(uint64 Id)
 {
-	FTimerNodePtr* TimerNodePtrPtr = TimerNodesIdMap.Find(Id);
-	if (TimerNodePtrPtr != nullptr)
+	FTimerNodePtr* NodePtrPtr = TimerNodesIdMap.Find(Id);
+	if (NodePtrPtr != nullptr)
 	{
-		//UE_LOG(TimingProfiler, Log, TEXT("Select and RequestScrollIntoView %s"), *(*TimerNodePtrPtr)->GetName().ToString());
-		TreeView->SetSelection(*TimerNodePtrPtr);
-		TreeView->RequestScrollIntoView(*TimerNodePtrPtr);
+		FTimerNodePtr NodePtr = *NodePtrPtr;
+
+		TreeView->SetSelection(NodePtr);
+		TreeView->RequestScrollIntoView(NodePtr);
 	}
 }
 
diff --git a/Engine/Source/Developer/TraceInsights/Private/Insights/Widgets/STimersView.h b/Engine/Source/Developer/TraceInsights/Private/Insights/Widgets/STimersView.h
--- a/Engine/Source/Developer/TraceInsights/Private/Insights/Widgets/STimersView.h
+++ b/Engine/Source/Developer/TraceInsights/Private/Insights/Widgets/STimersView.h
@@ -18,7 +18,6 @@
 // Insights
 #include "Insights/ViewModels/TimerGroupingAndSorting.h"
 #include "Insights/ViewModels/TimerNode.h"
-#include "Insights/ViewModels/TimersViewColumn.h"
 
 ////////////////////////////////////////////////////////////////////////////////////////////////////
 
@@ -27,7 +26,13 @@ class FMenuBuilder;
 namespace Trace
 {
 	class IAnalysisSession;
-	struct FTimelineEvent;
+}
+
+namespace Insights
+{
+	class FTable;
+	class FTableColumn;
+	class ITableCellValueSorter;
 }
 
 ////////////////////////////////////////////////////////////////////////////////////////////////////
@@ -60,11 +65,15 @@ public:
 	 */
 	void Construct(const FArguments& InArgs);
 
+	TSharedPtr<Insights::FTable> GetTable() const { return Table; }
+
+	void Reset();
+
 	/**
 	 * Rebuilds the tree (if necessary).
 	 * @param bResync - If true, it forces a resync with list of timers from Analysis, even if the list did not changed since last sync.
 	 */
-	void RebuildTree(bool bResync = true);
+	void RebuildTree(bool bResync);
 	void UpdateStats(double StartTime, double EndTime);
 
 	void SelectTimerNode(uint64 Id);
@@ -74,7 +83,6 @@ public:
 	const FTimerNodePtr* GetTimerNode(uint64 Id) const { return TimerNodesIdMap.Find(Id); }
 
 protected:
-
 	void UpdateTree();
 
 	/** Called when the analysis session has changed. */
@@ -100,9 +108,10 @@ protected:
 	// Tree View - Columns' Header
 
 	void InitializeAndShowHeaderColumns();
-	void TreeViewHeaderRow_CreateColumnArgs(const int32 ColumnIndex);
-	void TreeViewHeaderRow_ShowColumn(const FName ColumnId);
-	TSharedRef<SWidget> TreeViewHeaderRow_GenerateColumnMenu(const FTimersViewColumn& Column);
+
+	FText GetColumnHeaderText(const FName ColumnId) const;
+
+	TSharedRef<SWidget> TreeViewHeaderRow_GenerateColumnMenu(const Insights::FTableColumn& Column);
 
 	////////////////////////////////////////////////////////////////////////////////////////////////////
 	// Tree View - Misc
@@ -128,12 +137,11 @@ protected:
 	/** Called by STreeView to generate a table row for the specified item. */
 	TSharedRef<ITableRow> TreeView_OnGenerateRow(FTimerNodePtr TreeNode, const TSharedRef<STableViewBase>& OwnerTable);
 
-	bool TableRow_IsColumnVisible(const FName ColumnId) const;
-	void TableRow_SetHoveredTableCell(const FName ColumnId, const FTimerNodePtr TimerNodePtr);
+	void TableRow_SetHoveredCell(TSharedPtr<Insights::FTable> TablePtr, TSharedPtr<Insights::FTableColumn> ColumnPtr, const FTimerNodePtr NodePtr);
 	EHorizontalAlignment TableRow_GetColumnOutlineHAlignment(const FName ColumnId) const;
 
 	FText TableRow_GetHighlightText() const;
-	FName TableRow_GetHighlightedTimerName() const;
+	FName TableRow_GetHighlightedNodeName() const;
 
 	bool TableRow_ShouldBeEnabled(const uint32 TimerId) const;
 
@@ -143,15 +151,18 @@ protected:
 	/** Populates the group and stat tree with items based on the current data. */
 	void ApplyFiltering();
 
-	TSharedRef<SWidget> GetToggleButtonForTimerType(const ETimerNodeType StatType);
-	void FilterByTimerType_OnCheckStateChanged(ECheckBoxState NewRadioState, const ETimerNodeType InStatType);
-	ECheckBoxState FilterByTimerType_IsChecked(const ETimerNodeType InStatType) const;
+	void FilterOutZeroCountTimers_OnCheckStateChanged(ECheckBoxState NewRadioState);
+	ECheckBoxState FilterOutZeroCountTimers_IsChecked() const;
+
+	TSharedRef<SWidget> GetToggleButtonForTimerType(const ETimerNodeType InTimerType);
+	void FilterByTimerType_OnCheckStateChanged(ECheckBoxState NewRadioState, const ETimerNodeType InTimerType);
+	ECheckBoxState FilterByTimerType_IsChecked(const ETimerNodeType InTimerType) const;
 
 	bool SearchBox_IsEnabled() const;
 	void SearchBox_OnTextChanged(const FText& InFilterText);
 
 	////////////////////////////////////////////////////////////////////////////////////////////////////
-	// GroupBy
+	// Grouping
 
 	void CreateGroups();
 	void CreateGroupByOptionsSources();
@@ -167,17 +178,19 @@ protected:
 	////////////////////////////////////////////////////////////////////////////////////////////////////
 	// Sorting
 
-	static const EColumnSortMode::Type GetDefaultColumnSortMode();
 	static const FName GetDefaultColumnBeingSorted();
+	static const EColumnSortMode::Type GetDefaultColumnSortMode();
 
-	void SortTimers();
+	void CreateSortings();
+
+	void UpdateCurrentSortingByColumn();
+	void SortTreeNodes();
+	void SortTreeNodesRec(FTimerNode& Node, const Insights::ITableCellValueSorter& Sorter);
 
 	EColumnSortMode::Type GetSortModeForColumn(const FName ColumnId) const;
 	void SetSortModeForColumn(const FName& ColumnId, EColumnSortMode::Type SortMode);
 	void OnSortModeChanged(const EColumnSortPriority::Type SortPriority, const FName& ColumnId, const EColumnSortMode::Type SortMode);
 
-	//void TreeView_BuildSortByMenu(FMenuBuilder& MenuBuilder);
-
 	////////////////////////////////////////////////////////////////////////////////////////////////////
 	// Sorting actions
 
@@ -199,14 +212,18 @@ protected:
 	////////////////////////////////////////////////////////////////////////////////////////////////////
 	// Column visibility actions
 
-	// HideColumn (HeaderMenu)
-	bool HeaderMenu_HideColumn_CanExecute(const FName ColumnId) const;
-	void HeaderMenu_HideColumn_Execute(const FName ColumnId);
+	// ShowColumn
+	bool CanShowColumn(const FName ColumnId) const;
+	void ShowColumn(const FName ColumnId);
+
+	// HideColumn
+	bool CanHideColumn(const FName ColumnId) const;
+	void HideColumn(const FName ColumnId);
 
-	// ToggleColumn (ContextMenu)
-	bool ContextMenu_ToggleColumn_IsChecked(const FName ColumnId);
-	bool ContextMenu_ToggleColumn_CanExecute(const FName ColumnId) const;
-	void ContextMenu_ToggleColumn_Execute(const FName ColumnId);
+	// ToggleColumnVisibility
+	bool IsColumnVisible(const FName ColumnId);
+	bool CanToggleColumnVisibility(const FName ColumnId) const;
+	void ToggleColumnVisibility(const FName ColumnId);
 
 	// ShowAllColumns (ContextMenu)
 	bool ContextMenu_ShowAllColumns_CanExecute() const;
@@ -220,7 +237,12 @@ protected:
 	bool ContextMenu_ResetColumns_CanExecute() const;
 	void ContextMenu_ResetColumns_Execute();
 
+	////////////////////////////////////////////////////////////////////////////////////////////////////
+
 protected:
+	/** Table view model. */
+	TSharedPtr<Insights::FTable> Table;
+
 	/** A weak pointer to the profiler session used to populate this widget. */
 	TSharedPtr<const Trace::IAnalysisSession>/*Weak*/ Session;
 
@@ -230,12 +252,6 @@ protected:
 	/** The tree widget which holds the list of groups and timers corresponding with each group. */
 	TSharedPtr<STreeView<FTimerNodePtr>> TreeView;
 
-	/** Column metadata used to initialize column arguments, stored as PropertyName -> FEventGraphColumn. */
-	TMap<FName, FTimersViewColumn> TreeViewHeaderColumns;
-
-	/** Column arguments used to initialize a new header column in the tree view, stored as column name to column arguments mapping. */
-	TMap<FName, SHeaderRow::FColumn::FArguments> TreeViewHeaderColumnArgs;
-
 	/** Holds the tree view header row widget which display all columns in the tree view. */
 	TSharedPtr<SHeaderRow> TreeViewHeaderRow;
 
@@ -249,10 +265,10 @@ protected:
 	FName HoveredColumnId;
 
 	/** A shared pointer to the timer node currently being hovered by the mouse. */
-	FTimerNodePtr HoveredTimerNodePtr;
+	FTimerNodePtr HoveredNodePtr;
 
 	/** Name of the timer that should be drawn as highlighted. */
-	FName HighlightedTimerName;
+	FName HighlightedNodeName;
 
 	//////////////////////////////////////////////////
 	// Timer Nodes
@@ -266,9 +282,6 @@ protected:
 	/** All timer nodes. */
 	TSet<FTimerNodePtr> TimerNodes;
 
-	/** All timer nodes, stored as TimerName -> FTimerNodePtr. */
-	//TMap<FName, FTimerNodePtr> TimerNodesMap;
-
 	/** All timer nodes, stored as TimerId -> FTimerNodePtr. */
 	TMap<uint64, FTimerNodePtr> TimerNodesIdMap;
 
@@ -281,6 +294,8 @@ protected:
 	//////////////////////////////////////////////////
 	// Search box and filters
 
+	//bool bUseFiltering;
+
 	/** The search box widget used to filter items displayed in the stats and groups tree. */
 	TSharedPtr<SSearchBox> SearchBox;
 
@@ -291,11 +306,16 @@ protected:
 	TSharedPtr<FTimerNodeFilterCollection> Filters;
 
 	/** Holds the visibility of each timer type. */
-	bool bTimerNodeIsVisible[static_cast<int>(ETimerNodeType::InvalidOrMax)];
+	bool bTimerTypeIsVisible[static_cast<int>(ETimerNodeType::InvalidOrMax)];
+
+	/** Filter out the timers having zero total instance count (aggregated stats). */
+	bool bFilterOutZeroCountTimers;
 
 	//////////////////////////////////////////////////
 	// Grouping
 
+	//bool bUseGrouping;
+
 	TArray<TSharedPtr<ETimerGroupingMode>> GroupByOptionsSource;
 
 	TSharedPtr<SComboBox<TSharedPtr<ETimerGroupingMode>>> GroupByComboBox;
@@ -306,12 +326,20 @@ protected:
 	//////////////////////////////////////////////////
 	// Sorting
 
-	/** How we sort the timers? */
-	EColumnSortMode::Type ColumnSortMode;
+	//bool bUseSorting;
+
+	/** All available sorters. */
+	TArray<TSharedPtr<Insights::ITableCellValueSorter>> AvailableSorters;
 
-	/** Name of the column currently being sorted, NAME_None if sorting is disabled. */
+	/** Current sorter. It is nullptr if sorting is disabled. */
+	TSharedPtr<Insights::ITableCellValueSorter> CurrentSorter;
+
+	/** Name of the column currently being sorted. Can be NAME_None if sorting is disabled (CurrentSorting == nullptr) or if a complex sorting is used (CurrentSorting != nullptr). */
 	FName ColumnBeingSorted;
 
+	/** How we sort the nodes? Ascending or Descending. */
+	EColumnSortMode::Type ColumnSortMode;
+
 	//////////////////////////////////////////////////
 
 	double StatsStartTime;
diff --git a/Engine/Source/Developer/TraceInsights/Private/Insights/Widgets/STimersViewTooltip.cpp b/Engine/Source/Developer/TraceInsights/Private/Insights/Widgets/STimersViewTooltip.cpp
--- a/Engine/Source/Developer/TraceInsights/Private/Insights/Widgets/STimersViewTooltip.cpp
+++ b/Engine/Source/Developer/TraceInsights/Private/Insights/Widgets/STimersViewTooltip.cpp
@@ -4,7 +4,7 @@
 
 #include "EditorStyleSet.h"
 #include "SlateOptMacros.h"
-#include "Widgets/Images/SImage.h"
+#include "TraceServices/Model/AnalysisSession.h"
 #include "Widgets/Layout/SGridPanel.h"
 #include "Widgets/Layout/SSeparator.h"
 #include "Widgets/SBoxPanel.h"
@@ -13,19 +13,20 @@
 
 // Insights
 #include "Insights/Common/TimeUtils.h"
+#include "Insights/Table/ViewModels/Table.h"
+#include "Insights/Table/ViewModels/TableColumn.h"
 #include "Insights/ViewModels/TimerNode.h"
-#include "Insights/ViewModels/TimersViewColumn.h"
 #include "Insights/ViewModels/TimerNodeHelper.h"
 
 ////////////////////////////////////////////////////////////////////////////////////////////////////
 
-#define LOCTEXT_NAMESPACE "STimersViewTooltip"
+#define LOCTEXT_NAMESPACE "STimersView"
 
 ////////////////////////////////////////////////////////////////////////////////////////////////////
 
 BEGIN_SLATE_FUNCTION_BUILD_OPTIMIZATION
 
-TSharedPtr<SToolTip> STimersViewTooltip::GetColumnTooltip(const FTimersViewColumn& Column)
+TSharedPtr<SToolTip> STimersViewTooltip::GetTableTooltip(const Insights::FTable& Table)
 {
 	TSharedPtr<SToolTip> ColumnTooltip =
 		SNew(SToolTip)
@@ -37,7 +38,7 @@ TSharedPtr<SToolTip> STimersViewTooltip::GetColumnTooltip(const FTimersViewColum
 			.Padding(2.0f)
 			[
 				SNew(STextBlock)
-				.Text(Column.TitleName)
+				.Text(Table.GetDisplayName())
 				.TextStyle(FEditorStyle::Get(), TEXT("Profiler.TooltipBold"))
 			]
 
@@ -46,7 +47,7 @@ TSharedPtr<SToolTip> STimersViewTooltip::GetColumnTooltip(const FTimersViewColum
 			.Padding(2.0f)
 			[
 				SNew(STextBlock)
-				.Text(Column.Description)
+				.Text(Table.GetDescription())
 				.TextStyle(FEditorStyle::Get(), TEXT("Profiler.Tooltip"))
 			]
 		];
@@ -56,28 +57,53 @@ TSharedPtr<SToolTip> STimersViewTooltip::GetColumnTooltip(const FTimersViewColum
 
 ////////////////////////////////////////////////////////////////////////////////////////////////////
 
-TSharedPtr<SToolTip> STimersViewTooltip::GetTableCellTooltip(const TSharedPtr<FTimerNode> TimerNodePtr)
+TSharedPtr<SToolTip> STimersViewTooltip::GetColumnTooltip(const Insights::FTableColumn& Column)
 {
-	//const FSlateFontInfo TitleFont = FCoreStyle::GetDefaultFontStyle("Bold", 10);
-	//const FSlateFontInfo DescriptionFont = FCoreStyle::GetDefaultFontStyle("Regular", 8);
-	//const FSlateFontInfo DescriptionFontB = FCoreStyle::GetDefaultFontStyle("Bold", 8);
+	TSharedPtr<SToolTip> ColumnTooltip =
+		SNew(SToolTip)
+		[
+			SNew(SVerticalBox)
+
+			+ SVerticalBox::Slot()
+			.AutoHeight()
+			.Padding(2.0f)
+			[
+				SNew(STextBlock)
+				.Text(Column.GetTitleName())
+				.TextStyle(FEditorStyle::Get(), TEXT("Profiler.TooltipBold"))
+			]
+
+			+ SVerticalBox::Slot()
+			.AutoHeight()
+			.Padding(2.0f)
+			[
+				SNew(STextBlock)
+				.Text(Column.GetDescription())
+				.TextStyle(FEditorStyle::Get(), TEXT("Profiler.Tooltip"))
+			]
+		];
 
-	const FLinearColor DefaultColor(1.0f,1.0f,1.0f,1.0f);
-	const FLinearColor ThreadColor(5.0f, 0.0f, 0.0f, 1.0f);
-	const float Alpha = 0.0f;//TimerNodePtr->_FramePct * 0.01f;
-	const FLinearColor ColorAndOpacity = FMath::Lerp(DefaultColor, ThreadColor,Alpha);
+	return ColumnTooltip;
+}
+
+////////////////////////////////////////////////////////////////////////////////////////////////////
 
+TSharedPtr<SToolTip> STimersViewTooltip::GetCellTooltip(const TSharedPtr<FTimerNode> TimerNodePtr, const TSharedPtr<Insights::FTableColumn> ColumnPtr)
+{
 	const Trace::FAggregatedTimingStats& Stats = TimerNodePtr->GetAggregatedStats();
 
 	const int32 NumDigits = 5;
 
-	const FText TotalInclusiveTimeText = FText::FromString(TimeUtils::FormatTimeAuto(Stats.TotalInclusiveTime));
+	TCHAR FormatString[32];
+	FCString::Snprintf(FormatString, sizeof(FormatString), TEXT("%%.%dfs (%%s)"), NumDigits);
+
+	const FText TotalInclusiveTimeText = FText::FromString(FString::Printf(FormatString, Stats.TotalInclusiveTime, *TimeUtils::FormatTimeAuto(Stats.TotalInclusiveTime, 2)));
 	const FText MinInclusiveTimeText = FText::FromString(TimeUtils::FormatTimeMs(Stats.MinInclusiveTime, NumDigits, true));
 	const FText MaxInclusiveTimeText = FText::FromString(TimeUtils::FormatTimeMs(Stats.MaxInclusiveTime, NumDigits, true));
 	const FText AvgInclusiveTimeText = FText::FromString(TimeUtils::FormatTimeMs(Stats.AverageInclusiveTime, NumDigits, true));
 	const FText MedInclusiveTimeText = FText::FromString(TimeUtils::FormatTimeMs(Stats.MedianInclusiveTime, NumDigits, true));
 
-	const FText TotalExclusiveTimeText = FText::FromString(TimeUtils::FormatTimeAuto(Stats.TotalExclusiveTime));
+	const FText TotalExclusiveTimeText = FText::FromString(FString::Printf(FormatString, Stats.TotalExclusiveTime, *TimeUtils::FormatTimeAuto(Stats.TotalExclusiveTime, 2)));
 	const FText MinExclusiveTimeText = FText::FromString(TimeUtils::FormatTimeMs(Stats.MinExclusiveTime, NumDigits, true));
 	const FText MaxExclusiveTimeText = FText::FromString(TimeUtils::FormatTimeMs(Stats.MaxExclusiveTime, NumDigits, true));
 	const FText AvgExclusiveTimeText = FText::FromString(TimeUtils::FormatTimeMs(Stats.AverageExclusiveTime, NumDigits, true));
@@ -85,7 +111,6 @@ TSharedPtr<SToolTip> STimersViewTooltip::GetTableCellTooltip(const TSharedPtr<FT
 
 	const FText InstanceCountText = FText::AsNumber(Stats.InstanceCount);
 
-	//TSharedPtr<SHorizontalBox> HBoxCaption;
 	TSharedPtr<SGridPanel> GridPanel;
 	TSharedPtr<SHorizontalBox> HBox;
 
@@ -94,19 +119,12 @@ TSharedPtr<SToolTip> STimersViewTooltip::GetTableCellTooltip(const TSharedPtr<FT
 		[
 			SAssignNew(HBox, SHorizontalBox)
 
-			+SHorizontalBox::Slot()
+			+ SHorizontalBox::Slot()
 			.AutoWidth()
 			[
 				SNew(SVerticalBox)
 
-				//+SVerticalBox::Slot()
-				//.AutoHeight()
-				//.Padding(2.0f)
-				//[
-				//	SAssignNew(HBoxCaption, SHorizontalBox)
-				//]
-
-				+SVerticalBox::Slot()
+				+ SVerticalBox::Slot()
 				.AutoHeight()
 				.Padding(2.0f)
 				[
@@ -114,21 +132,21 @@ TSharedPtr<SToolTip> STimersViewTooltip::GetTableCellTooltip(const TSharedPtr<FT
 					.Orientation(Orient_Horizontal)
 				]
 
-				+SVerticalBox::Slot()
+				+ SVerticalBox::Slot()
 				.AutoHeight()
 				.Padding(2.0f)
 				[
 					SNew(SGridPanel)
 
 					// Id: [Id]
-					+SGridPanel::Slot(0, 0)
+					+ SGridPanel::Slot(0, 0)
 					.Padding(2.0f)
 					[
 						SNew(STextBlock)
 						.Text(LOCTEXT("TT_Id", "Id:"))
 						.TextStyle(FEditorStyle::Get(), TEXT("Profiler.TooltipBold"))
 					]
-					+SGridPanel::Slot(1, 0)
+					+ SGridPanel::Slot(1, 0)
 					.Padding(2.0f)
 					[
 						SNew(STextBlock)
@@ -137,14 +155,14 @@ TSharedPtr<SToolTip> STimersViewTooltip::GetTableCellTooltip(const TSharedPtr<FT
 					]
 
 					// Name: [Name]
-					+SGridPanel::Slot(0, 1)
+					+ SGridPanel::Slot(0, 1)
 					.Padding(2.0f)
 					[
 						SNew(STextBlock)
 						.Text(LOCTEXT("TT_Name", "Name:"))
 						.TextStyle(FEditorStyle::Get(), TEXT("Profiler.TooltipBold"))
 					]
-					+SGridPanel::Slot(1, 1)
+					+ SGridPanel::Slot(1, 1)
 					.Padding(2.0f)
 					[
 						SNew(STextBlock)
@@ -153,14 +171,14 @@ TSharedPtr<SToolTip> STimersViewTooltip::GetTableCellTooltip(const TSharedPtr<FT
 					]
 
 					//// Group: [MetaGroupName]
-					//+SGridPanel::Slot(0, 2)
+					//+ SGridPanel::Slot(0, 2)
 					//.Padding(2.0f)
 					//[
 					//	SNew(STextBlock)
 					//	.Text(LOCTEXT("TT_Group", "Group:"))
 					//	.TextStyle(FEditorStyle::Get(), TEXT("Profiler.TooltipBold"))
 					//]
-					//+SGridPanel::Slot(1, 2)
+					//+ SGridPanel::Slot(1, 2)
 					//.Padding(2.0f)
 					//[
 					//	SNew(STextBlock)
@@ -168,7 +186,7 @@ TSharedPtr<SToolTip> STimersViewTooltip::GetTableCellTooltip(const TSharedPtr<FT
 					//	.TextStyle(FEditorStyle::Get(), TEXT("Profiler.Tooltip"))
 					//]
 
-					// Type: [Type]
+					// Timer Type: [Type]
 					+ SGridPanel::Slot(0, 3)
 					.Padding(2.0f)
 					[
@@ -180,12 +198,12 @@ TSharedPtr<SToolTip> STimersViewTooltip::GetTableCellTooltip(const TSharedPtr<FT
 					.Padding(2.0f)
 					[
 						SNew(STextBlock)
-						.Text(TimerNodeTypeHelper::ToName(TimerNodePtr->GetType()))
+						.Text(TimerNodeTypeHelper::ToText(TimerNodePtr->GetType()))
 						.TextStyle(FEditorStyle::Get(), TEXT("Profiler.Tooltip"))
 					]
 				]
 
-				+SVerticalBox::Slot()
+				+ SVerticalBox::Slot()
 				.AutoHeight()
 				.Padding(2.0f)
 				[
@@ -215,7 +233,7 @@ TSharedPtr<SToolTip> STimersViewTooltip::GetTableCellTooltip(const TSharedPtr<FT
 					]
 				]
 
-				+SVerticalBox::Slot()
+				+ SVerticalBox::Slot()
 				.AutoHeight()
 				.Padding(2.0f)
 				[
@@ -238,7 +256,7 @@ TSharedPtr<SToolTip> STimersViewTooltip::GetTableCellTooltip(const TSharedPtr<FT
 						.TextStyle(FEditorStyle::Get(), TEXT("Profiler.TooltipBold"))
 					]
 					+ SGridPanel::Slot(2, 0)
-					.Padding(2.0f)
+					.Padding(FMargin(8.0f, 2.0f, 2.0f, 2.0f))
 					.HAlign(HAlign_Right)
 					[
 						SNew(STextBlock)
@@ -249,7 +267,7 @@ TSharedPtr<SToolTip> STimersViewTooltip::GetTableCellTooltip(const TSharedPtr<FT
 					// Stats are added here.
 				]
 
-				+SVerticalBox::Slot()
+				+ SVerticalBox::Slot()
 				.AutoHeight()
 				.Padding(2.0f)
 				[
@@ -266,260 +284,11 @@ TSharedPtr<SToolTip> STimersViewTooltip::GetTableCellTooltip(const TSharedPtr<FT
 	AddStatsRow(GridPanel, Row, LOCTEXT("TT_MedianTime",  "Median Time:"),  MedInclusiveTimeText,   MedExclusiveTimeText);
 	AddStatsRow(GridPanel, Row, LOCTEXT("TT_MinTime",     "Min Time:"),     MinInclusiveTimeText,   MinExclusiveTimeText);
 
-	/*
-	//TODO: We need Stats hierarchy (not the grouping hierarchy)!
-	const bool bHasParent = TimerNodePtr->GetStats()->GetParent().IsValid();
-	const bool bHasChildren = TimerNodePtr->GetStats()->GetChildren().Num() > 0;
-
-	if (bHasParent)
-	{
-		const FText ParentName = FText::FromName(TimerNodePtr->GetGroupPtr()->GetName());
-		HBoxCaption->AddSlot()
-			.AutoWidth()
-			[
-				SNew(STextBlock)
-				.Text(ParentName)
-				.TextStyle(FEditorStyle::Get(), TEXT("Profiler.Caption"))
-			];
-
-		HBoxCaption->AddSlot()
-			.AutoWidth()
-			[
-				SNew(SImage)
-				.Image(FEditorStyle::GetBrush("BreadcrumbTrail.Delimiter"))
-			];
-	}
-
-	const FText TimerName = FText::FromName(TimerNodePtr->GetName());
-	HBoxCaption->AddSlot()
-		.AutoWidth()
-		[
-			SNew(STextBlock)
-			.Text(TimerName)
-			.TextStyle(FEditorStyle::Get(), TEXT("Profiler.CaptionBold"))
-		];
-
-	if (bHasChildren)
-	{
-		typedef TKeyValuePair<float, FName> FEventNameAndPct;
-		TArray<FEventNameAndPct> MinimalChildren;
-
-		const TArray<FTimerNodePtr>& Children = TimerNodePtr->GetChildren();
-		for(int32 ChildIndex = 0; ChildIndex < Children.Num(); ChildIndex++)
-		{
-			const FTimerNodePtr Child = Children[ChildIndex];
-			float Percent = static_cast<float>(Child->GetStats().TotalInclusiveTime / TimerNodePtr->GetStats().TotalInclusiveTime * 100.0);
-			MinimalChildren.Add(FEventNameAndPct(Child->GetStats().TotalInclusiveTime, Child->GetName()));
-		}
-
-		struct FCompareByFloatDescending
-		{
-			FORCEINLINE bool operator()(const FEventNameAndPct& A, const FEventNameAndPct& B) const
-			{
-				return A.Key > B.Key;
-			}
-		};
-		MinimalChildren.Sort(FCompareByFloatDescending());
-
-		FString ChildrenNames;
-		const int32 NumChildrenToDisplay = FMath::Min(MinimalChildren.Num(), 3);
-		for(int32 SortedChildIndex = 0; SortedChildIndex < NumChildrenToDisplay; SortedChildIndex++)
-		{
-			const FEventNameAndPct& MinimalChild = MinimalChildren[SortedChildIndex];
-			ChildrenNames += FString::Printf(TEXT("%s (%.1f %%)"), *MinimalChild.Value.ToString(), MinimalChild.Key);
-
-			const bool bAddDelimiter = SortedChildIndex < NumChildrenToDisplay - 1;
-			if (bAddDelimiter)
-			{
-				ChildrenNames += TEXT(", ");
-			}
-		}
-
-		HBoxCaption->AddSlot()
-			.AutoWidth()
-			[
-				SNew(SImage)
-				.Image(FEditorStyle::GetBrush("BreadcrumbTrail.Delimiter"))
-			];
-
-		HBoxCaption->AddSlot()
-			.AutoWidth()
-			[
-				SNew(STextBlock)
-				.Text(FText::FromString(ChildrenNames))
-				.TextStyle(FEditorStyle::Get(), TEXT("Profiler.Caption"))
-			];
-	}
-	*/
-
 	return TableCellTooltip;
 }
 
 ////////////////////////////////////////////////////////////////////////////////////////////////////
 
-TSharedRef<SToolTip> STimersViewTooltip::GetTooltip()
-{
-	if (Session.IsValid())
-	{
-		const TSharedRef<SGridPanel> ToolTipGrid = SNew(SGridPanel);
-		int32 CurrentRowPos = 0;
-
-		AddHeader(ToolTipGrid, CurrentRowPos);
-		AddDescription(ToolTipGrid, CurrentRowPos);
-
-		AddNoDataInformation(ToolTipGrid, CurrentRowPos);
-
-		return SNew(SToolTip)
-			[
-				ToolTipGrid
-			];
-	}
-	else
-	{
-		return SNew(SToolTip)
-			.Text(LOCTEXT("NotImplemented", "Tooltip for multiple profiler instances has not been implemented yet"));
-	}
-}
-
-////////////////////////////////////////////////////////////////////////////////////////////////////
-
-void STimersViewTooltip::AddNoDataInformation(const TSharedRef<SGridPanel>& Grid, int32& RowPos)
-{
-	Grid->AddSlot(0, RowPos)
-		.Padding(2.0f)
-		.ColumnSpan(3)
-		[
-			SNew(STextBlock)
-			.TextStyle(FEditorStyle::Get(), TEXT("Profiler.TooltipBold"))
-			.Text(LOCTEXT("NoDataAvailable", "N/A"))
-		];
-	RowPos++;
-}
-
-////////////////////////////////////////////////////////////////////////////////////////////////////
-
-void STimersViewTooltip::AddHeader(const TSharedRef<SGridPanel>& Grid, int32& RowPos)
-{
-	const FString InstanceName = Session->GetName();
-
-	Grid->AddSlot(0, RowPos++)
-		.Padding(2.0f)
-		.ColumnSpan(3)
-		[
-			SNew(STextBlock)
-			.TextStyle(FEditorStyle::Get(), TEXT("Profiler.TooltipBold"))
-			.Text(LOCTEXT("StatInstance", "Stat information for profiler instance"))
-		];
-
-	Grid->AddSlot(0, RowPos++)
-		.Padding(2.0f)
-		.ColumnSpan(3)
-		[
-			SNew(STextBlock)
-			.TextStyle(FEditorStyle::Get(), TEXT("Profiler.Tooltip"))
-			.Text(FText::FromString(InstanceName))
-		];
-
-	AddSeparator(Grid, RowPos);
-}
-
-////////////////////////////////////////////////////////////////////////////////////////////////////
-
-void STimersViewTooltip::AddDescription(const TSharedRef<SGridPanel>& Grid, int32& RowPos)
-{
-	/*
-	const FProfilerStat& ProfilerStat = Session->GetMetaData()->GetTimerById(TimerId);
-	const ETimerNodeType SampleType = Session->GetMetaData()->GetSampleTypeForTimerId(TimerId);
-	const FSlateBrush* const StatIcon = STimersViewHelper::GetIconForStatType(SampleType);
-
-	Grid->AddSlot(0, RowPos)
-	.Padding(2.0f)
-	[
-		SNew(STextBlock)
-		.TextStyle(FEditorStyle::Get(), TEXT("Profiler.TooltipBold"))
-		.Text(LOCTEXT("GroupDesc","Group:"))
-	];
-
-	Grid->AddSlot(1, RowPos)
-	.Padding(2.0f)
-	.ColumnSpan(2)
-	[
-		SNew(STextBlock)
-		.TextStyle(FEditorStyle::Get(), TEXT("Profiler.Tooltip"))
-		.Text(FText::FromName(ProfilerStat.OwningGroup().Name()))
-	];
-	RowPos++;
-
-	Grid->AddSlot(0, RowPos)
-	.Padding(2.0f)
-	[
-		SNew(STextBlock)
-		.TextStyle(FEditorStyle::Get(), TEXT("Profiler.TooltipBold"))
-		.Text(LOCTEXT("NameDesc","Name:"))
-	];
-
-	Grid->AddSlot(1, RowPos)
-	.Padding(2.0f)
-	.ColumnSpan(2)
-	[
-		SNew(STextBlock)
-		.TextStyle(FEditorStyle::Get(), TEXT("Profiler.Tooltip"))
-		.Text(FText::FromName(ProfilerStat.Name()))
-	];
-	RowPos++;
-
-	Grid->AddSlot(0, RowPos)
-	.Padding(2.0f)
-	[
-		SNew(STextBlock)
-		.TextStyle(FEditorStyle::Get(), TEXT("Profiler.TooltipBold"))
-		.Text(LOCTEXT("TypeDesc","Type:"))
-	];
-
-	Grid->AddSlot(1, RowPos)
-	.Padding(2.0f)
-	.ColumnSpan(2)
-	[
-		SNew(SHorizontalBox)
-
-		+SHorizontalBox::Slot()
-		.AutoWidth()
-		[
-			SNew(SImage)
-			.Image(StatIcon)
-		]
-
-		+SHorizontalBox::Slot()
-		.AutoWidth()
-		.HAlign(HAlign_Left)
-		.VAlign(VAlign_Center)
-		[
-			SNew(STextBlock)
-			.Text(FText::FromString(ETimerNodeType::ToDescription(SampleType)))
-			.TextStyle(FEditorStyle::Get(), TEXT("Profiler.Tooltip"))
-		]
-	];
-	RowPos++;
-
-	AddSeparator(Grid, RowPos);
-	*/
-}
-
-////////////////////////////////////////////////////////////////////////////////////////////////////
-
-void STimersViewTooltip::AddSeparator(const TSharedRef<SGridPanel>& Grid, int32& RowPos)
-{
-	Grid->AddSlot(0, RowPos++)
-		.Padding(2.0f)
-		.ColumnSpan(3)
-		[
-			SNew(SSeparator)
-			.Orientation(Orient_Horizontal)
-		];
-}
-
-////////////////////////////////////////////////////////////////////////////////////////////////////
-
 void STimersViewTooltip::AddStatsRow(TSharedPtr<SGridPanel> Grid, int32& Row, const FText& Name, const FText& Value1, const FText& Value2)
 {
 	Grid->AddSlot(0, Row)
@@ -540,7 +309,7 @@ void STimersViewTooltip::AddStatsRow(TSharedPtr<SGridPanel> Grid, int32& Row, co
 		];
 
 	Grid->AddSlot(2, Row)
-		.Padding(2.0f)
+		.Padding(FMargin(8.0f, 2.0f, 2.0f, 2.0f))
 		.HAlign(HAlign_Right)
 		[
 			SNew(STextBlock)
diff --git a/Engine/Source/Developer/TraceInsights/Private/Insights/Widgets/STimersViewTooltip.h b/Engine/Source/Developer/TraceInsights/Private/Insights/Widgets/STimersViewTooltip.h
--- a/Engine/Source/Developer/TraceInsights/Private/Insights/Widgets/STimersViewTooltip.h
+++ b/Engine/Source/Developer/TraceInsights/Private/Insights/Widgets/STimersViewTooltip.h
@@ -3,50 +3,34 @@
 #pragma once
 
 #include "CoreMinimal.h"
-#include "EditorStyleSet.h"
-#include "SlateOptMacros.h"
-#include "TraceServices/AnalysisService.h"
-#include "Widgets/Images/SImage.h"
-#include "Widgets/Layout/SGridPanel.h"
-#include "Widgets/Layout/SSeparator.h"
-#include "Widgets/SToolTip.h"
-#include "Widgets/Text/STextBlock.h"
 
 // Insights
 #include "Insights/InsightsManager.h"
 
-// Insights
+class SToolTip;
+class SGridPanel;
+
+namespace Insights
+{
+	class FTable;
+	class FTableColumn;
+}
+
 class FTimerNode;
-class FTimersViewColumn;
 
 #define LOCTEXT_NAMESPACE "STimersView"
 
 /** Timers View Tooltip */
 class STimersViewTooltip
 {
-	const uint32 TimerId;
-	TSharedPtr<const Trace::IAnalysisSession> Session;
-
 public:
-	STimersViewTooltip(const uint32 InTimerId)
-		: TimerId(InTimerId)
-	{
-		Session = FInsightsManager::Get()->GetSession();
-	}
-
-	TSharedRef<SToolTip> GetTooltip();
+	STimersViewTooltip() = delete;
 
-protected:
-	void AddNoDataInformation(const TSharedRef<SGridPanel>& Grid, int32& RowPos);
-	void AddHeader(const TSharedRef<SGridPanel>& Grid, int32& RowPos);
-	void AddDescription(const TSharedRef<SGridPanel>& Grid, int32& RowPos);
-	void AddSeparator(const TSharedRef<SGridPanel>& Grid, int32& RowPos);
-
-public:
-	static TSharedPtr<SToolTip> GetColumnTooltip(const FTimersViewColumn& Column);
-	static TSharedPtr<SToolTip> GetTableCellTooltip(const TSharedPtr<FTimerNode> TimerNodePtr);
+	static TSharedPtr<SToolTip> GetTableTooltip(const Insights::FTable& Table);
+	static TSharedPtr<SToolTip> GetColumnTooltip(const Insights::FTableColumn& Column);
+	static TSharedPtr<SToolTip> GetCellTooltip(const TSharedPtr<FTimerNode> TreeNodePtr, const TSharedPtr<Insights::FTableColumn> ColumnPtr);
 
-protected:
+private:
 	static void AddStatsRow(TSharedPtr<SGridPanel> Grid, int32& Row, const FText& Name, const FText& Value1, const FText& Value2);
 };
 
diff --git a/Engine/Source/Developer/TraceInsights/Private/Insights/Widgets/STimingProfilerWindow.h b/Engine/Source/Developer/TraceInsights/Private/Insights/Widgets/STimingProfilerWindow.h
--- a/Engine/Source/Developer/TraceInsights/Private/Insights/Widgets/STimingProfilerWindow.h
+++ b/Engine/Source/Developer/TraceInsights/Private/Insights/Widgets/STimingProfilerWindow.h
@@ -18,13 +18,16 @@
 ////////////////////////////////////////////////////////////////////////////////////////////////////
 
 class FActiveTimerHandle;
-class SVerticalBox;
+class FMenuBuilder;
+
 class SFrameTrack;
 class SGraphTrack;
+class SLogView;
 class SStatsView;
 class STimersView;
+class STimerTreeView;
 class STimingView;
-class SLogView;
+class SVerticalBox;
 
 ////////////////////////////////////////////////////////////////////////////////////////////////////
 
@@ -36,6 +39,8 @@ struct FTimingProfilerTabs
 	static const FName GraphTrackID;
 	static const FName TimingViewID;
 	static const FName TimersID;
+	static const FName CallersID;
+	static const FName CalleesID;
 	static const FName StatsCountersID;
 	static const FName LogViewID;
 };
@@ -52,9 +57,11 @@ public:
 	/** Virtual destructor. */
 	virtual ~STimingProfilerWindow();
 
-	SLATE_BEGIN_ARGS(STimingProfilerWindow){}
+	SLATE_BEGIN_ARGS(STimingProfilerWindow) {}
 	SLATE_END_ARGS()
 
+	void Reset();
+
 	/** Constructs this widget. */
 	void Construct(const FArguments& InArgs, const TSharedRef<SDockTab>& ConstructUnderMajorTab, const TSharedPtr<SWindow>& ConstructUnderWindow);
 
@@ -64,22 +71,35 @@ public:
 
 	TSharedPtr<FTabManager> GetTabManager() const { return TabManager; }
 
+	TSharedPtr<STimingView> GetTimingView() const { return TimingView; }
+	TSharedPtr<STimersView> GetTimersView() const { return TimersView; }
+	TSharedPtr<STimerTreeView> GetCallersTreeView() const { return CallersTreeView; }
+	TSharedPtr<STimerTreeView> GetCalleesTreeView() const { return CalleesTreeView; }
+	TSharedPtr<SStatsView> GetStatsView() const { return StatsView; }
+	TSharedPtr<SLogView> GetLogView() const { return LogView; }
+
 private:
 	TSharedRef<SDockTab> SpawnTab_Toolbar(const FSpawnTabArgs& Args);
 	void OnToolbarTabClosed(TSharedRef<SDockTab> TabBeingClosed);
-	
+
 	TSharedRef<SDockTab> SpawnTab_FramesTrack(const FSpawnTabArgs& Args);
 	void OnFramesTrackTabClosed(TSharedRef<SDockTab> TabBeingClosed);
-	
+
 	TSharedRef<SDockTab> SpawnTab_GraphTrack(const FSpawnTabArgs& Args);
 	void OnGraphTrackTabClosed(TSharedRef<SDockTab> TabBeingClosed);
-	
+
 	TSharedRef<SDockTab> SpawnTab_TimingView(const FSpawnTabArgs& Args);
 	void OnTimingViewTabClosed(TSharedRef<SDockTab> TabBeingClosed);
-	
+
 	TSharedRef<SDockTab> SpawnTab_Timers(const FSpawnTabArgs& Args);
 	void OnTimersTabClosed(TSharedRef<SDockTab> TabBeingClosed);
-	
+
+	TSharedRef<SDockTab> SpawnTab_Callers(const FSpawnTabArgs& Args);
+	void OnCallersTabClosed(TSharedRef<SDockTab> TabBeingClosed);
+
+	TSharedRef<SDockTab> SpawnTab_Callees(const FSpawnTabArgs& Args);
+	void OnCalleesTabClosed(TSharedRef<SDockTab> TabBeingClosed);
+
 	TSharedRef<SDockTab> SpawnTab_StatsCounters(const FSpawnTabArgs& Args);
 	void OnStatsCountersTabClosed(TSharedRef<SDockTab> TabBeingClosed);
 
@@ -157,32 +177,37 @@ private:
 	 */
 	virtual FReply OnDragOver(const FGeometry& MyGeometry, const FDragDropEvent& DragDropEvent)  override;
 
-public:
-	/** Widget for the frame track */
+private:
+	/** The Frame track widget */
 	TSharedPtr<SFrameTrack> FrameTrack;
 
-	/** Widget for the graph track */
+	/** The graph track widget */
 	TSharedPtr<SGraphTrack> GraphTrack;
 
-	/** Widget for the timing track */
+	/** The Timing view (multi-track) widget */
 	TSharedPtr<STimingView> TimingView;
 
-	/** Holds the Timers view widget. */
+	/** The Timers view widget */
 	TSharedPtr<STimersView> TimersView;
 
-	/** Holds the Stats (Counters) view widget. */
+	/** The Callers tree view widget */
+	TSharedPtr<STimerTreeView> CallersTreeView;
+
+	/** The Callees tree view widget */
+	TSharedPtr<STimerTreeView> CalleesTreeView;
+
+	/** The Stats (Counters) view widget */
 	TSharedPtr<SStatsView> StatsView;
 
-	/** Widget for the log view */
+	/** The Log view widget */
 	TSharedPtr<SLogView> LogView;
 
-	/** The number of seconds the profiler has been active */
-	float DurationActive;
-
-private:
 	/** Holds the tab manager that manages the front-end's tabs. */
 	TSharedPtr<FTabManager> TabManager;
 
 	/** The handle to the active update duration tick */
 	TWeakPtr<FActiveTimerHandle> ActiveTimerHandle;
+
+	/** The number of seconds the profiler has been active */
+	float DurationActive;
 };
diff --git a/Engine/Source/Editor/UnrealEd/Classes/CookOnTheSide/CookOnTheFlyServer.h b/Engine/Source/Editor/UnrealEd/Classes/CookOnTheSide/CookOnTheFlyServer.h
--- a/Engine/Source/Editor/UnrealEd/Classes/CookOnTheSide/CookOnTheFlyServer.h
+++ b/Engine/Source/Editor/UnrealEd/Classes/CookOnTheSide/CookOnTheFlyServer.h
@@ -23,6 +23,7 @@ class IAssetRegistry;
 
 struct FPackageNameCache;
 struct FPackageTracker;
+class FSavePackageContext;
 
 enum class ECookInitializationFlags
 {
@@ -65,6 +66,7 @@ enum class ECookByTheBookOptions
 	NoInputPackages =					0x00000800, // don't include slate content (this cook will probably be missing content unless you know what you are doing)
 	DisableUnsolicitedPackages =		0x00001000, // don't cook any packages which aren't in the files to cook list (this is really dangerious as if you request a file it will not cook all it's dependencies automatically)
 	FullLoadAndSave =					0x00002000, // Load all packages into memory and save them all at once in one tick for speed reasons. This requires a lot of RAM for large games.
+	PackageStore =						0x00004000, // Cook package header information into a global package store
 };
 ENUM_CLASS_FLAGS(ECookByTheBookOptions);
 
@@ -445,6 +447,8 @@ public:
 	*/
 	bool IsCookByTheBookMode() const;
 
+	bool IsUsingPackageStore() const;
+
 	/**
 	* Helper function returns if we are in any cook on the fly mode
 	*
@@ -680,6 +684,11 @@ private:
 	*/
 	void InitializeSandbox();
 
+	/**
+	* Initialize the package store
+	*/
+	void InitializePackageStore(const TArray<FName>& TargetPlatformNames);
+
 	/**
 	* Initialize all target platforms
 	*/
@@ -851,9 +860,6 @@ private:
 		return (CookFlags & InCookFlags) != ECookInitializationFlags::None;
 	}
 
-	/** If true, the maximum file length of a package being saved will be reduced by 32 to compensate for compressed package intermediate files */
-	bool ShouldConsiderCompressedPackageFileLengthRequirements() const;
-
 	/**
 	*	Cook (save) the given package
 	*
@@ -928,6 +934,8 @@ private:
 	FPackageTracker*	PackageTracker;
 	FPackageNameCache*	PackageNameCache;
 
+	TArray<FSavePackageContext*> SavePackageContexts;
+
 	// temporary -- should eliminate the need for this. Only required right now because FullLoadAndSave 
 	// accesses maps directly
 	friend FPackageNameCache;
diff --git a/Engine/Source/Editor/UnrealEd/Classes/Editor/EditorEngine.h b/Engine/Source/Editor/UnrealEd/Classes/Editor/EditorEngine.h
--- a/Engine/Source/Editor/UnrealEd/Classes/Editor/EditorEngine.h
+++ b/Engine/Source/Editor/UnrealEd/Classes/Editor/EditorEngine.h
@@ -2293,14 +2293,15 @@ public:
 
 	/** The editor wrapper for UPackage::SavePackage. Auto-adds files to source control when necessary */
 	bool SavePackage( UPackage* InOuter, UObject* Base, EObjectFlags TopLevelFlags, const TCHAR* Filename, 
-		FOutputDevice* Error=GError, FLinkerLoad* Conform=NULL, bool bForceByteSwapping=false, bool bWarnOfLongFilename=true, 
+		FOutputDevice* Error=GError, FLinkerNull* Conform=NULL, bool bForceByteSwapping=false, bool bWarnOfLongFilename=true, 
 		uint32 SaveFlags=SAVE_None, const class ITargetPlatform* TargetPlatform = NULL, const FDateTime& FinalTimeStamp = FDateTime::MinValue(), bool bSlowTask = true );
 
 	/** The editor wrapper for UPackage::Save. Auto-adds files to source control when necessary */
 	FSavePackageResultStruct Save(UPackage* InOuter, UObject* Base, EObjectFlags TopLevelFlags, const TCHAR* Filename,
-		FOutputDevice* Error = GError, FLinkerLoad* Conform = NULL, bool bForceByteSwapping = false, bool bWarnOfLongFilename = true,
+		FOutputDevice* Error = GError, FLinkerNull* Conform = NULL, bool bForceByteSwapping = false, bool bWarnOfLongFilename = true,
 		uint32 SaveFlags = SAVE_None, const class ITargetPlatform* TargetPlatform = NULL, const FDateTime& FinalTimeStamp = FDateTime::MinValue(), 
-		bool bSlowTask = true, class FArchiveDiffMap* InOutDiffMap = nullptr);
+		bool bSlowTask = true, class FArchiveDiffMap* InOutDiffMap = nullptr,
+		FSavePackageContext* SavePackageContext = nullptr);
 
 	virtual bool InitializePhysicsSceneForSaveIfNecessary(UWorld* World, bool &bOutForceInitialized);
 	void CleanupPhysicsSceneThatWasInitializedForSave(UWorld* World, bool bForceInitialized);
diff --git a/Engine/Source/Editor/UnrealEd/Private/Commandlets/CookCommandlet.cpp b/Engine/Source/Editor/UnrealEd/Private/Commandlets/CookCommandlet.cpp
--- a/Engine/Source/Editor/UnrealEd/Private/Commandlets/CookCommandlet.cpp
+++ b/Engine/Source/Editor/UnrealEd/Private/Commandlets/CookCommandlet.cpp
@@ -286,7 +286,7 @@ namespace DetailedCookStats
 			UE_LOG(LogCookCommandlet, Display, TEXT("================="));
 			for (const auto& Attr : DDCSummaryStats)
 			{
-				UE_LOG(LogCookCommandlet, Display, TEXT("%-14s=%10s"), *Attr.Key, *Attr.Value);
+				UE_LOG(LogCookCommandlet, Display, TEXT("%-16s=%10s"), *Attr.Key, *Attr.Value);
 			}
 		}
 		if (DDCResourceUsageStats.Num() > 0)
@@ -597,7 +597,7 @@ int32 UCookCommandlet::Main(const FString& CmdLineParams)
 
 bool UCookCommandlet::CookByTheBook( const TArray<ITargetPlatform*>& Platforms, TArray<FString>& FilesInPath )
 {
-	TRACE_CPUPROFILER_EVENT_SCOPE_TEXT(TEXT("CookByTheBook"));
+	TRACE_CPUPROFILER_EVENT_SCOPE(CookByTheBook);
 
 	COOK_STAT(FScopedDurationTimer CookByTheBookTimer(DetailedCookStats::CookByTheBookTimeSec));
 	UCookOnTheFlyServer *CookOnTheFlyServer = NewObject<UCookOnTheFlyServer>();
@@ -813,6 +813,7 @@ bool UCookCommandlet::CookByTheBook( const TArray<ITargetPlatform*>& Platforms,
 	CookOptions |= Switches.Contains(TEXT("MAPSONLY")) ? ECookByTheBookOptions::MapsOnly : ECookByTheBookOptions::None;
 	CookOptions |= Switches.Contains(TEXT("NODEV")) ? ECookByTheBookOptions::NoDevContent : ECookByTheBookOptions::None;
 	CookOptions |= Switches.Contains(TEXT("FullLoadAndSave")) ? ECookByTheBookOptions::FullLoadAndSave : ECookByTheBookOptions::None;
+	CookOptions |= Switches.Contains(TEXT("PackageStore")) ? ECookByTheBookOptions::PackageStore : ECookByTheBookOptions::None;
 
 	const ECookByTheBookOptions SinglePackageFlags = ECookByTheBookOptions::NoAlwaysCookMaps | ECookByTheBookOptions::NoDefaultMaps | ECookByTheBookOptions::NoGameAlwaysCookPackages | ECookByTheBookOptions::NoInputPackages | ECookByTheBookOptions::NoSlatePackages | ECookByTheBookOptions::DisableUnsolicitedPackages | ECookByTheBookOptions::ForceDisableSaveGlobalShaders;
 	CookOptions |= bCookSinglePackage ? SinglePackageFlags : ECookByTheBookOptions::None;
