[Insights] Integrated CL 9470738 from dev-core. Fixed tree views (Timers, Counters, Generic Tree Table view and Net Stats) not being properly invalidated and showing outdated values for aggregated stats. Changed tooltips for tree views to be created only when needed.

diff --git a/Engine/Source/Developer/TraceInsights/Private/Insights/NetworkingProfiler/Widgets/SNetStatsTableCell.cpp b/Engine/Source/Developer/TraceInsights/Private/Insights/NetworkingProfiler/Widgets/SNetStatsTableCell.cpp
--- a/Engine/Source/Developer/TraceInsights/Private/Insights/NetworkingProfiler/Widgets/SNetStatsTableCell.cpp
+++ b/Engine/Source/Developer/TraceInsights/Private/Insights/NetworkingProfiler/Widgets/SNetStatsTableCell.cpp
@@ -15,7 +15,7 @@
 // Insights
 #include "Insights/Table/ViewModels/Table.h"
 #include "Insights/Table/ViewModels/TableColumn.h"
-#include "Insights/NetworkingProfiler/Widgets/SNetStatsViewTooltip.h"
+#include "Insights/NetworkingProfiler/Widgets/SNetStatsTableRow.h"
 
 #define LOCTEXT_NAMESPACE "SNetStatsView"
 
@@ -23,7 +23,7 @@
 
 BEGIN_SLATE_FUNCTION_BUILD_OPTIMIZATION
 
-void SNetStatsTableCell::Construct(const FArguments& InArgs, const TSharedRef<class ITableRow>& TableRow)
+void SNetStatsTableCell::Construct(const FArguments& InArgs, const TSharedRef<ITableRow>& TableRow)
 {
 	TablePtr = InArgs._TablePtr;
 	ColumnPtr = InArgs._ColumnPtr;
@@ -43,7 +43,7 @@ void SNetStatsTableCell::Construct(const FArguments& InArgs, const TSharedRef<cl
 
 ////////////////////////////////////////////////////////////////////////////////////////////////////
 
-TSharedRef<SWidget> SNetStatsTableCell::GenerateWidgetForColumn(const FArguments& InArgs, const TSharedRef<class ITableRow>& TableRow)
+TSharedRef<SWidget> SNetStatsTableCell::GenerateWidgetForColumn(const FArguments& InArgs, const TSharedRef<ITableRow>& TableRow)
 {
 	if (InArgs._IsNameColumn)
 	{
@@ -57,7 +57,7 @@ TSharedRef<SWidget> SNetStatsTableCell::GenerateWidgetForColumn(const FArguments
 
 ////////////////////////////////////////////////////////////////////////////////////////////////////
 
-TSharedRef<SWidget> SNetStatsTableCell::GenerateWidgetForNameColumn(const FArguments& InArgs, const TSharedRef<class ITableRow>& TableRow)
+TSharedRef<SWidget> SNetStatsTableCell::GenerateWidgetForNameColumn(const FArguments& InArgs, const TSharedRef<ITableRow>& TableRow)
 {
 	return
 		SNew(SHorizontalBox)
@@ -79,7 +79,7 @@ TSharedRef<SWidget> SNetStatsTableCell::GenerateWidgetForNameColumn(const FArgum
 			SNew(SImage)
 			.Visibility(this, &SNetStatsTableCell::GetHintIconVisibility)
 			.Image(FEditorStyle::GetBrush("Profiler.Tooltip.HintIcon10"))
-			.ToolTip(SNetStatsViewTooltip::GetCellTooltip(NetEventNodePtr, ColumnPtr))
+			.ToolTip(GetRowToolTip(TableRow))
 		]
 
 		// Name
@@ -101,10 +101,23 @@ TSharedRef<SWidget> SNetStatsTableCell::GenerateWidgetForNameColumn(const FArgum
 
 ////////////////////////////////////////////////////////////////////////////////////////////////////
 
-TSharedRef<SWidget> SNetStatsTableCell::GenerateWidgetForStatsColumn(const FArguments& InArgs, const TSharedRef<class ITableRow>& TableRow)
+TSharedPtr<IToolTip> SNetStatsTableCell::GetRowToolTip(const TSharedRef<ITableRow>& TableRow) const
 {
-	const FText CellText = NetEventNodePtr->IsGroup() ? FText::GetEmpty() : ColumnPtr->GetValueAsText(*NetEventNodePtr);
+	TSharedRef<SNetStatsTableRow> Row = StaticCastSharedRef<SNetStatsTableRow, ITableRow>(TableRow);
+	return Row->GetRowToolTip();
+}
+
+////////////////////////////////////////////////////////////////////////////////////////////////////
 
+FText SNetStatsTableCell::GetValueAsText() const
+{
+	return NetEventNodePtr->IsGroup() ? FText::GetEmpty() : ColumnPtr->GetValueAsText(*NetEventNodePtr);
+}
+
+////////////////////////////////////////////////////////////////////////////////////////////////////
+
+TSharedRef<SWidget> SNetStatsTableCell::GenerateWidgetForStatsColumn(const FArguments& InArgs, const TSharedRef<ITableRow>& TableRow)
+{
 	return
 		SNew(SHorizontalBox)
 
@@ -116,7 +129,7 @@ TSharedRef<SWidget> SNetStatsTableCell::GenerateWidgetForStatsColumn(const FArgu
 		.Padding(FMargin(2.0f, 0.0f))
 		[
 			SNew(STextBlock)
-			.Text(CellText)
+			.Text(this, &SNetStatsTableCell::GetValueAsText)
 			.TextStyle(FEditorStyle::Get(), TEXT("Profiler.Tooltip"))
 			.ColorAndOpacity(this, &SNetStatsTableCell::GetStatsColorAndOpacity)
 			.ShadowColorAndOpacity(this, &SNetStatsTableCell::GetShadowColorAndOpacity)
diff --git a/Engine/Source/Developer/TraceInsights/Private/Insights/NetworkingProfiler/Widgets/SNetStatsTableCell.h b/Engine/Source/Developer/TraceInsights/Private/Insights/NetworkingProfiler/Widgets/SNetStatsTableCell.h
--- a/Engine/Source/Developer/TraceInsights/Private/Insights/NetworkingProfiler/Widgets/SNetStatsTableCell.h
+++ b/Engine/Source/Developer/TraceInsights/Private/Insights/NetworkingProfiler/Widgets/SNetStatsTableCell.h
@@ -9,6 +9,8 @@
 // Insights
 #include "Insights/NetworkingProfiler/ViewModels/NetEventNode.h"
 
+class ITableRow;
+
 namespace Insights
 {
 	class FTable;
@@ -29,12 +31,12 @@ public:
 		SLATE_ARGUMENT(bool, IsNameColumn)
 	SLATE_END_ARGS()
 
-	void Construct(const FArguments& InArgs, const TSharedRef<class ITableRow>& TableRow);
+	void Construct(const FArguments& InArgs, const TSharedRef<ITableRow>& TableRow);
 
 protected:
-	TSharedRef<SWidget> GenerateWidgetForColumn(const FArguments& InArgs, const TSharedRef<class ITableRow>& TableRow);
-	TSharedRef<SWidget> GenerateWidgetForNameColumn(const FArguments& InArgs, const TSharedRef<class ITableRow>& TableRow);
-	TSharedRef<SWidget> GenerateWidgetForStatsColumn(const FArguments& InArgs, const TSharedRef<class ITableRow>& TableRow);
+	TSharedRef<SWidget> GenerateWidgetForColumn(const FArguments& InArgs, const TSharedRef<ITableRow>& TableRow);
+	TSharedRef<SWidget> GenerateWidgetForNameColumn(const FArguments& InArgs, const TSharedRef<ITableRow>& TableRow);
+	TSharedRef<SWidget> GenerateWidgetForStatsColumn(const FArguments& InArgs, const TSharedRef<ITableRow>& TableRow);
 
 	/**
 	 * The system will use this event to notify a widget that the cursor has entered it. This event is NOT bubbled.
@@ -98,11 +100,15 @@ protected:
 		return IsHovered() ? EVisibility::Visible : EVisibility::Hidden;
 	}
 
+	TSharedPtr<class IToolTip> GetRowToolTip(const TSharedRef<ITableRow>& TableRow) const;
+
 	FText GetDisplayName() const
 	{
 		return NetEventNodePtr->GetDisplayName();
 	}
 
+	FText GetValueAsText() const;
+
 	FSlateColor GetColorAndOpacity() const
 	{
 		const FLinearColor TextColor =
diff --git a/Engine/Source/Developer/TraceInsights/Private/Insights/NetworkingProfiler/Widgets/SNetStatsTableRow.cpp b/Engine/Source/Developer/TraceInsights/Private/Insights/NetworkingProfiler/Widgets/SNetStatsTableRow.cpp
--- a/Engine/Source/Developer/TraceInsights/Private/Insights/NetworkingProfiler/Widgets/SNetStatsTableRow.cpp
+++ b/Engine/Source/Developer/TraceInsights/Private/Insights/NetworkingProfiler/Widgets/SNetStatsTableRow.cpp
@@ -12,6 +12,7 @@
 #include "Insights/Table/ViewModels/Table.h"
 #include "Insights/Table/ViewModels/TableColumn.h"
 #include "Insights/NetworkingProfiler/Widgets/SNetStatsTableCell.h"
+#include "Insights/NetworkingProfiler/Widgets/SNetStatsViewTooltip.h"
 
 #define LOCTEXT_NAMESPACE "SNetStatsView"
 
@@ -23,8 +24,8 @@ void SNetStatsTableRow::Construct(const FArguments& InArgs, const TSharedRef<STa
 {
 	OnShouldBeEnabled = InArgs._OnShouldBeEnabled;
 	IsColumnVisibleDelegate = InArgs._OnIsColumnVisible;
-	SetHoveredCellDelegate = InArgs._OnSetHoveredCell;
 	GetColumnOutlineHAlignmentDelegate = InArgs._OnGetColumnOutlineHAlignmentDelegate;
+	SetHoveredCellDelegate = InArgs._OnSetHoveredCell;
 
 	HighlightText = InArgs._HighlightText;
 	HighlightedNodeName = InArgs._HighlightedNodeName;
@@ -32,6 +33,8 @@ void SNetStatsTableRow::Construct(const FArguments& InArgs, const TSharedRef<STa
 	TablePtr = InArgs._TablePtr;
 	NetEventNodePtr = InArgs._NetEventNodePtr;
 
+	RowToolTip = MakeShareable(new SNetEventTableRowToolTip(NetEventNodePtr));
+
 	SetEnabled(TAttribute<bool>(this, &SNetStatsTableRow::HandleShouldBeEnabled));
 
 	SMultiColumnTableRow<FNetEventNodePtr>::Construct(SMultiColumnTableRow<FNetEventNodePtr>::FArguments(), InOwnerTableView);
@@ -82,39 +85,28 @@ END_SLATE_FUNCTION_BUILD_OPTIMIZATION
 
 FReply SNetStatsTableRow::OnDragDetected(const FGeometry& MyGeometry, const FPointerEvent& MouseEvent)
 {
-	//im:TODO
-	//if (MouseEvent.IsMouseButtonDown(EKeys::LeftMouseButton))
-	//{
-	//	if (NetEventNode->IsGroup())
-	//	{
-	//		// Add all timer Ids for the group.
-	//		TArray<int32> NetEventIds;
-	//		const TArray<FNetEventNodePtr>& FilteredChildren = NetEventNode->GetFilteredChildren();
-	//		const int32 NumFilteredChildren = FilteredChildren.Num();
-	//
-	//		NetEventIds.Reserve(NumFilteredChildren);
-	//		for (int32 Nx = 0; Nx < NumFilteredChildren; ++Nx)
-	//		{
-	//			NetEventIds.Add(FilteredChildren[Nx]->GetId());
-	//		}
-	//
-	//		return FReply::Handled().BeginDragDrop(FStatIDDragDropOp::NewGroup(NetEventIds, NetEventNode->GetName().GetPlainNameString()));
-	//	}
-	//	else
-	//	{
-	//		return FReply::Handled().BeginDragDrop(FStatIDDragDropOp::NewSingle(NetEventNode->GetId(), NetEventNode->GetName().GetPlainNameString()));
-	//	}
-	//}
-
 	return SMultiColumnTableRow<FNetEventNodePtr>::OnDragDetected(MyGeometry, MouseEvent);
 }
 
 ////////////////////////////////////////////////////////////////////////////////////////////////////
 
+TSharedRef<IToolTip> SNetStatsTableRow::GetRowToolTip() const
+{
+	return RowToolTip.ToSharedRef();
+}
+
+////////////////////////////////////////////////////////////////////////////////////////////////////
+
+void SNetStatsTableRow::InvalidateContent()
+{
+	RowToolTip->InvalidateWidget();
+}
+
+////////////////////////////////////////////////////////////////////////////////////////////////////
+
 FSlateColor SNetStatsTableRow::GetBackgroundColorAndOpacity() const
 {
 	return GetBackgroundColorAndOpacity(NetEventNodePtr->GetAggregatedStats().TotalInclusive);
-	//return FLinearColor(0.0f, 0.0f, 0.0f, 1.0f)
 }
 
 ////////////////////////////////////////////////////////////////////////////////////////////////////
diff --git a/Engine/Source/Developer/TraceInsights/Private/Insights/NetworkingProfiler/Widgets/SNetStatsTableRow.h b/Engine/Source/Developer/TraceInsights/Private/Insights/NetworkingProfiler/Widgets/SNetStatsTableRow.h
--- a/Engine/Source/Developer/TraceInsights/Private/Insights/NetworkingProfiler/Widgets/SNetStatsTableRow.h
+++ b/Engine/Source/Developer/TraceInsights/Private/Insights/NetworkingProfiler/Widgets/SNetStatsTableRow.h
@@ -12,6 +12,9 @@
 // Insights
 #include "Insights/NetworkingProfiler/ViewModels/NetEventNodeHelper.h"
 
+class IToolTip;
+class SNetEventTableRowToolTip;
+
 namespace Insights
 {
 	class FTable;
@@ -23,7 +26,7 @@ DECLARE_DELEGATE_RetVal_OneParam(bool, FIsColumnVisibleDelegate, const FName /*C
 DECLARE_DELEGATE_RetVal_OneParam(EHorizontalAlignment, FGetColumnOutlineHAlignmentDelegate, const FName /*ColumnId*/);
 DECLARE_DELEGATE_ThreeParams(FSetHoveredNetEventTableCell, TSharedPtr<Insights::FTable> /*TablePtr*/, TSharedPtr<Insights::FTableColumn> /*ColumnPtr*/, const FNetEventNodePtr /*NetEventNodePtr*/);
 
-/** Widget that represents a table row in a timer tree control. Generates widgets for each column on demand. */
+/** Widget that represents a table row in the tree control. Generates widgets for each column on demand. */
 class SNetStatsTableRow : public SMultiColumnTableRow<FNetEventNodePtr>
 {
 public:
@@ -59,6 +62,9 @@ public:
 	 */
 	virtual FReply OnDragDetected(const FGeometry& MyGeometry, const FPointerEvent& MouseEvent) override;
 
+	TSharedRef<IToolTip> GetRowToolTip() const;
+	void InvalidateContent();
+
 protected:
 	FSlateColor GetBackgroundColorAndOpacity() const;
 	FSlateColor GetBackgroundColorAndOpacity(uint32 Size) const;
@@ -85,4 +91,6 @@ protected:
 
 	/** Name of the timer node that should be drawn as highlighted. */
 	TAttribute<FName> HighlightedNodeName;
+
+	TSharedPtr<SNetEventTableRowToolTip> RowToolTip;
 };
diff --git a/Engine/Source/Developer/TraceInsights/Private/Insights/NetworkingProfiler/Widgets/SNetStatsView.cpp b/Engine/Source/Developer/TraceInsights/Private/Insights/NetworkingProfiler/Widgets/SNetStatsView.cpp
--- a/Engine/Source/Developer/TraceInsights/Private/Insights/NetworkingProfiler/Widgets/SNetStatsView.cpp
+++ b/Engine/Source/Developer/TraceInsights/Private/Insights/NetworkingProfiler/Widgets/SNetStatsView.cpp
@@ -1636,24 +1636,32 @@ void SNetStatsView::UpdateStatsInternal()
 
 	if (Session.IsValid())
 	{
-		TUniquePtr<Trace::ITable<Trace::FNetProfilerAggregatedStats>> AggregatedStatsTable;
+		TUniquePtr<Trace::ITable<Trace::FNetProfilerAggregatedStats>> AggregationResultTable;
 
 		{
 			Trace::FAnalysisSessionReadScope SessionReadScope(*Session.Get());
 			const Trace::INetProfilerProvider& NetProfilerProvider = Trace::ReadNetProfilerProvider(*Session.Get());
-			AggregatedStatsTable.Reset(NetProfilerProvider.CreateAggregation(ConnectionIndex, ConnectionMode, StatsPacketStartIndex, StatsPacketEndIndex - 1, StatsStartPosition, StatsEndPosition));
+			AggregationResultTable.Reset(NetProfilerProvider.CreateAggregation(ConnectionIndex, ConnectionMode, StatsPacketStartIndex, StatsPacketEndIndex - 1, StatsStartPosition, StatsEndPosition));
 		}
 
-		if (AggregatedStatsTable.IsValid())
+		if (AggregationResultTable.IsValid())
 		{
-			TUniquePtr<Trace::ITableReader<Trace::FNetProfilerAggregatedStats>> TableReader(AggregatedStatsTable->CreateReader());
+			TUniquePtr<Trace::ITableReader<Trace::FNetProfilerAggregatedStats>> TableReader(AggregationResultTable->CreateReader());
 			while (TableReader->IsValid())
 			{
 				const Trace::FNetProfilerAggregatedStats* Row = TableReader->GetCurrentRow();
 				FNetEventNodePtr* NetEventNodePtrPtr = NetEventNodesIdMap.Find(static_cast<uint64>(Row->EventTypeIndex));
 				if (NetEventNodePtrPtr != nullptr)
 				{
-					(*NetEventNodePtrPtr)->SetAggregatedStats(*Row);
+					FNetEventNodePtr NetEventNodePtr = *NetEventNodePtrPtr;
+					NetEventNodePtr->SetAggregatedStats(*Row);
+
+					TSharedPtr<ITableRow> TableRowPtr = TreeView->WidgetFromItem(NetEventNodePtr);
+					if (TableRowPtr.IsValid())
+					{
+						TSharedPtr<SNetStatsTableRow> RowPtr = StaticCastSharedPtr<SNetStatsTableRow, ITableRow>(TableRowPtr);
+						RowPtr->InvalidateContent();
+					}
 				}
 				TableReader->NextRow();
 			}
diff --git a/Engine/Source/Developer/TraceInsights/Private/Insights/NetworkingProfiler/Widgets/SNetStatsViewTooltip.cpp b/Engine/Source/Developer/TraceInsights/Private/Insights/NetworkingProfiler/Widgets/SNetStatsViewTooltip.cpp
--- a/Engine/Source/Developer/TraceInsights/Private/Insights/NetworkingProfiler/Widgets/SNetStatsViewTooltip.cpp
+++ b/Engine/Source/Developer/TraceInsights/Private/Insights/NetworkingProfiler/Widgets/SNetStatsViewTooltip.cpp
@@ -87,10 +87,12 @@ TSharedPtr<SToolTip> SNetStatsViewTooltip::GetColumnTooltip(const Insights::FTab
 
 ////////////////////////////////////////////////////////////////////////////////////////////////////
 
-TSharedPtr<SToolTip> SNetStatsViewTooltip::GetCellTooltip(const TSharedPtr<FNetEventNode> NetEventNodePtr, const TSharedPtr<Insights::FTableColumn> ColumnPtr)
+TSharedPtr<SToolTip> SNetStatsViewTooltip::GetRowTooltip(const TSharedPtr<FNetEventNode> NetEventNodePtr)
 {
 	const Trace::FNetProfilerAggregatedStats& Stats = NetEventNodePtr->GetAggregatedStats();
 
+	const FText InstanceCountText = FText::AsNumber(Stats.InstanceCount);
+
 	const FText TotalInclusiveSizeText = FText::AsNumber(Stats.TotalInclusive);
 	const FText MaxInclusiveSizeText = FText::AsNumber(Stats.MaxInclusive);
 	const FText AvgInclusiveSizeText = FText::AsNumber(Stats.AverageInclusive);
@@ -100,8 +102,6 @@ TSharedPtr<SToolTip> SNetStatsViewTooltip::GetCellTooltip(const TSharedPtr<FNetE
 	//const FText AvgExclusiveSizeText = FText::AsNumber(Stats.AverageExclusive);
 	const FText AvgExclusiveSizeText;
 
-	const FText InstanceCountText = FText::AsNumber(Stats.InstanceCount);
-
 	TSharedPtr<SGridPanel> GridPanel;
 	TSharedPtr<SHorizontalBox> HBox;
 
@@ -177,7 +177,7 @@ TSharedPtr<SToolTip> SNetStatsViewTooltip::GetCellTooltip(const TSharedPtr<FNetE
 					//	.TextStyle(FEditorStyle::Get(), TEXT("Profiler.Tooltip"))
 					//]
 
-					// Sizer Type: [Type]
+					// Net Event Type: [Type]
 					+ SGridPanel::Slot(0, 3)
 					.Padding(2.0f)
 					[
diff --git a/Engine/Source/Developer/TraceInsights/Private/Insights/NetworkingProfiler/Widgets/SNetStatsViewTooltip.h b/Engine/Source/Developer/TraceInsights/Private/Insights/NetworkingProfiler/Widgets/SNetStatsViewTooltip.h
--- a/Engine/Source/Developer/TraceInsights/Private/Insights/NetworkingProfiler/Widgets/SNetStatsViewTooltip.h
+++ b/Engine/Source/Developer/TraceInsights/Private/Insights/NetworkingProfiler/Widgets/SNetStatsViewTooltip.h
@@ -3,11 +3,9 @@
 #pragma once
 
 #include "CoreMinimal.h"
+#include "Widgets/IToolTip.h"
+#include "Widgets/SToolTip.h"
 
-// Insights
-#include "Insights/InsightsManager.h"
-
-class SToolTip;
 class SGridPanel;
 
 namespace Insights
@@ -18,7 +16,7 @@ namespace Insights
 
 class FNetEventNode;
 
-#define LOCTEXT_NAMESPACE "SNetStatsView"
+////////////////////////////////////////////////////////////////////////////////////////////////////
 
 /** NetStats View Tooltip */
 class SNetStatsViewTooltip
@@ -28,10 +26,60 @@ public:
 
 	static TSharedPtr<SToolTip> GetTableTooltip(const Insights::FTable& Table);
 	static TSharedPtr<SToolTip> GetColumnTooltip(const Insights::FTableColumn& Column);
-	static TSharedPtr<SToolTip> GetCellTooltip(const TSharedPtr<FNetEventNode> TreeNodePtr, const TSharedPtr<Insights::FTableColumn> ColumnPtr);
+	static TSharedPtr<SToolTip> GetRowTooltip(const TSharedPtr<FNetEventNode> TreeNodePtr);
 
 private:
 	static void AddStatsRow(TSharedPtr<SGridPanel> Grid, int32& Row, const FText& Name, const FText& Value1, const FText& Value2);
 };
 
-#undef LOCTEXT_NAMESPACE
+////////////////////////////////////////////////////////////////////////////////////////////////////
+
+class SNetEventTableRowToolTip : public IToolTip
+{
+public:
+	SNetEventTableRowToolTip(const TSharedPtr<FNetEventNode> InTreeNodePtr) : TreeNodePtr(InTreeNodePtr) {}
+	virtual ~SNetEventTableRowToolTip() { }
+
+	virtual TSharedRef<class SWidget> AsWidget()
+	{
+		CreateToolTipWidget();
+		return ToolTipWidget.ToSharedRef();
+	}
+
+	virtual TSharedRef<SWidget> GetContentWidget()
+	{
+		CreateToolTipWidget();
+		return ToolTipWidget->GetContentWidget();
+	}
+
+	virtual void SetContentWidget(const TSharedRef<SWidget>& InContentWidget)
+	{
+		CreateToolTipWidget();
+		ToolTipWidget->SetContentWidget(InContentWidget);
+	}
+
+	void InvalidateWidget()
+	{
+		ToolTipWidget.Reset();
+	}
+
+	virtual bool IsEmpty() const { return false; }
+	virtual bool IsInteractive() const { return false; }
+	virtual void OnOpening() {}
+	virtual void OnClosed() {}
+
+private:
+	void CreateToolTipWidget()
+	{
+		if (!ToolTipWidget.IsValid())
+		{
+			ToolTipWidget = SNetStatsViewTooltip::GetRowTooltip(TreeNodePtr);
+		}
+	}
+
+private:
+	TSharedPtr<SToolTip> ToolTipWidget;
+	const TSharedPtr<FNetEventNode> TreeNodePtr;
+};
+
+////////////////////////////////////////////////////////////////////////////////////////////////////
diff --git a/Engine/Source/Developer/TraceInsights/Private/Insights/Table/Widgets/STableTreeViewCell.cpp b/Engine/Source/Developer/TraceInsights/Private/Insights/Table/Widgets/STableTreeViewCell.cpp
--- a/Engine/Source/Developer/TraceInsights/Private/Insights/Table/Widgets/STableTreeViewCell.cpp
+++ b/Engine/Source/Developer/TraceInsights/Private/Insights/Table/Widgets/STableTreeViewCell.cpp
@@ -14,7 +14,7 @@
 // Insights
 #include "Insights/Table/ViewModels/Table.h"
 #include "Insights/Table/ViewModels/TableColumn.h"
-#include "Insights/Table/Widgets/STableTreeViewTooltip.h"
+#include "Insights/Table/Widgets/STableTreeViewRow.h"
 
 #define LOCTEXT_NAMESPACE "STableTreeView"
 
@@ -81,7 +81,7 @@ TSharedRef<SWidget> STableTreeViewCell::GenerateWidgetForNameColumn(const FArgum
 			SNew(SImage)
 			.Visibility(this, &STableTreeViewCell::GetHintIconVisibility)
 			.Image(FEditorStyle::GetBrush("Profiler.Tooltip.HintIcon10"))
-			.ToolTip(STableTreeViewTooltip::GetCellTooltip(TableTreeNodePtr, ColumnPtr))
+			.ToolTip(GetRowToolTip(TableRow))
 		]
 
 		// Name
@@ -103,8 +103,26 @@ TSharedRef<SWidget> STableTreeViewCell::GenerateWidgetForNameColumn(const FArgum
 
 ////////////////////////////////////////////////////////////////////////////////////////////////////
 
+TSharedPtr<IToolTip> STableTreeViewCell::GetRowToolTip(const TSharedRef<ITableRow>& TableRow) const
+{
+	TSharedRef<STableTreeViewRow> Row = StaticCastSharedRef<STableTreeViewRow, ITableRow>(TableRow);
+	return Row->GetRowToolTip();
+}
+
+////////////////////////////////////////////////////////////////////////////////////////////////////
+
+FText STableTreeViewCell::GetValueAsText() const
+{
+	return ColumnPtr->GetValueAsText(*TableTreeNodePtr);
+}
+
+////////////////////////////////////////////////////////////////////////////////////////////////////
+
 TSharedRef<SWidget> STableTreeViewCell::GenerateWidgetForTableColumn(const FArguments& InArgs, const TSharedRef<ITableRow>& TableRow)
 {
+	// Note: For performance reason, init the cell text (value) only once.
+	//       If we'll need to update values without recreating the table row/cell widgets, bind .Text to STableTreeViewCell::GetValueAsText
+	//       or add API to explicitly update the text block.
 	const FText CellText = ColumnPtr->GetValueAsText(*TableTreeNodePtr);
 
 	return
diff --git a/Engine/Source/Developer/TraceInsights/Private/Insights/Table/Widgets/STableTreeViewCell.h b/Engine/Source/Developer/TraceInsights/Private/Insights/Table/Widgets/STableTreeViewCell.h
--- a/Engine/Source/Developer/TraceInsights/Private/Insights/Table/Widgets/STableTreeViewCell.h
+++ b/Engine/Source/Developer/TraceInsights/Private/Insights/Table/Widgets/STableTreeViewCell.h
@@ -100,11 +100,15 @@ protected:
 		return IsHovered() ? EVisibility::Visible : EVisibility::Hidden;
 	}
 
+	TSharedPtr<class IToolTip> GetRowToolTip(const TSharedRef<ITableRow>& TableRow) const;
+
 	FText GetDisplayName() const
 	{
 		return TableTreeNodePtr->GetDisplayName();
 	}
 
+	FText GetValueAsText() const;
+
 	FSlateColor GetColorAndOpacity() const
 	{
 		const FLinearColor TextColor =
diff --git a/Engine/Source/Developer/TraceInsights/Private/Insights/Table/Widgets/STableTreeViewRow.cpp b/Engine/Source/Developer/TraceInsights/Private/Insights/Table/Widgets/STableTreeViewRow.cpp
--- a/Engine/Source/Developer/TraceInsights/Private/Insights/Table/Widgets/STableTreeViewRow.cpp
+++ b/Engine/Source/Developer/TraceInsights/Private/Insights/Table/Widgets/STableTreeViewRow.cpp
@@ -12,6 +12,7 @@
 #include "Insights/Table/ViewModels/Table.h"
 #include "Insights/Table/ViewModels/TableColumn.h"
 #include "Insights/Table/Widgets/STableTreeViewCell.h"
+#include "Insights/Table/Widgets/STableTreeViewTooltip.h"
 
 #define LOCTEXT_NAMESPACE "STableTreeView"
 
@@ -26,8 +27,8 @@ void STableTreeViewRow::Construct(const FArguments& InArgs, const TSharedRef<STa
 {
 	OnShouldBeEnabled = InArgs._OnShouldBeEnabled;
 	IsColumnVisibleDelegate = InArgs._OnIsColumnVisible;
-	SetHoveredCellDelegate = InArgs._OnSetHoveredCell;
 	GetColumnOutlineHAlignmentDelegate = InArgs._OnGetColumnOutlineHAlignmentDelegate;
+	SetHoveredCellDelegate = InArgs._OnSetHoveredCell;
 
 	HighlightText = InArgs._HighlightText;
 	HighlightedNodeName = InArgs._HighlightedNodeName;
@@ -35,6 +36,8 @@ void STableTreeViewRow::Construct(const FArguments& InArgs, const TSharedRef<STa
 	TablePtr = InArgs._TablePtr;
 	TableTreeNodePtr = InArgs._TableTreeNodePtr;
 
+	RowToolTip = MakeShareable(new STableTreeRowToolTip(TableTreeNodePtr));
+
 	SetEnabled(TAttribute<bool>(this, &STableTreeViewRow::HandleShouldBeEnabled));
 
 	SMultiColumnTableRow<FTableTreeNodePtr>::Construct(SMultiColumnTableRow<FTableTreeNodePtr>::FArguments(), InOwnerTableView);
@@ -90,6 +93,20 @@ FReply STableTreeViewRow::OnDragDetected(const FGeometry& MyGeometry, const FPoi
 
 ////////////////////////////////////////////////////////////////////////////////////////////////////
 
+TSharedRef<IToolTip> STableTreeViewRow::GetRowToolTip() const
+{
+	return RowToolTip.ToSharedRef();
+}
+
+////////////////////////////////////////////////////////////////////////////////////////////////////
+
+void STableTreeViewRow::InvalidateContent()
+{
+	RowToolTip->InvalidateWidget();
+}
+
+////////////////////////////////////////////////////////////////////////////////////////////////////
+
 FSlateColor STableTreeViewRow::GetBackgroundColorAndOpacity() const
 {
 	return FLinearColor(0.0f, 0.0f, 0.0f, 1.0f);
@@ -102,7 +119,7 @@ FSlateColor STableTreeViewRow::GetBackgroundColorAndOpacity(double Time) const
 	const FLinearColor Color =	Time > TimeUtils::Second      ? FLinearColor(0.3f, 0.0f, 0.0f, 1.0f) :
 								Time > TimeUtils::Milisecond  ? FLinearColor(0.3f, 0.1f, 0.0f, 1.0f) :
 								Time > TimeUtils::Microsecond ? FLinearColor(0.0f, 0.1f, 0.0f, 1.0f) :
-																FLinearColor(0.0f, 0.0f, 0.0f, 1.0f);
+								                                FLinearColor(0.0f, 0.0f, 0.0f, 1.0f);
 	return Color;
 }
 
diff --git a/Engine/Source/Developer/TraceInsights/Private/Insights/Table/Widgets/STableTreeViewRow.h b/Engine/Source/Developer/TraceInsights/Private/Insights/Table/Widgets/STableTreeViewRow.h
--- a/Engine/Source/Developer/TraceInsights/Private/Insights/Table/Widgets/STableTreeViewRow.h
+++ b/Engine/Source/Developer/TraceInsights/Private/Insights/Table/Widgets/STableTreeViewRow.h
@@ -12,11 +12,14 @@
 // Insights
 #include "Insights/Table/ViewModels/TableTreeNode.h"
 
+class IToolTip;
+
 namespace Insights
 {
 
 class FTable;
 class FTableColumn;
+class STableTreeRowToolTip;
 
 DECLARE_DELEGATE_RetVal_OneParam(bool, FShouldBeEnabledDelegate, const uint32 /*TreeNodeId*/);
 DECLARE_DELEGATE_RetVal_OneParam(bool, FIsColumnVisibleDelegate, const FName /*ColumnId*/);
@@ -59,6 +62,9 @@ public:
 	 */
 	virtual FReply OnDragDetected(const FGeometry& MyGeometry, const FPointerEvent& MouseEvent) override;
 
+	TSharedRef<IToolTip> GetRowToolTip() const;
+	void InvalidateContent();
+
 protected:
 	FSlateColor GetBackgroundColorAndOpacity() const;
 	FSlateColor GetBackgroundColorAndOpacity(double Time) const;
@@ -85,6 +91,8 @@ protected:
 
 	/** Name of the tree node that should be drawn as highlighted. */
 	TAttribute<FName> HighlightedNodeName;
+
+	TSharedPtr<STableTreeRowToolTip> RowToolTip;
 };
 
 } // namespace Insights
diff --git a/Engine/Source/Developer/TraceInsights/Private/Insights/Table/Widgets/STableTreeViewTooltip.cpp b/Engine/Source/Developer/TraceInsights/Private/Insights/Table/Widgets/STableTreeViewTooltip.cpp
--- a/Engine/Source/Developer/TraceInsights/Private/Insights/Table/Widgets/STableTreeViewTooltip.cpp
+++ b/Engine/Source/Developer/TraceInsights/Private/Insights/Table/Widgets/STableTreeViewTooltip.cpp
@@ -4,7 +4,6 @@
 
 #include "EditorStyleSet.h"
 #include "SlateOptMacros.h"
-#include "Widgets/Images/SImage.h"
 #include "Widgets/Layout/SGridPanel.h"
 #include "Widgets/Layout/SSeparator.h"
 #include "Widgets/SBoxPanel.h"
@@ -88,7 +87,7 @@ TSharedPtr<SToolTip> STableTreeViewTooltip::GetColumnTooltip(const FTableColumn&
 
 ////////////////////////////////////////////////////////////////////////////////////////////////////
 
-TSharedPtr<SToolTip> STableTreeViewTooltip::GetCellTooltip(const TSharedPtr<FTableTreeNode> TreeNodePtr, const TSharedPtr<FTableColumn> ColumnPtr)
+TSharedPtr<SToolTip> STableTreeViewTooltip::GetRowTooltip(const TSharedPtr<FTableTreeNode> TreeNodePtr)
 {
 	TSharedPtr<SGridPanel> GridPanel;
 	TSharedPtr<SHorizontalBox> HBox;
diff --git a/Engine/Source/Developer/TraceInsights/Private/Insights/Table/Widgets/STableTreeViewTooltip.h b/Engine/Source/Developer/TraceInsights/Private/Insights/Table/Widgets/STableTreeViewTooltip.h
--- a/Engine/Source/Developer/TraceInsights/Private/Insights/Table/Widgets/STableTreeViewTooltip.h
+++ b/Engine/Source/Developer/TraceInsights/Private/Insights/Table/Widgets/STableTreeViewTooltip.h
@@ -3,8 +3,9 @@
 #pragma once
 
 #include "CoreMinimal.h"
+#include "Widgets/IToolTip.h"
+#include "Widgets/SToolTip.h"
 
-class SToolTip;
 class SGridPanel;
 
 namespace Insights
@@ -14,6 +15,8 @@ class FTable;
 class FTableColumn;
 class FTableTreeNode;
 
+////////////////////////////////////////////////////////////////////////////////////////////////////
+
 /** Tooltip for STableTreeView widget. */
 class STableTreeViewTooltip
 {
@@ -22,10 +25,62 @@ public:
 
 	static TSharedPtr<SToolTip> GetTableTooltip(const FTable& Table);
 	static TSharedPtr<SToolTip> GetColumnTooltip(const FTableColumn& Column);
-	static TSharedPtr<SToolTip> GetCellTooltip(const TSharedPtr<FTableTreeNode> TreeNodePtr, const TSharedPtr<FTableColumn> ColumnPtr);
+	static TSharedPtr<SToolTip> GetRowTooltip(const TSharedPtr<FTableTreeNode> TreeNodePtr);
 
 private:
 	static void AddGridRow(TSharedPtr<SGridPanel> Grid, int32& Row, const FText& Name, const FText& Value);
 };
 
+////////////////////////////////////////////////////////////////////////////////////////////////////
+
+class STableTreeRowToolTip : public IToolTip
+{
+public:
+	STableTreeRowToolTip(const TSharedPtr<FTableTreeNode> InTreeNodePtr) : TreeNodePtr(InTreeNodePtr) {}
+	virtual ~STableTreeRowToolTip() { }
+
+	virtual TSharedRef<class SWidget> AsWidget()
+	{
+		CreateToolTipWidget();
+		return ToolTipWidget.ToSharedRef();
+	}
+
+	virtual TSharedRef<SWidget> GetContentWidget()
+	{
+		CreateToolTipWidget();
+		return ToolTipWidget->GetContentWidget();
+	}
+
+	virtual void SetContentWidget(const TSharedRef<SWidget>& InContentWidget)
+	{
+		CreateToolTipWidget();
+		ToolTipWidget->SetContentWidget(InContentWidget);
+	}
+
+	void InvalidateWidget()
+	{
+		ToolTipWidget.Reset();
+	}
+
+	virtual bool IsEmpty() const { return false; }
+	virtual bool IsInteractive() const { return false; }
+	virtual void OnOpening() {}
+	virtual void OnClosed() {}
+
+private:
+	void CreateToolTipWidget()
+	{
+		if (!ToolTipWidget.IsValid())
+		{
+			ToolTipWidget = STableTreeViewTooltip::GetRowTooltip(TreeNodePtr);
+		}
+	}
+
+private:
+	TSharedPtr<SToolTip> ToolTipWidget;
+	const TSharedPtr<FTableTreeNode> TreeNodePtr;
+};
+
+////////////////////////////////////////////////////////////////////////////////////////////////////
+
 } // namespace Insights
diff --git a/Engine/Source/Developer/TraceInsights/Private/Insights/Widgets/SStatsTableCell.cpp b/Engine/Source/Developer/TraceInsights/Private/Insights/Widgets/SStatsTableCell.cpp
--- a/Engine/Source/Developer/TraceInsights/Private/Insights/Widgets/SStatsTableCell.cpp
+++ b/Engine/Source/Developer/TraceInsights/Private/Insights/Widgets/SStatsTableCell.cpp
@@ -16,7 +16,7 @@
 #include "Insights/InsightsStyle.h"
 #include "Insights/ViewModels/StatsViewColumn.h"
 #include "Insights/ViewModels/StatsViewColumnFactory.h"
-#include "Insights/Widgets/SStatsViewTooltip.h"
+#include "Insights/Widgets/SStatsTableRow.h"
 
 #define LOCTEXT_NAMESPACE "SStatsView"
 
@@ -24,12 +24,14 @@
 
 BEGIN_SLATE_FUNCTION_BUILD_OPTIMIZATION
 
-void SStatsTableCell::Construct(const FArguments& InArgs, const TSharedRef<class ITableRow>& TableRow)
+void SStatsTableCell::Construct(const FArguments& InArgs, const TSharedRef<ITableRow>& TableRow)
 {
-	SetHoveredTableCellDelegate = InArgs._OnSetHoveredTableCell;
 	StatsNodePtr = InArgs._StatsNodePtr;
 	ColumnId = InArgs._ColumnId;
-	//HighlightText = InArgs._HighlightText;
+
+	ensure(StatsNodePtr.IsValid());
+
+	SetHoveredTableCellDelegate = InArgs._OnSetHoveredTableCell;
 
 	ChildSlot
 		[
@@ -39,7 +41,7 @@ void SStatsTableCell::Construct(const FArguments& InArgs, const TSharedRef<class
 
 ////////////////////////////////////////////////////////////////////////////////////////////////////
 
-TSharedRef<SWidget> SStatsTableCell::GenerateWidgetForColumn(const FArguments& InArgs, const TSharedRef<class ITableRow>& TableRow)
+TSharedRef<SWidget> SStatsTableCell::GenerateWidgetForColumn(const FArguments& InArgs, const TSharedRef<ITableRow>& TableRow)
 {
 	if (InArgs._IsNameColumn)
 	{
@@ -53,7 +55,7 @@ TSharedRef<SWidget> SStatsTableCell::GenerateWidgetForColumn(const FArguments& I
 
 ////////////////////////////////////////////////////////////////////////////////////////////////////
 
-TSharedRef<SWidget> SStatsTableCell::GenerateWidgetForNameColumn(const FArguments& InArgs, const TSharedRef<class ITableRow>& TableRow)
+TSharedRef<SWidget> SStatsTableCell::GenerateWidgetForNameColumn(const FArguments& InArgs, const TSharedRef<ITableRow>& TableRow)
 {
 	const FStatsViewColumn& Column = *FStatsViewColumnFactory::Get().ColumnIdToPtrMapping.FindChecked(ColumnId);
 
@@ -77,7 +79,7 @@ TSharedRef<SWidget> SStatsTableCell::GenerateWidgetForNameColumn(const FArgument
 			SNew(SImage)
 			.Visibility(this, &SStatsTableCell::GetHintIconVisibility)
 			.Image(FEditorStyle::GetBrush("Profiler.Tooltip.HintIcon10"))
-			.ToolTip(SStatsViewTooltip::GetTableCellTooltip(StatsNodePtr))
+			.ToolTip(GetRowToolTip(TableRow))
 		]
 
 		// Color box
@@ -118,11 +120,25 @@ TSharedRef<SWidget> SStatsTableCell::GenerateWidgetForNameColumn(const FArgument
 
 ////////////////////////////////////////////////////////////////////////////////////////////////////
 
-TSharedRef<SWidget> SStatsTableCell::GenerateWidgetForStatsColumn(const FArguments& InArgs, const TSharedRef<class ITableRow>& TableRow)
+TSharedPtr<IToolTip> SStatsTableCell::GetRowToolTip(const TSharedRef<ITableRow>& TableRow) const
+{
+	TSharedRef<SStatsTableRow> Row = StaticCastSharedRef<SStatsTableRow, ITableRow>(TableRow);
+	return Row->GetRowToolTip();
+}
+
+////////////////////////////////////////////////////////////////////////////////////////////////////
+
+FText SStatsTableCell::GetValueAsText() const
+{
+	const FStatsViewColumn& Column = *FStatsViewColumnFactory::Get().ColumnIdToPtrMapping.FindChecked(ColumnId);
+	return Column.GetFormattedValue(*StatsNodePtr);
+}
+
+////////////////////////////////////////////////////////////////////////////////////////////////////
+
+TSharedRef<SWidget> SStatsTableCell::GenerateWidgetForStatsColumn(const FArguments& InArgs, const TSharedRef<ITableRow>& TableRow)
 {
 	const FStatsViewColumn& Column = *FStatsViewColumnFactory::Get().ColumnIdToPtrMapping.FindChecked(ColumnId);
-	const FText CellText = Column.GetFormattedValue(*StatsNodePtr);
-	//const FText CellText = ColumnPtr->GetValueAsText(*StatsNodePtr);
 
 	return
 		SNew(SHorizontalBox)
@@ -135,7 +151,7 @@ TSharedRef<SWidget> SStatsTableCell::GenerateWidgetForStatsColumn(const FArgumen
 		.Padding(FMargin(2.0f, 0.0f))
 		[
 			SNew(STextBlock)
-			.Text(CellText)
+			.Text(this, &SStatsTableCell::GetValueAsText)
 			.TextStyle(FEditorStyle::Get(), TEXT("Profiler.Tooltip"))
 			.ColorAndOpacity(this, &SStatsTableCell::GetStatsColorAndOpacity)
 			.ShadowColorAndOpacity(this, &SStatsTableCell::GetShadowColorAndOpacity)
diff --git a/Engine/Source/Developer/TraceInsights/Private/Insights/Widgets/SStatsTableCell.h b/Engine/Source/Developer/TraceInsights/Private/Insights/Widgets/SStatsTableCell.h
--- a/Engine/Source/Developer/TraceInsights/Private/Insights/Widgets/SStatsTableCell.h
+++ b/Engine/Source/Developer/TraceInsights/Private/Insights/Widgets/SStatsTableCell.h
@@ -9,6 +9,8 @@
 // Insights
 #include "Insights/ViewModels/StatsNode.h"
 
+class ITableRow;
+
 DECLARE_DELEGATE_TwoParams(FSetHoveredStatsTableCell, const FName /*ColumnId*/, const FStatsNodePtr /*SamplePtr*/);
 DECLARE_DELEGATE_RetVal_OneParam(EHorizontalAlignment, FGetColumnOutlineHAlignmentDelegate, const FName /*ColumnId*/);
 
@@ -23,12 +25,12 @@ public:
 		SLATE_ARGUMENT(bool, IsNameColumn)
 	SLATE_END_ARGS()
 
-	void Construct(const FArguments& InArgs, const TSharedRef<class ITableRow>& TableRow);
+	void Construct(const FArguments& InArgs, const TSharedRef<ITableRow>& TableRow);
 
 protected:
-	TSharedRef<SWidget> GenerateWidgetForColumn(const FArguments& InArgs, const TSharedRef<class ITableRow>& TableRow);
-	TSharedRef<SWidget> GenerateWidgetForNameColumn(const FArguments& InArgs, const TSharedRef<class ITableRow>& TableRow);
-	TSharedRef<SWidget> GenerateWidgetForStatsColumn(const FArguments& InArgs, const TSharedRef<class ITableRow>& TableRow);
+	TSharedRef<SWidget> GenerateWidgetForColumn(const FArguments& InArgs, const TSharedRef<ITableRow>& TableRow);
+	TSharedRef<SWidget> GenerateWidgetForNameColumn(const FArguments& InArgs, const TSharedRef<ITableRow>& TableRow);
+	TSharedRef<SWidget> GenerateWidgetForStatsColumn(const FArguments& InArgs, const TSharedRef<ITableRow>& TableRow);
 
 	/**
 	 * The system will use this event to notify a widget that the cursor has entered it. This event is NOT bubbled.
@@ -97,25 +99,32 @@ protected:
 		return IsHovered() ? EVisibility::Visible : EVisibility::Hidden;
 	}
 
+	TSharedPtr<class IToolTip> GetRowToolTip(const TSharedRef<ITableRow>& TableRow) const;
+
 	FText GetDisplayName() const
 	{
 		return StatsNodePtr->GetDisplayName();
 	}
 
+	FText GetValueAsText() const;
+
 	FSlateColor GetColorAndOpacity() const
 	{
 		const FLinearColor TextColor =
-			StatsNodePtr->IsFiltered() ? FLinearColor(1.0f, 1.0f, 1.0f, 0.5f) :
-			FLinearColor::White;
+			StatsNodePtr->IsFiltered() ?
+				FLinearColor(1.0f, 1.0f, 1.0f, 0.5f) :
+				FLinearColor(1.0f, 1.0f, 1.0f, 1.0f);
 		return TextColor;
 	}
 
 	FSlateColor GetStatsColorAndOpacity() const
 	{
 		const FLinearColor TextColor =
-			StatsNodePtr->IsFiltered() ? FLinearColor(1.0f, 1.0f, 1.0f, 0.5f) :
-			StatsNodePtr->GetAggregatedStats().Count == 0 ? FLinearColor(1.0f, 1.0f, 1.0f, 0.6f) :
-			FLinearColor::White;
+			StatsNodePtr->IsFiltered() ?
+				FLinearColor(1.0f, 1.0f, 1.0f, 0.5f) :
+			StatsNodePtr->GetAggregatedStats().Count == 0 ?
+				FLinearColor(1.0f, 1.0f, 1.0f, 0.6f) :
+				FLinearColor(1.0f, 1.0f, 1.0f, 1.0f);
 		return TextColor;
 	}
 
@@ -127,8 +136,9 @@ protected:
 	FLinearColor GetShadowColorAndOpacity() const
 	{
 		const FLinearColor ShadowColor =
-			StatsNodePtr->IsFiltered() ? FLinearColor(0.f, 0.f, 0.f, 0.25f) :
-			FLinearColor(0.0f, 0.0f, 0.0f, 0.5f);
+			StatsNodePtr->IsFiltered() ?
+				FLinearColor(0.f, 0.f, 0.f, 0.25f) :
+				FLinearColor(0.0f, 0.0f, 0.0f, 0.5f);
 		return ShadowColor;
 	}
 
diff --git a/Engine/Source/Developer/TraceInsights/Private/Insights/Widgets/SStatsTableRow.cpp b/Engine/Source/Developer/TraceInsights/Private/Insights/Widgets/SStatsTableRow.cpp
--- a/Engine/Source/Developer/TraceInsights/Private/Insights/Widgets/SStatsTableRow.cpp
+++ b/Engine/Source/Developer/TraceInsights/Private/Insights/Widgets/SStatsTableRow.cpp
@@ -10,8 +10,8 @@
 // Insights
 #include "Insights/Common/TimeUtils.h"
 #include "Insights/ViewModels/StatsViewColumnFactory.h"
-#include "Insights/Widgets/SStatsViewTooltip.h"
 #include "Insights/Widgets/SStatsTableCell.h"
+#include "Insights/Widgets/SStatsViewTooltip.h"
 
 #define LOCTEXT_NAMESPACE "SStatsView"
 
@@ -23,14 +23,16 @@ void SStatsTableRow::Construct(const FArguments& InArgs, const TSharedRef<STable
 {
 	OnShouldBeEnabled = InArgs._OnShouldBeEnabled;
 	IsColumnVisibleDelegate = InArgs._OnIsColumnVisible;
-	SetHoveredTableCellDelegate = InArgs._OnSetHoveredTableCell;
 	GetColumnOutlineHAlignmentDelegate = InArgs._OnGetColumnOutlineHAlignmentDelegate;
+	SetHoveredTableCellDelegate = InArgs._OnSetHoveredTableCell;
 
 	HighlightText = InArgs._HighlightText;
 	HighlightedNodeName = InArgs._HighlightedNodeName;
 
 	StatsNodePtr = InArgs._StatsNodePtr;
 
+	RowToolTip = MakeShareable(new SStatsCounterTableRowToolTip(StatsNodePtr));
+
 	SetEnabled(TAttribute<bool>(this, &SStatsTableRow::HandleShouldBeEnabled));
 
 	SMultiColumnTableRow<FStatsNodePtr>::Construct(SMultiColumnTableRow<FStatsNodePtr>::FArguments(), InOwnerTableView);
@@ -78,38 +80,27 @@ END_SLATE_FUNCTION_BUILD_OPTIMIZATION
 
 FReply SStatsTableRow::OnDragDetected(const FGeometry& MyGeometry, const FPointerEvent& MouseEvent)
 {
-	//im:TODO
-	//if (MouseEvent.IsMouseButtonDown(EKeys::LeftMouseButton))
-	//{
-	//	if (StatsNode->IsGroup())
-	//	{
-	//		// Add all timer Ids for the group.
-	//		TArray<int32> StatsIds;
-	//		const TArray<FStatsNodePtr>& FilteredChildren = StatsNode->GetFilteredChildren();
-	//		const int32 NumFilteredChildren = FilteredChildren.Num();
-	//
-	//		StatsIds.Reserve(NumFilteredChildren);
-	//		for (int32 Nx = 0; Nx < NumFilteredChildren; ++Nx)
-	//		{
-	//			StatsIds.Add(FilteredChildren[Nx]->GetId());
-	//		}
-	//
-	//		return FReply::Handled().BeginDragDrop(FStatIDDragDropOp::NewGroup(StatsIds, StatsNode->GetName().GetPlainNameString()));
-	//	}
-	//	else
-	//	{
-	//		return FReply::Handled().BeginDragDrop(FStatIDDragDropOp::NewSingle(StatsNode->GetId(), StatsNode->GetName().GetPlainNameString()));
-	//	}
-	//}
-
 	return SMultiColumnTableRow<FStatsNodePtr>::OnDragDetected(MyGeometry, MouseEvent);
 }
 
 ////////////////////////////////////////////////////////////////////////////////////////////////////
 
+TSharedRef<IToolTip> SStatsTableRow::GetRowToolTip() const
+{
+	return RowToolTip.ToSharedRef();
+}
+
+////////////////////////////////////////////////////////////////////////////////////////////////////
+
+void SStatsTableRow::InvalidateContent()
+{
+	RowToolTip->InvalidateWidget();
+}
+
+////////////////////////////////////////////////////////////////////////////////////////////////////
+
 FSlateColor SStatsTableRow::GetBackgroundColorAndOpacity() const
 {
-	//return GetBackgroundColorAndOpacity(StatsNodePtr->GetAggregatedStats().Sum);
 	return FLinearColor(0.0f, 0.0f, 0.0f, 1.0f);
 }
 
@@ -120,7 +111,7 @@ FSlateColor SStatsTableRow::GetBackgroundColorAndOpacity(double Time) const
 	const FLinearColor Color =	Time > TimeUtils::Second      ? FLinearColor(0.3f, 0.0f, 0.0f, 1.0f) :
 								Time > TimeUtils::Milisecond  ? FLinearColor(0.3f, 0.1f, 0.0f, 1.0f) :
 								Time > TimeUtils::Microsecond ? FLinearColor(0.0f, 0.1f, 0.0f, 1.0f) :
-																FLinearColor(0.0f, 0.0f, 0.0f, 1.0f);
+								                                FLinearColor(0.0f, 0.0f, 0.0f, 1.0f);
 	return Color;
 }
 
diff --git a/Engine/Source/Developer/TraceInsights/Private/Insights/Widgets/SStatsTableRow.h b/Engine/Source/Developer/TraceInsights/Private/Insights/Widgets/SStatsTableRow.h
--- a/Engine/Source/Developer/TraceInsights/Private/Insights/Widgets/SStatsTableRow.h
+++ b/Engine/Source/Developer/TraceInsights/Private/Insights/Widgets/SStatsTableRow.h
@@ -12,20 +12,23 @@
 // Insights
 #include "Insights/ViewModels/StatsNodeHelper.h"
 
+class IToolTip;
+class SStatsCounterTableRowToolTip;
+
 DECLARE_DELEGATE_RetVal_OneParam(bool, FShouldBeEnabledDelegate, const uint32 /*StatsId*/);
 DECLARE_DELEGATE_RetVal_OneParam(bool, FIsColumnVisibleDelegate, const FName /*ColumnId*/);
-DECLARE_DELEGATE_TwoParams(FSetHoveredStatsTableCell, const FName /*ColumnId*/, const FStatsNodePtr /*StatsNodePtr*/);
 DECLARE_DELEGATE_RetVal_OneParam(EHorizontalAlignment, FGetColumnOutlineHAlignmentDelegate, const FName /*ColumnId*/);
+DECLARE_DELEGATE_TwoParams(FSetHoveredStatsTableCell, const FName /*ColumnId*/, const FStatsNodePtr /*StatsNodePtr*/);
 
-/** Widget that represents a table row in the stats' tree control. Generates widgets for each column on demand. */
+/** Widget that represents a table row in the Stats Counters' tree control. Generates widgets for each column on demand. */
 class SStatsTableRow : public SMultiColumnTableRow<FStatsNodePtr>
 {
 public:
 	SLATE_BEGIN_ARGS(SStatsTableRow) {}
 		SLATE_EVENT(FShouldBeEnabledDelegate, OnShouldBeEnabled)
 		SLATE_EVENT(FIsColumnVisibleDelegate, OnIsColumnVisible)
-		SLATE_EVENT(FSetHoveredStatsTableCell, OnSetHoveredTableCell)
 		SLATE_EVENT(FGetColumnOutlineHAlignmentDelegate, OnGetColumnOutlineHAlignmentDelegate)
+		SLATE_EVENT(FSetHoveredStatsTableCell, OnSetHoveredTableCell)
 		SLATE_ATTRIBUTE(FText, HighlightText)
 		SLATE_ATTRIBUTE(FName, HighlightedNodeName)
 		SLATE_ARGUMENT(FStatsNodePtr, StatsNodePtr)
@@ -52,6 +55,9 @@ public:
 	 */
 	virtual FReply OnDragDetected(const FGeometry& MyGeometry, const FPointerEvent& MouseEvent) override;
 
+	TSharedRef<IToolTip> GetRowToolTip() const;
+	void InvalidateContent();
+
 protected:
 	FSlateColor GetBackgroundColorAndOpacity() const;
 	FSlateColor GetBackgroundColorAndOpacity(double Time) const;
@@ -70,9 +76,11 @@ protected:
 	FSetHoveredStatsTableCell SetHoveredTableCellDelegate;
 	FGetColumnOutlineHAlignmentDelegate GetColumnOutlineHAlignmentDelegate;
 
-	/** Text to be highlighted on stats name. */
+	/** Text to be highlighted on stats counter name. */
 	TAttribute<FText> HighlightText;
 
-	/** Name of the stats node that should be drawn as highlighted. */
+	/** Name of the stats counter node that should be drawn as highlighted. */
 	TAttribute<FName> HighlightedNodeName;
+
+	TSharedPtr<SStatsCounterTableRowToolTip> RowToolTip;
 };
diff --git a/Engine/Source/Developer/TraceInsights/Private/Insights/Widgets/SStatsView.cpp b/Engine/Source/Developer/TraceInsights/Private/Insights/Widgets/SStatsView.cpp
--- a/Engine/Source/Developer/TraceInsights/Private/Insights/Widgets/SStatsView.cpp
+++ b/Engine/Source/Developer/TraceInsights/Private/Insights/Widgets/SStatsView.cpp
@@ -1648,7 +1648,10 @@ void TTimeCalculationHelper<double>::EnumerateValues(const Trace::ICounter& Coun
 
 	Counter.EnumerateFloatValues(IntervalStartTime, IntervalEndTime, [this, &StatsEx, Callback](double Time, double Value)
 	{
-		Callback(StatsEx, Value);
+		if (Time >= IntervalStartTime && Time < IntervalEndTime)
+		{
+			Callback(StatsEx, Value);
+		}
 	});
 }
 
@@ -1670,7 +1673,10 @@ void TTimeCalculationHelper<int64>::EnumerateValues(const Trace::ICounter& Count
 
 	Counter.EnumerateValues(IntervalStartTime, IntervalEndTime, [this, &StatsEx, Callback](double Time, int64 Value)
 	{
-		Callback(StatsEx, Value);
+		if (Time >= IntervalStartTime && Time < IntervalEndTime)
+		{
+			Callback(StatsEx, Value);
+		}
 	});
 }
 
@@ -1830,10 +1836,11 @@ void TTimeCalculationHelper<double>::PostProcess(TMap<uint64, FStatsNodePtr>& St
 		PostProcess(KV.Value, bComputeMedian);
 
 		// Update the stats node.
-		FStatsNodePtr* StatsNodePtrPtr = StatsNodesIdMap.Find(KV.Key);
-		if (StatsNodePtrPtr != nullptr)
+		FStatsNodePtr* NodePtrPtr = StatsNodesIdMap.Find(KV.Key);
+		if (NodePtrPtr != nullptr)
 		{
-			(*StatsNodePtrPtr)->SetAggregatedStats(KV.Value.BaseStats);
+			FStatsNodePtr NodePtr = *NodePtrPtr;
+			NodePtr->SetAggregatedStats(KV.Value.BaseStats);
 		}
 	}
 }
@@ -1849,10 +1856,11 @@ void TTimeCalculationHelper<int64>::PostProcess(TMap<uint64, FStatsNodePtr>& Sta
 		PostProcess(KV.Value, bComputeMedian);
 
 		// Update the stats node.
-		FStatsNodePtr* StatsNodePtrPtr = StatsNodesIdMap.Find(KV.Key);
-		if (StatsNodePtrPtr != nullptr)
+		FStatsNodePtr* NodePtrPtr = StatsNodesIdMap.Find(KV.Key);
+		if (NodePtrPtr != nullptr)
 		{
-			(*StatsNodePtrPtr)->SetAggregatedIntegerStats(KV.Value.BaseStats);
+			FStatsNodePtr NodePtr = *NodePtrPtr;
+			NodePtr->SetAggregatedIntegerStats(KV.Value.BaseStats);
 		}
 	}
 }
@@ -1929,6 +1937,17 @@ void SStatsView::UpdateStats(double StartTime, double EndTime)
 		CalculationHelperInt.PostProcess(StatsNodesIdMap, bComputeMedian);
 	}
 
+	// Invalidate all tree table rows.
+	for (const FStatsNodePtr NodePtr : StatsNodes)
+	{
+		TSharedPtr<ITableRow> TableRowPtr = TreeView->WidgetFromItem(NodePtr);
+		if (TableRowPtr.IsValid())
+		{
+			TSharedPtr<SStatsTableRow> StatsTableRowPtr = StaticCastSharedPtr<SStatsTableRow, ITableRow>(TableRowPtr);
+			StatsTableRowPtr->InvalidateContent();
+		}
+	}
+
 	UpdateTree();
 }
 
@@ -1936,12 +1955,13 @@ void SStatsView::UpdateStats(double StartTime, double EndTime)
 
 void SStatsView::SelectStatsNode(uint64 Id)
 {
-	FStatsNodePtr* StatsNodePtrPtr = StatsNodesIdMap.Find(Id);
-	if (StatsNodePtrPtr != nullptr)
+	FStatsNodePtr* NodePtrPtr = StatsNodesIdMap.Find(Id);
+	if (NodePtrPtr != nullptr)
 	{
-		//UE_LOG(TimingProfiler, Log, TEXT("Select and RequestScrollIntoView %s"), *(*StatsNodePtrPtr)->GetName().ToString());
-		TreeView->SetSelection(*StatsNodePtrPtr);
-		TreeView->RequestScrollIntoView(*StatsNodePtrPtr);
+		FStatsNodePtr NodePtr = *NodePtrPtr;
+
+		TreeView->SetSelection(NodePtr);
+		TreeView->RequestScrollIntoView(NodePtr);
 	}
 }
 
diff --git a/Engine/Source/Developer/TraceInsights/Private/Insights/Widgets/SStatsViewTooltip.cpp b/Engine/Source/Developer/TraceInsights/Private/Insights/Widgets/SStatsViewTooltip.cpp
--- a/Engine/Source/Developer/TraceInsights/Private/Insights/Widgets/SStatsViewTooltip.cpp
+++ b/Engine/Source/Developer/TraceInsights/Private/Insights/Widgets/SStatsViewTooltip.cpp
@@ -56,17 +56,8 @@ TSharedPtr<SToolTip> SStatsViewTooltip::GetColumnTooltip(const FStatsViewColumn&
 
 ////////////////////////////////////////////////////////////////////////////////////////////////////
 
-TSharedPtr<SToolTip> SStatsViewTooltip::GetTableCellTooltip(const TSharedPtr<FStatsNode> StatsNodePtr)
+TSharedPtr<SToolTip> SStatsViewTooltip::GetRowTooltip(const TSharedPtr<FStatsNode> StatsNodePtr)
 {
-	//const FSlateFontInfo TitleFont = FCoreStyle::GetDefaultFontStyle("Bold", 10);
-	//const FSlateFontInfo DescriptionFont = FCoreStyle::GetDefaultFontStyle("Regular", 8);
-	//const FSlateFontInfo DescriptionFontB = FCoreStyle::GetDefaultFontStyle("Bold", 8);
-
-	const FLinearColor DefaultColor(1.0f,1.0f,1.0f,1.0f);
-	const FLinearColor ThreadColor(5.0f, 0.0f, 0.0f, 1.0f);
-	const float Alpha = 0.0f;//StatsNodePtr->_FramePct * 0.01f;
-	const FLinearColor ColorAndOpacity = FMath::Lerp(DefaultColor, ThreadColor,Alpha);
-
 	const FText InstanceCountText = FText::AsNumber(StatsNodePtr->GetAggregatedStats().Count);
 
 	FText SumText = StatsNodePtr->GetTextForAggregatedStatsSum();
@@ -74,10 +65,9 @@ TSharedPtr<SToolTip> SStatsViewTooltip::GetTableCellTooltip(const TSharedPtr<FSt
 	FText MaxText = StatsNodePtr->GetTextForAggregatedStatsMax();
 	FText AvgText = StatsNodePtr->GetTextForAggregatedStatsAverage();
 	FText MedText = StatsNodePtr->GetTextForAggregatedStatsMedian();
-	FText LowText = StatsNodePtr->GetTextForAggregatedStatsLowerQuartile();
-	FText UppText = StatsNodePtr->GetTextForAggregatedStatsUpperQuartile();
+	//FText LowText = StatsNodePtr->GetTextForAggregatedStatsLowerQuartile();
+	//FText UppText = StatsNodePtr->GetTextForAggregatedStatsUpperQuartile();
 
-	//TSharedPtr<SHorizontalBox> HBoxCaption;
 	TSharedPtr<SGridPanel> GridPanel;
 	TSharedPtr<SHorizontalBox> HBox;
 
@@ -86,19 +76,12 @@ TSharedPtr<SToolTip> SStatsViewTooltip::GetTableCellTooltip(const TSharedPtr<FSt
 		[
 			SAssignNew(HBox, SHorizontalBox)
 
-			+SHorizontalBox::Slot()
+			+ SHorizontalBox::Slot()
 			.AutoWidth()
 			[
 				SNew(SVerticalBox)
 
-				//+SVerticalBox::Slot()
-				//.AutoHeight()
-				//.Padding(2.0f)
-				//[
-				//	SAssignNew(HBoxCaption, SHorizontalBox)
-				//]
-
-				+SVerticalBox::Slot()
+				+ SVerticalBox::Slot()
 				.AutoHeight()
 				.Padding(2.0f)
 				[
@@ -106,24 +89,22 @@ TSharedPtr<SToolTip> SStatsViewTooltip::GetTableCellTooltip(const TSharedPtr<FSt
 					.Orientation(Orient_Horizontal)
 				]
 
-				+SVerticalBox::Slot()
+				+ SVerticalBox::Slot()
 				.AutoHeight()
 				.Padding(2.0f)
 				[
 					SNew(SGridPanel)
 
 					// Id: [Id]
-					+SGridPanel::Slot(0, 0)
+					+ SGridPanel::Slot(0, 0)
 					.Padding(2.0f)
-					//.HAlign(HAlign_Right)
 					[
 						SNew(STextBlock)
 						.Text(LOCTEXT("TT_Id", "Id:"))
 						.TextStyle(FEditorStyle::Get(), TEXT("Profiler.TooltipBold"))
 					]
-					+SGridPanel::Slot(1, 0)
+					+ SGridPanel::Slot(1, 0)
 					.Padding(2.0f)
-					.HAlign(HAlign_Left)
 					[
 						SNew(STextBlock)
 						.Text(FText::AsNumber(StatsNodePtr->GetId()))
@@ -131,45 +112,24 @@ TSharedPtr<SToolTip> SStatsViewTooltip::GetTableCellTooltip(const TSharedPtr<FSt
 					]
 
 					// Name: [Name]
-					+SGridPanel::Slot(0, 1)
+					+ SGridPanel::Slot(0, 1)
 					.Padding(2.0f)
-					//.HAlign(HAlign_Right)
 					[
 						SNew(STextBlock)
 						.Text(LOCTEXT("TT_Name", "Name:"))
 						.TextStyle(FEditorStyle::Get(), TEXT("Profiler.TooltipBold"))
 					]
-					+SGridPanel::Slot(1, 1)
+					+ SGridPanel::Slot(1, 1)
 					.Padding(2.0f)
-					.HAlign(HAlign_Left)
 					[
 						SNew(STextBlock)
 						.Text(FText::FromName(StatsNodePtr->GetName()))
 						.TextStyle(FEditorStyle::Get(), TEXT("Profiler.Tooltip"))
 					]
 
-					//// Group: [MetaGroupName]
-					//+SGridPanel::Slot(0, 2)
-					//.Padding(2.0f)
-					////.HAlign(HAlign_Right)
-					//[
-					//	SNew(STextBlock)
-					//	.Text(LOCTEXT("TT_Group", "Group:"))
-					//	.TextStyle(FEditorStyle::Get(), TEXT("Profiler.TooltipBold"))
-					//]
-					//+SGridPanel::Slot(1, 2)
-					//.Padding(2.0f)
-					//.HAlign(HAlign_Left)
-					//[
-					//	SNew(STextBlock)
-					//	.Text(FText::FromName(StatsNodePtr->GetMetaGroupName()))
-					//	.TextStyle(FEditorStyle::Get(), TEXT("Profiler.Tooltip"))
-					//]
-
-					// Type: [Type]
+					// Counter Type: [Type]
 					+ SGridPanel::Slot(0, 3)
 					.Padding(2.0f)
-					//.HAlign(HAlign_Right)
 					[
 						SNew(STextBlock)
 						.Text(LOCTEXT("TT_Type", "Type:"))
@@ -177,7 +137,6 @@ TSharedPtr<SToolTip> SStatsViewTooltip::GetTableCellTooltip(const TSharedPtr<FSt
 					]
 					+ SGridPanel::Slot(1, 3)
 					.Padding(2.0f)
-					.HAlign(HAlign_Left)
 					[
 						SNew(STextBlock)
 						.Text(StatsNodeTypeHelper::ToName(StatsNodePtr->GetType()))
@@ -185,7 +144,7 @@ TSharedPtr<SToolTip> SStatsViewTooltip::GetTableCellTooltip(const TSharedPtr<FSt
 					]
 				]
 
-				+SVerticalBox::Slot()
+				+ SVerticalBox::Slot()
 				.AutoHeight()
 				.Padding(2.0f)
 				[
@@ -201,7 +160,6 @@ TSharedPtr<SToolTip> SStatsViewTooltip::GetTableCellTooltip(const TSharedPtr<FSt
 
 					+ SGridPanel::Slot(0, 0)
 					.Padding(2.0f)
-					//.HAlign(HAlign_Right)
 					[
 						SNew(STextBlock)
 						.Text(LOCTEXT("TT_NumInstances", "Num Instances:"))
@@ -209,7 +167,6 @@ TSharedPtr<SToolTip> SStatsViewTooltip::GetTableCellTooltip(const TSharedPtr<FSt
 					]
 					+ SGridPanel::Slot(1, 0)
 					.Padding(2.0f)
-					.HAlign(HAlign_Left)
 					[
 						SNew(STextBlock)
 						.Text(InstanceCountText)
@@ -217,7 +174,7 @@ TSharedPtr<SToolTip> SStatsViewTooltip::GetTableCellTooltip(const TSharedPtr<FSt
 					]
 				]
 
-				+SVerticalBox::Slot()
+				+ SVerticalBox::Slot()
 				.AutoHeight()
 				.Padding(2.0f)
 				[
@@ -234,7 +191,7 @@ TSharedPtr<SToolTip> SStatsViewTooltip::GetTableCellTooltip(const TSharedPtr<FSt
 					// Aggregated stats are added here.
 				]
 
-				+SVerticalBox::Slot()
+				+ SVerticalBox::Slot()
 				.AutoHeight()
 				.Padding(2.0f)
 				[
@@ -247,271 +204,21 @@ TSharedPtr<SToolTip> SStatsViewTooltip::GetTableCellTooltip(const TSharedPtr<FSt
 	int32 Row = 1;
 	AddAggregatedStatsRow(GridPanel, Row, LOCTEXT("TT_Sum",     "Sum:"),            SumText);
 	AddAggregatedStatsRow(GridPanel, Row, LOCTEXT("TT_Max",     "Max:"),            MaxText);
-	AddAggregatedStatsRow(GridPanel, Row, LOCTEXT("TT_UpperQ",  "Upper Quartile:"), UppText);
+	//AddAggregatedStatsRow(GridPanel, Row, LOCTEXT("TT_UpperQ",  "Upper Quartile:"), UppText);
 	AddAggregatedStatsRow(GridPanel, Row, LOCTEXT("TT_Average", "Average:"),        AvgText);
 	AddAggregatedStatsRow(GridPanel, Row, LOCTEXT("TT_Median",  "Median:"),         MedText);
-	AddAggregatedStatsRow(GridPanel, Row, LOCTEXT("TT_LowerQ",  "Lower Quartile:"), LowText);
+	//AddAggregatedStatsRow(GridPanel, Row, LOCTEXT("TT_LowerQ",  "Lower Quartile:"), LowText);
 	AddAggregatedStatsRow(GridPanel, Row, LOCTEXT("TT_Min",     "Min:"),            MinText);
 
-	/*
-	//TODO: We need Stats hierarchy (not the grouping hierarchy)!
-	const bool bHasParent = StatsNodePtr->GetStats()->GetParent().IsValid();
-	const bool bHasChildren = StatsNodePtr->GetStats()->GetChildren().Num() > 0;
-
-	if (bHasParent)
-	{
-		const FText ParentName = FText::FromName(StatsNodePtr->GetGroupPtr()->GetName());
-		HBoxCaption->AddSlot()
-			.AutoWidth()
-			[
-				SNew(STextBlock)
-				.Text(ParentName)
-				.TextStyle(FEditorStyle::Get(), TEXT("Profiler.Caption"))
-			];
-
-		HBoxCaption->AddSlot()
-			.AutoWidth()
-			[
-				SNew(SImage)
-				.Image(FEditorStyle::GetBrush("BreadcrumbTrail.Delimiter"))
-			];
-	}
-
-	const FText StatsName = FText::FromName(StatsNodePtr->GetName());
-	HBoxCaption->AddSlot()
-		.AutoWidth()
-		[
-			SNew(STextBlock)
-			.Text(StatsName)
-			.TextStyle(FEditorStyle::Get(), TEXT("Profiler.CaptionBold"))
-		];
-
-	if (bHasChildren)
-	{
-		typedef TKeyValuePair<float, FName> FEventNameAndPct;
-		TArray<FEventNameAndPct> MinimalChildren;
-
-		const TArray<FStatsNodePtr>& Children = StatsNodePtr->GetChildren();
-		for(int32 ChildIndex = 0; ChildIndex < Children.Num(); ChildIndex++)
-		{
-			const FStatsNodePtr Child = Children[ChildIndex];
-			float Percent = static_cast<float>(Child->GetStats().TotalInclusiveTime / StatsNodePtr->GetStats().TotalInclusiveTime * 100.0);
-			MinimalChildren.Add(FEventNameAndPct(Child->GetStats().TotalInclusiveTime, Child->GetName()));
-		}
-
-		struct FCompareByFloatDescending
-		{
-			FORCEINLINE bool operator()(const FEventNameAndPct& A, const FEventNameAndPct& B) const
-			{
-				return A.Key > B.Key;
-			}
-		};
-		MinimalChildren.Sort(FCompareByFloatDescending());
-
-		FString ChildrenNames;
-		const int32 NumChildrenToDisplay = FMath::Min(MinimalChildren.Num(), 3);
-		for(int32 SortedChildIndex = 0; SortedChildIndex < NumChildrenToDisplay; SortedChildIndex++)
-		{
-			const FEventNameAndPct& MinimalChild = MinimalChildren[SortedChildIndex];
-			ChildrenNames += FString::Printf(TEXT("%s (%.1f %%)"), *MinimalChild.Value.ToString(), MinimalChild.Key);
-
-			const bool bAddDelimiter = SortedChildIndex < NumChildrenToDisplay - 1;
-			if (bAddDelimiter)
-			{
-				ChildrenNames += TEXT(", ");
-			}
-		}
-
-		HBoxCaption->AddSlot()
-			.AutoWidth()
-			[
-				SNew(SImage)
-				.Image(FEditorStyle::GetBrush("BreadcrumbTrail.Delimiter"))
-			];
-
-		HBoxCaption->AddSlot()
-			.AutoWidth()
-			[
-				SNew(STextBlock)
-				.Text(FText::FromString(ChildrenNames))
-				.TextStyle(FEditorStyle::Get(), TEXT("Profiler.Caption"))
-			];
-	}
-	*/
-
 	return TableCellTooltip;
 }
 
 ////////////////////////////////////////////////////////////////////////////////////////////////////
 
-TSharedRef<SToolTip> SStatsViewTooltip::GetTooltip()
-{
-	if (Session.IsValid())
-	{
-		const TSharedRef<SGridPanel> ToolTipGrid = SNew(SGridPanel);
-		int32 CurrentRowPos = 0;
-
-		AddHeader(ToolTipGrid, CurrentRowPos);
-		AddDescription(ToolTipGrid, CurrentRowPos);
-
-		AddNoDataInformation(ToolTipGrid, CurrentRowPos);
-
-		return SNew(SToolTip)
-			[
-				ToolTipGrid
-			];
-	}
-	else
-	{
-		return SNew(SToolTip)
-			.Text(LOCTEXT("NotImplemented", "Tooltip for multiple profiler instances has not been implemented yet"));
-	}
-}
-
-////////////////////////////////////////////////////////////////////////////////////////////////////
-
-void SStatsViewTooltip::AddNoDataInformation(const TSharedRef<SGridPanel>& Grid, int32& RowPos)
-{
-	Grid->AddSlot(0, RowPos)
-		.Padding(2.0f)
-		.ColumnSpan(3)
-		[
-			SNew(STextBlock)
-			.TextStyle(FEditorStyle::Get(), TEXT("Profiler.TooltipBold"))
-			.Text(LOCTEXT("NoDataAvailable", "N/A"))
-		];
-	RowPos++;
-}
-
-////////////////////////////////////////////////////////////////////////////////////////////////////
-
-void SStatsViewTooltip::AddHeader(const TSharedRef<SGridPanel>& Grid, int32& RowPos)
-{
-	const FString InstanceName = Session->GetName();
-
-	Grid->AddSlot(0, RowPos++)
-		.Padding(2.0f)
-		.ColumnSpan(3)
-		[
-			SNew(STextBlock)
-			.TextStyle(FEditorStyle::Get(), TEXT("Profiler.TooltipBold"))
-			.Text(LOCTEXT("StatInstance", "Stat information for profiler instance"))
-		];
-
-	Grid->AddSlot(0, RowPos++)
-		.Padding(2.0f)
-		.ColumnSpan(3)
-		[
-			SNew(STextBlock)
-			.TextStyle(FEditorStyle::Get(), TEXT("Profiler.Tooltip"))
-			.Text(FText::FromString(InstanceName))
-		];
-
-	AddSeparator(Grid, RowPos);
-}
-
-////////////////////////////////////////////////////////////////////////////////////////////////////
-
-void SStatsViewTooltip::AddDescription(const TSharedRef<SGridPanel>& Grid, int32& RowPos)
-{
-	/*
-	const FProfilerStat& ProfilerStat = Session->GetMetaData()->GetStatsById(StatsId);
-	const EStatsNodeType SampleType = Session->GetMetaData()->GetSampleTypeForStatsId(StatsId);
-	const FSlateBrush* const StatIcon = SStatsViewHelper::GetIconForStatType(SampleType);
-
-	Grid->AddSlot(0, RowPos)
-	.Padding(2.0f)
-	[
-		SNew(STextBlock)
-		.TextStyle(FEditorStyle::Get(), TEXT("Profiler.TooltipBold"))
-		.Text(LOCTEXT("GroupDesc","Group:"))
-	];
-
-	Grid->AddSlot(1, RowPos)
-	.Padding(2.0f)
-	.ColumnSpan(2)
-	[
-		SNew(STextBlock)
-		.TextStyle(FEditorStyle::Get(), TEXT("Profiler.Tooltip"))
-		.Text(FText::FromName(ProfilerStat.OwningGroup().Name()))
-	];
-	RowPos++;
-
-	Grid->AddSlot(0, RowPos)
-	.Padding(2.0f)
-	[
-		SNew(STextBlock)
-		.TextStyle(FEditorStyle::Get(), TEXT("Profiler.TooltipBold"))
-		.Text(LOCTEXT("NameDesc","Name:"))
-	];
-
-	Grid->AddSlot(1, RowPos)
-	.Padding(2.0f)
-	.ColumnSpan(2)
-	[
-		SNew(STextBlock)
-		.TextStyle(FEditorStyle::Get(), TEXT("Profiler.Tooltip"))
-		.Text(FText::FromName(ProfilerStat.Name()))
-	];
-	RowPos++;
-
-	Grid->AddSlot(0, RowPos)
-	.Padding(2.0f)
-	[
-		SNew(STextBlock)
-		.TextStyle(FEditorStyle::Get(), TEXT("Profiler.TooltipBold"))
-		.Text(LOCTEXT("TypeDesc","Type:"))
-	];
-
-	Grid->AddSlot(1, RowPos)
-	.Padding(2.0f)
-	.ColumnSpan(2)
-	[
-		SNew(SHorizontalBox)
-
-		+SHorizontalBox::Slot()
-		.AutoWidth()
-		[
-			SNew(SImage)
-			.Image(StatIcon)
-		]
-
-		+SHorizontalBox::Slot()
-		.AutoWidth()
-		.HAlign(HAlign_Left)
-		.VAlign(VAlign_Center)
-		[
-			SNew(STextBlock)
-			.Text(FText::FromString(EStatsNodeType::ToDescription(SampleType)))
-			.TextStyle(FEditorStyle::Get(), TEXT("Profiler.Tooltip"))
-		]
-	];
-	RowPos++;
-
-	AddSeparator(Grid, RowPos);
-	*/
-}
-
-////////////////////////////////////////////////////////////////////////////////////////////////////
-
-void SStatsViewTooltip::AddSeparator(const TSharedRef<SGridPanel>& Grid, int32& RowPos)
-{
-	Grid->AddSlot(0, RowPos++)
-		.Padding(2.0f)
-		.ColumnSpan(3)
-		[
-			SNew(SSeparator)
-			.Orientation(Orient_Horizontal)
-		];
-}
-
-////////////////////////////////////////////////////////////////////////////////////////////////////
-
 void SStatsViewTooltip::AddAggregatedStatsRow(TSharedPtr<SGridPanel> Grid, int32& Row, const FText& Name, const FText& Value)
 {
 	Grid->AddSlot(0, Row)
 		.Padding(2.0f)
-		//.HAlign(HAlign_Right)
 		[
 			SNew(STextBlock)
 			.Text(Name)
diff --git a/Engine/Source/Developer/TraceInsights/Private/Insights/Widgets/SStatsViewTooltip.h b/Engine/Source/Developer/TraceInsights/Private/Insights/Widgets/SStatsViewTooltip.h
--- a/Engine/Source/Developer/TraceInsights/Private/Insights/Widgets/SStatsViewTooltip.h
+++ b/Engine/Source/Developer/TraceInsights/Private/Insights/Widgets/SStatsViewTooltip.h
@@ -3,48 +3,92 @@
 #pragma once
 
 #include "CoreMinimal.h"
+#include "Widgets/IToolTip.h"
+#include "Widgets/SToolTip.h"
 
 // Insights
 #include "Insights/InsightsManager.h"
 
-class SToolTip;
 class SGridPanel;
 
-namespace Trace { class IAnalysisSession; }
+namespace Insights
+{
+	//TODO: class FTable;
+	//TODO: class FTableColumn;
+}
 
-// Insights
 class FStatsNode;
 class FStatsViewColumn;
 
 #define LOCTEXT_NAMESPACE "SStatsView"
 
-/** Stats View Tooltip */
+////////////////////////////////////////////////////////////////////////////////////////////////////
+
+/** Stats Counters View Tooltip */
 class SStatsViewTooltip
 {
-	const uint32 StatsId;
-	TSharedPtr<const Trace::IAnalysisSession> Session;
+public:
+	SStatsViewTooltip() = delete;
 
+	//TODO: static TSharedPtr<SToolTip> GetTableTooltip(const Insights::FTable& Table);
+	//TODO: static TSharedPtr<SToolTip> GetColumnTooltip(const Insights::FTableColumn& Column);
+	static TSharedPtr<SToolTip> GetColumnTooltip(const FStatsViewColumn& Column);
+	static TSharedPtr<SToolTip> GetRowTooltip(const TSharedPtr<FStatsNode> StatsNodePtr);
+
+private:
+	static void AddAggregatedStatsRow(TSharedPtr<SGridPanel> Grid, int32& Row, const FText& Name, const FText& Value);
+};
+
+////////////////////////////////////////////////////////////////////////////////////////////////////
+
+class SStatsCounterTableRowToolTip : public IToolTip
+{
 public:
-	SStatsViewTooltip(const uint32 InStatsId)
-		: StatsId(InStatsId)
+	SStatsCounterTableRowToolTip(const TSharedPtr<FStatsNode> InTreeNodePtr) : TreeNodePtr(InTreeNodePtr) {}
+	virtual ~SStatsCounterTableRowToolTip() { }
+
+	virtual TSharedRef<class SWidget> AsWidget()
+	{
+		CreateToolTipWidget();
+		return ToolTipWidget.ToSharedRef();
+	}
+
+	virtual TSharedRef<SWidget> GetContentWidget()
 	{
-		Session = FInsightsManager::Get()->GetSession();
+		CreateToolTipWidget();
+		return ToolTipWidget->GetContentWidget();
 	}
 
-	TSharedRef<SToolTip> GetTooltip();
+	virtual void SetContentWidget(const TSharedRef<SWidget>& InContentWidget)
+	{
+		CreateToolTipWidget();
+		ToolTipWidget->SetContentWidget(InContentWidget);
+	}
 
-protected:
-	void AddNoDataInformation(const TSharedRef<SGridPanel>& Grid, int32& RowPos);
-	void AddHeader(const TSharedRef<SGridPanel>& Grid, int32& RowPos);
-	void AddDescription(const TSharedRef<SGridPanel>& Grid, int32& RowPos);
-	void AddSeparator(const TSharedRef<SGridPanel>& Grid, int32& RowPos);
+	void InvalidateWidget()
+	{
+		ToolTipWidget.Reset();
+	}
 
-public:
-	static TSharedPtr<SToolTip> GetColumnTooltip(const FStatsViewColumn& Column);
-	static TSharedPtr<SToolTip> GetTableCellTooltip(const TSharedPtr<FStatsNode> StatsNodePtr);
+	virtual bool IsEmpty() const { return false; }
+	virtual bool IsInteractive() const { return false; }
+	virtual void OnOpening() {}
+	virtual void OnClosed() {}
 
-protected:
-	static void AddAggregatedStatsRow(TSharedPtr<SGridPanel> Grid, int32& Row, const FText& Name, const FText& Value);
+private:
+	void CreateToolTipWidget()
+	{
+		if (!ToolTipWidget.IsValid())
+		{
+			ToolTipWidget = SStatsViewTooltip::GetRowTooltip(TreeNodePtr);
+		}
+	}
+
+private:
+	TSharedPtr<SToolTip> ToolTipWidget;
+	const TSharedPtr<FStatsNode> TreeNodePtr;
 };
 
+////////////////////////////////////////////////////////////////////////////////////////////////////
+
 #undef LOCTEXT_NAMESPACE
diff --git a/Engine/Source/Developer/TraceInsights/Private/Insights/Widgets/STimerTableCell.cpp b/Engine/Source/Developer/TraceInsights/Private/Insights/Widgets/STimerTableCell.cpp
--- a/Engine/Source/Developer/TraceInsights/Private/Insights/Widgets/STimerTableCell.cpp
+++ b/Engine/Source/Developer/TraceInsights/Private/Insights/Widgets/STimerTableCell.cpp
@@ -15,7 +15,7 @@
 // Insights
 #include "Insights/Table/ViewModels/Table.h"
 #include "Insights/Table/ViewModels/TableColumn.h"
-#include "Insights/Widgets/STimersViewTooltip.h"
+#include "Insights/Widgets/STimerTableRow.h"
 
 #define LOCTEXT_NAMESPACE "STimersView"
 
@@ -23,7 +23,7 @@
 
 BEGIN_SLATE_FUNCTION_BUILD_OPTIMIZATION
 
-void STimerTableCell::Construct(const FArguments& InArgs, const TSharedRef<class ITableRow>& TableRow)
+void STimerTableCell::Construct(const FArguments& InArgs, const TSharedRef<ITableRow>& TableRow)
 {
 	TablePtr = InArgs._TablePtr;
 	ColumnPtr = InArgs._ColumnPtr;
@@ -43,7 +43,7 @@ void STimerTableCell::Construct(const FArguments& InArgs, const TSharedRef<class
 
 ////////////////////////////////////////////////////////////////////////////////////////////////////
 
-TSharedRef<SWidget> STimerTableCell::GenerateWidgetForColumn(const FArguments& InArgs, const TSharedRef<class ITableRow>& TableRow)
+TSharedRef<SWidget> STimerTableCell::GenerateWidgetForColumn(const FArguments& InArgs, const TSharedRef<ITableRow>& TableRow)
 {
 	if (InArgs._IsNameColumn)
 	{
@@ -57,7 +57,7 @@ TSharedRef<SWidget> STimerTableCell::GenerateWidgetForColumn(const FArguments& I
 
 ////////////////////////////////////////////////////////////////////////////////////////////////////
 
-TSharedRef<SWidget> STimerTableCell::GenerateWidgetForNameColumn(const FArguments& InArgs, const TSharedRef<class ITableRow>& TableRow)
+TSharedRef<SWidget> STimerTableCell::GenerateWidgetForNameColumn(const FArguments& InArgs, const TSharedRef<ITableRow>& TableRow)
 {
 	return
 		SNew(SHorizontalBox)
@@ -91,7 +91,7 @@ TSharedRef<SWidget> STimerTableCell::GenerateWidgetForNameColumn(const FArgument
 				SNew(SImage)
 				.Visibility(this, &STimerTableCell::GetHintIconVisibility)
 				.Image(FEditorStyle::GetBrush("Profiler.Tooltip.HintIcon10"))
-				.ToolTip(STimersViewTooltip::GetCellTooltip(TimerNodePtr, ColumnPtr))
+				.ToolTip(GetRowToolTip(TableRow))
 			]
 		]
 
@@ -114,10 +114,23 @@ TSharedRef<SWidget> STimerTableCell::GenerateWidgetForNameColumn(const FArgument
 
 ////////////////////////////////////////////////////////////////////////////////////////////////////
 
-TSharedRef<SWidget> STimerTableCell::GenerateWidgetForStatsColumn(const FArguments& InArgs, const TSharedRef<class ITableRow>& TableRow)
+TSharedPtr<IToolTip> STimerTableCell::GetRowToolTip(const TSharedRef<ITableRow>& TableRow) const
 {
-	const FText CellText = ColumnPtr->GetValueAsText(*TimerNodePtr);
+	TSharedRef<STimerTableRow> Row = StaticCastSharedRef<STimerTableRow, ITableRow>(TableRow);
+	return Row->GetRowToolTip();
+}
+
+////////////////////////////////////////////////////////////////////////////////////////////////////
 
+FText STimerTableCell::GetValueAsText() const
+{
+	return ColumnPtr->GetValueAsText(*TimerNodePtr);
+}
+
+////////////////////////////////////////////////////////////////////////////////////////////////////
+
+TSharedRef<SWidget> STimerTableCell::GenerateWidgetForStatsColumn(const FArguments& InArgs, const TSharedRef<ITableRow>& TableRow)
+{
 	return
 		SNew(SHorizontalBox)
 
@@ -129,7 +142,7 @@ TSharedRef<SWidget> STimerTableCell::GenerateWidgetForStatsColumn(const FArgumen
 		.Padding(FMargin(2.0f, 0.0f))
 		[
 			SNew(STextBlock)
-			.Text(CellText)
+			.Text(this, &STimerTableCell::GetValueAsText)
 			.TextStyle(FEditorStyle::Get(), TEXT("Profiler.Tooltip"))
 			.ColorAndOpacity(this, &STimerTableCell::GetStatsColorAndOpacity)
 			.ShadowColorAndOpacity(this, &STimerTableCell::GetShadowColorAndOpacity)
diff --git a/Engine/Source/Developer/TraceInsights/Private/Insights/Widgets/STimerTableCell.h b/Engine/Source/Developer/TraceInsights/Private/Insights/Widgets/STimerTableCell.h
--- a/Engine/Source/Developer/TraceInsights/Private/Insights/Widgets/STimerTableCell.h
+++ b/Engine/Source/Developer/TraceInsights/Private/Insights/Widgets/STimerTableCell.h
@@ -9,6 +9,8 @@
 // Insights
 #include "Insights/ViewModels/TimerNode.h"
 
+class ITableRow;
+
 namespace Insights
 {
 	class FTable;
@@ -29,12 +31,12 @@ public:
 		SLATE_ARGUMENT(bool, IsNameColumn)
 	SLATE_END_ARGS()
 
-	void Construct(const FArguments& InArgs, const TSharedRef<class ITableRow>& TableRow);
+	void Construct(const FArguments& InArgs, const TSharedRef<ITableRow>& TableRow);
 
 protected:
-	TSharedRef<SWidget> GenerateWidgetForColumn(const FArguments& InArgs, const TSharedRef<class ITableRow>& TableRow);
-	TSharedRef<SWidget> GenerateWidgetForNameColumn(const FArguments& InArgs, const TSharedRef<class ITableRow>& TableRow);
-	TSharedRef<SWidget> GenerateWidgetForStatsColumn(const FArguments& InArgs, const TSharedRef<class ITableRow>& TableRow);
+	TSharedRef<SWidget> GenerateWidgetForColumn(const FArguments& InArgs, const TSharedRef<ITableRow>& TableRow);
+	TSharedRef<SWidget> GenerateWidgetForNameColumn(const FArguments& InArgs, const TSharedRef<ITableRow>& TableRow);
+	TSharedRef<SWidget> GenerateWidgetForStatsColumn(const FArguments& InArgs, const TSharedRef<ITableRow>& TableRow);
 
 	/**
 	 * The system will use this event to notify a widget that the cursor has entered it. This event is NOT bubbled.
@@ -103,11 +105,15 @@ protected:
 		return IsHovered() ? EVisibility::Visible : EVisibility::Hidden;
 	}
 
+	TSharedPtr<class IToolTip> GetRowToolTip(const TSharedRef<ITableRow>& TableRow) const;
+
 	FText GetDisplayName() const
 	{
 		return TimerNodePtr->GetDisplayName();
 	}
 
+	FText GetValueAsText() const;
+
 	FSlateColor GetColorAndOpacity() const
 	{
 		const FLinearColor TextColor =
diff --git a/Engine/Source/Developer/TraceInsights/Private/Insights/Widgets/STimerTableRow.cpp b/Engine/Source/Developer/TraceInsights/Private/Insights/Widgets/STimerTableRow.cpp
--- a/Engine/Source/Developer/TraceInsights/Private/Insights/Widgets/STimerTableRow.cpp
+++ b/Engine/Source/Developer/TraceInsights/Private/Insights/Widgets/STimerTableRow.cpp
@@ -5,12 +5,14 @@
 #include "EditorStyleSet.h"
 #include "SlateOptMacros.h"
 #include "Widgets/SOverlay.h"
+#include "Widgets/SToolTip.h"
 #include "Widgets/Images/SImage.h"
 
 // Insights
 #include "Insights/Common/TimeUtils.h"
 #include "Insights/Table/ViewModels/Table.h"
 #include "Insights/Table/ViewModels/TableColumn.h"
+#include "Insights/Widgets/STimersViewTooltip.h"
 #include "Insights/Widgets/STimerTableCell.h"
 
 #define LOCTEXT_NAMESPACE "STimersView"
@@ -23,8 +25,8 @@ void STimerTableRow::Construct(const FArguments& InArgs, const TSharedRef<STable
 {
 	OnShouldBeEnabled = InArgs._OnShouldBeEnabled;
 	IsColumnVisibleDelegate = InArgs._OnIsColumnVisible;
-	SetHoveredCellDelegate = InArgs._OnSetHoveredCell;
 	GetColumnOutlineHAlignmentDelegate = InArgs._OnGetColumnOutlineHAlignmentDelegate;
+	SetHoveredCellDelegate = InArgs._OnSetHoveredCell;
 
 	HighlightText = InArgs._HighlightText;
 	HighlightedNodeName = InArgs._HighlightedNodeName;
@@ -32,6 +34,8 @@ void STimerTableRow::Construct(const FArguments& InArgs, const TSharedRef<STable
 	TablePtr = InArgs._TablePtr;
 	TimerNodePtr = InArgs._TimerNodePtr;
 
+	RowToolTip = MakeShareable(new STimerTableRowToolTip(TimerNodePtr));
+
 	SetEnabled(TAttribute<bool>(this, &STimerTableRow::HandleShouldBeEnabled));
 
 	SMultiColumnTableRow<FTimerNodePtr>::Construct(SMultiColumnTableRow<FTimerNodePtr>::FArguments(), InOwnerTableView);
@@ -111,10 +115,23 @@ FReply STimerTableRow::OnDragDetected(const FGeometry& MyGeometry, const FPointe
 
 ////////////////////////////////////////////////////////////////////////////////////////////////////
 
+TSharedRef<IToolTip> STimerTableRow::GetRowToolTip() const
+{
+	return RowToolTip.ToSharedRef();
+}
+
+////////////////////////////////////////////////////////////////////////////////////////////////////
+
+void STimerTableRow::InvalidateContent()
+{
+	RowToolTip->InvalidateWidget();
+}
+
+////////////////////////////////////////////////////////////////////////////////////////////////////
+
 FSlateColor STimerTableRow::GetBackgroundColorAndOpacity() const
 {
 	return GetBackgroundColorAndOpacity(TimerNodePtr->GetAggregatedStats().TotalInclusiveTime);
-	//return FLinearColor(0.0f, 0.0f, 0.0f, 1.0f)
 }
 
 ////////////////////////////////////////////////////////////////////////////////////////////////////
@@ -124,7 +141,7 @@ FSlateColor STimerTableRow::GetBackgroundColorAndOpacity(double Time) const
 	const FLinearColor Color =	Time > TimeUtils::Second      ? FLinearColor(0.3f, 0.0f, 0.0f, 1.0f) :
 								Time > TimeUtils::Milisecond  ? FLinearColor(0.3f, 0.1f, 0.0f, 1.0f) :
 								Time > TimeUtils::Microsecond ? FLinearColor(0.0f, 0.1f, 0.0f, 1.0f) :
-																FLinearColor(0.0f, 0.0f, 0.0f, 1.0f);
+								                                FLinearColor(0.0f, 0.0f, 0.0f, 1.0f);
 	return Color;
 }
 
diff --git a/Engine/Source/Developer/TraceInsights/Private/Insights/Widgets/STimerTableRow.h b/Engine/Source/Developer/TraceInsights/Private/Insights/Widgets/STimerTableRow.h
--- a/Engine/Source/Developer/TraceInsights/Private/Insights/Widgets/STimerTableRow.h
+++ b/Engine/Source/Developer/TraceInsights/Private/Insights/Widgets/STimerTableRow.h
@@ -12,6 +12,9 @@
 // Insights
 #include "Insights/ViewModels/TimerNodeHelper.h"
 
+class IToolTip;
+class STimerTableRowToolTip;
+
 namespace Insights
 {
 	class FTable;
@@ -23,7 +26,7 @@ DECLARE_DELEGATE_RetVal_OneParam(bool, FIsColumnVisibleDelegate, const FName /*C
 DECLARE_DELEGATE_RetVal_OneParam(EHorizontalAlignment, FGetColumnOutlineHAlignmentDelegate, const FName /*ColumnId*/);
 DECLARE_DELEGATE_ThreeParams(FSetHoveredTimerTableCell, TSharedPtr<Insights::FTable> /*TablePtr*/, TSharedPtr<Insights::FTableColumn> /*ColumnPtr*/, const FTimerNodePtr /*TimerNodePtr*/);
 
-/** Widget that represents a table row in a timer tree control. Generates widgets for each column on demand. */
+/** Widget that represents a table row in the tree control. Generates widgets for each column on demand. */
 class STimerTableRow : public SMultiColumnTableRow<FTimerNodePtr>
 {
 public:
@@ -59,6 +62,9 @@ public:
 	 */
 	virtual FReply OnDragDetected(const FGeometry& MyGeometry, const FPointerEvent& MouseEvent) override;
 
+	TSharedRef<IToolTip> GetRowToolTip() const;
+	void InvalidateContent();
+
 protected:
 	FSlateColor GetBackgroundColorAndOpacity() const;
 	FSlateColor GetBackgroundColorAndOpacity(double Time) const;
@@ -85,4 +91,6 @@ protected:
 
 	/** Name of the timer node that should be drawn as highlighted. */
 	TAttribute<FName> HighlightedNodeName;
+
+	TSharedPtr<STimerTableRowToolTip> RowToolTip;
 };
diff --git a/Engine/Source/Developer/TraceInsights/Private/Insights/Widgets/STimersView.cpp b/Engine/Source/Developer/TraceInsights/Private/Insights/Widgets/STimersView.cpp
--- a/Engine/Source/Developer/TraceInsights/Private/Insights/Widgets/STimersView.cpp
+++ b/Engine/Source/Developer/TraceInsights/Private/Insights/Widgets/STimersView.cpp
@@ -1613,22 +1613,34 @@ void STimersView::UpdateStats(double StartTime, double EndTime)
 		const bool bIsGpuTrackVisible = TimingView.IsValid() && TimingView->IsGpuTrackVisible();
 
 		TUniquePtr<Trace::ITable<Trace::FTimingProfilerAggregatedStats>> AggregationResultTable;
+
 		{
 			Trace::FAnalysisSessionReadScope SessionReadScope(*Session.Get());
 			const Trace::ITimingProfilerProvider& TimingProfilerProvider = *Trace::ReadTimingProfilerProvider(*Session.Get());
 			AggregationResultTable.Reset(TimingProfilerProvider.CreateAggregation(StartTime, EndTime, ThreadFilter, bIsGpuTrackVisible));
 		}
 
-		TUniquePtr<Trace::ITableReader<Trace::FTimingProfilerAggregatedStats>> AggregationResultTableReader(AggregationResultTable->CreateReader());
-		while (AggregationResultTableReader->IsValid())
+		if (AggregationResultTable.IsValid())
 		{
-			const Trace::FTimingProfilerAggregatedStats* Row = AggregationResultTableReader->GetCurrentRow();
-			FTimerNodePtr* TimerNodePtrPtr = TimerNodesIdMap.Find(Row->Timer->Id);
-			if (TimerNodePtrPtr != nullptr)
+			TUniquePtr<Trace::ITableReader<Trace::FTimingProfilerAggregatedStats>> TableReader(AggregationResultTable->CreateReader());
+			while (TableReader->IsValid())
 			{
-				(*TimerNodePtrPtr)->SetAggregatedStats(*Row);
+				const Trace::FTimingProfilerAggregatedStats* Row = TableReader->GetCurrentRow();
+				FTimerNodePtr* TimerNodePtrPtr = TimerNodesIdMap.Find(Row->Timer->Id);
+				if (TimerNodePtrPtr != nullptr)
+				{
+					FTimerNodePtr TimerNodePtr = *TimerNodePtrPtr;
+					TimerNodePtr->SetAggregatedStats(*Row);
+
+					TSharedPtr<ITableRow> TableRowPtr = TreeView->WidgetFromItem(TimerNodePtr);
+					if (TableRowPtr.IsValid())
+					{
+						TSharedPtr<STimerTableRow> RowPtr = StaticCastSharedPtr<STimerTableRow, ITableRow>(TableRowPtr);
+						RowPtr->InvalidateContent();
+					}
+				}
+				TableReader->NextRow();
 			}
-			AggregationResultTableReader->NextRow();
 		}
 	}
 
diff --git a/Engine/Source/Developer/TraceInsights/Private/Insights/Widgets/STimersViewTooltip.cpp b/Engine/Source/Developer/TraceInsights/Private/Insights/Widgets/STimersViewTooltip.cpp
--- a/Engine/Source/Developer/TraceInsights/Private/Insights/Widgets/STimersViewTooltip.cpp
+++ b/Engine/Source/Developer/TraceInsights/Private/Insights/Widgets/STimersViewTooltip.cpp
@@ -88,10 +88,12 @@ TSharedPtr<SToolTip> STimersViewTooltip::GetColumnTooltip(const Insights::FTable
 
 ////////////////////////////////////////////////////////////////////////////////////////////////////
 
-TSharedPtr<SToolTip> STimersViewTooltip::GetCellTooltip(const TSharedPtr<FTimerNode> TimerNodePtr, const TSharedPtr<Insights::FTableColumn> ColumnPtr)
+TSharedPtr<SToolTip> STimersViewTooltip::GetRowTooltip(const TSharedPtr<FTimerNode> TimerNodePtr)
 {
 	const Trace::FAggregatedTimingStats& Stats = TimerNodePtr->GetAggregatedStats();
 
+	const FText InstanceCountText = FText::AsNumber(Stats.InstanceCount);
+
 	const int32 NumDigits = 5;
 
 	TCHAR FormatString[32];
@@ -109,8 +111,6 @@ TSharedPtr<SToolTip> STimersViewTooltip::GetCellTooltip(const TSharedPtr<FTimerN
 	const FText AvgExclusiveTimeText = FText::FromString(TimeUtils::FormatTimeMs(Stats.AverageExclusiveTime, NumDigits, true));
 	const FText MedExclusiveTimeText = FText::FromString(TimeUtils::FormatTimeMs(Stats.MedianExclusiveTime, NumDigits, true));
 
-	const FText InstanceCountText = FText::AsNumber(Stats.InstanceCount);
-
 	TSharedPtr<SGridPanel> GridPanel;
 	TSharedPtr<SHorizontalBox> HBox;
 
diff --git a/Engine/Source/Developer/TraceInsights/Private/Insights/Widgets/STimersViewTooltip.h b/Engine/Source/Developer/TraceInsights/Private/Insights/Widgets/STimersViewTooltip.h
--- a/Engine/Source/Developer/TraceInsights/Private/Insights/Widgets/STimersViewTooltip.h
+++ b/Engine/Source/Developer/TraceInsights/Private/Insights/Widgets/STimersViewTooltip.h
@@ -3,11 +3,9 @@
 #pragma once
 
 #include "CoreMinimal.h"
+#include "Widgets/IToolTip.h"
+#include "Widgets/SToolTip.h"
 
-// Insights
-#include "Insights/InsightsManager.h"
-
-class SToolTip;
 class SGridPanel;
 
 namespace Insights
@@ -18,7 +16,7 @@ namespace Insights
 
 class FTimerNode;
 
-#define LOCTEXT_NAMESPACE "STimersView"
+////////////////////////////////////////////////////////////////////////////////////////////////////
 
 /** Timers View Tooltip */
 class STimersViewTooltip
@@ -28,10 +26,60 @@ public:
 
 	static TSharedPtr<SToolTip> GetTableTooltip(const Insights::FTable& Table);
 	static TSharedPtr<SToolTip> GetColumnTooltip(const Insights::FTableColumn& Column);
-	static TSharedPtr<SToolTip> GetCellTooltip(const TSharedPtr<FTimerNode> TreeNodePtr, const TSharedPtr<Insights::FTableColumn> ColumnPtr);
+	static TSharedPtr<SToolTip> GetRowTooltip(const TSharedPtr<FTimerNode> TreeNodePtr);
 
 private:
 	static void AddStatsRow(TSharedPtr<SGridPanel> Grid, int32& Row, const FText& Name, const FText& Value1, const FText& Value2);
 };
 
-#undef LOCTEXT_NAMESPACE
+////////////////////////////////////////////////////////////////////////////////////////////////////
+
+class STimerTableRowToolTip : public IToolTip
+{
+public:
+	STimerTableRowToolTip(const TSharedPtr<FTimerNode> InTreeNodePtr) : TreeNodePtr(InTreeNodePtr) {}
+	virtual ~STimerTableRowToolTip() { }
+
+	virtual TSharedRef<class SWidget> AsWidget()
+	{
+		CreateToolTipWidget();
+		return ToolTipWidget.ToSharedRef();
+	}
+
+	virtual TSharedRef<SWidget> GetContentWidget()
+	{
+		CreateToolTipWidget();
+		return ToolTipWidget->GetContentWidget();
+	}
+
+	virtual void SetContentWidget(const TSharedRef<SWidget>& InContentWidget)
+	{
+		CreateToolTipWidget();
+		ToolTipWidget->SetContentWidget(InContentWidget);
+	}
+
+	void InvalidateWidget()
+	{
+		ToolTipWidget.Reset();
+	}
+
+	virtual bool IsEmpty() const { return false; }
+	virtual bool IsInteractive() const { return false; }
+	virtual void OnOpening() {}
+	virtual void OnClosed() {}
+
+private:
+	void CreateToolTipWidget()
+	{
+		if (!ToolTipWidget.IsValid())
+		{
+			ToolTipWidget = STimersViewTooltip::GetRowTooltip(TreeNodePtr);
+		}
+	}
+
+private:
+	TSharedPtr<SToolTip> ToolTipWidget;
+	const TSharedPtr<FTimerNode> TreeNodePtr;
+};
+
+////////////////////////////////////////////////////////////////////////////////////////////////////
