[Insights] Integrated CL 9487375 from //UE4/Dev-Core. Start Page: Fixed session list to correctly update when a trace file is renamed.

diff --git a/Engine/Source/Developer/TraceInsights/Private/Insights/Widgets/SStartPageWindow.cpp b/Engine/Source/Developer/TraceInsights/Private/Insights/Widgets/SStartPageWindow.cpp
--- a/Engine/Source/Developer/TraceInsights/Private/Insights/Widgets/SStartPageWindow.cpp
+++ b/Engine/Source/Developer/TraceInsights/Private/Insights/Widgets/SStartPageWindow.cpp
@@ -893,6 +893,8 @@ void SStartPageWindow::RefreshTraceSessionList()
 
 	AvailableSessionCount = AvailableSessions.Num();
 
+	bool bSessionChanged = false;
+
 	// Count number of live sessions and update file sizes.
 	int32 OldLiveSessionCount = LiveSessionCount;
 	LiveSessionCount = 0;
@@ -910,11 +912,16 @@ void SStartPageWindow::RefreshTraceSessionList()
 		{
 			(*TraceSessionPtrPtr)->Size = SessionInfo.Size;
 		}
+		else
+		{
+			bSessionChanged = true;
+		}
 	}
 
 	// If session list has changed on analysis side, recreate the TraceSessions list view widget.
 	//TODO: if (AvailableSessionsChangeNumber != SessionService->GetAvailableSessionsChangeNumber())
-	if (AvailableSessionCount != TraceSessions.Num() ||
+	if (bSessionChanged ||
+		AvailableSessionCount != TraceSessions.Num() ||
 		LiveSessionCount != OldLiveSessionCount)
 	{
 		TSharedPtr<FTraceSession> NewSelectedTraceSession;
