[Insights] Integrated CL 9546052 and CL 9547719 from //UE4/Dev-Core.
  - Networking / Packets view: Fixed infinite loop if zooming in (vertically) too much.
  - Networking / Packets view: Fixed rare crash (trying to use an invalid ConnectionIndex).
  - Added GetNetTraceVersion() API to INetProfilerProvider (by [REDACTED_PERSON])
  - Changed the Networking Insights tabs to open automatically only if the trace has net events.

diff --git a/Engine/Source/Developer/TraceInsights/Private/Insights/InsightsManager.cpp b/Engine/Source/Developer/TraceInsights/Private/Insights/InsightsManager.cpp
--- a/Engine/Source/Developer/TraceInsights/Private/Insights/InsightsManager.cpp
+++ b/Engine/Source/Developer/TraceInsights/Private/Insights/InsightsManager.cpp
@@ -5,6 +5,7 @@
 #include "Framework/Application/SlateApplication.h"
 #include "Modules/ModuleManager.h"
 #include "Templates/UniquePtr.h"
+#include "TraceServices/Model/NetProfiler.h"
 
 // Insights
 #include "Insights/LoadingProfiler/LoadingProfilerManager.h"
@@ -36,6 +37,7 @@ FInsightsManager::FInsightsManager(TSharedRef<Trace::IAnalysisService> InTraceAn
 	, ActionManager(this)
 	, Settings()
 	, bIsDebugInfoEnabled(false)
+	, bIsNetworkingProfilerAvailable(false)
 {
 }
 
@@ -44,8 +46,8 @@ FInsightsManager::FInsightsManager(TSharedRef<Trace::IAnalysisService> InTraceAn
 void FInsightsManager::PostConstructor()
 {
 	// Register tick functions.
-	//OnTick = FTickerDelegate::CreateSP(this, &FInsightsManager::Tick);
-	//OnTickHandle = FTicker::GetCoreTicker().AddTicker(OnTick, 1.0f);
+	OnTick = FTickerDelegate::CreateSP(this, &FInsightsManager::Tick);
+	OnTickHandle = FTicker::GetCoreTicker().AddTicker(OnTick, 1.0f);
 
 	FInsightsCommands::Register();
 	BindCommands();
@@ -67,7 +69,7 @@ FInsightsManager::~FInsightsManager()
 	FInsightsCommands::Unregister();
 
 	// Unregister tick function.
-	//FTicker::GetCoreTicker().RemoveTicker(OnTickHandle);
+	FTicker::GetCoreTicker().RemoveTicker(OnTickHandle);
 }
 
 ////////////////////////////////////////////////////////////////////////////////////////////////////
@@ -123,7 +125,30 @@ FInsightsSettings& FInsightsManager::GetSettings()
 
 bool FInsightsManager::Tick(float DeltaTime)
 {
-	//SCOPE_CYCLE_COUNTER(STAT_IM_Tick);
+	if (!bIsNetworkingProfilerAvailable)
+	{
+		uint32 NetTraceVersion = 0;
+
+		if (Session.IsValid())
+		{
+			Trace::FAnalysisSessionReadScope SessionReadScope(*Session.Get());
+			const Trace::INetProfilerProvider& NetProfilerProvider = Trace::ReadNetProfilerProvider(*Session.Get());
+			NetTraceVersion = NetProfilerProvider.GetNetTraceVersion();
+		}
+
+		if (NetTraceVersion > 0)
+		{
+			bIsNetworkingProfilerAvailable = true;
+
+			if (FGlobalTabmanager::Get()->HasTabSpawner(FInsightsManagerTabs::NetworkingProfilerTabId))
+			{
+				FGlobalTabmanager::Get()->InvokeTab(FInsightsManagerTabs::NetworkingProfilerTabId);
+				FGlobalTabmanager::Get()->InvokeTab(FInsightsManagerTabs::NetworkingProfilerTabId);
+			}
+
+			ActivateTimingInsightsTab();
+		}
+	}
 
 	return true;
 }
@@ -136,6 +161,7 @@ void FInsightsManager::ResetSession()
 	{
 		Session.Reset();
 		CurrentSessionHandle = 0;
+		bIsNetworkingProfilerAvailable = false;
 		OnSessionChanged();
 	}
 }
@@ -169,28 +195,35 @@ void FInsightsManager::OnSessionChanged()
 
 void FInsightsManager::SpawnAndActivateTabs()
 {
-	// Open tabs for profilers.
+	// Open Timing Insights tab.
 	if (FGlobalTabmanager::Get()->HasTabSpawner(FInsightsManagerTabs::TimingProfilerTabId))
 	{
 		FGlobalTabmanager::Get()->InvokeTab(FInsightsManagerTabs::TimingProfilerTabId);
 	}
+
+	// Open Asset Loading Insights tab.
 	if (FGlobalTabmanager::Get()->HasTabSpawner(FInsightsManagerTabs::LoadingProfilerTabId))
 	{
 		FGlobalTabmanager::Get()->InvokeTab(FInsightsManagerTabs::LoadingProfilerTabId);
 	}
+
+	// Close the existing Networking Insights tabs.
 	if (FGlobalTabmanager::Get()->HasTabSpawner(FInsightsManagerTabs::NetworkingProfilerTabId))
 	{
-		// Close previous tabs.
 		TSharedPtr<SDockTab> NetworkingProfilerTab;
 		while ((NetworkingProfilerTab = FGlobalTabmanager::Get()->FindExistingLiveTab(FInsightsManagerTabs::NetworkingProfilerTabId)).IsValid())
 		{
 			NetworkingProfilerTab->RequestCloseTab();
 		}
-
-		FGlobalTabmanager::Get()->InvokeTab(FInsightsManagerTabs::NetworkingProfilerTabId);
-		FGlobalTabmanager::Get()->InvokeTab(FInsightsManagerTabs::NetworkingProfilerTabId);
 	}
 
+	ActivateTimingInsightsTab();
+}
+
+////////////////////////////////////////////////////////////////////////////////////////////////////
+
+void FInsightsManager::ActivateTimingInsightsTab()
+{
 	// Ensure Timing Insights / Timing View is the active tab / view.
 	if (TSharedPtr<SDockTab> TimingInsightsTab = FGlobalTabmanager::Get()->FindExistingLiveTab(FInsightsManagerTabs::TimingProfilerTabId))
 	{
@@ -297,6 +330,7 @@ void FInsightsManager::LoadSession(Trace::FSessionHandle SessionHandle)
 	{
 		Session = AnalysisService->StartAnalysis(SessionInfo.Name, MoveTemp(DataStream));
 		CurrentSessionHandle = SessionHandle;
+		bIsNetworkingProfilerAvailable = false;
 		SpawnAndActivateTabs();
 		OnSessionChanged();
 	}
@@ -313,6 +347,7 @@ void FInsightsManager::LoadTraceFile(const FString& TraceFilepath)
 	{
 		Session = AnalysisService->StartAnalysis(*TraceFilepath, MoveTemp(DataStream));
 		CurrentSessionHandle = 0;
+		bIsNetworkingProfilerAvailable = false;
 		SpawnAndActivateTabs();
 		OnSessionChanged();
 	}
diff --git a/Engine/Source/Developer/TraceInsights/Private/Insights/InsightsManager.h b/Engine/Source/Developer/TraceInsights/Private/Insights/InsightsManager.h
--- a/Engine/Source/Developer/TraceInsights/Private/Insights/InsightsManager.h
+++ b/Engine/Source/Developer/TraceInsights/Private/Insights/InsightsManager.h
@@ -177,6 +177,8 @@ private:
 
 	void SpawnAndActivateTabs();
 
+	void ActivateTimingInsightsTab();
+
 private:
 	/** The delegate to be invoked when this manager ticks. */
 	FTickerDelegate OnTick;
@@ -211,4 +213,6 @@ private:
 
 	/** A shared pointer to the global instance of the main manager. */
 	static TSharedPtr<FInsightsManager> Instance;
+
+	bool bIsNetworkingProfilerAvailable;
 };
diff --git a/Engine/Source/Developer/TraceInsights/Private/Insights/NetworkingProfiler/Widgets/SPacketView.cpp b/Engine/Source/Developer/TraceInsights/Private/Insights/NetworkingProfiler/Widgets/SPacketView.cpp
--- a/Engine/Source/Developer/TraceInsights/Private/Insights/NetworkingProfiler/Widgets/SPacketView.cpp
+++ b/Engine/Source/Developer/TraceInsights/Private/Insights/NetworkingProfiler/Widgets/SPacketView.cpp
@@ -235,11 +235,24 @@ void SPacketView::UpdateState()
 		Trace::FAnalysisSessionReadScope SessionReadScope(*Session.Get());
 		const Trace::INetProfilerProvider& NetProfilerProvider = Trace::ReadNetProfilerProvider(*Session.Get());
 
+		// TODO: if (NetProfilerProvider.IsConnectionValid(GameInstanceIndex, ConnectionIndex, ConnectionMode))
 		const uint32 GameInstanceCount = NetProfilerProvider.GetGameInstanceCount();
 		if (GameInstanceIndex < GameInstanceCount)
 		{
-			//const uint32 ConnectionCount = NetProfilerProvider.GetConnectionCount(GameInstanceIndex);
-			//if (ConnectionIndex < ConnectionCount) TODO: (if NetProfilerProvider.IsConnectionValid(GameInstanceIndex))
+			bool bIsValidConnection = false;
+			NetProfilerProvider.ReadConnections(GameInstanceIndex, [this, &bIsValidConnection](const Trace::FNetProfilerConnection& Connection)
+			{
+				if (ConnectionIndex == Connection.ConnectionIndex)
+				{
+					if ((ConnectionMode == Trace::ENetProfilerConnectionMode::Outgoing && Connection.bHasOutgoingData) ||
+						(ConnectionMode == Trace::ENetProfilerConnectionMode::Incoming && Connection.bHasIncomingData))
+					{
+						bIsValidConnection = true;
+					}
+				}
+			});
+
+			if (bIsValidConnection)
 			{
 				const uint32 NumPackets = NetProfilerProvider.GetPacketCount(ConnectionIndex, ConnectionMode);
 
@@ -731,7 +744,11 @@ void SPacketView::DrawVerticalAxisGrid(FDrawContext& DrawContext, const FSlateBr
 
 	if (Delta > 0.0)
 	{
-		const int64 DeltaBits = static_cast<int64>(Delta);
+		int64 DeltaBits = static_cast<int64>(Delta);
+		if (DeltaBits <= 0)
+		{
+			DeltaBits = 1;
+		}
 
 		// Compute rounding based on magnitude of visible range of values (Delta).
 		int64 Power10 = 1;
diff --git a/Engine/Source/Developer/TraceServices/Private/Analyzers/NetTraceAnalyzer.cpp b/Engine/Source/Developer/TraceServices/Private/Analyzers/NetTraceAnalyzer.cpp
--- a/Engine/Source/Developer/TraceServices/Private/Analyzers/NetTraceAnalyzer.cpp
+++ b/Engine/Source/Developer/TraceServices/Private/Analyzers/NetTraceAnalyzer.cpp
@@ -74,6 +74,8 @@ bool FNetTraceAnalyzer::OnEvent(uint16 RouteId, const FOnEventContext& Context)
 			const uint64 TimestampCycles = EventData.GetValue<uint64>("Timestamp");
 
 			LastTimeStamp = Context.SessionContext.TimestampFromCycle(TimestampCycles);
+
+			NetProfilerProvider.SetNetTraceVersion(NetTraceVersion);
 		}
 		break;
 
diff --git a/Engine/Source/Developer/TraceServices/Private/Model/NetProfilerProvider.cpp b/Engine/Source/Developer/TraceServices/Private/Model/NetProfilerProvider.cpp
--- a/Engine/Source/Developer/TraceServices/Private/Model/NetProfilerProvider.cpp
+++ b/Engine/Source/Developer/TraceServices/Private/Model/NetProfilerProvider.cpp
@@ -12,6 +12,7 @@ const FName FNetProfilerProvider::ProviderName("NetProfilerProvider");
 
 FNetProfilerProvider::FNetProfilerProvider(IAnalysisSession& InSession)
 	: Session(InSession)
+	, NetTraceVersion(0U)
 	, Connections(InSession.GetLinearAllocator(), 4096)
 	, ConnectionChangeCount(0u)
 {
@@ -23,6 +24,18 @@ FNetProfilerProvider::~FNetProfilerProvider()
 {
 }
 
+void FNetProfilerProvider::SetNetTraceVersion(uint32 Version)
+{
+	Session.WriteAccessCheck();
+
+	NetTraceVersion = Version;
+}
+
+uint32 FNetProfilerProvider::GetNetTraceVersion() const
+{
+	return NetTraceVersion;
+}
+
 uint32 FNetProfilerProvider::AddNetProfilerName(const TCHAR* Name)
 {
 	Session.WriteAccessCheck();
diff --git a/Engine/Source/Developer/TraceServices/Private/Model/NetProfilerProvider.h b/Engine/Source/Developer/TraceServices/Private/Model/NetProfilerProvider.h
--- a/Engine/Source/Developer/TraceServices/Private/Model/NetProfilerProvider.h
+++ b/Engine/Source/Developer/TraceServices/Private/Model/NetProfilerProvider.h
@@ -53,6 +53,8 @@ public:
 	explicit FNetProfilerProvider(IAnalysisSession& InSession);
 	virtual ~FNetProfilerProvider();
 
+	void SetNetTraceVersion(uint32 NetTraceVersion);
+
 	uint32 AddNetProfilerName(const TCHAR* Name);
 
 	uint32 AddNetProfilerEventType(uint32 NameIndex, uint32 Level);
@@ -71,6 +73,8 @@ public:
 
 	// INetProfilerProvider Interface
 
+	virtual uint32 GetNetTraceVersion() const override;
+
 	// Access Names
 	virtual uint32 GetNameCount() const override { return Names.Num(); }
 	virtual void ReadNames(TFunctionRef<void(const FNetProfilerName*, uint64)> Callback) const override;
@@ -127,6 +131,8 @@ private:
 
 	IAnalysisSession& Session;
 
+	uint32 NetTraceVersion;
+
 	TArray<FNetProfilerName> Names;
 
 	TArray<FNetProfilerEventType> EventTypes;
diff --git a/Engine/Source/Developer/TraceServices/Public/TraceServices/Model/NetProfiler.h b/Engine/Source/Developer/TraceServices/Public/TraceServices/Model/NetProfiler.h
--- a/Engine/Source/Developer/TraceServices/Public/TraceServices/Model/NetProfiler.h
+++ b/Engine/Source/Developer/TraceServices/Public/TraceServices/Model/NetProfiler.h
@@ -123,6 +123,10 @@ struct FNetProfilerAggregatedStats
 class INetProfilerProvider : public IProvider
 {
 public:
+	// Return the version reported in the trace
+	// A return value of 0 indicates no network trace data
+	virtual uint32 GetNetTraceVersion() const = 0;
+
 	// Access Names
 	virtual uint32 GetNameCount() const = 0;
 	virtual void ReadNames(TFunctionRef<void(const FNetProfilerName*, uint64)> Callback) const = 0;
