[Insights] Asset Loading Insights: Fixed crash when selecting Object Type Aggregation Tab on second open Trace.

diff --git a/Engine/Source/Developer/TraceInsights/Private/Insights/LoadingProfiler/Widgets/SLoadingProfilerWindow.cpp b/Engine/Source/Developer/TraceInsights/Private/Insights/LoadingProfiler/Widgets/SLoadingProfilerWindow.cpp
--- a/Engine/Source/Developer/TraceInsights/Private/Insights/LoadingProfiler/Widgets/SLoadingProfilerWindow.cpp
+++ b/Engine/Source/Developer/TraceInsights/Private/Insights/LoadingProfiler/Widgets/SLoadingProfilerWindow.cpp
@@ -91,6 +91,51 @@ void SLoadingProfilerWindow::Reset()
 	{
 		ExportDetailsTreeView->Reset();
 	}
+
+	UpdateTableTreeViews();
+}
+
+////////////////////////////////////////////////////////////////////////////////////////////////////
+
+void SLoadingProfilerWindow::UpdateTableTreeViews()
+{
+	TSharedPtr<const Trace::IAnalysisSession> Session = FInsightsManager::Get()->GetSession();
+	if (Session.IsValid() && Trace::ReadLoadTimeProfilerProvider(*Session.Get()))
+	{
+		Trace::FAnalysisSessionReadScope SessionReadScope(*Session.Get());
+		const Trace::ILoadTimeProfilerProvider& LoadTimeProfilerProvider = *Trace::ReadLoadTimeProfilerProvider(*Session.Get());
+
+		const double SelectionStartTime = TimingView ? TimingView->GetSelectionStartTime() : 0.0;
+		const double SelectionEndTime = TimingView ? TimingView->GetSelectionEndTime() : 0.0;
+
+		if (EventAggregationTreeView)
+		{
+			Trace::ITable<Trace::FLoadTimeProfilerAggregatedStats>* EventAggregationTable = LoadTimeProfilerProvider.CreateEventAggregation(SelectionStartTime, SelectionEndTime);
+			EventAggregationTreeView->GetTable()->UpdateSourceTable(MakeShareable(EventAggregationTable));
+			EventAggregationTreeView->RebuildTree(true);
+		}
+
+		if (ObjectTypeAggregationTreeView)
+		{
+			Trace::ITable<Trace::FLoadTimeProfilerAggregatedStats>* ObjectTypeAggregationTable = LoadTimeProfilerProvider.CreateObjectTypeAggregation(SelectionStartTime, SelectionEndTime);
+			ObjectTypeAggregationTreeView->GetTable()->UpdateSourceTable(MakeShareable(ObjectTypeAggregationTable));
+			ObjectTypeAggregationTreeView->RebuildTree(true);
+		}
+
+		if (PackageDetailsTreeView)
+		{
+			Trace::ITable<Trace::FPackagesTableRow>* PackageDetailsTable = LoadTimeProfilerProvider.CreatePackageDetailsTable(SelectionStartTime, SelectionEndTime);
+			PackageDetailsTreeView->GetTable()->UpdateSourceTable(MakeShareable(PackageDetailsTable));
+			PackageDetailsTreeView->RebuildTree(true);
+		}
+
+		if (ExportDetailsTreeView)
+		{
+			Trace::ITable<Trace::FExportsTableRow>* ExportDetailsTable = LoadTimeProfilerProvider.CreateExportDetailsTable(SelectionStartTime, SelectionEndTime);
+			ExportDetailsTreeView->GetTable()->UpdateSourceTable(MakeShareable(ExportDetailsTable));
+			ExportDetailsTreeView->RebuildTree(true);
+		}
+	}
 }
 
 ////////////////////////////////////////////////////////////////////////////////////////////////////
diff --git a/Engine/Source/Developer/TraceInsights/Private/Insights/LoadingProfiler/Widgets/SLoadingProfilerWindow.h b/Engine/Source/Developer/TraceInsights/Private/Insights/LoadingProfiler/Widgets/SLoadingProfilerWindow.h
--- a/Engine/Source/Developer/TraceInsights/Private/Insights/LoadingProfiler/Widgets/SLoadingProfilerWindow.h
+++ b/Engine/Source/Developer/TraceInsights/Private/Insights/LoadingProfiler/Widgets/SLoadingProfilerWindow.h
@@ -56,6 +56,7 @@ public:
 	SLATE_END_ARGS()
 
 	void Reset();
+	void UpdateTableTreeViews();
 
 	/** Constructs this widget. */
 	void Construct(const FArguments& InArgs, const TSharedRef<SDockTab>& ConstructUnderMajorTab, const TSharedPtr<SWindow>& ConstructUnderWindow);
diff --git a/Engine/Source/Developer/TraceInsights/Private/Insights/Widgets/STimingView.cpp b/Engine/Source/Developer/TraceInsights/Private/Insights/Widgets/STimingView.cpp
--- a/Engine/Source/Developer/TraceInsights/Private/Insights/Widgets/STimingView.cpp
+++ b/Engine/Source/Developer/TraceInsights/Private/Insights/Widgets/STimingView.cpp
@@ -2036,49 +2036,7 @@ void STimingView::UpdateAggregatedStats()
 		TSharedPtr<SLoadingProfilerWindow> LoadingProfilerWnd = FLoadingProfilerManager::Get()->GetProfilerWindow();
 		if (LoadingProfilerWnd.IsValid())
 		{
-			TSharedPtr<Insights::STableTreeView> EventAggregationTreeView = LoadingProfilerWnd->GetEventAggregationTreeView();
-			TSharedPtr<Insights::STableTreeView> ObjectTypeAggregationTreeView = LoadingProfilerWnd->GetObjectTypeAggregationTreeView();
-			TSharedPtr<Insights::STableTreeView> PackageDetailsTreeView = LoadingProfilerWnd->GetPackageDetailsTreeView();
-			TSharedPtr<Insights::STableTreeView> ExportDetailsTreeView = LoadingProfilerWnd->GetExportDetailsTreeView();
-
-			TSharedPtr<const Trace::IAnalysisSession> Session = FInsightsManager::Get()->GetSession();
-			if (Session.IsValid() && Trace::ReadLoadTimeProfilerProvider(*Session.Get()))
-			{
-				Trace::FAnalysisSessionReadScope SessionReadScope(*Session.Get());
-				const Trace::ILoadTimeProfilerProvider& LoadTimeProfilerProvider = *Trace::ReadLoadTimeProfilerProvider(*Session.Get());
-
-				if (EventAggregationTreeView.IsValid())
-				{
-					//TODO: LoadTimeProfilerProvider.UpdateEventAggregation(EventAggregationTreeView->GetTable(), SelectionStartTime, SelectionEndTime)
-					Trace::ITable<Trace::FLoadTimeProfilerAggregatedStats>* EventAggregationTable = LoadTimeProfilerProvider.CreateEventAggregation(SelectionStartTime, SelectionEndTime);
-					EventAggregationTreeView->GetTable()->UpdateSourceTable(MakeShareable(EventAggregationTable));
-					EventAggregationTreeView->RebuildTree(true);
-				}
-
-				if (ObjectTypeAggregationTreeView.IsValid())
-				{
-					//TODO: LoadTimeProfilerProvider.UpdateObjectTypeAggregation(ObjectTypeAggregationTreeView->GetTable(), SelectionStartTime, SelectionEndTime)
-					Trace::ITable<Trace::FLoadTimeProfilerAggregatedStats>* ObjectTypeAggregationTable = LoadTimeProfilerProvider.CreateObjectTypeAggregation(SelectionStartTime, SelectionEndTime);
-					ObjectTypeAggregationTreeView->GetTable()->UpdateSourceTable(MakeShareable(ObjectTypeAggregationTable));
-					ObjectTypeAggregationTreeView->RebuildTree(true);
-				}
-
-				if (PackageDetailsTreeView.IsValid())
-				{
-					//TODO: LoadTimeProfilerProvider.UpdatePackageDetails(PackageDetailsTreeView->GetTable(), SelectionStartTime, SelectionEndTime)
-					Trace::ITable<Trace::FPackagesTableRow>* PackageDetailsTable = LoadTimeProfilerProvider.CreatePackageDetailsTable(SelectionStartTime, SelectionEndTime);
-					PackageDetailsTreeView->GetTable()->UpdateSourceTable(MakeShareable(PackageDetailsTable));
-					PackageDetailsTreeView->RebuildTree(true);
-				}
-
-				if (ExportDetailsTreeView.IsValid())
-				{
-					//TODO: LoadTimeProfilerProvider.UpdateExportDetails(ExportDetailsTreeView->GetTable(), SelectionStartTime, SelectionEndTime)
-					Trace::ITable<Trace::FExportsTableRow>* ExportDetailsTable = LoadTimeProfilerProvider.CreateExportDetailsTable(SelectionStartTime, SelectionEndTime);
-					ExportDetailsTreeView->GetTable()->UpdateSourceTable(MakeShareable(ExportDetailsTable));
-					ExportDetailsTreeView->RebuildTree(true);
-				}
-			}
+			LoadingProfilerWnd->UpdateTableTreeViews();
 		}
 	}
 }
