Fix compile errors
#rnx

diff --git a/Engine/Source/Developer/TraceInsights/Private/Insights/ViewModels/MarkersTimingTrack.cpp b/Engine/Source/Developer/TraceInsights/Private/Insights/ViewModels/MarkersTimingTrack.cpp
--- a/Engine/Source/Developer/TraceInsights/Private/Insights/ViewModels/MarkersTimingTrack.cpp
+++ b/Engine/Source/Developer/TraceInsights/Private/Insights/ViewModels/MarkersTimingTrack.cpp
@@ -424,7 +424,7 @@ void FTimeMarkerTrackBuilder::Flush(float AvailableTextW)
 			FString CategoryStr(LastCategory);
 			if (CategoryStr.StartsWith(TEXT("Log")))
 			{
-				CategoryStr.RightChopInline(3, fase);
+				CategoryStr.RightChopInline(3, false);
 			}
 
 			const int32 LastWholeCharacterIndexCategory = FontMeasureService->FindLastWholeCharacterIndexBeforeOffset(CategoryStr, Font, FMath::RoundToInt(AvailableTextW - 2.0f));
