[TraceInsights]
  -  Added "Stats Counters" view (somehow similar with Timers view). Features: sorting, grouping, tooltip, columns for aggregated statisitcs (Sum, Max, Avg, Med, Min). Lower and Upper Quartile aggregated statistics are not yet computed.
  - Graph track: Fixed crash when resizing the widget and its height becomes negative.
  - Frames track: Fixed crash when resizing the widget and its width becomes negative.
  - Frames track: Increased initial/max zoom level (improved default zoom for short sessions).
  - Frames track: Added time limit for how often can make updates (sync with state from analysis) -- 100ms.
  - Timers/Stats views: Added time limit for how often can make updates (sync with state from analysis) -- 200ms.
  - Timers view: Removed filter by Compute scope (as we do not have yet support for compute timing events).
  - Timers view: Moved "Min" column at the end (after the Average and Median columns). So columns with show in order: Total, Max, Avg, Med, Min.
  - Renamed "Timing Profiler" to "Timing Insights".
  - Renamed "I/O Profiler" to "Asset Loading Insights".
  - Removed Mock and Settings buttons from "Timing Insights" toolbar.
  - Added a main menu ("MENU") to "Timing Insights" (above toolbar) to  be able to open any closed view.
  - Changed Timing Insights layout to use a tab manager (each view is now a dockable tab).
  - Added Live and Load buttons on Start Page (WIP -- no action for now).
  - Added icons to all major tabs and to all Timing Insights views. Also started to use FInsightsStyle as a proxy for those brush resources.
  - Fixed few coding standard issues (int vs. int32; removed "Nocopyable" inheritance; renamed few "too short" field members; renamed/refactored DrawContext / Builder classes; etc.) (based on review from Steve.Robb -- Note: not all raised issues are fixed by this commit).
  - Increased version to 0.06

diff --git a/Engine/Source/Developer/TraceInsights/Private/Insights/Common/FixedCircularBuffer.h b/Engine/Source/Developer/TraceInsights/Private/Insights/Common/FixedCircularBuffer.h
--- a/Engine/Source/Developer/TraceInsights/Private/Insights/Common/FixedCircularBuffer.h
+++ b/Engine/Source/Developer/TraceInsights/Private/Insights/Common/FixedCircularBuffer.h
@@ -2,12 +2,12 @@
 
 #pragma once
 
-template <typename T, int Size = 64>
+template <typename T, int32 Size = 64>
 struct TFixedCircularBuffer
 {
 	T Buffer[Size];
-	int BufferStartOffset = 0;
-	int BufferLength = 0;
+	int32 BufferStartOffset = 0;
+	int32 BufferLength = 0;
 
 	void Reset()
 	{
@@ -20,12 +20,16 @@ struct TFixedCircularBuffer
 		Buffer[BufferStartOffset] = Value;
 		BufferStartOffset++;
 		if (BufferStartOffset == Size)
+		{
 			BufferStartOffset = 0;
+		}
 		if (BufferLength < Size)
+		{
 			BufferLength++;
+		}
 	}
 
-	const T GetValue(const int RecentIndex) const
+	const T GetValue(const int32 RecentIndex) const
 	{
 		return Buffer[(BufferStartOffset + Size - RecentIndex - 1) % Size];
 	}
@@ -33,10 +37,14 @@ struct TFixedCircularBuffer
 	const T ComputeAverage() const
 	{
 		if (BufferLength == 0)
+		{
 			return 0;
+		}
 		T Total = 0;
-		for (int RecentIndex = 0; RecentIndex < BufferLength; ++RecentIndex)
+		for (int32 RecentIndex = 0; RecentIndex < BufferLength; ++RecentIndex)
+		{
 			Total += Buffer[(BufferStartOffset + Size - RecentIndex - 1) % Size];
+		}
 		return Total / BufferLength;
 	}
 };
diff --git a/Engine/Source/Developer/TraceInsights/Private/Insights/Common/PaintUtils.h b/Engine/Source/Developer/TraceInsights/Private/Insights/Common/PaintUtils.h
--- a/Engine/Source/Developer/TraceInsights/Private/Insights/Common/PaintUtils.h
+++ b/Engine/Source/Developer/TraceInsights/Private/Insights/Common/PaintUtils.h
@@ -14,8 +14,10 @@ class WidgetStyle;
 enum class ESlateDrawEffect : uint8;
 class FSlateWindowElementList;
 
-/// Holds current state provided by OnPaint function, used to simplify drawing.
-struct FDrawContext : public FNoncopyable
+/**
+ * Holds current state provided by OnPaint function, used to simplify drawing.
+ */
+struct FDrawContext
 {
 	FDrawContext(const FGeometry& InGeometry,
 				 const FSlateRect& InCullingRect,
@@ -31,6 +33,12 @@ struct FDrawContext : public FNoncopyable
 		, LayerId(InOutLayerId)
 	{}
 
+	/**
+	 * Non-copyable
+	 */
+	FDrawContext(const FDrawContext&) = delete;
+	FDrawContext& operator=(const FDrawContext&) = delete;
+
 	inline void DrawBox(const float X, const float Y, const float W, const float H, const FSlateBrush* Brush, const FLinearColor& Color) const
 	{
 		FSlateDrawElement::MakeBox(ElementList, LayerId, MAKE_PAINT_GEOMETRY_RC(Geometry, X, Y, W, H), Brush, DrawEffects, Color);
diff --git a/Engine/Source/Developer/TraceInsights/Private/Insights/Common/TimeUtils.cpp b/Engine/Source/Developer/TraceInsights/Private/Insights/Common/TimeUtils.cpp
--- a/Engine/Source/Developer/TraceInsights/Private/Insights/Common/TimeUtils.cpp
+++ b/Engine/Source/Developer/TraceInsights/Private/Insights/Common/TimeUtils.cpp
@@ -172,7 +172,7 @@ FString FormatTimeAuto(const double InDuration)
 
 ////////////////////////////////////////////////////////////////////////////////////////////////////
 
-FString FormatTimeMs(const double InDuration, const int NumDigits, bool bAddTimeUnit)
+FString FormatTimeMs(const double InDuration, const int32 NumDigits, bool bAddTimeUnit)
 {
 	if (FMath::IsNearlyZero(InDuration, TimeUtils::Picosecond))
 	{
@@ -394,7 +394,7 @@ void SplitTime(const double InTime, FTimeSplit& OutTimeSplit)
 	{
 		OutTimeSplit.bIsZero = false;
 		OutTimeSplit.bIsInfinite = true;
-		for (int Index = 0; Index < 8; ++Index)
+		for (int32 Index = 0; Index < 8; ++Index)
 		{
 			OutTimeSplit.Units[Index] = 0;
 		}
diff --git a/Engine/Source/Developer/TraceInsights/Private/Insights/Common/TimeUtils.h b/Engine/Source/Developer/TraceInsights/Private/Insights/Common/TimeUtils.h
--- a/Engine/Source/Developer/TraceInsights/Private/Insights/Common/TimeUtils.h
+++ b/Engine/Source/Developer/TraceInsights/Private/Insights/Common/TimeUtils.h
@@ -55,7 +55,7 @@ namespace TimeUtils
 	};
 
 	FString FormatTimeAuto(const double Duration);
-	FString FormatTimeMs(const double Duration, const int NumDigits = 2, bool bAddTimeUnit = false);
+	FString FormatTimeMs(const double Duration, const int32 NumDigits = 2, bool bAddTimeUnit = false);
 	FString FormatTime(const double Time, const double Precision = 0.0);
 	FString FormatTimeHMS(const double Time, const double Precision = 0.0);
 	void SplitTime(const double Time, FTimeSplit& OutTimeSplit);
diff --git a/Engine/Source/Developer/TraceInsights/Private/Insights/InsightsCommands.cpp b/Engine/Source/Developer/TraceInsights/Private/Insights/InsightsCommands.cpp
--- a/Engine/Source/Developer/TraceInsights/Private/Insights/InsightsCommands.cpp
+++ b/Engine/Source/Developer/TraceInsights/Private/Insights/InsightsCommands.cpp
@@ -37,7 +37,7 @@ PRAGMA_DISABLE_OPTIMIZATION
 void FInsightsCommands::RegisterCommands()
 {
 	UI_COMMAND(InsightsManager_Live, "Live", "Loads profiler data from a live trace session", EUserInterfaceActionType::Button, FInputChord());
-	UI_COMMAND(InsightsManager_Load, "Load", "Loads profiler data from a trace file", EUserInterfaceActionType::Button, FInputChord(EModifierKey::Control, EKeys::L));
+	UI_COMMAND(InsightsManager_Load, "Load...", "Loads profiler data from a trace file", EUserInterfaceActionType::Button, FInputChord(EModifierKey::Control, EKeys::L));
 	UI_COMMAND(InsightsManager_Mock, "Mock", "Creates mock profiler data", EUserInterfaceActionType::Button, FInputChord());
 	UI_COMMAND(ToggleDebugInfo, "Debug", "Toggles the display of debug info", EUserInterfaceActionType::ToggleButton, FInputChord(EModifierKey::Control, EKeys::D));
 	UI_COMMAND(OpenSettings, "Settings", "Opens the Unreal Insights settings", EUserInterfaceActionType::Button, FInputChord(EModifierKey::Control, EKeys::O));
diff --git a/Engine/Source/Developer/TraceInsights/Private/Insights/InsightsManager.cpp b/Engine/Source/Developer/TraceInsights/Private/Insights/InsightsManager.cpp
--- a/Engine/Source/Developer/TraceInsights/Private/Insights/InsightsManager.cpp
+++ b/Engine/Source/Developer/TraceInsights/Private/Insights/InsightsManager.cpp
@@ -119,8 +119,9 @@ bool FInsightsManager::Tick(float DeltaTime)
 void FInsightsManager::ResetSession()
 {
 	Session = nullptr;
-	//TODO: SessionChangedEvent.Broadcast();
+
 	FTimingProfilerManager::Get()->OnSessionChanged();
+	SessionChangedEvent.Broadcast();
 }
 
 ////////////////////////////////////////////////////////////////////////////////////////////////////
@@ -131,8 +132,8 @@ void FInsightsManager::CreateMockSession()
 
 	Session = AnalysisService->MockAnalysis();
 
-	RequestTimersUpdateEvent.Broadcast();
-	//TODO: SessionChangedEvent.Broadcast();
+	FTimingProfilerManager::Get()->OnSessionChanged();
+	SessionChangedEvent.Broadcast();
 }
 
 ////////////////////////////////////////////////////////////////////////////////////////////////////
@@ -164,8 +165,8 @@ void FInsightsManager::CreateLiveSession()
 		Session = AnalysisService->StartAnalysis(TEXT("Latest session"), MoveTemp(DataStream));
 	}
 
-	RequestTimersUpdateEvent.Broadcast();
-	//TODO: SessionChangedEvent.Broadcast();
+	FTimingProfilerManager::Get()->OnSessionChanged();
+	SessionChangedEvent.Broadcast();
 }
 
 ////////////////////////////////////////////////////////////////////////////////////////////////////
@@ -179,8 +180,8 @@ void FInsightsManager::LoadTraceFile(const FString& TraceFilepath)
 
 	Session = AnalysisService->StartAnalysis(*TraceFilepath, MoveTemp(DataStream));
 
-	RequestTimersUpdateEvent.Broadcast();
-	//TODO: SessionChangedEvent.Broadcast();
+	FTimingProfilerManager::Get()->OnSessionChanged();
+	SessionChangedEvent.Broadcast();
 }
 
 ////////////////////////////////////////////////////////////////////////////////////////////////////
diff --git a/Engine/Source/Developer/TraceInsights/Private/Insights/InsightsManager.h b/Engine/Source/Developer/TraceInsights/Private/Insights/InsightsManager.h
--- a/Engine/Source/Developer/TraceInsights/Private/Insights/InsightsManager.h
+++ b/Engine/Source/Developer/TraceInsights/Private/Insights/InsightsManager.h
@@ -120,15 +120,15 @@ public:
 	void SetDebugInfo(const bool bDebugInfoEnabledState) { bIsDebugInfoEnabled = bDebugInfoEnabledState; }
 
 	////////////////////////////////////////////////////////////////////////////////////////////////////
-	// RequestTimersUpdateEvent
+	// SessionChangedEvent
 
 public:
-	/// The event to execute when the filter and presets widget should be updated with the latest data.
-	DECLARE_EVENT(FTimingProfilerManager, FRequestTimersUpdateEvent);
-	FRequestTimersUpdateEvent& OnRequestTimersUpdate() { return RequestTimersUpdateEvent; }
+	/// The event to execute when the session has changed.
+	DECLARE_EVENT(FTimingProfilerManager, FSessionChangedEvent);
+	FSessionChangedEvent& GetSessionChangedEvent() { return SessionChangedEvent; }
 protected:
-	/// The event to execute when the filter and presets widget should be updated with the latest data.
-	FRequestTimersUpdateEvent RequestTimersUpdateEvent;
+	/// The event to execute when the session has changed.
+	FSessionChangedEvent SessionChangedEvent;
 
 	////////////////////////////////////////////////////////////////////////////////////////////////////
 
diff --git a/Engine/Source/Developer/TraceInsights/Private/Insights/InsightsStyle.cpp b/Engine/Source/Developer/TraceInsights/Private/Insights/InsightsStyle.cpp
--- a/Engine/Source/Developer/TraceInsights/Private/Insights/InsightsStyle.cpp
+++ b/Engine/Source/Developer/TraceInsights/Private/Insights/InsightsStyle.cpp
@@ -7,7 +7,7 @@
 #define BOX_BRUSH(RelativePath, ...)    FSlateBoxBrush   (FPaths::EngineContentDir() / "Editor/Slate"/ RelativePath + TEXT(".png"), __VA_ARGS__)
 #define BORDER_BRUSH(RelativePath, ...) FSlateBorderBrush(FPaths::EngineContentDir() / "Editor/Slate"/ RelativePath + TEXT(".png"), __VA_ARGS__)
 
-TSharedPtr< FSlateStyleSet > FInsightsStyle::StyleInstance = nullptr;
+TSharedPtr<FSlateStyleSet> FInsightsStyle::StyleInstance = nullptr;
 
 void FInsightsStyle::Initialize()
 {
@@ -31,7 +31,7 @@ FName FInsightsStyle::GetStyleSetName()
 	return StyleSetName;
 }
 
-TSharedRef< FSlateStyleSet > FInsightsStyle::Create()
+TSharedRef<FSlateStyleSet> FInsightsStyle::Create()
 {
 	TSharedRef<FSlateStyleSet> StyleRef = MakeShareable(new FSlateStyleSet(FInsightsStyle::GetStyleSetName()));
 	StyleRef->SetContentRoot(FPaths::EngineContentDir() / TEXT("Editor/Slate"));
@@ -39,9 +39,52 @@ TSharedRef< FSlateStyleSet > FInsightsStyle::Create()
 
 	FSlateStyleSet& Style = StyleRef.Get();
 
-	// Widgets
-	Style.Set("Insights.Background", new BOX_BRUSH("Insights/Background", FMargin(5.f / 12.f)));
-	Style.Set("Insights.BorderPadding", FVector2D(1,0));
+	//////////////////////////////////////////////////
+	// Icons for major components
+
+	Style.Set("StartPage.Icon.Large", new IMAGE_BRUSH("/Icons/icon_tab_Tools_16x", FVector2D(32.0f, 32.0f)));
+	Style.Set("StartPage.Icon.Small", new IMAGE_BRUSH("/Icons/icon_tab_Tools_16x", FVector2D(16.0f, 16.0f)));
+
+	Style.Set("TimingInsights.Icon.Large", new IMAGE_BRUSH("Icons/Profiler/profiler_stats_40x", FVector2D(32.0f, 32.0f)));
+	Style.Set("TimingInsights.Icon.Small", new IMAGE_BRUSH("Icons/Profiler/profiler_stats_40x", FVector2D(16.0f, 16.0f)));
+
+	Style.Set("AssetLoadingInsights.Icon.Large", new IMAGE_BRUSH("Icons/Profiler/profiler_stats_40x", FVector2D(32.0f, 32.0f)));
+	Style.Set("AssetLoadingInsights.Icon.Small", new IMAGE_BRUSH("Icons/Profiler/profiler_stats_40x", FVector2D(16.0f, 16.0f)));
+
+	Style.Set("Toolbar.Icon.Large", new IMAGE_BRUSH("/Icons/icon_tab_Tools_16x", FVector2D(32.0f, 32.0f)));
+	Style.Set("Toolbar.Icon.Small", new IMAGE_BRUSH("/Icons/icon_tab_Tools_16x", FVector2D(16.0f, 16.0f)));
+
+	//////////////////////////////////////////////////
+	// Timing Insights
+
+	Style.Set("FramesTrack.Icon.Large", new IMAGE_BRUSH("Icons/Profiler/profiler_stats_40x", FVector2D(32.0f, 32.0f)));
+	Style.Set("FramesTrack.Icon.Small", new IMAGE_BRUSH("Icons/Profiler/profiler_stats_40x", FVector2D(16.0f, 16.0f)));
+
+	Style.Set("GraphTrack.Icon.Large", new IMAGE_BRUSH("Icons/Profiler/profiler_stats_40x", FVector2D(32.0f, 32.0f)));
+	Style.Set("GraphTrack.Icon.Small", new IMAGE_BRUSH("Icons/Profiler/profiler_stats_40x", FVector2D(16.0f, 16.0f)));
+
+	Style.Set("TimingView.Icon.Large", new IMAGE_BRUSH("Icons/Profiler/profiler_stats_40x", FVector2D(32.0f, 32.0f)));
+	Style.Set("TimingView.Icon.Small", new IMAGE_BRUSH("Icons/Profiler/profiler_stats_40x", FVector2D(16.0f, 16.0f)));
+
+	Style.Set("TimersView.Icon.Large", new IMAGE_BRUSH("Icons/Profiler/Profiler_Data_Capture_40x", FVector2D(32.0f, 32.0f)));
+	Style.Set("TimersView.Icon.Small", new IMAGE_BRUSH("Icons/Profiler/Profiler_Data_Capture_40x", FVector2D(16.0f, 16.0f)));
+
+	Style.Set("StatsCountersView.Icon.Large", new IMAGE_BRUSH("Icons/Profiler/Profiler_Data_Capture_40x", FVector2D(32.0f, 32.0f)));
+	Style.Set("StatsCountersView.Icon.Small", new IMAGE_BRUSH("Icons/Profiler/Profiler_Data_Capture_40x", FVector2D(16.0f, 16.0f)));
+
+	Style.Set("LogView.Icon.Large", new IMAGE_BRUSH("Icons/Profiler/profiler_CopyToClipboard_32x", FVector2D(32.0f, 32.0f)));
+	Style.Set("LogView.Icon.Small", new IMAGE_BRUSH("Icons/Profiler/profiler_CopyToClipboard_32x", FVector2D(16.0f, 16.0f)));
+
+	//////////////////////////////////////////////////
+	// Start Page buttons
+
+	Style.Set("Live.Icon.Large", new IMAGE_BRUSH("Icons/Profiler/Profiler_LoadMultiple_Profiler_40x", FVector2D(32.0f, 32.0f)));
+	Style.Set("Live.Icon.Small", new IMAGE_BRUSH("Icons/Profiler/Profiler_Load_Profiler_40x", FVector2D(16.0f, 16.0f)));
+
+	Style.Set("Load.Icon.Large", new IMAGE_BRUSH("Icons/LV_Load", FVector2D(32.0f, 32.0f)));
+	Style.Set("Load.Icon.Small", new IMAGE_BRUSH("Icons/LV_Load", FVector2D(16.0f, 16.0f)));
+
+	//////////////////////////////////////////////////
 
 	return StyleRef;
 }
diff --git a/Engine/Source/Developer/TraceInsights/Private/Insights/InsightsStyle.h b/Engine/Source/Developer/TraceInsights/Private/Insights/InsightsStyle.h
--- a/Engine/Source/Developer/TraceInsights/Private/Insights/InsightsStyle.h
+++ b/Engine/Source/Developer/TraceInsights/Private/Insights/InsightsStyle.h
@@ -5,7 +5,7 @@
 #include "CoreMinimal.h"
 #include "Styling/SlateStyle.h"
 
-/// Style data for Insights tools
+/** Style data for Insights tools */
 class FInsightsStyle
 {
 public:
@@ -14,6 +14,16 @@ public:
 	static const ISlateStyle& Get();
 	static FName GetStyleSetName();
 
+	static const FLinearColor& GetColor(FName PropertyName, const ANSICHAR* Specifier = NULL)
+	{
+		return StyleInstance->GetColor(PropertyName, Specifier);
+	}
+
+	static const FSlateBrush* GetBrush(FName PropertyName, const ANSICHAR* Specifier = NULL)
+	{
+		return StyleInstance->GetBrush(PropertyName, Specifier);
+	}
+
 private:
 	static TSharedRef<class FSlateStyleSet> Create();
 
diff --git a/Engine/Source/Developer/TraceInsights/Private/Insights/TimingProfilerCommands.cpp b/Engine/Source/Developer/TraceInsights/Private/Insights/TimingProfilerCommands.cpp
--- a/Engine/Source/Developer/TraceInsights/Private/Insights/TimingProfilerCommands.cpp
+++ b/Engine/Source/Developer/TraceInsights/Private/Insights/TimingProfilerCommands.cpp
@@ -54,8 +54,9 @@ void FTimingProfilerCommands::RegisterCommands()
 {
 	UI_COMMAND(ToggleFramesTrackVisibility, "Frames", "Toggles the visibility of the Frames track", EUserInterfaceActionType::ToggleButton, FInputChord(EModifierKey::Control, EKeys::F));
 	UI_COMMAND(ToggleGraphTrackVisibility, "Graph", "Toggles the visibility of the Overview Graph track", EUserInterfaceActionType::ToggleButton, FInputChord(EModifierKey::Control, EKeys::G));
-	UI_COMMAND(ToggleTimingTrackVisibility, "Timing", "Toggles the visibility of the main Timing view", EUserInterfaceActionType::ToggleButton, FInputChord(EModifierKey::Control, EKeys::T));
-	UI_COMMAND(ToggleTimersViewVisibility, "Timers", "Toggles the visibility of the Timers view", EUserInterfaceActionType::ToggleButton, FInputChord(EKeys::T));
+	UI_COMMAND(ToggleTimingViewVisibility, "Timing", "Toggles the visibility of the main Timing view", EUserInterfaceActionType::ToggleButton, FInputChord(EModifierKey::Control, EKeys::T));
+	UI_COMMAND(ToggleTimersViewVisibility, "Timers", "Toggles the visibility of the Timers view", EUserInterfaceActionType::ToggleButton, FInputChord());
+	UI_COMMAND(ToggleStatsCountersViewVisibility, "Stats", "Toggles the visibility of the Stats Counters view", EUserInterfaceActionType::ToggleButton, FInputChord());
 	UI_COMMAND(ToggleLogViewVisibility, "Log", "Toggles the visibility of the Log view", EUserInterfaceActionType::ToggleButton, FInputChord(EKeys::L));
 }
 PRAGMA_ENABLE_OPTIMIZATION
@@ -99,8 +100,9 @@ PRAGMA_ENABLE_OPTIMIZATION
 
 IMPLEMENT_TOGGLE_COMMAND(ToggleFramesTrackVisibility, IsFramesTrackVisible, SetFramesTrackVisible)
 IMPLEMENT_TOGGLE_COMMAND(ToggleGraphTrackVisibility, IsGraphTrackVisible, SetGraphTrackVisible)
-IMPLEMENT_TOGGLE_COMMAND(ToggleTimingTrackVisibility, IsTimingTrackVisible, SetTimingTrackVisible)
+IMPLEMENT_TOGGLE_COMMAND(ToggleTimingViewVisibility, IsTimingViewVisible, SetTimingViewVisible)
 IMPLEMENT_TOGGLE_COMMAND(ToggleTimersViewVisibility, IsTimersViewVisible, SetTimersViewVisible)
+IMPLEMENT_TOGGLE_COMMAND(ToggleStatsCountersViewVisibility, IsStatsCountersViewVisible, SetStatsCountersViewVisible)
 IMPLEMENT_TOGGLE_COMMAND(ToggleLogViewVisibility, IsLogViewVisible, SetLogViewVisible)
 
 #undef IMPLEMENT_TOGGLE_COMMAND
diff --git a/Engine/Source/Developer/TraceInsights/Private/Insights/TimingProfilerCommands.h b/Engine/Source/Developer/TraceInsights/Private/Insights/TimingProfilerCommands.h
--- a/Engine/Source/Developer/TraceInsights/Private/Insights/TimingProfilerCommands.h
+++ b/Engine/Source/Developer/TraceInsights/Private/Insights/TimingProfilerCommands.h
@@ -9,16 +9,17 @@
 class FMenuBuilder;
 
 ////////////////////////////////////////////////////////////////////////////////////////////////////
-/// Class that holds all profiler commands.
-
+/**
+ * Class that holds all profiler commands.
+ */
 class FTimingProfilerCommands
 	: public TCommands<FTimingProfilerCommands>
 {
 public:
-	/// Default constructor.
+	/** Default constructor. */
 	FTimingProfilerCommands();
 
-	/// Initialize commands.
+	/** Initialize commands. */
 	virtual void RegisterCommands() override;
 
 public:
@@ -29,45 +30,52 @@ public:
 	//     const FUIAction <CommandName>_Custom(...) const;
 	//////////////////////////////////////////////////
 
-	/// Toggles visibility for the Frames Track. Global and custom command.
+	/** Toggles visibility for the Frames Track. Global and custom command. */
 	TSharedPtr<FUICommandInfo> ToggleFramesTrackVisibility;
 
-	/// Toggles visibility for the Graph Track. Global and custom command.
+	/** Toggles visibility for the Graph Track. Global and custom command. */
 	TSharedPtr<FUICommandInfo> ToggleGraphTrackVisibility;
 
-	/// Toggles visibility for the Timing Track. Global and custom command.
-	TSharedPtr<FUICommandInfo> ToggleTimingTrackVisibility;
+	/** Toggles visibility for the Timing View. Global and custom command. */
+	TSharedPtr<FUICommandInfo> ToggleTimingViewVisibility;
 
-	/// Toggles visibility for the Timers View. Global and custom command.
+	/** Toggles visibility for the Timers View. Global and custom command. */
 	TSharedPtr<FUICommandInfo> ToggleTimersViewVisibility;
 
-	/// Toggles visibility for the Log View. Global and custom command.
+	/** Toggles visibility for the Stats Counters View. Global and custom command. */
+	TSharedPtr<FUICommandInfo> ToggleStatsCountersViewVisibility;
+
+	/** Toggles visibility for the Log View. Global and custom command. */
 	TSharedPtr<FUICommandInfo> ToggleLogViewVisibility;
 };
 
 ////////////////////////////////////////////////////////////////////////////////////////////////////
-
+/**
+ * Menu builder. Helper class for adding a customized menu entry using the global UI command info.
+ */
 class FTimingProfilerMenuBuilder
 {
 public:
-	/// Helper method for adding a customized menu entry using the global UI command info.
-	/// FUICommandInfo cannot be executed with custom parameters, so we need to create a custom FUIAction,
-	/// but sometime we have global and local version for the UI command, so reuse data from the global UI command info.
-	/// Ex:
-	///     SessionInstance_ToggleCapture          - Global version will toggle capture process for all active session instances
-	///     SessionInstance_ToggleCapture_OneParam - Local version will toggle capture process only for the specified session instance
-	///
-	/// @param MenuBuilder The menu to add items to
-	/// @param FUICommandInfo A shared pointer to the UI command info
-	/// @param UIAction Customized version of the UI command info stored in an UI action
-	///
+	/**
+	 * Helper method for adding a customized menu entry using the global UI command info.
+	 * FUICommandInfo cannot be executed with custom parameters, so we need to create a custom FUIAction,
+	 * but sometime we have global and local version for the UI command, so reuse data from the global UI command info.
+	 * Ex:
+	 *     SessionInstance_ToggleCapture          - Global version will toggle capture process for all active session instances
+	 *     SessionInstance_ToggleCapture_OneParam - Local version will toggle capture process only for the specified session instance
+	 *
+	 * @param MenuBuilder The menu to add items to
+	 * @param FUICommandInfo A shared pointer to the UI command info
+	 * @param UIAction Customized version of the UI command info stored in an UI action
+	 */
 	static void AddMenuEntry(FMenuBuilder& MenuBuilder, const TSharedPtr<FUICommandInfo>& UICommandInfo, const FUIAction& UIAction);
 };
 
 ////////////////////////////////////////////////////////////////////////////////////////////////////
-/// Class that provides helper functions for the commands to avoid cluttering profiler manager with many small functions.
-/// Can't contain any variables. Directly operates on the profiler manager instance.
-
+/**
+ * Class that provides helper functions for the commands to avoid cluttering profiler manager with many small functions.
+ * Can't contain any variables. Directly operates on the profiler manager instance.
+ */
 class FTimingProfilerActionManager
 {
 	friend class FTimingProfilerManager;
@@ -92,8 +100,9 @@ protected:\
 
 	DECLARE_TOGGLE_COMMAND(ToggleFramesTrackVisibility)
 	DECLARE_TOGGLE_COMMAND(ToggleGraphTrackVisibility)
-	DECLARE_TOGGLE_COMMAND(ToggleTimingTrackVisibility)
+	DECLARE_TOGGLE_COMMAND(ToggleTimingViewVisibility)
 	DECLARE_TOGGLE_COMMAND(ToggleTimersViewVisibility)
+	DECLARE_TOGGLE_COMMAND(ToggleStatsCountersViewVisibility)
 	DECLARE_TOGGLE_COMMAND(ToggleLogViewVisibility)
 #undef DECLARE_TOGGLE_COMMAND
 
@@ -101,15 +110,15 @@ protected:\
 	// OpenSettings
 
 public:
-	void Map_OpenSettings_Global(); ///< Maps UI command info OpenSettings with the specified UI command list.
-	const FUIAction OpenSettings_Custom() const; ///< UI action for OpenSettings command.
+	void Map_OpenSettings_Global(); /**< Maps UI command info OpenSettings with the specified UI command list. */
+	const FUIAction OpenSettings_Custom() const; /**< UI action for OpenSettings command. */
 protected:
-	void OpenSettings_Execute(); ///< Handles FExecuteAction for OpenSettings.
-	bool OpenSettings_CanExecute() const; ///< Handles FCanExecuteAction for OpenSettings.
+	void OpenSettings_Execute(); /**< Handles FExecuteAction for OpenSettings. */
+	bool OpenSettings_CanExecute() const; /**< Handles FCanExecuteAction for OpenSettings. */
 
 	//////////////////////////////////////////////////
 
 protected:
-	/// Reference to the global instance of the profiler manager.
+	/** Reference to the global instance of the profiler manager. */
 	class FTimingProfilerManager* This;
 };
diff --git a/Engine/Source/Developer/TraceInsights/Private/Insights/TimingProfilerManager.cpp b/Engine/Source/Developer/TraceInsights/Private/Insights/TimingProfilerManager.cpp
--- a/Engine/Source/Developer/TraceInsights/Private/Insights/TimingProfilerManager.cpp
+++ b/Engine/Source/Developer/TraceInsights/Private/Insights/TimingProfilerManager.cpp
@@ -14,6 +14,7 @@
 #include "Insights/Widgets/SFrameTrack.h"
 #include "Insights/Widgets/SGraphTrack.h"
 #include "Insights/Widgets/SLogView.h"
+#include "Insights/Widgets/SStatsView.h"
 #include "Insights/Widgets/STimersView.h"
 #include "Insights/Widgets/STimingView.h"
 #include "Insights/Widgets/STimingProfilerWindow.h"
@@ -39,8 +40,9 @@ FTimingProfilerManager::FTimingProfilerManager(TSharedRef<FUICommandList> InComm
 	, ProfilerWindow(nullptr)
 	, bIsFramesTrackVisible(true)
 	, bIsGraphTrackVisible(false)
-	, bIsTimingTrackVisible(true)
+	, bIsTimingViewVisible(true)
 	, bIsTimersViewVisible(true)
+	, bIsStatsCountersViewVisible(true)
 	, bIsLogViewVisible(true)
 {
 }
@@ -63,8 +65,9 @@ void FTimingProfilerManager::BindCommands()
 {
 	ActionManager.Map_ToggleFramesTrackVisibility_Global();
 	ActionManager.Map_ToggleGraphTrackVisibility_Global();
-	ActionManager.Map_ToggleTimingTrackVisibility_Global();
+	ActionManager.Map_ToggleTimingViewVisibility_Global();
 	ActionManager.Map_ToggleTimersViewVisibility_Global();
+	ActionManager.Map_ToggleStatsCountersViewVisibility_Global();
 	ActionManager.Map_ToggleLogViewVisibility_Global();
 }
 
@@ -122,11 +125,75 @@ void FTimingProfilerManager::OnSessionChanged()
 	TSharedPtr<STimingProfilerWindow> Wnd = GetProfilerWindow();
 	if (Wnd.IsValid())
 	{
-		Wnd->FrameTrack->Reset();
-		Wnd->GraphTrack->Reset();
-		Wnd->TimingView->Reset();
-		Wnd->TimersView->RebuildTree();
-		Wnd->LogView->Reset();
+		if (Wnd->FrameTrack)
+		{
+			Wnd->FrameTrack->Reset();
+		}
+		if (Wnd->GraphTrack)
+		{
+			Wnd->GraphTrack->Reset();
+		}
+		if (Wnd->TimingView)
+		{
+			Wnd->TimingView->Reset();
+		}
+		if (Wnd->TimersView)
+		{
+			Wnd->TimersView->RebuildTree();
+		}
+		if (Wnd->StatsView)
+		{
+			Wnd->StatsView->RebuildTree();
+		}
+		if (Wnd->LogView)
+		{
+			Wnd->LogView->Reset();
+		}
+	}
+}
+
+////////////////////////////////////////////////////////////////////////////////////////////////////
+
+void FTimingProfilerManager::SetFramesTrackVisible(const bool bFramesTrackVisibleState)
+{
+	bool bWasFramesTrackVisible = bIsFramesTrackVisible;
+
+	bIsFramesTrackVisible = bFramesTrackVisibleState;
+
+	TSharedPtr<STimingProfilerWindow> Wnd = GetProfilerWindow();
+	if (Wnd.IsValid())
+	{
+		Wnd->ShowHideTab(FTimingProfilerTabs::FramesTrackID, bIsFramesTrackVisible);
+	}
+}
+
+////////////////////////////////////////////////////////////////////////////////////////////////////
+
+void FTimingProfilerManager::SetGraphTrackVisible(const bool bGraphTrackVisibleState)
+{
+	bool bWasGraphTrackVisible = bIsGraphTrackVisible;
+
+	bIsGraphTrackVisible = bGraphTrackVisibleState;
+
+	TSharedPtr<STimingProfilerWindow> Wnd = GetProfilerWindow();
+	if (Wnd.IsValid())
+	{
+		Wnd->ShowHideTab(FTimingProfilerTabs::GraphTrackID, bIsGraphTrackVisible);
+	}
+}
+
+////////////////////////////////////////////////////////////////////////////////////////////////////
+
+void FTimingProfilerManager::SetTimingViewVisible(const bool bTimingViewVisibleState)
+{
+	bool bWasTimingViewVisible = bIsTimingViewVisible;
+
+	bIsTimingViewVisible = bTimingViewVisibleState;
+
+	TSharedPtr<STimingProfilerWindow> Wnd = GetProfilerWindow();
+	if (Wnd.IsValid())
+	{
+		Wnd->ShowHideTab(FTimingProfilerTabs::TimingViewID, bIsTimingViewVisible);
 	}
 }
 
@@ -138,13 +205,40 @@ void FTimingProfilerManager::SetTimersViewVisible(const bool bTimersViewVisibleS
 
 	bIsTimersViewVisible = bTimersViewVisibleState;
 
-	if (bIsTimersViewVisible && !bWasTimersViewVisible && FInsightsManager::Get()->GetSession().IsValid())
+	TSharedPtr<STimingProfilerWindow> Wnd = GetProfilerWindow();
+	if (Wnd.IsValid())
 	{
-		TSharedPtr<STimingProfilerWindow> Wnd = GetProfilerWindow();
-		if (Wnd.IsValid())
-		{
-			Wnd->TimersView->RebuildTree();
-		}
+		Wnd->ShowHideTab(FTimingProfilerTabs::TimersID, bIsTimersViewVisible);
+	}
+}
+
+////////////////////////////////////////////////////////////////////////////////////////////////////
+
+void FTimingProfilerManager::SetStatsCountersViewVisible(const bool bStatsCountersViewVisibleState)
+{
+	bool bWasStatsCountersViewVisible = bIsStatsCountersViewVisible;
+
+	bIsStatsCountersViewVisible = bStatsCountersViewVisibleState;
+
+	TSharedPtr<STimingProfilerWindow> Wnd = GetProfilerWindow();
+	if (Wnd.IsValid())
+	{
+		Wnd->ShowHideTab(FTimingProfilerTabs::StatsCountersID, bIsStatsCountersViewVisible);
+	}
+}
+
+////////////////////////////////////////////////////////////////////////////////////////////////////
+
+void FTimingProfilerManager::SetLogViewVisible(const bool bLogViewVisibleState)
+{
+	bool bWasLogViewVisible = bIsLogViewVisible;
+
+	bIsLogViewVisible = bLogViewVisibleState;
+
+	TSharedPtr<STimingProfilerWindow> Wnd = GetProfilerWindow();
+	if (Wnd.IsValid())
+	{
+		Wnd->ShowHideTab(FTimingProfilerTabs::LogViewID, bIsLogViewVisible);
 	}
 }
 
diff --git a/Engine/Source/Developer/TraceInsights/Private/Insights/TimingProfilerManager.h b/Engine/Source/Developer/TraceInsights/Private/Insights/TimingProfilerManager.h
--- a/Engine/Source/Developer/TraceInsights/Private/Insights/TimingProfilerManager.h
+++ b/Engine/Source/Developer/TraceInsights/Private/Insights/TimingProfilerManager.h
@@ -20,21 +20,22 @@ namespace Trace
 }
 
 ////////////////////////////////////////////////////////////////////////////////////////////////////
-/// This class manages the Timing Profiler state and settings.
-
+/**
+ * This class manages the Timing Profiler state and settings.
+ */
 class FTimingProfilerManager
 	: public TSharedFromThis<FTimingProfilerManager>
 {
 	friend class FTimingProfilerActionManager;
 
 public:
-	/// Creates the Timing Profiler manager, only one instance can exist.
+	/** Creates the Timing Profiler manager, only one instance can exist. */
 	FTimingProfilerManager(TSharedRef<FUICommandList> InCommandList);
 
-	/// Virtual destructor.
+	/** Virtual destructor. */
 	virtual ~FTimingProfilerManager();
 
-	/// Creates an instance of the profiler manager.
+	/** Creates an instance of the profiler manager. */
 	static TSharedPtr<FTimingProfilerManager> Initialize()
 	{
 		if (FTimingProfilerManager::Instance.IsValid())
@@ -48,34 +49,36 @@ public:
 		return FTimingProfilerManager::Instance;
 	}
 
-	/// Shutdowns the Timing Profiler manager.
+	/** Shutdowns the Timing Profiler manager. */
 	void Shutdown()
 	{
 		FTimingProfilerManager::Instance.Reset();
 	}
 
 protected:
-	/// Finishes initialization of the profiler manager
+	/** Finishes initialization of the profiler manager */
 	void PostConstructor();
 
-	/// Binds our UI commands to delegates.
+	/** Binds our UI commands to delegates. */
 	void BindCommands();
 
 public:
-	/// @return the global instance of the Timing Profiler manager (FTimingProfilerManager).
-	/// This is an internal singleton and cannot be used outside ProfilerModule.
-	/// For external use:
-	///     IProfilerModule& ProfilerModule = FModuleManager::Get().LoadModuleChecked<IProfilerModule>("Profiler");
-	///     ProfilerModule.GetProfilerManager();
+	/**
+	 * @return the global instance of the Timing Profiler manager (FTimingProfilerManager).
+	 * This is an internal singleton and cannot be used outside ProfilerModule.
+	 * For external use:
+	 *     IProfilerModule& ProfilerModule = FModuleManager::Get().LoadModuleChecked<IProfilerModule>("Profiler");
+	 *     ProfilerModule.GetProfilerManager();
+	 */
 	static TSharedPtr<FTimingProfilerManager> Get();
 
-	/// @returns UI command list for the Timing Profiler manager.
+	/** @returns UI command list for the Timing Profiler manager. */
 	const TSharedRef<FUICommandList> GetCommandList() const;
 
-	/// @return an instance of the Timing Profiler commands.
+	/** @return an instance of the Timing Profiler commands. */
 	static const FTimingProfilerCommands& GetCommands();
 
-	/// @return an instance of the Timing Profiler action manager.
+	/** @return an instance of the Timing Profiler action manager. */
 	static FTimingProfilerActionManager& GetActionManager();
 
 	void AssignProfilerWindow(const TSharedRef<STimingProfilerWindow>& InProfilerWindow)
@@ -83,8 +86,10 @@ public:
 		ProfilerWindow = InProfilerWindow;
 	}
 
-	/// Converts profiler window weak pointer to a shared pointer and returns it.
-	/// Make sure the returned pointer is valid before trying to dereference it.
+	/**
+	 * Converts profiler window weak pointer to a shared pointer and returns it.
+	 * Make sure the returned pointer is valid before trying to dereference it.
+	 */
 	TSharedPtr<class STimingProfilerWindow> GetProfilerWindow() const
 	{
 		return ProfilerWindow.Pin();
@@ -93,65 +98,72 @@ public:
 	////////////////////////////////////////////////////////////////////////////////////////////////////
 	// Getters and setters used by Toggle Commands.
 
-	/// @return true, if Frames Track is visible
+	/** @return true, if Frames Track is visible */
 	const bool IsFramesTrackVisible() const { return bIsFramesTrackVisible; }
-	void SetFramesTrackVisible(const bool bFramesTrackVisibleState) { bIsFramesTrackVisible = bFramesTrackVisibleState; }
+	void SetFramesTrackVisible(const bool bFramesTrackVisibleState);
 
-	/// @return true, if Graph Track is visible
+	/** @return true, if Graph Track is visible */
 	const bool IsGraphTrackVisible() const { return bIsGraphTrackVisible; }
-	void SetGraphTrackVisible(const bool bGraphTrackVisibleState) { bIsGraphTrackVisible = bGraphTrackVisibleState; }
+	void SetGraphTrackVisible(const bool bGraphTrackVisibleState);
 
-	/// @return true, if Timing Track is visible
-	const bool IsTimingTrackVisible() const { return bIsTimingTrackVisible; }
-	void SetTimingTrackVisible(const bool bTimingTrackVisibleState) { bIsTimingTrackVisible = bTimingTrackVisibleState; }
+	/** @return true, if Timing View is visible */
+	const bool IsTimingViewVisible() const { return bIsTimingViewVisible; }
+	void SetTimingViewVisible(const bool bTimingViewVisibleState);
 
-	/// @return true, if Timers View is visible
+	/** @return true, if Timers View is visible */
 	const bool IsTimersViewVisible() const { return bIsTimersViewVisible; }
 	void SetTimersViewVisible(const bool bTimersViewVisibleState);
 
-	/// @return true, if Log View is visible
+	/** @return true, if Stats Counters View is visible */
+	const bool IsStatsCountersViewVisible() const { return bIsStatsCountersViewVisible; }
+	void SetStatsCountersViewVisible(const bool bStatsCountersViewVisibleState);
+
+	/** @return true, if Log View is visible */
 	const bool IsLogViewVisible() const { return bIsLogViewVisible; }
-	void SetLogViewVisible(const bool bLogViewVisibleState) { bIsLogViewVisible = bLogViewVisibleState; }
+	void SetLogViewVisible(const bool bLogViewVisibleState);
 
 	////////////////////////////////////////////////////////////////////////////////////////////////////
 
 	void OnSessionChanged();
 
 protected:
-	/// Updates this manager, done through FCoreTicker.
+	/** Updates this manager, done through FCoreTicker. */
 	bool Tick(float DeltaTime);
 
 protected:
-	/// The delegate to be invoked when this manager ticks.
+	/** The delegate to be invoked when this manager ticks. */
 	FTickerDelegate OnTick;
 
-	/// Handle to the registered OnTick.
+	/** Handle to the registered OnTick. */
 	FDelegateHandle OnTickHandle;
 
-	/// List of UI commands for this manager. This will be filled by this and corresponding classes.
+	/** List of UI commands for this manager. This will be filled by this and corresponding classes. */
 	TSharedRef<FUICommandList> CommandList;
 
-	/// An instance of the Timing Manager action manager.
+	/** An instance of the Timing Manager action manager. */
 	FTimingProfilerActionManager ActionManager;
 
-	/// A weak pointer to the Timing Profiler window.
+	/** A weak pointer to the Timing Profiler window. */
 	TWeakPtr<class STimingProfilerWindow> ProfilerWindow;
 
-	/// If the Frames Track is visible or hidden.
+	/** If the Frames Track is visible or hidden. */
 	bool bIsFramesTrackVisible;
 
-	/// If the Graph Track is visible or hidden.
+	/** If the Graph Track is visible or hidden. */
 	bool bIsGraphTrackVisible;
 
-	/// If the Timing Track is visible or hidden.
-	bool bIsTimingTrackVisible;
+	/** If the Timing View is visible or hidden. */
+	bool bIsTimingViewVisible;
 
-	/// If the Timers View is visible or hidden.
+	/** If the Timers View is visible or hidden. */
 	bool bIsTimersViewVisible;
 
-	/// If the Log View is visible or hidden.
+	/** If the Stats Counters View is visible or hidden. */
+	bool bIsStatsCountersViewVisible;
+
+	/** If the Log View is visible or hidden. */
 	bool bIsLogViewVisible;
 
-	/// A shared pointer to the global instance of the profiler manager.
+	/** A shared pointer to the global instance of the profiler manager. */
 	static TSharedPtr<FTimingProfilerManager> Instance;
 };
diff --git a/Engine/Source/Developer/TraceInsights/Private/Insights/TraceInsightsModule.cpp b/Engine/Source/Developer/TraceInsights/Private/Insights/TraceInsightsModule.cpp
--- a/Engine/Source/Developer/TraceInsights/Private/Insights/TraceInsightsModule.cpp
+++ b/Engine/Source/Developer/TraceInsights/Private/Insights/TraceInsightsModule.cpp
@@ -14,24 +14,25 @@
 //#include "WorkspaceMenuStructureModule.h"
 
 // Insights
-#include "Insights/IUnrealInsightsModule.h"
 #include "Insights/InsightsManager.h"
+#include "Insights/InsightsStyle.h"
 #include "Insights/IoProfilerManager.h"
+#include "Insights/IUnrealInsightsModule.h"
 #include "Insights/TimingProfilerManager.h"
 #include "Insights/Widgets/SIoProfilerWindow.h"
+#include "Insights/Widgets/SStartPageWindow.h"
 #include "Insights/Widgets/STimingProfilerWindow.h"
 
 ////////////////////////////////////////////////////////////////////////////////////////////////////
 
-static const FName StartPageTabName("StartPage");
-static const FName TimingProfilerTabName("TimingProfiler");
-static const FName IoProfilerTabName("IoProfiler");
-//static const FName MemoryProfilerTabName("MemoryProfiler");
+static const FName StartPageTabName(TEXT("StartPage"));
+static const FName TimingProfilerTabName(TEXT("TimingProfiler"));
+static const FName IoProfilerTabName(TEXT("IoProfiler"));
 
 ////////////////////////////////////////////////////////////////////////////////////////////////////
-///
-/// Implements the Trace Insights module (TimingProfiler, IoProfiler, etc.).
-///
+/**
+ * Implements the Trace Insights module (TimingProfiler, IoProfiler, etc.).
+ */
 class FTraceInsightsModule : public IUnrealInsightsModule
 {
 public:
@@ -47,22 +48,22 @@ public:
 	virtual void OnLayoutRestored(TSharedPtr<FTabManager> TabManager) override;
 
 protected:
-	/// Callback called when a major tab is closed.
+	/** Callback called when a major tab is closed. */
 	void OnTabBeingClosed(TSharedRef<SDockTab> TabBeingClosed);
 
-	/// Callback called when the Timing Profiler major tab is closed.
+	/** Callback called when the Timing Profiler major tab is closed. */
 	void OnTimingProfilerTabBeingClosed(TSharedRef<SDockTab> TabBeingClosed);
 
-	/// Callback called when the Io Profiler major tab is closed.
+	/** Callback called when the Io Profiler major tab is closed. */
 	void OnIoProfilerTabBeingClosed(TSharedRef<SDockTab> TabBeingClosed);
 
-	/// Start Page
+	/** Start Page */
 	TSharedRef<SDockTab> SpawnStartPageTab(const FSpawnTabArgs& Args);
 
-	/// Timing Profiler
+	/** Timing Profiler */
 	TSharedRef<SDockTab> SpawnTimingProfilerTab(const FSpawnTabArgs& Args);
 
-	/// I/O Profiler
+	/** I/O Profiler */
 	TSharedRef<SDockTab> SpawnIoProfilerTab(const FSpawnTabArgs& Args);
 
 protected:
@@ -85,6 +86,8 @@ void FTraceInsightsModule::StartupModule()
 	TraceSessionService = TraceServicesModule.GetSessionService();
 	TraceSessionService->StartRecorderServer();
 
+	FInsightsStyle::Initialize();
+
 	FInsightsManager::Initialize(TraceAnalysisService.ToSharedRef(), TraceSessionService.ToSharedRef());
 	FTimingProfilerManager::Initialize();
 	FIoProfilerManager::Initialize();
@@ -93,15 +96,15 @@ void FTraceInsightsModule::StartupModule()
 	auto& StartPageTabSpawnerEntry = FGlobalTabmanager::Get()->RegisterTabSpawner(StartPageTabName,
 		FOnSpawnTab::CreateRaw(this, &FTraceInsightsModule::SpawnStartPageTab))
 		.SetDisplayName(NSLOCTEXT("FTraceInsightsModule", "StartPageTabTitle", "Start Page"))
-		.SetTooltipText(NSLOCTEXT("FTraceInsightsModule", "StartPageTooltipText", "Open the Unreal Insights start page."));
-		//.SetIcon(FSlateIcon(FEditorStyle::GetStyleSetName(), "Profiler.Tab")); //im:TODO
+		.SetTooltipText(NSLOCTEXT("FTraceInsightsModule", "StartPageTooltipText", "Open the start page for Unreal Insights."))
+		.SetIcon(FSlateIcon(FInsightsStyle::GetStyleSetName(), "StartPage.Icon.Small"));
 
-	// Register tab spawner for the Timing Profiler.
+	// Register tab spawner for the Timing Insights.
 	auto& TimingProfilerTabSpawnerEntry = FGlobalTabmanager::Get()->RegisterTabSpawner(TimingProfilerTabName,
 		FOnSpawnTab::CreateRaw(this, &FTraceInsightsModule::SpawnTimingProfilerTab))
-		.SetDisplayName(NSLOCTEXT("FTraceInsightsModule", "TimingProfilerTabTitle", "Timing Profiler"))
-		.SetTooltipText(NSLOCTEXT("FTraceInsightsModule", "TimingProfilerTooltipText", "Open the Timing Profiler tab."));
-		//.SetIcon(FSlateIcon(FEditorStyle::GetStyleSetName(), "Profiler.Tab")); //im:TODO
+		.SetDisplayName(NSLOCTEXT("FTraceInsightsModule", "TimingProfilerTabTitle", "Timing Insights"))
+		.SetTooltipText(NSLOCTEXT("FTraceInsightsModule", "TimingProfilerTooltipText", "Open the Timing Insights tab."))
+		.SetIcon(FSlateIcon(FInsightsStyle::GetStyleSetName(), "TimingInsights.Icon.Small"));
 
 //#if WITH_EDITOR
 //	TimingProfilerTabSpawnerEntry.SetGroup(WorkspaceMenu::GetMenuStructure().GetDeveloperToolsMiscCategory());
@@ -109,12 +112,12 @@ void FTraceInsightsModule::StartupModule()
 //	TimingProfilerTabSpawnerEntry.SetGroup(WorkspaceMenu::GetMenuStructure().GetToolsCategory());
 //#endif
 
-	// Register tab spawner for the I/O Profiler.
+	// Register tab spawner for the Asset Loading Insights.
 	auto& IoProfilerTabSpawnerEntry = FGlobalTabmanager::Get()->RegisterTabSpawner(IoProfilerTabName,
 		FOnSpawnTab::CreateRaw(this, &FTraceInsightsModule::SpawnIoProfilerTab))
-		.SetDisplayName(NSLOCTEXT("FTraceInsightsModule", "IoProfilerTabTitle", "I/O Profiler"))
-		.SetTooltipText(NSLOCTEXT("FTraceInsightsModule", "IoProfilerTooltipText", "Open the Timing Profiler tab."));
-		//.SetIcon(FSlateIcon(FEditorStyle::GetStyleSetName(), "Profiler.Tab")); //im:TODO
+		.SetDisplayName(NSLOCTEXT("FTraceInsightsModule", "IoProfilerTabTitle", "Asset Loading Insights"))
+		.SetTooltipText(NSLOCTEXT("FTraceInsightsModule", "IoProfilerTooltipText", "Open the Asset Loading Insights tab."))
+		.SetIcon(FSlateIcon(FInsightsStyle::GetStyleSetName(), "AssetLoadingInsights.Icon.Small"));
 
 //#if WITH_EDITOR
 //	IoProfilerTabSpawnerEntry.SetGroup(WorkspaceMenu::GetMenuStructure().GetDeveloperToolsMiscCategory());
@@ -184,14 +187,11 @@ void FTraceInsightsModule::OnLayoutRestored(TSharedPtr<FTabManager> TabManager)
 TSharedRef<SDockTab> FTraceInsightsModule::SpawnStartPageTab(const FSpawnTabArgs& Args)
 {
 	const TSharedRef<SDockTab> DockTab = SNew(SDockTab)
-		//.Icon(FEditorStyle::GetBrush("Launcher.TabIcon")) //TODO
 		.TabRole(ETabRole::NomadTab);
 
 	// Create the Start Page widget.
-	TSharedRef<STextBlock> ProfilerWindow = SNew(STextBlock)
-		.Text(NSLOCTEXT("FTraceInsightsModule", "StartPageContent", "TODO"))
-		.ColorAndOpacity(FLinearColor(0.0f, 0.0f, 0.0f, 1.0f));
-	DockTab->SetContent(ProfilerWindow);
+	TSharedRef<SStartPageWindow> Window = SNew(SStartPageWindow);
+	DockTab->SetContent(Window);
 
 	return DockTab;
 }
@@ -201,16 +201,15 @@ TSharedRef<SDockTab> FTraceInsightsModule::SpawnStartPageTab(const FSpawnTabArgs
 TSharedRef<SDockTab> FTraceInsightsModule::SpawnTimingProfilerTab(const FSpawnTabArgs& Args)
 {
 	const TSharedRef<SDockTab> DockTab = SNew(SDockTab)
-		//.Icon(FEditorStyle::GetBrush("Profiler.Tab")) //TODO
 		.TabRole(ETabRole::NomadTab);
 
 	// Register OnTabClosed to handle Timing profiler manager shutdown.
 	DockTab->SetOnTabClosed(SDockTab::FOnTabClosedCallback::CreateRaw(this, &FTraceInsightsModule::OnTimingProfilerTabBeingClosed));
 
 	// Create the STimingProfilerWindow widget.
-	TSharedRef<STimingProfilerWindow> ProfilerWindow = SNew(STimingProfilerWindow);
-	FTimingProfilerManager::Get()->AssignProfilerWindow(ProfilerWindow);
-	DockTab->SetContent(ProfilerWindow);
+	TSharedRef<STimingProfilerWindow> Window = SNew(STimingProfilerWindow, DockTab, Args.GetOwnerWindow());
+	FTimingProfilerManager::Get()->AssignProfilerWindow(Window);
+	DockTab->SetContent(Window);
 
 	return DockTab;
 }
@@ -231,16 +230,15 @@ void FTraceInsightsModule::OnTimingProfilerTabBeingClosed(TSharedRef<SDockTab> T
 TSharedRef<SDockTab> FTraceInsightsModule::SpawnIoProfilerTab(const FSpawnTabArgs& Args)
 {
 	const TSharedRef<SDockTab> DockTab = SNew(SDockTab)
-		//.Icon(FEditorStyle::GetBrush("Profiler.Tab")) //TODO
 		.TabRole(ETabRole::NomadTab);
 
 	// Register OnTabClosed to handle I/O profiler manager shutdown.
 	DockTab->SetOnTabClosed(SDockTab::FOnTabClosedCallback::CreateRaw(this, &FTraceInsightsModule::OnIoProfilerTabBeingClosed));
 
 	// Create the SIoProfilerWindow widget.
-	TSharedRef<SIoProfilerWindow> ProfilerWindow = SNew(SIoProfilerWindow);
-	FIoProfilerManager::Get()->AssignProfilerWindow(ProfilerWindow);
-	DockTab->SetContent(ProfilerWindow);
+	TSharedRef<SIoProfilerWindow> Window = SNew(SIoProfilerWindow);
+	FIoProfilerManager::Get()->AssignProfilerWindow(Window);
+	DockTab->SetContent(Window);
 
 	return DockTab;
 }
diff --git a/Engine/Source/Developer/TraceInsights/Private/Insights/ViewModels/BaseTimingTrack.h b/Engine/Source/Developer/TraceInsights/Private/Insights/ViewModels/BaseTimingTrack.h
--- a/Engine/Source/Developer/TraceInsights/Private/Insights/ViewModels/BaseTimingTrack.h
+++ b/Engine/Source/Developer/TraceInsights/Private/Insights/ViewModels/BaseTimingTrack.h
@@ -3,11 +3,26 @@
 #pragma once
 
 #include "CoreMinimal.h"
+#include "Misc/EnumClassFlags.h"
 
 ////////////////////////////////////////////////////////////////////////////////////////////////////
 
 class FTimingTrackViewport;
 
+////////////////////////////////////////////////////////////////////////////////////////////////////
+
+enum class ETimingTrackFlags : uint32
+{
+	None = 0,
+	IsVisible = (1 << 0),
+	IsSelected = (1 << 1),
+	IsHovered = (1 << 2),
+	IsHeaderHovered = (1 << 3),
+};
+ENUM_CLASS_FLAGS(ETimingTrackFlags);
+
+////////////////////////////////////////////////////////////////////////////////////////////////////
+
 class FBaseTimingTrack
 {
 	friend class FTimingViewDrawHelper;
@@ -15,11 +30,9 @@ class FBaseTimingTrack
 protected:
 	FBaseTimingTrack(uint64 InId)
 		: Id(InId)
-		, Y(0.0f)
-		, H(0.0f)
-		, bIsVisible(true)
-		, bIsSelected(false)
-		, bIsHovered(false)
+		, PosY(0.0f)
+		, Height(0.0f)
+		, Flags(ETimingTrackFlags::IsVisible)
 	{}
 
 	virtual ~FBaseTimingTrack()
@@ -29,51 +42,49 @@ public:
 	virtual void Reset()
 	{
 		Id = 0;
-		Y = 0.0f;
-		H = 0.0f;
-		bIsVisible = true;
-		bIsSelected = false;
-		bIsHovered = false;
+		PosY = 0.0f;
+		Height = 0.0f;
+		Flags = ETimingTrackFlags::IsVisible;
 	}
 
 	uint64 GetId() const { return Id; }
 
-	float GetPosY() const { return Y; }
-	virtual void SetPosY(float InPosY) { Y = InPosY; }
+	float GetPosY() const { return PosY; }
+	virtual void SetPosY(float InPosY) { PosY = InPosY; }
 
-	float GetHeight() const { return H; }
-	virtual void SetHeight(float InHeight) { H = InHeight; }
+	float GetHeight() const { return Height; }
+	virtual void SetHeight(float InHeight) { Height = InHeight; }
 
-	bool IsVisible() const { return bIsVisible; }
-	void Show() { SetVisibilityFlag(true); }
-	void Hide() { SetVisibilityFlag(false); }
-	void ToggleVisibility() { SetVisibilityFlag(!IsVisible()); }
-	virtual void SetVisibilityFlag(bool bInIsVisible) { bIsVisible = bInIsVisible; }
+	bool IsVisible() const { return EnumHasAnyFlags(Flags, ETimingTrackFlags::IsVisible); }
+	void Show() { Flags |= ETimingTrackFlags::IsVisible; OnVisibilityChanged(); }
+	void Hide() { Flags &= ~ETimingTrackFlags::IsVisible; OnVisibilityChanged(); }
+	void ToggleVisibility() { Flags ^= ETimingTrackFlags::IsVisible; OnVisibilityChanged(); }
+	void SetVisibilityFlag(bool bIsVisible) { bIsVisible ? Show() : Hide(); }
+	virtual void OnVisibilityChanged() {}
 
-	bool IsSelected() const { return bIsSelected; }
-	void Select() { SetSelectedFlag(true); }
-	void Unselect() { SetSelectedFlag(false); }
-	virtual void SetSelectedFlag(bool bInIsSelected) { bIsSelected = bInIsSelected; }
+	bool IsSelected() const { return EnumHasAnyFlags(Flags, ETimingTrackFlags::IsSelected); }
+	void Select() { Flags |= ETimingTrackFlags::IsSelected; OnSelectedFlagChanged(); }
+	void Unselect() { Flags &= ~ETimingTrackFlags::IsSelected; OnSelectedFlagChanged(); }
+	void ToggleSelectedFlag() { Flags ^= ETimingTrackFlags::IsSelected; OnSelectedFlagChanged(); }
+	void SetSelectedFlag(bool bIsSelected) { bIsSelected ? Select() : Unselect(); }
+	virtual void OnSelectedFlagChanged() {}
 
-	bool IsHovered() const { return bIsHovered; }
-	void SetHoveredState(bool bInIsHovered) { bIsHovered = bInIsHovered; }
-	bool IsHeaderHovered() const { return bIsHovered && bIsHeaderHovered; }
-	void SetHeaderHoveredState(bool bInIsHeaderHovered) { bIsHeaderHovered = bInIsHeaderHovered; }
-	virtual void UpdateHoveredState(float MX, float MY, const FTimingTrackViewport& Viewport) = 0;
+	bool IsHovered() const { return EnumHasAnyFlags(Flags, ETimingTrackFlags::IsHovered); }
+	void SetHoveredState(bool bIsHovered) { bIsHovered ? Flags |= ETimingTrackFlags::IsHovered : Flags &= ~ETimingTrackFlags::IsHovered; }
+	bool IsHeaderHovered() const { return EnumHasAllFlags(Flags, ETimingTrackFlags::IsHovered | ETimingTrackFlags::IsHeaderHovered); }
+	void SetHeaderHoveredState(bool bIsHeaderHovered) { bIsHeaderHovered ? Flags |= ETimingTrackFlags::IsHeaderHovered : Flags &= ~ETimingTrackFlags::IsHeaderHovered; }
+	virtual void UpdateHoveredState(float MouseX, float MouseY, const FTimingTrackViewport& Viewport) = 0;
 
-	virtual void Tick(const double InCurrentTime, const float InDeltaTime) {}
-	virtual void Update(const FTimingTrackViewport& InViewport) {}
+	virtual void Tick(const double CurrentTime, const float DeltaTime) {}
+	virtual void Update(const FTimingTrackViewport& Viewport) {}
 
-protected:
+private:
 	uint64 Id;
 
-	float Y; // y position, in Slate units
-	float H; // height, in Slate units
+	float PosY; // y position, in Slate units
+	float Height; // height, in Slate units
 
-	bool bIsVisible;
-	bool bIsSelected;
-	bool bIsHovered;
-	bool bIsHeaderHovered;
+	ETimingTrackFlags Flags;
 };
 
 ////////////////////////////////////////////////////////////////////////////////////////////////////
diff --git a/Engine/Source/Developer/TraceInsights/Private/Insights/ViewModels/FrameTrackHelper.cpp b/Engine/Source/Developer/TraceInsights/Private/Insights/ViewModels/FrameTrackHelper.cpp
--- a/Engine/Source/Developer/TraceInsights/Private/Insights/ViewModels/FrameTrackHelper.cpp
+++ b/Engine/Source/Developer/TraceInsights/Private/Insights/ViewModels/FrameTrackHelper.cpp
@@ -18,97 +18,66 @@
 #include <limits>
 
 ////////////////////////////////////////////////////////////////////////////////////////////////////
-// FFrameTrackCacheContext
+// FFrameTrackTimelineBuilder
 ////////////////////////////////////////////////////////////////////////////////////////////////////
 
-FFrameTrackCacheContext::FFrameTrackCacheContext(const FFrameTrackViewport& InViewport)
-	: Viewport(InViewport)
-{
-}
-
-////////////////////////////////////////////////////////////////////////////////////////////////////
-
-void FFrameTrackCacheContext::Begin()
+FFrameTrackTimelineBuilder::FFrameTrackTimelineBuilder(FFrameTrackTimeline& InTimeline, const FFrameTrackViewport& InViewport)
+	: Timeline(InTimeline)
+	, Viewport(InViewport)
+	, NumAddedFrames(0)
 {
 	SampleW = Viewport.GetSampleWidth();
 	FramesPerSample = Viewport.GetNumFramesPerSample();
-	NumSamples = FMath::CeilToInt(Viewport.Width / SampleW);
+	NumSamples = FMath::Max(0, FMath::CeilToInt(Viewport.Width / SampleW));
 	FirstFrameIndex = Viewport.GetFirstFrameIndex();
 
-	NumFrames = 0;
-}
-
-////////////////////////////////////////////////////////////////////////////////////////////////////
-
-bool FFrameTrackCacheContext::BeginTimeline(FFrameTrackTimeline& CachedTimeline)
-{
-	CurrentCachedTimeline = &CachedTimeline;
-
-	CurrentCachedTimeline->NumAggregatedFrames = 0;
-	CurrentCachedTimeline->Samples.Reset(NumSamples);
-	CurrentCachedTimeline->Samples.AddDefaulted(NumSamples);
-
-	return true;
+	Timeline.NumAggregatedFrames = 0;
+	Timeline.Samples.Reset();
+	Timeline.Samples.AddDefaulted(NumSamples);
 }
 
 ////////////////////////////////////////////////////////////////////////////////////////////////////
 
-void FFrameTrackCacheContext::AddEvent(const Trace::FFrame& Event)
+void FFrameTrackTimelineBuilder::AddFrame(const Trace::FFrame& Frame)
 {
-	NumFrames++;
+	NumAddedFrames++;
 
-	const int32 FrameIndex = Event.Index;
+	const int32 FrameIndex = Frame.Index;
 
-	int SampleIndex = (FrameIndex - FirstFrameIndex) / FramesPerSample;
+	int32 SampleIndex = (FrameIndex - FirstFrameIndex) / FramesPerSample;
 	if (SampleIndex >= 0 && SampleIndex < NumSamples)
 	{
-		FFrameTrackSample& Sample = CurrentCachedTimeline->Samples[SampleIndex];
+		FFrameTrackSample& Sample = Timeline.Samples[SampleIndex];
 		Sample.NumFrames++;
 
-		double Duration = Event.EndTime - Event.StartTime;
+		double Duration = Frame.EndTime - Frame.StartTime;
 		Sample.TotalDuration += Duration;
-		if (Event.StartTime < Sample.StartTime)
+		if (Frame.StartTime < Sample.StartTime)
 		{
-			Sample.StartTime = Event.StartTime;
+			Sample.StartTime = Frame.StartTime;
 		}
-		if (Event.EndTime > Sample.EndTime)
+		if (Frame.EndTime > Sample.EndTime)
 		{
-			Sample.EndTime = Event.EndTime;
+			Sample.EndTime = Frame.EndTime;
 		}
 		if (Duration > Sample.LargestFrameDuration)
 		{
 			Sample.LargestFrameIndex = FrameIndex;
-			Sample.LargestFrameStartTime = Event.StartTime;
+			Sample.LargestFrameStartTime = Frame.StartTime;
 			Sample.LargestFrameDuration = Duration;
 		}
 
-		CurrentCachedTimeline->NumAggregatedFrames++;
+		Timeline.NumAggregatedFrames++;
 	}
 }
 
 ////////////////////////////////////////////////////////////////////////////////////////////////////
-
-void FFrameTrackCacheContext::EndTimeline()
-{
-	//CurrentCachedTimeline = nullptr;
-}
-
+// FFrameTrackDrawHelper
 ////////////////////////////////////////////////////////////////////////////////////////////////////
 
-void FFrameTrackCacheContext::End()
-{
-}
-
-////////////////////////////////////////////////////////////////////////////////////////////////////
-// FFrameTrackDrawContext
-////////////////////////////////////////////////////////////////////////////////////////////////////
-
-FFrameTrackDrawContext::FFrameTrackDrawContext(const FFrameTrackViewport& InViewport, const FGeometry& InAllottedGeometry, int32& InLayerId, FSlateWindowElementList& InElementList)
-	: Viewport(InViewport)
-	, AllottedGeometry(InAllottedGeometry)
-	, LayerId(InLayerId)
-	, ElementList(InElementList)
-	, DrawEffects(ESlateDrawEffect::None)
+FFrameTrackDrawHelper::FFrameTrackDrawHelper(const FDrawContext& InDrawContext, const FFrameTrackViewport& InViewport)
+	: DrawContext(InDrawContext)
+	, Viewport(InViewport)
 	, WhiteBrush(FCoreStyle::Get().GetBrush("WhiteBrush"))
 	, BorderBrush(FEditorStyle::GetBrush("PlainBorder"))
 	, NumFrames(0)
@@ -118,7 +87,7 @@ FFrameTrackDrawContext::FFrameTrackDrawContext(const FFrameTrackViewport& InView
 
 ////////////////////////////////////////////////////////////////////////////////////////////////////
 
-void FFrameTrackDrawContext::DrawBackground(const FWidgetStyle& InWidgetStyle) const
+void FFrameTrackDrawHelper::DrawBackground() const
 {
 	const FSlateBrush* AreaBrush = WhiteBrush;
 	const FLinearColor ValidAreaColor(0.07f, 0.07f, 0.07f, 1.0f);
@@ -132,44 +101,20 @@ void FFrameTrackDrawContext::DrawBackground(const FWidgetStyle& InWidgetStyle) c
 	if (X0 >= W || X1 <= 0.0f)
 	{
 		// Draw invalid area (entire view).
-		FSlateDrawElement::MakeBox
-		(
-			ElementList,
-			LayerId,
-			MAKE_PAINT_GEOMETRY_RC(AllottedGeometry, 0.0f, 0.0f, W, H),
-			AreaBrush,
-			DrawEffects,
-			InvalidAreaColor
-		);
+		DrawContext.DrawBox(0.0f, 0.0f, W, H, AreaBrush, InvalidAreaColor);
 	}
 	else // X0 < W && X1 > 0
 	{
 		if (X0 > 0.0f)
 		{
 			// Draw invalid area (left).
-			FSlateDrawElement::MakeBox
-			(
-				ElementList,
-				LayerId,
-				MAKE_PAINT_GEOMETRY_RC(AllottedGeometry, 0.0f, 0.0f, X0, H),
-				AreaBrush,
-				DrawEffects,
-				InvalidAreaColor
-			);
+			DrawContext.DrawBox(0.0f, 0.0f, X0, H, AreaBrush, InvalidAreaColor);
 		}
 
 		if (X1 < W)
 		{
 			// Draw invalid area (right).
-			FSlateDrawElement::MakeBox
-			(
-				ElementList,
-				LayerId,
-				MAKE_PAINT_GEOMETRY_RC(AllottedGeometry, X1, 0.0f, W - X1, H),
-				AreaBrush,
-				DrawEffects,
-				InvalidAreaColor
-			);
+			DrawContext.DrawBox(X1, 0.0f, W - X1, H, AreaBrush, InvalidAreaColor);
 		}
 
 		X0 = FMath::Max(X0, 0.0f);
@@ -178,24 +123,16 @@ void FFrameTrackDrawContext::DrawBackground(const FWidgetStyle& InWidgetStyle) c
 		if (X1 > X0)
 		{
 			// Draw valid area.
-			FSlateDrawElement::MakeBox
-			(
-				ElementList,
-				LayerId,
-				MAKE_PAINT_GEOMETRY_RC(AllottedGeometry, X0, 0.0f, X1 - X0, H),
-				AreaBrush,
-				DrawEffects,
-				ValidAreaColor
-			);
+			DrawContext.DrawBox(X0, 0.0f, X1 - X0, H, AreaBrush, ValidAreaColor);
 		}
 	}
 
-	LayerId++;
+	DrawContext.LayerId++;
 }
 
 ////////////////////////////////////////////////////////////////////////////////////////////////////
 
-FLinearColor FFrameTrackDrawContext::GetColorById(int Id) const
+FLinearColor FFrameTrackDrawHelper::GetColorById(int32 Id) const
 {
 	const float Alpha = 0.9;
 	switch (Id)
@@ -203,7 +140,7 @@ FLinearColor FFrameTrackDrawContext::GetColorById(int Id) const
 	case 0: // Game Frames
 		return FLinearColor(0.75, 1.0, 1.0, Alpha);
 
-	case 1: // Render Frames
+	case 1: // Rendering Frames
 		return FLinearColor(1.0, 0.75, 0.75, Alpha);
 
 	case 2:
@@ -216,25 +153,25 @@ FLinearColor FFrameTrackDrawContext::GetColorById(int Id) const
 
 ////////////////////////////////////////////////////////////////////////////////////////////////////
 
-void FFrameTrackDrawContext::DrawCached(const FFrameTrackTimeline& CachedTimeline) const
+void FFrameTrackDrawHelper::DrawCached(const FFrameTrackTimeline& Timeline) const
 {
-	if (CachedTimeline.NumAggregatedFrames == 0)
+	if (Timeline.NumAggregatedFrames == 0)
 	{
 		return;
 	}
 
-	NumFrames += CachedTimeline.NumAggregatedFrames;
+	NumFrames += Timeline.NumAggregatedFrames;
 
-	FLinearColor TimelineColor = GetColorById(CachedTimeline.Id);
+	FLinearColor TimelineColor = GetColorById(static_cast<int32>(Timeline.Id));
 
 	const float SampleW = Viewport.GetSampleWidth();
-	const int NumSamples = CachedTimeline.Samples.Num();
+	const int32 NumSamples = Timeline.Samples.Num();
 
-	const float VY1 = FMath::RoundToFloat(Viewport.GetViewportYForValue(0.0));
+	const float Y1 = FMath::RoundToFloat(Viewport.GetViewportYForValue(0.0));
 
-	for (int SampleIndex = 0; SampleIndex < NumSamples; SampleIndex++)
+	for (int32 SampleIndex = 0; SampleIndex < NumSamples; SampleIndex++)
 	{
-		const FFrameTrackSample& Sample = CachedTimeline.Samples[SampleIndex];
+		const FFrameTrackSample& Sample = Timeline.Samples[SampleIndex];
 		if (Sample.NumFrames == 0)
 		{
 			continue;
@@ -243,20 +180,20 @@ void FFrameTrackDrawContext::DrawCached(const FFrameTrackTimeline& CachedTimelin
 		NumDrawSamples++;
 
 		const float X = SampleIndex * SampleW;
-		float VY2;
+		float Y2;
 
 		FLinearColor ColorFill = TimelineColor;
 
 		if (Sample.LargestFrameDuration == std::numeric_limits<double>::infinity())
 		{
-			VY2 = Viewport.Height;
+			Y2 = Viewport.Height;
 			ColorFill.R = 0.0f;
 			ColorFill.G = 0.0f;
 			ColorFill.B = 0.0f;
 		}
 		else
 		{
-			VY2 = FMath::RoundToFloat(Viewport.GetViewportYForValue(Sample.LargestFrameDuration));
+			Y2 = FMath::RoundToFloat(Viewport.GetViewportYForValue(Sample.LargestFrameDuration));
 			if (Sample.LargestFrameDuration > 1.0 / 30.0)
 			{
 				ColorFill.G *= 0.5f;
@@ -268,79 +205,75 @@ void FFrameTrackDrawContext::DrawCached(const FFrameTrackTimeline& CachedTimelin
 			}
 		}
 
-		const float Y = FMath::RoundToFloat(Viewport.Height) - VY2;
-		const float H = VY2 - VY1;
+		const float Y = FMath::RoundToFloat(Viewport.Height) - Y2;
+		const float H = Y2 - Y1;
 
 		const FLinearColor ColorBorder(ColorFill.R * 0.75f, ColorFill.G * 0.75f, ColorFill.B * 0.75f, 1.0);
 
 		if (SampleW > 2.0f)
 		{
-			FPaintGeometry GeoFill = MAKE_PAINT_GEOMETRY_RC(AllottedGeometry, X + 1.0f, Y + 1.0f, SampleW - 2.0f, H - 2.0f);
-			FSlateDrawElement::MakeBox(ElementList, LayerId, GeoFill, WhiteBrush, DrawEffects, ColorFill);
-
-			//FPaintGeometry GeoBorder = MAKE_PAINT_GEOMETRY_RC(AllottedGeometry, X, Y, SampleW, H);
-			//FSlateDrawElement::MakeBox(ElementList, LayerId+1, GeoBorder, BorderBrush, DrawEffects, ColorBorder);
-
-			#define LOCAL_DRAW_RC(x, y, w, h) FSlateDrawElement::MakeBox(ElementList, LayerId, MAKE_PAINT_GEOMETRY_RC(AllottedGeometry, x, y, w, h), WhiteBrush, DrawEffects, ColorBorder)
-			LOCAL_DRAW_RC(X, Y, 1.0, H);
-			LOCAL_DRAW_RC(X + SampleW - 1.0f, Y, 1.0, H);
-			LOCAL_DRAW_RC(X + 1.0f, Y, SampleW - 2.0f, 1.0f);
-			LOCAL_DRAW_RC(X + 1.0f, Y + H - 1.0f, SampleW - 2.0f, 1.0f);
-			#undef DRAW_RC
+			DrawContext.DrawBox(X + 1.0f, Y + 1.0f, SampleW - 2.0f, H - 2.0f, WhiteBrush, ColorFill);
+
+			// Draw border.
+			DrawContext.DrawBox(X, Y, 1.0, H, WhiteBrush, ColorBorder);
+			DrawContext.DrawBox(X + SampleW - 1.0f, Y, 1.0, H, WhiteBrush, ColorBorder);
+			DrawContext.DrawBox(X + 1.0f, Y, SampleW - 2.0f, 1.0f, WhiteBrush, ColorBorder);
+			DrawContext.DrawBox(X + 1.0f, Y + H - 1.0f, SampleW - 2.0f, 1.0f, WhiteBrush, ColorBorder);
 		}
 		else
 		{
-			FPaintGeometry GeoBorder = MAKE_PAINT_GEOMETRY_RC(AllottedGeometry, X, Y, SampleW, H);
-			FSlateDrawElement::MakeBox(ElementList, LayerId, GeoBorder, WhiteBrush, DrawEffects, ColorBorder);
+			DrawContext.DrawBox(X, Y, SampleW, H, WhiteBrush, ColorBorder);
 		}
 	}
 
-	LayerId++;
+	DrawContext.LayerId++;
 }
 
 ////////////////////////////////////////////////////////////////////////////////////////////////////
 
-void FFrameTrackDrawContext::DrawHoveredSample(const FFrameTrackSample& Sample) const
+void FFrameTrackDrawHelper::DrawHoveredSample(const FFrameTrackSample& Sample) const
 {
 	const float SampleW = Viewport.GetSampleWidth();
 	const int32 FramesPerSample = Viewport.GetNumFramesPerSample();
 	const int32 FirstFrameIndex = Viewport.GetFirstFrameIndex();
-	const int SampleIndex = (Sample.LargestFrameIndex - FirstFrameIndex) / FramesPerSample;
+	const int32 SampleIndex = (Sample.LargestFrameIndex - FirstFrameIndex) / FramesPerSample;
 	const float X = SampleIndex * SampleW;
 
-	const float VY1 = FMath::RoundToFloat(Viewport.GetViewportYForValue(0.0));
-	float VY2;
+	const float Y1 = FMath::RoundToFloat(Viewport.GetViewportYForValue(0.0));
+	float Y2;
 	if (Sample.LargestFrameDuration == std::numeric_limits<double>::infinity())
 	{
-		VY2 = Viewport.Height;
+		Y2 = Viewport.Height;
 	}
 	else
 	{
-		VY2 = FMath::RoundToFloat(Viewport.GetViewportYForValue(Sample.LargestFrameDuration));
+		Y2 = FMath::RoundToFloat(Viewport.GetViewportYForValue(Sample.LargestFrameDuration));
 	}
-	const float Y = FMath::RoundToFloat(Viewport.Height) - VY2;
-	const float H = VY2 - VY1;
+	const float Y = FMath::RoundToFloat(Viewport.Height) - Y2;
+	const float H = Y2 - Y1;
 
 	const FLinearColor ColorBorder(1.0f, 1.0f, 0.0f, 1.0);
-	FPaintGeometry GeoBorder = MAKE_PAINT_GEOMETRY_RC(AllottedGeometry, X - 1.0f, Y - 1.0f, SampleW + 2.0f, H + 2.0f);
-	FSlateDrawElement::MakeBox(ElementList, LayerId, GeoBorder, BorderBrush, DrawEffects, ColorBorder);
-
-	LayerId++;
+	DrawContext.DrawBox(X - 1.0f, Y - 1.0f, SampleW + 2.0f, H + 2.0f, BorderBrush, ColorBorder);
+	DrawContext.LayerId++;
 }
 
 ////////////////////////////////////////////////////////////////////////////////////////////////////
 
-void FFrameTrackDrawContext::DrawHighlightedInterval(const FFrameTrackTimeline& CachedTimeline, const double StartTime, const double EndTime) const
+void FFrameTrackDrawHelper::DrawHighlightedInterval(const FFrameTrackTimeline& Timeline, const double StartTime, const double EndTime) const
 {
-	const int NumSamples = CachedTimeline.Samples.Num();
+	const int32 NumSamples = Timeline.Samples.Num();
 
 	//TODO: binary search
-	int Index1 = 0;
-	int Index2 = NumSamples - 1;
-	while (Index1 < NumSamples && CachedTimeline.Samples[Index1].EndTime < StartTime)
+	int32 Index1 = 0;
+	int32 Index2 = NumSamples - 1;
+	while (Index1 < NumSamples && Timeline.Samples[Index1].EndTime < StartTime)
+	{
 		Index1++;
-	while (Index2 >= Index1 && CachedTimeline.Samples[Index2].StartTime > EndTime)
+	}
+	while (Index2 >= Index1 && Timeline.Samples[Index2].StartTime > EndTime)
+	{
 		Index2--;
+	}
 
 	if (Index1 <= Index2)
 	{
@@ -348,42 +281,40 @@ void FFrameTrackDrawContext::DrawHighlightedInterval(const FFrameTrackTimeline&
 		float X1 = Index1 * SampleW;
 		float X2 = Index2 * SampleW;
 
-		const float Y1 = 10.0f;
+		const float Y1 = 10.0f; // allows 10px for the horizontal scrollbar (one displayed on top of the track)
 		const float Y2 = Viewport.Height;
-		const float D = 2.0f;
-		const float H = 10.0f;
+		const float D = 2.0f; // line thickness (for both horizontal and vertical lines)
+		const float H = 10.0f; // height of corner lines
 
 		const FLinearColor Color(1.0f, 1.0f, 1.0f, 1.0f);
 
 		if (X1 >= 0.0f && X1 < Viewport.Width - 2.0f)
 		{
-			FPaintGeometry GeoL1 = MAKE_PAINT_GEOMETRY_RC(AllottedGeometry, X1 - D, Y1, D, H);
-			FSlateDrawElement::MakeBox(ElementList, LayerId, GeoL1, WhiteBrush, DrawEffects, Color);
-
-			FPaintGeometry GeoL2 = MAKE_PAINT_GEOMETRY_RC(AllottedGeometry, X1 - D, Y2 - H, D, H);
-			FSlateDrawElement::MakeBox(ElementList, LayerId, GeoL2, WhiteBrush, DrawEffects, Color);
+			// Draw left side vertical lines.
+			DrawContext.DrawBox(X1 - D, Y1, D, H, WhiteBrush, Color);
+			DrawContext.DrawBox(X1 - D, Y2 - H, D, H, WhiteBrush, Color);
 		}
 
 		if (X2 >= -2.0f && X2 < Viewport.Width)
 		{
-			FPaintGeometry GeoR1 = MAKE_PAINT_GEOMETRY_RC(AllottedGeometry, X2, Y1, D, H);
-			FSlateDrawElement::MakeBox(ElementList, LayerId, GeoR1, WhiteBrush, DrawEffects, Color);
-
-			FPaintGeometry GeoR2 = MAKE_PAINT_GEOMETRY_RC(AllottedGeometry, X2, Y2 - H, D, H);
-			FSlateDrawElement::MakeBox(ElementList, LayerId, GeoR2, WhiteBrush, DrawEffects, Color);
+			// Draw right side vertical lines.
+			DrawContext.DrawBox(X2, Y1, D, H, WhiteBrush, Color);
+			DrawContext.DrawBox(X2, Y2 - H, D, H, WhiteBrush, Color);
 		}
 
 		if (X1 < 0)
+		{
 			X1 = 0.0f;
+		}
 		if (X2 > Viewport.Width)
+		{
 			X2 = Viewport.Width;
+		}
 		if (X1 < X2)
 		{
-			FPaintGeometry GeoT = MAKE_PAINT_GEOMETRY_RC(AllottedGeometry, X1, Y1, X2 - X1, D);
-			FSlateDrawElement::MakeBox(ElementList, LayerId, GeoT, WhiteBrush, DrawEffects, Color);
-
-			FPaintGeometry GeoB = MAKE_PAINT_GEOMETRY_RC(AllottedGeometry, X1, Y2 - D, X2 - X1, D);
-			FSlateDrawElement::MakeBox(ElementList, LayerId, GeoB, WhiteBrush, DrawEffects, Color);
+			// Draw horizontal lines.
+			DrawContext.DrawBox(X1, Y1, X2 - X1, D, WhiteBrush, Color);
+			DrawContext.DrawBox(X1, Y2 - D, X2 - X1, D, WhiteBrush, Color);
 		}
 	}
 }
diff --git a/Engine/Source/Developer/TraceInsights/Private/Insights/ViewModels/FrameTrackHelper.h b/Engine/Source/Developer/TraceInsights/Private/Insights/ViewModels/FrameTrackHelper.h
--- a/Engine/Source/Developer/TraceInsights/Private/Insights/ViewModels/FrameTrackHelper.h
+++ b/Engine/Source/Developer/TraceInsights/Private/Insights/ViewModels/FrameTrackHelper.h
@@ -7,6 +7,7 @@
 
 enum class ESlateDrawEffect : uint8;
 
+struct FDrawContext;
 struct FGeometry;
 struct FSlateBrush;
 
@@ -18,7 +19,9 @@ namespace Trace
 	struct FFrame;
 }
 
-enum EFrameType
+////////////////////////////////////////////////////////////////////////////////////////////////////
+
+enum class EFrameType
 {
 	Unknown = -1,
 
@@ -26,6 +29,8 @@ enum EFrameType
 	Render
 };
 
+////////////////////////////////////////////////////////////////////////////////////////////////////
+
 struct FFrameTrackSample
 {
 	int32 NumFrames;
@@ -47,74 +52,86 @@ struct FFrameTrackSample
 	{}
 };
 
+////////////////////////////////////////////////////////////////////////////////////////////////////
+
 struct FFrameTrackTimeline
 {
 	uint64 Id;
-	//EFrameType Type;
+	//TODO: EFrameType Type;
 
-	//double MinValue; // min value
-	//double MaxValue; // max value
-	//double AvgValue; // average value
+	//TODO: double MinValue; // min value
+	//TODO: double MaxValue; // max value
 
-	//float SampleW; // width of a sample, in Slate units
-	//int FramesPerSample; // number of frames in a sample
-	//int FirstFrameIndex; // index of first frame in first sample; can be negative
-	//int NumSamples; // total number of samples
-	int NumAggregatedFrames; // total number of frames aggregated in samples; ie. sum of all Sample.NumFrames
+	//TODO: float SampleW; // width of a sample, in Slate units
+	//TODO: int32 FramesPerSample; // number of frames in a sample
+	//TODO: int32 FirstFrameIndex; // index of first frame in first sample; can be negative
+	//TODO: int32 NumSamples; // total number of samples
+	int32 NumAggregatedFrames; // total number of frames aggregated in samples; ie. sum of all Sample.NumFrames
 
 	TArray<FFrameTrackSample> Samples;
 };
 
-class FFrameTrackCacheContext : public FNoncopyable
+////////////////////////////////////////////////////////////////////////////////////////////////////
+
+class FFrameTrackTimelineBuilder
 {
 public:
-	FFrameTrackCacheContext(const FFrameTrackViewport& InViewport);
-	//~FFrameTrackCacheContext();
+	FFrameTrackTimelineBuilder(FFrameTrackTimeline& InTimeline, const FFrameTrackViewport& InViewport);
 
-	void Begin();
-	bool BeginTimeline(FFrameTrackTimeline& CachedTimeline);
-	void AddEvent(const Trace::FFrame& Event);
-	void EndTimeline();
-	void End();
+	/**
+	 * Non-copyable
+	 */
+	FFrameTrackTimelineBuilder(const FFrameTrackTimelineBuilder&) = delete;
+	FFrameTrackTimelineBuilder& operator=(const FFrameTrackTimelineBuilder&) = delete;
 
-public:
-	const FFrameTrackViewport& Viewport;
+	void AddFrame(const Trace::FFrame& Frame);
 
-	FFrameTrackTimeline* CurrentCachedTimeline; // pointer to current FFrameTrackTimeline; valid only between BeginTimeline - EndTimeline calls.
+	int32 GetNumAddedFrames() const { return NumAddedFrames; }
+
+private:
+	FFrameTrackTimeline& Timeline;
+	const FFrameTrackViewport& Viewport;
 
 	float SampleW; // width of a sample, in Slate units
-	int FramesPerSample; // number of frames in a sample
-	int FirstFrameIndex; // index of first frame in first sample; can be negative
-	int NumSamples; // total number of samples
+	int32 FramesPerSample; // number of frames in a sample
+	int32 FirstFrameIndex; // index of first frame in first sample; can be negative
+	int32 NumSamples; // total number of samples
 
 	// Debug stats.
-	int NumFrames; // number of added frames
+	int32 NumAddedFrames; // counts total number of added frame events
 };
 
-class FFrameTrackDrawContext : public FNoncopyable
+////////////////////////////////////////////////////////////////////////////////////////////////////
+
+class FFrameTrackDrawHelper
 {
 public:
-	FFrameTrackDrawContext(const FFrameTrackViewport& InViewport, const FGeometry& InAllottedGeometry, int32& InLayerId, FSlateWindowElementList& InElementList);
-	//~FFrameTrackDrawContext();
+	FFrameTrackDrawHelper(const FDrawContext& InDrawContext, const FFrameTrackViewport& InViewport);
 
-	void DrawBackground(const FWidgetStyle& InWidgetStyle) const;
-	void DrawCached(const FFrameTrackTimeline& CachedTimeline) const;
+	/**
+	 * Non-copyable
+	 */
+	FFrameTrackDrawHelper(const FFrameTrackDrawHelper&) = delete;
+	FFrameTrackDrawHelper& operator=(const FFrameTrackDrawHelper&) = delete;
+
+	void DrawBackground() const;
+	void DrawCached(const FFrameTrackTimeline& Timeline) const;
 	void DrawHoveredSample(const FFrameTrackSample& Sample) const;
-	void DrawHighlightedInterval(const FFrameTrackTimeline& CachedTimeline, const double StartTime, const double EndTime) const;
+	void DrawHighlightedInterval(const FFrameTrackTimeline& Timeline, const double StartTime, const double EndTime) const;
 
-	FLinearColor GetColorById(int Id) const;
+	FLinearColor GetColorById(int32 Id) const;
 
-public:
+	int32 GetNumFrames() const { return NumFrames; }
+	int32 GetNumDrawSamples() const { return NumDrawSamples; }
+
+private:
+	const FDrawContext& DrawContext;
 	const FFrameTrackViewport& Viewport;
-	const FGeometry& AllottedGeometry;
-	int32& LayerId;
-	FSlateWindowElementList& ElementList;
-	const ESlateDrawEffect DrawEffects;
 
 	const FSlateBrush* WhiteBrush;
 	const FSlateBrush* BorderBrush;
 
 	// Debug stats.
-	mutable int NumFrames;
-	mutable int NumDrawSamples;
+	mutable int32 NumFrames;
+	mutable int32 NumDrawSamples;
 };
diff --git a/Engine/Source/Developer/TraceInsights/Private/Insights/ViewModels/FrameTrackViewport.h b/Engine/Source/Developer/TraceInsights/Private/Insights/ViewModels/FrameTrackViewport.h
--- a/Engine/Source/Developer/TraceInsights/Private/Insights/ViewModels/FrameTrackViewport.h
+++ b/Engine/Source/Developer/TraceInsights/Private/Insights/ViewModels/FrameTrackViewport.h
@@ -27,8 +27,8 @@ public:
 		MaxX = 0.0f;
 		PosX = 0.0f;
 		MinScaleX = 0.0001f;
-		MaxScaleX = 10.0f;
-		ScaleX = 1.0f;
+		MaxScaleX = 16.0f;
+		ScaleX = MaxScaleX;
 
 		MinValue = 0.0;
 		MaxValue = 1.0f;
diff --git a/Engine/Source/Developer/TraceInsights/Private/Insights/ViewModels/GraphTrack.cpp b/Engine/Source/Developer/TraceInsights/Private/Insights/ViewModels/GraphTrack.cpp
--- a/Engine/Source/Developer/TraceInsights/Private/Insights/ViewModels/GraphTrack.cpp
+++ b/Engine/Source/Developer/TraceInsights/Private/Insights/ViewModels/GraphTrack.cpp
@@ -43,6 +43,7 @@ FGraphTrackSeries::~FGraphTrackSeries()
 
 FGraphTrack::FGraphTrack(uint64 InTrackId)
 	: FBaseTimingTrack(InTrackId)
+	//, AllSeries()
 	, WhiteBrush(FCoreStyle::Get().GetBrush("WhiteBrush"))
 	, PointBrush(FEditorStyle::GetBrush("Graph.ExecutionBubble"))
 	, BorderBrush(new FSlateBorderBrush(NAME_None, FMargin(1.0f)))
@@ -55,6 +56,10 @@ FGraphTrack::FGraphTrack(uint64 InTrackId)
 	, bDrawBoxes(true)
 	, BaselineY(0.0f)
 	, ScaleY(1.0)
+	, NumAddedEvents(0)
+	, NumDrawPoints(0)
+	, NumDrawLines(0)
+	, NumDrawBoxes(0)
 {
 }
 
@@ -67,24 +72,19 @@ FGraphTrack::~FGraphTrack()
 
 ////////////////////////////////////////////////////////////////////////////////////////////////////
 
-void FGraphTrack::UpdateHoveredState(float MX, float MY, const FTimingTrackViewport& Viewport)
+void FGraphTrack::UpdateHoveredState(float MouseX, float MouseY, const FTimingTrackViewport& Viewport)
 {
-	if (MY >= Y && MY < Y + H)
-	{
-		bIsHovered = true;
+	static const float HeaderWidth = 100.0f;
+	static const float HeaderHeight = 14.0f;
 
-		if (MX < 100.0f && MY < Y + 14.0f)
-		{
-			bIsHeaderHovered = true;
-		}
-		else
-		{
-			bIsHeaderHovered = false;
-		}
+	if (MouseY >= GetPosY() && MouseY < GetPosY() + GetHeight())
+	{
+		SetHoveredState(true);
+		SetHeaderHoveredState(MouseX < HeaderWidth && MouseY < GetPosY() + HeaderHeight);
 	}
 	else
 	{
-		bIsHovered = false;
+		SetHoveredState(false);
 	}
 }
 
@@ -106,106 +106,110 @@ void FGraphTrack::UpdateStats()
 
 ////////////////////////////////////////////////////////////////////////////////////////////////////
 
-void FGraphTrack::Draw(FDrawContext& DC, const FTimingTrackViewport& Viewport) const
+void FGraphTrack::Draw(FDrawContext& DrawContext, const FTimingTrackViewport& Viewport) const
 {
-	DC.DrawBox(0.0f, Y, Viewport.Width, H, WhiteBrush, FLinearColor(0.1f, 0.1f, 0.1f, 1.0f));
-	DC.LayerId++;
+	DrawContext.DrawBox(0.0f, GetPosY(), Viewport.Width, GetHeight(), WhiteBrush, FLinearColor(0.1f, 0.1f, 0.1f, 1.0f));
+	DrawContext.LayerId++;
 
-	//TODO: Set clipping to (0, Y, Viewport.Width, H)!
+	//TODO: Set clipping to (0, GetPosY(), Viewport.Width, GetHeight())!
 
 	for (const FGraphTrackSeries& Series : AllSeries)
 	{
-		DrawSeries(Series, DC, Viewport);
+		DrawSeries(Series, DrawContext, Viewport);
 	}
 
-	// Draw bottom of the track (debug).
-	//DC.DrawBox(0.0f, Viewport.Height - 1.0f, Viewport.Width, 1.0f, WhiteBrush, FLinearColor(0.5f, 0.0f, 0.0f, 1.0f));
-	//DC.LayerId++;
-
 	// Draw baseline (Value == 0).
-	DC.DrawBox(0.0f, Y + BaselineY - 1.0f, Viewport.Width, 1.0f, WhiteBrush, FLinearColor(0.05f, 0.05f, 0.05f, 1.0f));
-	DC.LayerId++;
+	DrawContext.DrawBox(0.0f, GetPosY() + BaselineY - 1.0f, Viewport.Width, 1.0f, WhiteBrush, FLinearColor(0.05f, 0.05f, 0.05f, 1.0f));
+	DrawContext.LayerId++;
 }
 
 ////////////////////////////////////////////////////////////////////////////////////////////////////
 
-void FGraphTrack::DrawSeries(const FGraphTrackSeries& Series, FDrawContext& DC, const FTimingTrackViewport& Viewport) const
+void FGraphTrack::DrawSeries(const FGraphTrackSeries& Series, FDrawContext& DrawContext, const FTimingTrackViewport& Viewport) const
 {
 	if (bDrawBoxes)
 	{
-		int NumBoxes = Series.Boxes.Num();
-		for (int Index = 0; Index < NumBoxes; ++Index)
+		int32 NumBoxes = Series.Boxes.Num();
+		for (int32 Index = 0; Index < NumBoxes; ++Index)
 		{
 			const FGraphBox& Box = Series.Boxes[Index];
-			DC.DrawBox(Box.X, Y + Box.Y, Box.W, BaselineY - Box.Y, WhiteBrush, Series.Color);
+			DrawContext.DrawBox(Box.X, GetPosY() + Box.Y, Box.W, BaselineY - Box.Y, WhiteBrush, Series.Color);
 		}
-		DC.LayerId++;
+		DrawContext.LayerId++;
 	}
 
 	if (bDrawLines)
 	{
-		FPaintGeometry Geo = DC.Geometry.ToPaintGeometry();
-		Geo.AppendTransform(FSlateLayoutTransform(FVector2D(0, Y)));
-		FSlateDrawElement::MakeLines(DC.ElementList, DC.LayerId, Geo, Series.LinePoints, DC.DrawEffects, Series.Color, false, 1.0f);
-		DC.LayerId++;
+		FPaintGeometry Geo = DrawContext.Geometry.ToPaintGeometry();
+		Geo.AppendTransform(FSlateLayoutTransform(FVector2D(0, GetPosY())));
+		FSlateDrawElement::MakeLines(DrawContext.ElementList, DrawContext.LayerId, Geo, Series.LinePoints, DrawContext.DrawEffects, Series.Color, false, 1.0f);
+		DrawContext.LayerId++;
 	}
 
 	if (bDrawPoints)
 	{
-		int NumPoints = Series.Points.Num();
+		int32 NumPoints = Series.Points.Num();
+
+#define INSIGHTS_GRAPH_TRACK_DRAW_POINTS_AS_RECTANGLES 0
+#if !INSIGHTS_GRAPH_TRACK_DRAW_POINTS_AS_RECTANGLES
 
 		if (bDrawPointsWithBorder)
 		{
 			// Draw points (border).
-			for (int Index = 0; Index < NumPoints; ++Index)
+			for (int32 Index = 0; Index < NumPoints; ++Index)
 			{
 				const FVector2D& Pt = Series.Points[Index];
 				const float PtX = Pt.X - GraphTrackPointVisualSize / 2.0f - 1.0f;
-				const float PtY = Y + Pt.Y - GraphTrackPointVisualSize / 2.0f - 1.0f;
-				DC.DrawBox(PtX, PtY, GraphTrackPointVisualSize + 2.0f, GraphTrackPointVisualSize + 2.0f, PointBrush, Series.BorderColor);
+				const float PtY = GetPosY() + Pt.Y - GraphTrackPointVisualSize / 2.0f - 1.0f;
+				DrawContext.DrawBox(PtX, PtY, GraphTrackPointVisualSize + 2.0f, GraphTrackPointVisualSize + 2.0f, PointBrush, Series.BorderColor);
 			}
-			DC.LayerId++;
+			DrawContext.LayerId++;
 		}
 
 		// Draw points (interior).
-		for (int Index = 0; Index < NumPoints; ++Index)
+		for (int32 Index = 0; Index < NumPoints; ++Index)
 		{
 			const FVector2D& Pt = Series.Points[Index];
 			const float PtX = Pt.X - GraphTrackPointVisualSize / 2.0f;
-			const float PtY = Y + Pt.Y - GraphTrackPointVisualSize / 2.0f;
-			DC.DrawBox(PtX, PtY, GraphTrackPointVisualSize, GraphTrackPointVisualSize, PointBrush, Series.Color);
+			const float PtY = GetPosY() + Pt.Y - GraphTrackPointVisualSize / 2.0f;
+			DrawContext.DrawBox(PtX, PtY, GraphTrackPointVisualSize, GraphTrackPointVisualSize, PointBrush, Series.Color);
 		}
-		DC.LayerId++;
+		DrawContext.LayerId++;
+
+#else // Alternative way of drawing points; kept here for debugging purposes.
 
-		/*
 		//const float Angle = FMath::DegreesToRadians(45.0f);
 
+		if (bDrawPointsWithBorder)
+		{
+			// Draw borders.
+			const float BorderPtSize = GraphTrackPointVisualSize;
+			//FVector2D BorderRotationPoint(BorderPtSize / 2.0f, BorderPtSize / 2.0f);
+			for (int32 Index = 0; Index < NumPoints; ++Index)
+			{
+				const FVector2D& Pt = Series.Points[Index];
+				const float PtX = Pt.X - BorderPtSize / 2.0f;
+				const float PtY = GetPosY() + Pt.Y - BorderPtSize / 2.0f;
+				DrawContext.DrawBox(PtX, PtY, BorderPtSize, BorderPtSize, BorderBrush, Series.BorderColor);
+				//DrawContext.DrawRotatedBox(PtX, PtY, BorderPtSize, BorderPtSize, BorderBrush, Series.BorderColor, Angle, BorderRotationPoint);
+			}
+			DrawContext.LayerId++;
+		}
+
 		// Draw points as rectangles.
 		const float PtSize = GraphTrackPointVisualSize - 2.0f;
 		//FVector2D RotationPoint(PtSize / 2.0f, PtSize / 2.0f);
-		for (int Index = 0; Index < NumPoints; ++Index)
+		for (int32 Index = 0; Index < NumPoints; ++Index)
 		{
 			const FVector2D& Pt = Series.Points[Index];
 			const float PtX = Pt.X - PtSize / 2.0f;
-			const float PtY = Y + Pt.Y - PtSize / 2.0f;
-			DC.DrawBox(PtX, PtY, PtSize, PtSize, PointBrush, Series.Color);
-			//DC.DrawRotatedBox(PtX, PtY, PtSize, PtSize, WhiteBrush, Series.Color, Angle, RotationPoint);
+			const float PtY = GetPosY() + Pt.Y - PtSize / 2.0f;
+			DrawContext.DrawBox(PtX, PtY, PtSize, PtSize, PointBrush, Series.Color);
+			//DrawContext.DrawRotatedBox(PtX, PtY, PtSize, PtSize, WhiteBrush, Series.Color, Angle, RotationPoint);
 		}
-		DC.LayerId++;
+		DrawContext.LayerId++;
 
-		// Draw borders.
-		const float BorderPtSize = GraphTrackPointVisualSize;
-		//FVector2D BorderRotationPoint(BorderPtSize / 2.0f, BorderPtSize / 2.0f);
-		for (int Index = 0; Index < NumPoints; ++Index)
-		{
-			const FVector2D& Pt = Series.Points[Index];
-			const float PtX = Pt.X - BorderPtSize / 2.0f;
-			const float PtY = Y + Pt.Y - BorderPtSize / 2.0f;
-			DC.DrawBox(PtX, PtY, BorderPtSize, BorderPtSize, BorderBrush, Series.BorderColor);
-			//DC.DrawRotatedBox(PtX, PtY, BorderPtSize, BorderPtSize, BorderBrush, Series.BorderColor, Angle, BorderRotationPoint);
-		}
-		DC.LayerId++;
-		*/
+#endif
 	}
 }
 
@@ -240,9 +244,10 @@ void FRandomGraphTrack::Update(const FTimingTrackViewport& Viewport)
 	NumAddedEvents = 0;
 
 	// TODO: vertical panning and zooming
-	BaselineY = H;
+	BaselineY = GetHeight();
 	ScaleY = 1.0;
 
+	ensure(AllSeries.Num() == 1);
 	GenerateSeries(AllSeries[0], Viewport, 1000000);
 
 	UpdateStats();
@@ -250,15 +255,15 @@ void FRandomGraphTrack::Update(const FTimingTrackViewport& Viewport)
 
 ////////////////////////////////////////////////////////////////////////////////////////////////////
 
-void FRandomGraphTrack::GenerateSeries(FGraphTrackSeries& Series, const FTimingTrackViewport& Viewport, const int EventCount)
+void FRandomGraphTrack::GenerateSeries(FGraphTrackSeries& Series, const FTimingTrackViewport& Viewport, const int32 EventCount)
 {
 	//////////////////////////////////////////////////
 	// Generate random events.
 
-	static const double DT1 = 0.0000001; // 100ns
-	static const double DT2 = 0.01; // 100ms
-	const float V1 = 0;
-	const float V2 = H;
+	static const double MinDeltaTime = 0.0000001; // 100ns
+	static const double MaxDeltaTime = 0.01; // 100ms
+	const float MinValue = 0;
+	const float MaxValue = GetHeight();
 
 	struct FGraphEvent
 	{
@@ -272,50 +277,45 @@ void FRandomGraphTrack::GenerateSeries(FGraphTrackSeries& Series, const FTimingT
 
 	FRandomStream RandomStream(0);
 	double NextT = 0.0;
-	for (int Index = 0; Index < EventCount; ++Index)
+	for (int32 Index = 0; Index < EventCount; ++Index)
 	{
 		FGraphEvent Ev;
 		Ev.Time = NextT;
-		NextT += DT1 + RandomStream.GetFraction() * (DT2 - DT1);
-		Ev.Duration = DT1 + RandomStream.GetFraction() * (NextT - Ev.Time - DT1);
-		Ev.Value = V1 + RandomStream.GetFraction() * (V2 - V1);
+		const double TimeAdvance = RandomStream.GetFraction() * (MaxDeltaTime - MinDeltaTime);
+		NextT += MinDeltaTime + TimeAdvance;
+		Ev.Duration = MinDeltaTime + RandomStream.GetFraction() * TimeAdvance;
+		Ev.Value = MinValue + RandomStream.GetFraction() * (MaxValue - MinValue);
 		Events.Add(Ev);
 	}
 
-	// Sort events by time.
-	//Events.Sort([](const FGraphEvent& A, const FGraphEvent& B) { return A.Time < B.Time; });
-
 	//////////////////////////////////////////////////
 	// Optimize and build draw lists.
-
-	FGraphTrackBuilder Builder(*this, Series, Viewport);
-
-	Builder.Begin();
-
-	int Index = 0;
-	while (Index < EventCount && Events[Index].Time < Viewport.StartTime)
-	{
-		++Index;
-	}
-	if (Index > 0)
 	{
-		Index--; // one point outside screen (left side)
-	}
-	while (Index < EventCount)
-	{
-		const FGraphEvent& Ev = Events[Index];
-		Builder.AddEvent(Ev.Time, Ev.Duration, Ev.Value);
+		FGraphTrackBuilder Builder(*this, Series, Viewport);
 
-		if (Ev.Time > Viewport.EndTime)
+		int32 Index = 0;
+		while (Index < EventCount && Events[Index].Time < Viewport.StartTime)
+		{
+			++Index;
+		}
+		if (Index > 0)
 		{
-			// one point outside screen (right side)
-			break;
+			Index--; // one point outside screen (left side)
 		}
+		while (Index < EventCount)
+		{
+			const FGraphEvent& Ev = Events[Index];
+			Builder.AddEvent(Ev.Time, Ev.Duration, Ev.Value);
 
-		++Index;
-	}
+			if (Ev.Time > Viewport.EndTime)
+			{
+				// one point outside screen (right side)
+				break;
+			}
 
-	Builder.End();
+			++Index;
+		}
+	}
 }
 
 ////////////////////////////////////////////////////////////////////////////////////////////////////
@@ -334,11 +334,9 @@ FFramesGraphTrack::FFramesGraphTrack(uint64 InTrackId)
 	AllSeries.AddDefaulted(2);
 
 	FGraphTrackSeries& GameFramesSeries = AllSeries[0];
-	//GameFramesSeries.SetColor(FLinearColor(1.0f, 0.5f, 0.0f, 1.0f), FLinearColor(1.0f, 0.8f, 0.3f, 1.0f));
 	GameFramesSeries.SetColor(FLinearColor(0.3f, 0.3f, 1.0f, 1.0f), FLinearColor(1.0f, 1.0f, 1.0f, 1.0f));
 
 	FGraphTrackSeries& RenderFramesSeries = AllSeries[1];
-	//RenderFramesSeries.SetColor(FLinearColor(1.0f, 0.0f, 0.5f, 1.0f), FLinearColor(1.0f, 0.3f, 0.8f, 1.0f));
 	RenderFramesSeries.SetColor(FLinearColor(1.0f, 0.3f, 0.3f, 1.0f), FLinearColor(1.0f, 1.0f, 1.0f, 1.0f));
 }
 
@@ -355,9 +353,10 @@ void FFramesGraphTrack::Update(const FTimingTrackViewport& Viewport)
 	NumAddedEvents = 0;
 
 	// TODO: Vertical panning and zooming needs to be moved out in a Viewport like controller.
-	BaselineY = H;
+	BaselineY = GetHeight();
 	ScaleY = 200.0 / 0.1; // 200px = 100ms
 
+	ensure(AllSeries.Num() == 2);
 	UpdateSeries(AllSeries[0], Viewport, TraceFrameType_Game);
 	UpdateSeries(AllSeries[1], Viewport, TraceFrameType_Rendering);
 
@@ -370,12 +369,10 @@ void FFramesGraphTrack::UpdateSeries(FGraphTrackSeries& Series, const FTimingTra
 {
 	FGraphTrackBuilder Builder(*this, Series, Viewport);
 
-	Builder.Begin();
-
 	TSharedPtr<const Trace::IAnalysisSession> Session = FInsightsManager::Get()->GetSession();
 	if (Session.IsValid())
 	{
-		Trace::FAnalysisSessionReadScope _(*Session.Get());
+		Trace::FAnalysisSessionReadScope SessionReadScope(*Session.Get());
 		Session->ReadFramesProvider([this, &Builder, FrameType](const Trace::IFrameProvider& FramesProvider)
 		{
 			uint64 FrameCount = FramesProvider.GetFrameCount(FrameType);
@@ -386,8 +383,6 @@ void FFramesGraphTrack::UpdateSeries(FGraphTrackSeries& Series, const FTimingTra
 			});
 		});
 	}
-
-	Builder.End();
 }
 
 ////////////////////////////////////////////////////////////////////////////////////////////////////
@@ -398,12 +393,6 @@ FGraphTrackBuilder::FGraphTrackBuilder(FGraphTrack& InTrack, FGraphTrackSeries&
 	: Track(InTrack)
 	, Series(InSeries)
 	, Viewport(InViewport)
-{
-}
-
-////////////////////////////////////////////////////////////////////////////////////////////////////
-
-void FGraphTrackBuilder::Begin()
 {
 	BeginPoints();
 	BeginConnectedLines();
@@ -412,7 +401,7 @@ void FGraphTrackBuilder::Begin()
 
 ////////////////////////////////////////////////////////////////////////////////////////////////////
 
-void FGraphTrackBuilder::End()
+FGraphTrackBuilder::~FGraphTrackBuilder()
 {
 	EndPoints();
 	EndConnectedLines();
@@ -456,10 +445,12 @@ void FGraphTrackBuilder::BeginPoints()
 
 	PointsCurrentX = -DBL_MAX;
 
-	int MaxPointsPerLineScan = FMath::CeilToInt(Track.H / GraphTrackPointDY);
 	PointsAtCurrentX.Reset();
-	PointsAtCurrentX.AddDefaulted(MaxPointsPerLineScan);
-	//PointsAtCurrentX.SetNum(MaxPointsPerLineScan);
+	int32 MaxPointsPerLineScan = FMath::CeilToInt(Track.GetHeight() / GraphTrackPointDY);
+	if (MaxPointsPerLineScan > 0)
+	{
+		PointsAtCurrentX.AddDefaulted(MaxPointsPerLineScan);
+	}
 }
 
 ////////////////////////////////////////////////////////////////////////////////////////////////////
@@ -483,7 +474,7 @@ void FGraphTrackBuilder::AddPoint(double Time, double Value)
 
 	const float Y = Track.GetYForValue(Value);
 
-	int Index = FMath::RoundToInt(Y / GraphTrackPointDY);
+	int32 Index = FMath::RoundToInt(Y / GraphTrackPointDY);
 	if (Index >= 0 && Index < PointsAtCurrentX.Num())
 	{
 		FPointInfo& Pt = PointsAtCurrentX[Index];
@@ -497,13 +488,12 @@ void FGraphTrackBuilder::AddPoint(double Time, double Value)
 
 void FGraphTrackBuilder::FlushPoints()
 {
-	for (int Index = 0; Index < PointsAtCurrentX.Num(); ++Index)
+	for (int32 Index = 0; Index < PointsAtCurrentX.Num(); ++Index)
 	{
 		FPointInfo& Pt = PointsAtCurrentX[Index];
 		if (Pt.bValid)
 		{
 			Pt.bValid = false;
-			//const float Y = Index * GraphTrackPointDY + GraphTrackPointDY / 2.0f;
 			Series.Points.Add(FVector2D(Pt.X, Pt.Y));
 		}
 	}
@@ -670,8 +660,16 @@ void FGraphTrackBuilder::AddBox(double Time, double Duration, double Value)
 
 ////////////////////////////////////////////////////////////////////////////////////////////////////
 
+void FGraphTrackBuilder::FlushBox()
+{
+	// TODO: reduction algorithm
+}
+
+////////////////////////////////////////////////////////////////////////////////////////////////////
+
 void FGraphTrackBuilder::EndBoxes()
 {
+	FlushBox();
 }
 
 ////////////////////////////////////////////////////////////////////////////////////////////////////
diff --git a/Engine/Source/Developer/TraceInsights/Private/Insights/ViewModels/GraphTrack.h b/Engine/Source/Developer/TraceInsights/Private/Insights/ViewModels/GraphTrack.h
--- a/Engine/Source/Developer/TraceInsights/Private/Insights/ViewModels/GraphTrack.h
+++ b/Engine/Source/Developer/TraceInsights/Private/Insights/ViewModels/GraphTrack.h
@@ -73,21 +73,21 @@ public:
 	FGraphTrack(uint64 InTrackId);
 	virtual ~FGraphTrack();
 
-	virtual void UpdateHoveredState(float MX, float MY, const FTimingTrackViewport& Viewport);
+	virtual void UpdateHoveredState(float MouseX, float MouseY, const FTimingTrackViewport& Viewport);
 
 	virtual void Update(const FTimingTrackViewport& Viewport) override = 0;
 
-	void Draw(FDrawContext& DC, const FTimingTrackViewport& Viewport) const;
+	void Draw(FDrawContext& DrawContext, const FTimingTrackViewport& Viewport) const;
 
-	int GetNumAddedEvents() const { return NumAddedEvents; }
-	int GetNumDrawPoints() const { return NumDrawPoints; }
-	int GetNumDrawLines() const { return NumDrawLines; }
-	int GetNumDrawBoxes() const { return NumDrawBoxes; }
+	int32 GetNumAddedEvents() const { return NumAddedEvents; }
+	int32 GetNumDrawPoints() const { return NumDrawPoints; }
+	int32 GetNumDrawLines() const { return NumDrawLines; }
+	int32 GetNumDrawBoxes() const { return NumDrawBoxes; }
 
 protected:
 	void UpdateStats();
 
-	void DrawSeries(const FGraphTrackSeries& Series, FDrawContext& DC, const FTimingTrackViewport& Viewport) const;
+	void DrawSeries(const FGraphTrackSeries& Series, FDrawContext& DrawContext, const FTimingTrackViewport& Viewport) const;
 
 	float GetYForValue(double Value) const
 	{
@@ -113,10 +113,10 @@ protected:
 	double ScaleY; // scale between Value units and viewport units; in pixels (Slate units) / Value unit
 
 	// Stats
-	int NumAddedEvents;
-	int NumDrawPoints;
-	int NumDrawLines;
-	int NumDrawBoxes;
+	int32 NumAddedEvents;
+	int32 NumDrawPoints;
+	int32 NumDrawLines;
+	int32 NumDrawBoxes;
 };
 
 ////////////////////////////////////////////////////////////////////////////////////////////////////
@@ -130,7 +130,7 @@ public:
 	virtual void Update(const FTimingTrackViewport& Viewport) override;
 
 protected:
-	void GenerateSeries(FGraphTrackSeries& Series, const FTimingTrackViewport& Viewport, const int EventCount);
+	void GenerateSeries(FGraphTrackSeries& Series, const FTimingTrackViewport& Viewport, const int32 EventCount);
 };
 
 ////////////////////////////////////////////////////////////////////////////////////////////////////
@@ -149,9 +149,9 @@ protected:
 
 ////////////////////////////////////////////////////////////////////////////////////////////////////
 
-class FGraphTrackBuilder : public FNoncopyable
+class FGraphTrackBuilder
 {
-protected:
+private:
 	struct FPointInfo
 	{
 		bool bValid;
@@ -163,15 +163,21 @@ protected:
 
 public:
 	FGraphTrackBuilder(FGraphTrack& InTrack, FGraphTrackSeries& InSeries, const FTimingTrackViewport& InViewport);
-	//~FGraphTrackBuilder();
+	~FGraphTrackBuilder();
+
+	/**
+	 * Non-copyable
+	 */
+	FGraphTrackBuilder(const FGraphTrackBuilder&) = delete;
+	FGraphTrackBuilder& operator=(const FGraphTrackBuilder&) = delete;
 
-	const FTimingTrackViewport& GetViewport() { return Viewport; }
+	FGraphTrack& GetTrack() const { return Track; }
+	FGraphTrackSeries& GetSeries() const { return Series; }
+	const FTimingTrackViewport& GetViewport() const { return Viewport; }
 
-	void Begin();
 	void AddEvent(double Time, double Duration, double Value);
-	void End();
 
-protected:
+private:
 	void BeginPoints();
 	void AddPoint(double Time, double Value);
 	void FlushPoints();
@@ -184,9 +190,10 @@ protected:
 
 	void BeginBoxes();
 	void AddBox(double Time, double Duration, double Value);
+	void FlushBox();
 	void EndBoxes();
 
-protected:
+private:
 	FGraphTrack& Track;
 	FGraphTrackSeries& Series;
 	const FTimingTrackViewport& Viewport;
diff --git a/Engine/Source/Developer/TraceInsights/Private/Insights/ViewModels/LogFilter.cpp b/Engine/Source/Developer/TraceInsights/Private/Insights/ViewModels/LogFilter.cpp
--- a/Engine/Source/Developer/TraceInsights/Private/Insights/ViewModels/LogFilter.cpp
+++ b/Engine/Source/Developer/TraceInsights/Private/Insights/ViewModels/LogFilter.cpp
@@ -144,7 +144,7 @@ bool FLogFilter::IsMessageAllowed(const FLogMessageRecord& Message)
 void FLogFilter::SyncAvailableCategories(const TSet<FName> Categories)
 {
 	// Remove obsolete categories.
-	int OldNumAvailableLogCategories = AvailableLogCategories.Num();
+	int32 OldNumAvailableLogCategories = AvailableLogCategories.Num();
 	for (int32 Index = OldNumAvailableLogCategories - 1; Index >= 0; --Index)
 	{
 		if (!Categories.Contains(AvailableLogCategories[Index]))
@@ -273,35 +273,22 @@ void FLogFilteringAsyncTask::DoWork()
 	if (Filter.IsFilterSetByText())
 	{
 		UE_LOG(TimingProfiler, Log, TEXT("[LogView] FLogFilteringAsyncTask::DoWork [%d to %d] by Text (\"%s\")"), StartIndex, EndIndex, *Filter.GetFilterText().ToString());
-
-		for (int Index = StartIndex; Index < EndIndex && !bCanceled; ++Index)
-		{
-			TSharedPtr<FLogMessageRecord> RecordPtr = LogView->GetCache().GetUncached(Index, true);
-
-			if (RecordPtr.IsValid() && Filter.IsMessageAllowed(*RecordPtr))
-			{
-				FilteredMessages.Add(Index);
-			}
-
-			bCanceled = LogView->IsFilteringAsyncTaskCancelRequested();
-		}
 	}
 	else
 	{
 		UE_LOG(TimingProfiler, Log, TEXT("[LogView] FLogFilteringAsyncTask::DoWork [%d to %d]"), StartIndex, EndIndex);
+	}
 
-		// If we do not filter by text, we use a faster way to read log messages (i.e. no message formatting is needed in this case).
-		for (int Index = StartIndex; Index < EndIndex && !bCanceled; ++Index)
-		{
-			TSharedPtr<FLogMessageRecord> RecordPtr = LogView->GetCache().GetUncached(Index, false);
-
-			if (RecordPtr.IsValid() && Filter.IsMessageAllowed(*RecordPtr))
-			{
-				FilteredMessages.Add(Index);
-			}
+	for (int32 Index = StartIndex; Index < EndIndex && !bCanceled; ++Index)
+	{
+		TSharedPtr<FLogMessageRecord> RecordPtr = LogView->GetCache().GetUncached(Index);
 
-			bCanceled = LogView->IsFilteringAsyncTaskCancelRequested();
+		if (RecordPtr.IsValid() && Filter.IsMessageAllowed(*RecordPtr))
+		{
+			FilteredMessages.Add(Index);
 		}
+
+		bCanceled = LogView->IsFilteringAsyncTaskCancelRequested();
 	}
 
 	if (bCanceled)
diff --git a/Engine/Source/Developer/TraceInsights/Private/Insights/ViewModels/LogFilter.h b/Engine/Source/Developer/TraceInsights/Private/Insights/ViewModels/LogFilter.h
--- a/Engine/Source/Developer/TraceInsights/Private/Insights/ViewModels/LogFilter.h
+++ b/Engine/Source/Developer/TraceInsights/Private/Insights/ViewModels/LogFilter.h
@@ -160,7 +160,7 @@ class FAsyncLogFilterProxy;
 class FLogFilteringAsyncTask : public FNonAbandonableTask
 {
 public:
-	FLogFilteringAsyncTask(int InStartIndex, int InEndIndex, const FLogFilter& InFilter, TSharedPtr<SLogView> InLogView)
+	FLogFilteringAsyncTask(int32 InStartIndex, int32 InEndIndex, const FLogFilter& InFilter, TSharedPtr<SLogView> InLogView)
 		: StartIndex(InStartIndex)
 		, EndIndex(InEndIndex)
 		, Filter(InFilter)
@@ -170,8 +170,8 @@ public:
 
 	void DoWork();
 
-	int GetStartIndex() const { return StartIndex; }
-	int GetEndIndex() const { return EndIndex; }
+	int32 GetStartIndex() const { return StartIndex; }
+	int32 GetEndIndex() const { return EndIndex; }
 	const FLogFilter& GetFilter() const { return Filter; }
 
 	const TArray<uint32>& GetFilteredMessages() const { return FilteredMessages; }
@@ -183,8 +183,8 @@ public:
 
 private:
 	// [StartIndex, EndIndex) is the range of log messages to filter; accessible from worker thread.
-	int StartIndex;
-	int EndIndex;
+	int32 StartIndex;
+	int32 EndIndex;
 
 	/** A copy of the filter settings. */
 	FLogFilter Filter;
