[TraceInsights]

diff --git a/Engine/Source/Developer/TraceInsights/Private/Insights/ViewModels/LogMessage.cpp b/Engine/Source/Developer/TraceInsights/Private/Insights/ViewModels/LogMessage.cpp
--- a/Engine/Source/Developer/TraceInsights/Private/Insights/ViewModels/LogMessage.cpp
+++ b/Engine/Source/Developer/TraceInsights/Private/Insights/ViewModels/LogMessage.cpp
@@ -39,7 +39,9 @@ FLogMessageRecord::FLogMessageRecord(const Trace::FLogMessage& TraceLogMessage)
 	// Strip the "Log" prefix.
 	FString CategoryStr(TraceLogMessage.Category->Name);
 	if (CategoryStr.StartsWith(TEXT("Log")))
+	{
 		CategoryStr = CategoryStr.RightChop(3);
+	}
 	Category = FText::FromString(CategoryStr);
 
 	// Strip the trailing whitespaces (ex. some messages ends with "\n" and we do not want the LogView rows to have an unnecessary increased height).
@@ -142,19 +144,19 @@ FLogMessageRecord& FLogMessageCache::Get(uint64 Index)
 
 	if (Session.IsValid())
 	{
-		Trace::FAnalysisSessionReadScope _(*Session.Get());
+		Trace::FAnalysisSessionReadScope SessionReadScope(*Session.Get());
 		Session->ReadLogProvider(
 			[this, Index](const Trace::ILogProvider& LogProvider)
-		{
-			LogProvider.ReadMessage(
-				Index,
-				[this, Index](const Trace::FLogMessage& Message)
 			{
-				FScopeLock Lock(&CriticalSection);
-				FLogMessageRecord Entry(Message);
-				Map.Add(Index, Entry);
+				LogProvider.ReadMessage(
+					Index,
+					[this, Index](const Trace::FLogMessage& Message)
+					{
+						FScopeLock Lock(&CriticalSection);
+						FLogMessageRecord Entry(Message);
+						Map.Add(Index, Entry);
+					});
 			});
-		});
 	}
 
 	{
@@ -170,23 +172,23 @@ FLogMessageRecord& FLogMessageCache::Get(uint64 Index)
 
 ////////////////////////////////////////////////////////////////////////////////////////////////////
 
-TSharedPtr<FLogMessageRecord> FLogMessageCache::GetUncached(uint64 Index, bool bResolveFormatString) const
+TSharedPtr<FLogMessageRecord> FLogMessageCache::GetUncached(uint64 Index) const
 {
 	TSharedPtr<FLogMessageRecord> EntryPtr;
 
 	if (Session.IsValid())
 	{
-		Trace::FAnalysisSessionReadScope _(*Session.Get());
+		Trace::FAnalysisSessionReadScope SessionReadScope(*Session.Get());
 		Session->ReadLogProvider(
-			[Index, bResolveFormatString, &EntryPtr](const Trace::ILogProvider& LogProvider)
-		{
-			LogProvider.ReadMessage(
-				Index,
-				[&EntryPtr](const Trace::FLogMessage& Message)
+			[Index, &EntryPtr](const Trace::ILogProvider& LogProvider)
 			{
-				EntryPtr = MakeShareable(new FLogMessageRecord(Message));
+				LogProvider.ReadMessage(
+					Index,
+					[&EntryPtr](const Trace::FLogMessage& Message)
+					{
+						EntryPtr = MakeShareable(new FLogMessageRecord(Message));
+					});
 			});
-		});
 	}
 
 	return EntryPtr;
diff --git a/Engine/Source/Developer/TraceInsights/Private/Insights/ViewModels/LogMessage.h b/Engine/Source/Developer/TraceInsights/Private/Insights/ViewModels/LogMessage.h
--- a/Engine/Source/Developer/TraceInsights/Private/Insights/ViewModels/LogMessage.h
+++ b/Engine/Source/Developer/TraceInsights/Private/Insights/ViewModels/LogMessage.h
@@ -51,7 +51,7 @@ public:
 	void Reset();
 
 	FLogMessageRecord& Get(uint64 Index);
-	TSharedPtr<FLogMessageRecord> GetUncached(uint64 Index, bool bResolveFormatString) const;
+	TSharedPtr<FLogMessageRecord> GetUncached(uint64 Index) const;
 
 private:
 	FCriticalSection CriticalSection;
diff --git a/Engine/Source/Developer/TraceInsights/Private/Insights/ViewModels/MarkersTimingTrack.cpp b/Engine/Source/Developer/TraceInsights/Private/Insights/ViewModels/MarkersTimingTrack.cpp
--- a/Engine/Source/Developer/TraceInsights/Private/Insights/ViewModels/MarkersTimingTrack.cpp
+++ b/Engine/Source/Developer/TraceInsights/Private/Insights/ViewModels/MarkersTimingTrack.cpp
@@ -70,25 +70,37 @@ void FMarkersTimingTrack::Reset()
 
 ////////////////////////////////////////////////////////////////////////////////////////////////////
 
-void FMarkersTimingTrack::UpdateHoveredState(float MX, float MY, const FTimingTrackViewport& Viewport)
+void FMarkersTimingTrack::UpdateHeight()
 {
-	if (MY >= Y && MY < Y + H)
+	static const float BookmarksTrackHeight = 14.0f;
+	static const float TimeMarkersTrackHeight = 28.0f;
+
+	if (bUseOnlyBookmarks)
 	{
-		bIsHovered = true;
-		TargetHoveredAnimPercent = 1.0f;
+		SetHeight(BookmarksTrackHeight);
+	}
+	else
+	{
+		SetHeight(TimeMarkersTrackHeight);
+	}
+}
 
-		if (MX < 80.0f && MY < Y + 14.0f)
-		{
-			bIsHeaderHovered = true;
-		}
-		else
-		{
-			bIsHeaderHovered = false;
-		}
+////////////////////////////////////////////////////////////////////////////////////////////////////
+
+void FMarkersTimingTrack::UpdateHoveredState(float MouseX, float MouseY, const FTimingTrackViewport& Viewport)
+{
+	static const float HeaderWidth = 80.0f;
+	static const float HeaderHeight = 14.0f;
+
+	if (MouseY >= GetPosY() && MouseY < GetPosY() + GetHeight())
+	{
+		SetHoveredState(true);
+		SetHeaderHoveredState(MouseX < HeaderWidth && MouseY < GetPosY() + HeaderHeight);
+		TargetHoveredAnimPercent = 1.0f;
 	}
 	else
 	{
-		bIsHovered = false;
+		SetHoveredState(false);
 		TargetHoveredAnimPercent = 0.0f;
 	}
 }
@@ -124,42 +136,39 @@ void FMarkersTimingTrack::Tick(const double InCurrentTime, const float InDeltaTi
 
 void FMarkersTimingTrack::Update(const FTimingTrackViewport& InViewport)
 {
-	FTimeMarkerTrackBuilder CB(*this, InViewport);
-	CB.Begin();
+	FTimeMarkerTrackBuilder Builder(*this, InViewport);
 
 	TSharedPtr<const Trace::IAnalysisSession> Session = FInsightsManager::Get()->GetSession();
 	if (Session.IsValid())
 	{
-		Trace::FAnalysisSessionReadScope _(*Session.Get());
+		Trace::FAnalysisSessionReadScope SessionReadScope(*Session.Get());
 
-		Session->ReadLogProvider([&CB](const Trace::ILogProvider& LogProvider)
+		Session->ReadLogProvider([&Builder](const Trace::ILogProvider& LogProvider)
 		{
-			CB.BeginLog(LogProvider);
+			Builder.BeginLog(LogProvider);
 
 			LogProvider.EnumerateMessages(
-				CB.GetViewport().StartTime,
-				CB.GetViewport().EndTime,
-				[&CB](const Trace::FLogMessage& Message) { CB.AddLogMessage(Message); });
+				Builder.GetViewport().StartTime,
+				Builder.GetViewport().EndTime,
+				[&Builder](const Trace::FLogMessage& Message) { Builder.AddLogMessage(Message); });
 
-			CB.EndLog();
+			Builder.EndLog();
 		});
 	}
-
-	CB.End();
 }
 
 ////////////////////////////////////////////////////////////////////////////////////////////////////
 
-void FMarkersTimingTrack::Draw(FDrawContext& DC, const FTimingTrackViewport& Viewport) const
+void FMarkersTimingTrack::Draw(FDrawContext& DrawContext, const FTimingTrackViewport& Viewport) const
 {
 	// Draw background.
-	DC.DrawBox(0.0f, Y, Viewport.Width, H, WhiteBrush, FLinearColor(0.04f, 0.04f, 0.04f, 1.0f));
-	DC.LayerId++;
+	DrawContext.DrawBox(0.0f, GetPosY(), Viewport.Width, GetHeight(), WhiteBrush, FLinearColor(0.04f, 0.04f, 0.04f, 1.0f));
+	DrawContext.LayerId++;
 
 	// Draw the track's header, in background.
 	if (!IsHovered() || CurrentHoveredAnimPercent < 1.0)
 	{
-		DrawHeader(DC, true);
+		DrawHeader(DrawContext, true);
 	}
 
 	//////////////////////////////////////////////////
@@ -169,8 +178,8 @@ void FMarkersTimingTrack::Draw(FDrawContext& DC, const FTimingTrackViewport& Vie
 	float BoxY, BoxH;
 	if (IsCollapsed())
 	{
-		BoxY = Y;
-		BoxH = H;
+		BoxY = GetPosY();
+		BoxH = GetHeight();
 	}
 	else
 	{
@@ -182,17 +191,16 @@ void FMarkersTimingTrack::Draw(FDrawContext& DC, const FTimingTrackViewport& Vie
 	for (int32 BoxIndex = 0; BoxIndex < NumBoxes; BoxIndex++)
 	{
 		const FTimeMarkerBoxInfo& Box = TimeMarkerBoxes[BoxIndex];
-		DC.DrawBox(Box.X, BoxY, Box.W, BoxH, WhiteBrush, Box.Color);
-		//NumDrawBoxes++;
+		DrawContext.DrawBox(Box.X, BoxY, Box.W, BoxH, WhiteBrush, Box.Color);
 	}
-	DC.LayerId++;
+	DrawContext.LayerId++;
 	NumDrawBoxes = NumBoxes;
 
 	//////////////////////////////////////////////////
 	// Draw texts (strings are already truncated).
 
-	const float CategoryY = Y + 2.0f;
-	const float MessageY = Y + (IsBookmarksTrack() ? 1.0f : 14.0f);
+	const float CategoryY = GetPosY() + 2.0f;
+	const float MessageY = GetPosY() + (IsBookmarksTrack() ? 1.0f : 14.0f);
 
 	int32 NumTexts = TimeMarkerTexts.Num();
 	for (int32 TextIndex = 0; TextIndex < NumTexts; TextIndex++)
@@ -201,46 +209,44 @@ void FMarkersTimingTrack::Draw(FDrawContext& DC, const FTimingTrackViewport& Vie
 
 		if (!IsBookmarksTrack() && TextInfo.Category.Len() > 0)
 		{
-			DC.DrawText(TextInfo.X, CategoryY, TextInfo.Category, Font, TextInfo.Color);
+			DrawContext.DrawText(TextInfo.X, CategoryY, TextInfo.Category, Font, TextInfo.Color);
 			NumDrawTexts++;
 		}
 
 		if (TextInfo.Message.Len() > 0)
 		{
-			DC.DrawText(TextInfo.X, MessageY, TextInfo.Message, Font, TextInfo.Color);
+			DrawContext.DrawText(TextInfo.X, MessageY, TextInfo.Message, Font, TextInfo.Color);
 			NumDrawTexts++;
 		}
 	}
-	DC.LayerId++;
-	//NumDrawTexts = NumTexts * 2;
+	DrawContext.LayerId++;
 
 	//////////////////////////////////////////////////
 
 	// When hovered, the track's header is draw on top.
 	if (IsHovered() || CurrentHoveredAnimPercent > 0.0)
 	{
-		DrawHeader(DC, false);
+		DrawHeader(DrawContext, false);
 	}
 }
 
-
 ////////////////////////////////////////////////////////////////////////////////////////////////////
 
-void FMarkersTimingTrack::DrawHeader(FDrawContext& DC, bool bFirstDraw) const
+void FMarkersTimingTrack::DrawHeader(FDrawContext& DrawContext, bool bFirstDraw) const
 {
 	const FString Name = IsBookmarksTrack() ? TEXT("Bookmarks") : TEXT("Logs");
 	//const float ArrowSize = 6.0f;
 	const float ArrowSizeX = 4.0f;
 	const float ArrowSizeY = 8.0f;
 	const float ArrowX = IsBookmarksTrack() ? 64.0f : 31.0f;
-	const float ArrowY = Y + 3.0f;
+	const float ArrowY = GetPosY() + 3.0f;
 	const float HeaderW = ArrowX + ArrowSizeX + 4.0f;
 	const float HeaderH = 14.0f;
 
 	if (!bFirstDraw)
 	{
-		DC.DrawBox(0.0f, Y, HeaderW, HeaderH, WhiteBrush, FLinearColor(0.04f, 0.04f, 0.04f, CurrentHoveredAnimPercent));
-		DC.LayerId++;
+		DrawContext.DrawBox(0.0f, GetPosY(), HeaderW, HeaderH, WhiteBrush, FLinearColor(0.04f, 0.04f, 0.04f, CurrentHoveredAnimPercent));
+		DrawContext.LayerId++;
 	}
 
 	FLinearColor Color;
@@ -258,41 +264,30 @@ void FMarkersTimingTrack::DrawHeader(FDrawContext& DC, bool bFirstDraw) const
 	}
 
 	// Draw "Bookmarks" or "Logs" text.
-	DC.DrawText(2.0f, Y + 1.0f, Name, Font, Color);
+	DrawContext.DrawText(2.0f, GetPosY() + 1.0f, Name, Font, Color);
 
 	if (IsCollapsed())
 	{
-		// Draw "down arrow".
-		//for (float A = 0.0f; A < ArrowSize; A += 1.0f)
-		//{
-		//	DC.DrawBox(ArrowX - A, ArrowY + ArrowSize - A , 2.0f * A + 1.0f, 1.0f, WhiteBrush, Color);
-		//}
-
 		// Draw "right empty arrow".
-		TArray<FVector2D> Points;
-		Points.SetNum(4);
-		Points[0] = FVector2D(ArrowX, ArrowY);
-		Points[1] = FVector2D(ArrowX, ArrowY + ArrowSizeY);
-		Points[2] = FVector2D(ArrowX + ArrowSizeX, ArrowY + ArrowSizeY / 2.0f);
-		Points[3] = FVector2D(ArrowX, ArrowY);
-		FSlateDrawElement::MakeLines(DC.ElementList, DC.LayerId, DC.Geometry.ToPaintGeometry(), Points, DC.DrawEffects, Color, false, 1.0f);
+		TArray<FVector2D> Points =
+		{
+			FVector2D(ArrowX, ArrowY),
+			FVector2D(ArrowX, ArrowY + ArrowSizeY),
+			FVector2D(ArrowX + ArrowSizeX, ArrowY + ArrowSizeY / 2.0f),
+			FVector2D(ArrowX, ArrowY)
+		};
+		FSlateDrawElement::MakeLines(DrawContext.ElementList, DrawContext.LayerId, DrawContext.Geometry.ToPaintGeometry(), Points, DrawContext.DrawEffects, Color, false, 1.0f);
 	}
 	else
 	{
-		// Draw "up arrow".
-		//for (float A = 0.0f; A < ArrowSize; A += 1.0f)
-		//{
-		//	DC.DrawBox(ArrowX - A, ArrowY + A, 2.0f * A + 1.0f, 1.0f, WhiteBrush, Color);
-		//}
-
 		// Draw "down-right filled arrow".
 		for (float A = 1.0f; A < ArrowSizeY; A += 1.0f)
 		{
-			DC.DrawBox(ArrowX - 3.0 + ArrowSizeY - A, ArrowY + A - 1.0f, A, 1.0f, WhiteBrush, Color);
+			DrawContext.DrawBox(ArrowX - 3.0 + ArrowSizeY - A, ArrowY + A - 1.0f, A, 1.0f, WhiteBrush, Color);
 		}
 	}
 
-	DC.LayerId++;
+	DrawContext.LayerId++;
 }
 
 ////////////////////////////////////////////////////////////////////////////////////////////////////
@@ -304,12 +299,6 @@ FTimeMarkerTrackBuilder::FTimeMarkerTrackBuilder(FMarkersTimingTrack& InTrack, c
 	, Viewport(InViewport)
 	, FontMeasureService(FSlateApplication::Get().GetRenderer()->GetFontMeasureService())
 	, Font(FCoreStyle::GetDefaultFontStyle("Regular", 8))
-{
-}
-
-////////////////////////////////////////////////////////////////////////////////////////////////////
-
-void FTimeMarkerTrackBuilder::Begin()
 {
 	Track.ResetCache();
 	Track.NumLogMessages = 0;
@@ -323,9 +312,10 @@ void FTimeMarkerTrackBuilder::BeginLog(const Trace::ILogProvider& LogProvider)
 
 	LastX1 = -1000.0f;
 	LastX2 = -1000.0f;
-	LastCategory = nullptr;
-	LastVerbosity = (ELogVerbosity::Type)0;
 	LastLogIndex = 0;
+	LastVerbosity = ELogVerbosity::NoLogging;
+	LastCategory = nullptr;
+	LastMessage = nullptr;
 }
 
 ////////////////////////////////////////////////////////////////////////////////////////////////////
@@ -351,7 +341,7 @@ void FTimeMarkerTrackBuilder::AddLogMessage(const Trace::FLogMessage& Message)
 		{
 			X = -1.0f;
 		}
-		AddTimeMarker(X, Message.Category->Name, Message.Verbosity, Message.Index);
+		AddTimeMarker(X, Message.Index, Message.Verbosity, Message.Category->Name, Message.Message);
 	}
 }
 
@@ -362,7 +352,9 @@ FLinearColor FTimeMarkerTrackBuilder::GetColorByCategory(const TCHAR* const Cate
 	// Strip the "Log" prefix.
 	FString CategoryStr(Category);
 	if (CategoryStr.StartsWith(TEXT("Log")))
+	{
 		CategoryStr = CategoryStr.RightChop(3);
+	}
 
 	uint32 Hash = 0;
 	for (const TCHAR* c = *CategoryStr; *c; ++c)
@@ -430,42 +422,39 @@ void FTimeMarkerTrackBuilder::Flush(float AvailableTextW)
 
 		if (AvailableTextW > 6.0f)
 		{
-			LogProviderPtr->ReadMessage(
-				LastLogIndex,
-				[this, AvailableTextW, &Color](const Trace::FLogMessage& Message)
+			// Strip the "Log" prefix.
+			FString CategoryStr(LastCategory);
+			if (CategoryStr.StartsWith(TEXT("Log")))
 			{
-				// Strip the "Log" prefix.
-				FString CategoryStr(Message.Category->Name);
-				if (CategoryStr.StartsWith(TEXT("Log")))
-					CategoryStr = CategoryStr.RightChop(3);
+				CategoryStr = CategoryStr.RightChop(3);
+			}
 
-				const int32 LastWholeCharacterIndexCategory = FontMeasureService->FindLastWholeCharacterIndexBeforeOffset(CategoryStr, Font, FMath::RoundToInt(AvailableTextW - 2.0f));
-				const int32 LastWholeCharacterIndexMessage = FontMeasureService->FindLastWholeCharacterIndexBeforeOffset(Message.Message, Font, FMath::RoundToInt(AvailableTextW - 2.0f));
+			const int32 LastWholeCharacterIndexCategory = FontMeasureService->FindLastWholeCharacterIndexBeforeOffset(CategoryStr, Font, FMath::RoundToInt(AvailableTextW - 2.0f));
+			const int32 LastWholeCharacterIndexMessage = FontMeasureService->FindLastWholeCharacterIndexBeforeOffset(LastMessage, Font, FMath::RoundToInt(AvailableTextW - 2.0f));
 
-				if (LastWholeCharacterIndexCategory >= 0 ||
-					LastWholeCharacterIndexMessage >= 0)
+			if (LastWholeCharacterIndexCategory >= 0 ||
+				LastWholeCharacterIndexMessage >= 0)
+			{
+				// Add new Text info to cache.
+				FTimeMarkerTextInfo& TextInfo = Track.TimeMarkerTexts[Track.TimeMarkerTexts.AddDefaulted()];
+				TextInfo.X = LastX2 + 2.0f;
+				TextInfo.Color = Color;
+				if (LastWholeCharacterIndexCategory >= 0)
 				{
-					// Add new Text info to cache.
-					FTimeMarkerTextInfo& TextInfo = Track.TimeMarkerTexts[Track.TimeMarkerTexts.AddDefaulted()];
-					TextInfo.X = LastX2 + 2.0f;
-					TextInfo.Color = Color;
-					if (LastWholeCharacterIndexCategory >= 0)
-					{
-						TextInfo.Category = CategoryStr.Left(LastWholeCharacterIndexCategory + 1);
-					}
-					if (LastWholeCharacterIndexMessage >= 0)
-					{
-						TextInfo.Message.AppendChars(Message.Message, LastWholeCharacterIndexMessage + 1);
-					}
+					TextInfo.Category = CategoryStr.Left(LastWholeCharacterIndexCategory + 1);
 				}
-			});
+				if (LastWholeCharacterIndexMessage >= 0)
+				{
+					TextInfo.Message.AppendChars(LastMessage, LastWholeCharacterIndexMessage + 1);
+				}
+			}
 		}
 	}
 }
 
 ////////////////////////////////////////////////////////////////////////////////////////////////////
 
-void FTimeMarkerTrackBuilder::AddTimeMarker(const float X, const TCHAR* const Category, const ELogVerbosity::Type Verbosity, const uint64 LogIndex)
+void FTimeMarkerTrackBuilder::AddTimeMarker(const float X, const uint64 LogIndex, const ELogVerbosity::Type Verbosity, const TCHAR* const Category, const TCHAR* Message)
 {
 	float W = X - LastX2;
 
@@ -524,6 +513,7 @@ void FTimeMarkerTrackBuilder::AddTimeMarker(const float X, const TCHAR* const Ca
 	LastCategory = Category;
 	LastVerbosity = Verbosity;
 	LastLogIndex = LogIndex;
+	LastMessage = Message;
 }
 
 ////////////////////////////////////////////////////////////////////////////////////////////////////
@@ -535,10 +525,4 @@ void FTimeMarkerTrackBuilder::EndLog()
 
 ////////////////////////////////////////////////////////////////////////////////////////////////////
 
-void FTimeMarkerTrackBuilder::End()
-{
-}
-
-////////////////////////////////////////////////////////////////////////////////////////////////////
-
 #undef LOCTEXT_NAMESPACE
diff --git a/Engine/Source/Developer/TraceInsights/Private/Insights/ViewModels/MarkersTimingTrack.h b/Engine/Source/Developer/TraceInsights/Private/Insights/ViewModels/MarkersTimingTrack.h
--- a/Engine/Source/Developer/TraceInsights/Private/Insights/ViewModels/MarkersTimingTrack.h
+++ b/Engine/Source/Developer/TraceInsights/Private/Insights/ViewModels/MarkersTimingTrack.h
@@ -20,11 +20,6 @@ class FTimingTrackViewport;
 
 ////////////////////////////////////////////////////////////////////////////////////////////////////
 
-static const float BookmarksTrackHeight = 14.0f;
-static const float TimeMarkersTrackHeight = 28.0f;
-
-////////////////////////////////////////////////////////////////////////////////////////////////////
-
 struct FTimeMarkerBoxInfo
 {
 	float X;
@@ -58,7 +53,7 @@ public:
 	void ToggleCollapsed() { bIsCollapsed = !bIsCollapsed; }
 
 	bool IsBookmarksTrack() const { return bUseOnlyBookmarks; }
-	void SetIsBookmarksTrackFlag(bool bInUseOnlyBookmarks)
+	void SetBookmarksTrackFlag(bool bInUseOnlyBookmarks)
 	{
 		bUseOnlyBookmarks = bInUseOnlyBookmarks;
 		UpdateHeight();
@@ -69,31 +64,28 @@ public:
 	void ClearDirtyFlag() { bIsDirty = false; }
 
 	// Stats
-	int GetNumLogMessages() const { return NumLogMessages; }
-	int GetNumBoxes() const { return TimeMarkerBoxes.Num(); }
-	int GetNumTexts() const { return TimeMarkerTexts.Num(); }
+	int32 GetNumLogMessages() const { return NumLogMessages; }
+	int32 GetNumBoxes() const { return TimeMarkerBoxes.Num(); }
+	int32 GetNumTexts() const { return TimeMarkerTexts.Num(); }
 
-	virtual void UpdateHoveredState(float MX, float MY, const FTimingTrackViewport& Viewport) override;
-	virtual void Tick(const double InCurrentTime, const float InDeltaTime) override;
+	virtual void UpdateHoveredState(float MouseX, float MouseY, const FTimingTrackViewport& Viewport) override;
+	virtual void Tick(const double CurrentTime, const float DeltaTime) override;
 	virtual void Update(const FTimingTrackViewport& Viewport) override;
 
-	void Draw(FDrawContext& DC, const FTimingTrackViewport& Viewport) const;
+	void Draw(FDrawContext& DrawContext, const FTimingTrackViewport& Viewport) const;
 
-protected:
+private:
 	void ResetCache()
 	{
 		TimeMarkerBoxes.Reset();
 		TimeMarkerTexts.Reset();
 	}
 
-	void UpdateHeight()
-	{
-		H = bUseOnlyBookmarks ? BookmarksTrackHeight : TimeMarkersTrackHeight;
-	}
+	void UpdateHeight();
 
-	void DrawHeader(FDrawContext& DC, bool bFirstDraw) const;
+	void DrawHeader(FDrawContext& DrawContext, bool bFirstDraw) const;
 
-protected:
+private:
 	TArray<FTimeMarkerBoxInfo> TimeMarkerBoxes;
 	TArray<FTimeMarkerTextInfo> TimeMarkerTexts;
 
@@ -105,9 +97,9 @@ protected:
 	float CurrentHoveredAnimPercent;
 
 	// Stats
-	int NumLogMessages;
-	mutable int NumDrawBoxes;
-	mutable int NumDrawTexts;
+	int32 NumLogMessages;
+	mutable int32 NumDrawBoxes;
+	mutable int32 NumDrawTexts;
 
 	// Slate resources
 	const FSlateBrush* WhiteBrush;
@@ -116,28 +108,31 @@ protected:
 
 ////////////////////////////////////////////////////////////////////////////////////////////////////
 
-class FTimeMarkerTrackBuilder : public FNoncopyable
+class FTimeMarkerTrackBuilder
 {
 public:
 	FTimeMarkerTrackBuilder(FMarkersTimingTrack& InTrack, const FTimingTrackViewport& InViewport);
-	//~FTimeMarkerTrackBuilder();
+
+	/**
+	 * Non-copyable
+	 */
+	FTimeMarkerTrackBuilder(const FTimeMarkerTrackBuilder&) = delete;
+	FTimeMarkerTrackBuilder& operator=(const FTimeMarkerTrackBuilder&) = delete;
 
 	const FTimingTrackViewport& GetViewport() { return Viewport; }
 
-	void Begin();
 	void BeginLog(const Trace::ILogProvider& LogProvider);
 	void AddLogMessage(const Trace::FLogMessage& Message);
 	void EndLog();
-	void End();
 
 	static FLinearColor GetColorByCategory(const TCHAR* const Category);
 	static FLinearColor GetColorByVerbosity(const ELogVerbosity::Type Verbosity);
 
 private:
 	void Flush(float AvailableTextW);
-	void AddTimeMarker(const float X, const TCHAR* const Category, const ELogVerbosity::Type Verbosity, const uint64 LogIndex);
+	void AddTimeMarker(const float X, const uint64 LogIndex, const ELogVerbosity::Type Verbosity, const TCHAR* const Category, const TCHAR* Message);
 
-public:
+private:
 	FMarkersTimingTrack& Track;
 	const FTimingTrackViewport& Viewport;
 
@@ -148,9 +143,10 @@ public:
 
 	float LastX1;
 	float LastX2;
-	const TCHAR* LastCategory;
-	ELogVerbosity::Type LastVerbosity;
 	uint64 LastLogIndex;
+	ELogVerbosity::Type LastVerbosity;
+	const TCHAR* LastCategory;
+	const TCHAR* LastMessage;
 };
 
 ////////////////////////////////////////////////////////////////////////////////////////////////////
diff --git a/Engine/Source/Developer/TraceInsights/Private/Insights/ViewModels/TimeRulerTrack.cpp b/Engine/Source/Developer/TraceInsights/Private/Insights/ViewModels/TimeRulerTrack.cpp
--- a/Engine/Source/Developer/TraceInsights/Private/Insights/ViewModels/TimeRulerTrack.cpp
+++ b/Engine/Source/Developer/TraceInsights/Private/Insights/ViewModels/TimeRulerTrack.cpp
@@ -32,30 +32,35 @@ FTimeRulerTrack::~FTimeRulerTrack()
 
 ////////////////////////////////////////////////////////////////////////////////////////////////////
 
-void FTimeRulerTrack::UpdateHoveredState(float MX, float MY, const FTimingTrackViewport& Viewport)
+void FTimeRulerTrack::Reset()
 {
-	if (MY >= Y && MY < Y + H)
+	FBaseTimingTrack::Reset();
+
+	static const float TimeRulerHeight = 22.0f;
+	SetHeight(TimeRulerHeight);
+}
+
+////////////////////////////////////////////////////////////////////////////////////////////////////
+
+void FTimeRulerTrack::UpdateHoveredState(float MouseX, float MouseY, const FTimingTrackViewport& Viewport)
+{
+	//static const float HeaderWidth = 100.0f;
+	//static const float HeaderHeight = 14.0f;
+
+	if (MouseY >= GetPosY() && MouseY < GetPosY() + GetHeight())
 	{
-		bIsHovered = true;
-
-		//if (MX < 100.0f && MY < Y + 14.0f)
-		//{
-		//	bIsHeaderHovered = true;
-		//}
-		//else
-		//{
-		//	bIsHeaderHovered = false;
-		//}
+		SetHoveredState(true);
+		//SetHeaderHoveredState(MouseX < HeaderWidth && MouseY < GetPosY() + HeaderHeight);
 	}
 	else
 	{
-		bIsHovered = false;
+		SetHoveredState(false);
 	}
 }
 
 ////////////////////////////////////////////////////////////////////////////////////////////////////
 
-void FTimeRulerTrack::Draw(FDrawContext& DC, const FTimingTrackViewport& Viewport, const FVector2D& MousePosition, const bool bIsSelecting, const double SelectionStartTime, const double SelectionEndTime) const
+void FTimeRulerTrack::Draw(FDrawContext& DrawContext, const FTimingTrackViewport& Viewport, const FVector2D& MousePosition, const bool bIsSelecting, const double SelectionStartTime, const double SelectionEndTime) const
 {
 	const float MinorTickMark = 5.0f;
 	const float MajorTickMark = 20 * MinorTickMark;
@@ -63,7 +68,7 @@ void FTimeRulerTrack::Draw(FDrawContext& DC, const FTimingTrackViewport& Viewpor
 	const float MinorTickMarkHeight = 5.0f;
 	const float MajorTickMarkHeight = 9.0f;
 
-	const float TextY = Y + MajorTickMarkHeight;
+	const float TextY = GetPosY() + MajorTickMarkHeight;
 
 	double MinorTickMarkTime = Viewport.GetDurationForViewportDX(MinorTickMark);
 	double MajorTickMarkTime = Viewport.GetDurationForViewportDX(MajorTickMark);
@@ -75,23 +80,23 @@ void FTimeRulerTrack::Draw(FDrawContext& DC, const FTimingTrackViewport& Viewpor
 	float MajorOX = static_cast<float>(FMath::RoundToDouble(MajorN * static_cast<double>(MajorTickMark) - VX));
 
 	// Draw the time ruler's background.
-	DC.DrawBox(0.0f, Y, Viewport.Width, TimeRulerHeight, WhiteBrush, FLinearColor(0.09f, 0.09f, 0.09f, 1.0f));
-	DC.LayerId++;
+	DrawContext.DrawBox(0.0f, GetPosY(), Viewport.Width, GetHeight(), WhiteBrush, FLinearColor(0.09f, 0.09f, 0.09f, 1.0f));
+	DrawContext.LayerId++;
 
 	// Draw the minor tick marks.
 	for (float X = MinorOX; X < Viewport.Width; X += MinorTickMark)
 	{
-		const bool bIsTenth = ((int)(((X - MajorOX) / MinorTickMark) + 0.4f) % 2 == 0);
+		const bool bIsTenth = ((int32)(((X - MajorOX) / MinorTickMark) + 0.4f) % 2 == 0);
 		const float MinorTickH = bIsTenth ? MinorTickMarkHeight : MinorTickMarkHeight - 1.0f;
-		DC.DrawBox(X, Y, 1.0f, MinorTickH, WhiteBrush,
+		DrawContext.DrawBox(X, GetPosY(), 1.0f, MinorTickH, WhiteBrush,
 			bIsTenth ? FLinearColor(0.3f, 0.3f, 0.3f, 1.0f) : FLinearColor(0.25f, 0.25f, 0.25f, 1.0f));
 	}
 	// Draw the major tick marks.
 	for (float X = MajorOX; X < Viewport.Width; X += MajorTickMark)
 	{
-		DC.DrawBox(X, Y, 1.0f, MajorTickMarkHeight, WhiteBrush, FLinearColor(0.4f, 0.4f, 0.4f, 1.0f));
+		DrawContext.DrawBox(X, GetPosY(), 1.0f, MajorTickMarkHeight, WhiteBrush, FLinearColor(0.4f, 0.4f, 0.4f, 1.0f));
 	}
-	DC.LayerId++;
+	DrawContext.LayerId++;
 
 	const double DT = static_cast<double>(MajorTickMark) / Viewport.ScaleX;
 	const double Precision = FMath::Max(DT / 10.0, TimeUtils::Nanosecond);
@@ -104,10 +109,10 @@ void FTimeRulerTrack::Draw(FDrawContext& DC, const FTimingTrackViewport& Viewpor
 		const double T = Viewport.SlateUnitsToTime(X);
 		FString Text = TimeUtils::FormatTime(T, Precision);
 		const float TextWidth = FontMeasureService->Measure(Text, Font).X;
-		DC.DrawText(X - TextWidth / 2, TextY, Text, Font,
+		DrawContext.DrawText(X - TextWidth / 2, TextY, Text, Font,
 			(T < 0 || T >= Viewport.MaxValidTime) ? FLinearColor(0.7f, 0.5f, 0.5f, 1.0f) : FLinearColor(0.8f, 0.8f, 0.8f, 1.0f));
 	}
-	DC.LayerId++;
+	DrawContext.LayerId++;
 
 	bool bShowMousePos = !MousePosition.IsZero();
 	if (bShowMousePos)
@@ -121,7 +126,7 @@ void FTimeRulerTrack::Draw(FDrawContext& DC, const FTimingTrackViewport& Viewpor
 
 		const double MousePosTime = Viewport.SlateUnitsToTime(MousePosition.X);
 		const double MousePosPrecision = FMath::Max(DT / 100.0, TimeUtils::Nanosecond);
-		if (MousePosition.Y >= Y && MousePosition.Y < Y + H)
+		if (MousePosition.Y >= GetPosY() && MousePosition.Y < GetPosY() + GetHeight())
 		{
 			// If mouse is hovering the time ruller, format time with a better precision (split seconds in ms, us, ns and ps).
 			MousePosText = TimeUtils::FormatTimeSplit(MousePosTime, MousePosPrecision);
@@ -148,30 +153,34 @@ void FTimeRulerTrack::Draw(FDrawContext& DC, const FTimingTrackViewport& Viewpor
 			float SelectionX1 = Viewport.TimeToSlateUnitsRounded(SelectionStartTime);
 			float SelectionX2 = Viewport.TimeToSlateUnitsRounded(SelectionEndTime);
 			if (X - SelectionX1 > 1.0f)
+			{
 				X = SelectionX2 + W / 2;
+			}
 			if (SelectionX2 - X > 1.0f)
+			{
 				X = SelectionX1 - W / 2;
+			}
 			MousePosTextForegroundColor = FLinearColor(FColor(32, 64, 128, 255));
 		}
 		else
 		{
 			// Draw horizontal line at mouse position.
-			//DC.DrawBox(0.0f, MousePosition.Y, Viewport.Width, 1.0f, WhiteBrush, MousePosLineColor);
+			//DrawContext.DrawBox(0.0f, MousePosition.Y, Viewport.Width, 1.0f, WhiteBrush, MousePosLineColor);
 
 			// Draw vertical line at mouse position.
-			DC.DrawBox(MousePosition.X, 0.0f, 1.0f, Viewport.Height, WhiteBrush, MousePosLineColor);
+			DrawContext.DrawBox(MousePosition.X, 0.0f, 1.0f, Viewport.Height, WhiteBrush, MousePosLineColor);
 
 			// Stroke the vertical line above current time box.
-			DC.DrawBox(MousePosition.X, 0.0f, 1.0f, TextY, WhiteBrush, MousePosTextBackgroundColor);
+			DrawContext.DrawBox(MousePosition.X, 0.0f, 1.0f, TextY, WhiteBrush, MousePosTextBackgroundColor);
 		}
 
 		// Fill the current time box.
-		DC.DrawBox(X - W / 2, TextY, W, 12.0f, WhiteBrush, MousePosTextBackgroundColor);
-		DC.LayerId++;
+		DrawContext.DrawBox(X - W / 2, TextY, W, 12.0f, WhiteBrush, MousePosTextBackgroundColor);
+		DrawContext.LayerId++;
 
 		// Draw current time text.
-		DC.DrawText(X - MousePosTextWidth / 2, TextY, MousePosText, Font, MousePosTextForegroundColor);
-		DC.LayerId++;
+		DrawContext.DrawText(X - MousePosTextWidth / 2, TextY, MousePosText, Font, MousePosTextForegroundColor);
+		DrawContext.LayerId++;
 	}
 }
 
diff --git a/Engine/Source/Developer/TraceInsights/Private/Insights/ViewModels/TimeRulerTrack.h b/Engine/Source/Developer/TraceInsights/Private/Insights/ViewModels/TimeRulerTrack.h
--- a/Engine/Source/Developer/TraceInsights/Private/Insights/ViewModels/TimeRulerTrack.h
+++ b/Engine/Source/Developer/TraceInsights/Private/Insights/ViewModels/TimeRulerTrack.h
@@ -14,26 +14,16 @@ class FTimingTrackViewport;
 
 ////////////////////////////////////////////////////////////////////////////////////////////////////
 
-static const float TimeRulerHeight = 22.0f;
-
-////////////////////////////////////////////////////////////////////////////////////////////////////
-
 class FTimeRulerTrack : public FBaseTimingTrack
 {
 public:
 	FTimeRulerTrack(uint64 InTrackId);
 	virtual ~FTimeRulerTrack();
 
-	virtual void Reset() override
-	{
-		FBaseTimingTrack::Reset();
-
-		H = TimeRulerHeight;
-	}
-
-	virtual void UpdateHoveredState(float MX, float MY, const FTimingTrackViewport& Viewport);
+	virtual void Reset() override;
+	virtual void UpdateHoveredState(float MouseX, float MouseY, const FTimingTrackViewport& Viewport);
 
-	void Draw(FDrawContext& DC, const FTimingTrackViewport& Viewport,
+	void Draw(FDrawContext& DrawContext, const FTimingTrackViewport& Viewport,
 				const FVector2D& MousePosition = FVector2D(0.0f, 0.0f),
 				const bool bIsSelecting = false,
 				const double SelectionStartTime = 0.0,
diff --git a/Engine/Source/Developer/TraceInsights/Private/Insights/ViewModels/TimerGroupingAndSorting.h b/Engine/Source/Developer/TraceInsights/Private/Insights/ViewModels/TimerGroupingAndSorting.h
--- a/Engine/Source/Developer/TraceInsights/Private/Insights/ViewModels/TimerGroupingAndSorting.h
+++ b/Engine/Source/Developer/TraceInsights/Private/Insights/ViewModels/TimerGroupingAndSorting.h
@@ -8,19 +8,19 @@
 ////////////////////////////////////////////////////////////////////////////////////////////////////
 // Grouping
 
-/// Enumerates types of grouping or sorting for the timer nodes.
-enum class ETimerGroupingMode : uint32
+/** Enumerates types of grouping or sorting for the timer nodes. */
+enum class ETimerGroupingMode
 {
-	/// Creates a single group for all timers.
+	/** Creates a single group for all timers. */
 	Flat,
 
-	/// Creates one group for one letter.
+	/** Creates one group for one letter. */
 	ByName,
 
-	/// Creates groups based on timer metadata group names.
+	/** Creates groups based on timer metadata group names. */
 	ByMetaGroupName,
 
-	/// Creates one group for each timer type.
+	/** Creates one group for each timer type. */
 	ByType,
 
 	ByTotalInclusiveTime,
@@ -29,11 +29,11 @@ enum class ETimerGroupingMode : uint32
 
 	ByInstanceCount,
 
-	/// Invalid enum type, may be used as a number of enumerations.
+	/** Invalid enum type, may be used as a number of enumerations. */
 	InvalidOrMax,
 };
 
-/// Type definition for shared pointers to instances of ETimerGroupingMode.
+/** Type definition for shared pointers to instances of ETimerGroupingMode. */
 typedef TSharedPtr<ETimerGroupingMode> ETimerGroupingModePtr;
 
 ////////////////////////////////////////////////////////////////////////////////////////////////////
diff --git a/Engine/Source/Developer/TraceInsights/Private/Insights/ViewModels/TimerNode.cpp b/Engine/Source/Developer/TraceInsights/Private/Insights/ViewModels/TimerNode.cpp
--- a/Engine/Source/Developer/TraceInsights/Private/Insights/ViewModels/TimerNode.cpp
+++ b/Engine/Source/Developer/TraceInsights/Private/Insights/ViewModels/TimerNode.cpp
@@ -6,16 +6,16 @@
 
 ////////////////////////////////////////////////////////////////////////////////////////////////////
 
-void FTimerNode::ResetStats()
+void FTimerNode::ResetAggregatedStats()
 {
-	Stats.Reset();
+	AggregatedStats.Reset();
 }
 
 ////////////////////////////////////////////////////////////////////////////////////////////////////
 
-void FTimerNode::SetStats(FTimerStats& InStats)
+void FTimerNode::SetAggregatedStats(FTimerAggregatedStats& InAggregatedStats)
 {
-	Stats = InStats;
+	AggregatedStats = InAggregatedStats;
 }
 
 ////////////////////////////////////////////////////////////////////////////////////////////////////
diff --git a/Engine/Source/Developer/TraceInsights/Private/Insights/ViewModels/TimerNode.h b/Engine/Source/Developer/TraceInsights/Private/Insights/ViewModels/TimerNode.h
--- a/Engine/Source/Developer/TraceInsights/Private/Insights/ViewModels/TimerNode.h
+++ b/Engine/Source/Developer/TraceInsights/Private/Insights/ViewModels/TimerNode.h
@@ -18,7 +18,7 @@ typedef TSharedRef<const class FTimerNode> FTimerNodeRefConst;
 /** Type definition for weak references to instances of FTimerNode. */
 typedef TWeakPtr<class FTimerNode> FTimerNodeWeak;
 
-enum ETimerNodeType
+enum class ETimerNodeType
 {
 	/** The TimerNode is a CPU Scope timer. */
 	CpuScope,
@@ -36,8 +36,13 @@ enum ETimerNodeType
 	InvalidOrMax,
 };
 
-struct FTimerStats
+struct FTimerAggregatedStats
 {
+	//FAggregatedStats InclusiveTime; /** Aggregated stats for inclusive time, in seconds. */
+	//FAggregatedStats ExclusiveTime; /** Aggregated stats for exclusive time, in seconds. */
+	////FAggregatedStats InstanceCountPerFrame; /** Aggregated stats for variation of number of instances per frame. */
+	//uint64 GetInstanceCount() { return InclusiveTime.Count; } /** Number of instances. */
+
 	double TotalInclusiveTime; /** Total inclusive time, in seconds. */
 	double MinInclusiveTime; /** Min inclusive time, in seconds. */
 	double MaxInclusiveTime; /** Max inclusive time, in seconds. */
@@ -51,20 +56,16 @@ struct FTimerStats
 	double MedianExclusiveTime; /** Median exclusive time, in seconds. */
 
 	uint64 InstanceCount; /** Number of instances. */
-	//uint64 MinInstanceCountPerFrame;
-	//uint64 MaxInstanceCountPerFrame;
-	//uint64 AverageInstanceCountPerFrame;
-	//uint64 MedianInstanceCountPerFrame;
 
-	FTimerStats()
+	FTimerAggregatedStats()
 		: TotalInclusiveTime(0.0)
 		, MinInclusiveTime(DBL_MAX)
-		, MaxInclusiveTime(DBL_MIN)
+		, MaxInclusiveTime(-DBL_MAX)
 		, AverageInclusiveTime(0.0)
 		, MedianInclusiveTime(0.0)
 		, TotalExclusiveTime(0.0)
 		, MinExclusiveTime(DBL_MAX)
-		, MaxExclusiveTime(DBL_MIN)
+		, MaxExclusiveTime(-DBL_MAX)
 		, AverageExclusiveTime(0.0)
 		, MedianExclusiveTime(0.0)
 		, InstanceCount(0)
@@ -75,12 +76,12 @@ struct FTimerStats
 	{
 		TotalInclusiveTime = 0.0;
 		MinInclusiveTime = DBL_MAX;
-		MaxInclusiveTime = DBL_MIN;
+		MaxInclusiveTime = -DBL_MAX;
 		AverageInclusiveTime = 0.0;
 		MedianInclusiveTime = 0.0;
 		TotalExclusiveTime = 0.0;
 		MinExclusiveTime = DBL_MAX;
-		MaxExclusiveTime = DBL_MIN;
+		MaxExclusiveTime = -DBL_MAX;
 		AverageExclusiveTime = 0.0;
 		MedianExclusiveTime = 0.0;
 		InstanceCount = 0;
@@ -104,7 +105,7 @@ public:
 		, Type(InType)
 		, bForceExpandGroupNode(false)
 	{
-		ResetStats();
+		ResetAggregatedStats();
 	}
 
 	/** Initialization constructor for the group node. */
@@ -114,7 +115,7 @@ public:
 		, Type(ETimerNodeType::Group)
 		, bForceExpandGroupNode(false)
 	{
-		ResetStats();
+		ResetAggregatedStats();
 	}
 
 	/**
@@ -163,16 +164,16 @@ public:
 	}
 
 	/**
-	 * @return the stats of this timer.
+	 * @return the aggregated stats for this timer.
 	 */
-	const FTimerStats& GetStats() const
+	const FTimerAggregatedStats& GetAggregatedStats() const
 	{
-		return Stats;
+		return AggregatedStats;
 	}
 
-	void ResetStats();
+	void ResetAggregatedStats();
 
-	void SetStats(FTimerStats& Stats);
+	void SetAggregatedStats(FTimerAggregatedStats& AggregatedStats);
 
 	/**
 	 * @return a const reference to the child nodes of this group.
@@ -259,8 +260,8 @@ protected:
 	/** Holds the type of this timer; for the group, this is Group. */
 	const ETimerNodeType Type;
 
-	/** Stats */
-	FTimerStats Stats;
+	/** Aggregated stats. */
+	FTimerAggregatedStats AggregatedStats;
 
 	/** Children of this node. */
 	TArray<FTimerNodePtr> Children;
diff --git a/Engine/Source/Developer/TraceInsights/Private/Insights/ViewModels/TimerNodeHelper.cpp b/Engine/Source/Developer/TraceInsights/Private/Insights/ViewModels/TimerNodeHelper.cpp
--- a/Engine/Source/Developer/TraceInsights/Private/Insights/ViewModels/TimerNodeHelper.cpp
+++ b/Engine/Source/Developer/TraceInsights/Private/Insights/ViewModels/TimerNodeHelper.cpp
@@ -11,14 +11,14 @@
 
 FText TimerNodeTypeHelper::ToName(const ETimerNodeType NodeType)
 {
-	static_assert(ETimerNodeType::InvalidOrMax == 4, "Not all cases are handled in switch below!?");
+	static_assert(static_cast<int>(ETimerNodeType::InvalidOrMax) == 4, "Not all cases are handled in switch below!?");
 	switch (NodeType)
 	{
-	case ETimerNodeType::GpuScope:			return LOCTEXT("Timer_Name_GpuScope", "GPU");
-	case ETimerNodeType::ComputeScope:		return LOCTEXT("Timer_Name_ComputeScope", "Compute");
-	case ETimerNodeType::CpuScope:			return LOCTEXT("Timer_Name_CpuScope", "CPU");
-	case ETimerNodeType::Group:				return LOCTEXT("Timer_Name_Group", "Group"); //TODO
-	default: return LOCTEXT("InvalidOrMax", "InvalidOrMax");
+		case ETimerNodeType::GpuScope:		return LOCTEXT("Timer_Name_Gpu", "GPU");
+		case ETimerNodeType::ComputeScope:	return LOCTEXT("Timer_Name_Compute", "Compute");
+		case ETimerNodeType::CpuScope:		return LOCTEXT("Timer_Name_Cpu", "CPU");
+		case ETimerNodeType::Group:			return LOCTEXT("Timer_Name_Group", "Group");
+		default:							return LOCTEXT("InvalidOrMax", "InvalidOrMax");
 	}
 }
 
@@ -26,14 +26,14 @@ FText TimerNodeTypeHelper::ToName(const ETimerNodeType NodeType)
 
 FText TimerNodeTypeHelper::ToDescription(const ETimerNodeType NodeType)
 {
-	static_assert(ETimerNodeType::InvalidOrMax == 4, "Not all cases are handled in switch below!?");
+	static_assert(static_cast<int>(ETimerNodeType::InvalidOrMax) == 4, "Not all cases are handled in switch below!?");
 	switch (NodeType)
 	{
-	case ETimerNodeType::GpuScope:			return LOCTEXT("Timer_Desc_GpuScope", "GPU scope timer");
-	case ETimerNodeType::ComputeScope:		return LOCTEXT("Timer_Desc_ComputeScope", "Compute scope timer");
-	case ETimerNodeType::CpuScope:			return LOCTEXT("Timer_Desc_CpuScope", "CPU scope timer");
-	case ETimerNodeType::Group:				return LOCTEXT("Timer_Desc_Group", "Group timer node"); //TODO
-	default: return LOCTEXT("InvalidOrMax", "InvalidOrMax");
+		case ETimerNodeType::GpuScope:		return LOCTEXT("Timer_Desc_Gpu", "GPU scope timer");
+		case ETimerNodeType::ComputeScope:	return LOCTEXT("Timer_Desc_Compute", "Compute scope timer");
+		case ETimerNodeType::CpuScope:		return LOCTEXT("Timer_Desc_Cpu", "CPU scope timer");
+		case ETimerNodeType::Group:			return LOCTEXT("Timer_Desc_Group", "Group timer node");
+		default:							return LOCTEXT("InvalidOrMax", "InvalidOrMax");
 	}
 }
 
@@ -41,14 +41,14 @@ FText TimerNodeTypeHelper::ToDescription(const ETimerNodeType NodeType)
 
 FName TimerNodeTypeHelper::ToBrushName(const ETimerNodeType NodeType)
 {
-	static_assert(ETimerNodeType::InvalidOrMax == 4, "Not all cases are handled in switch below!?");
+	static_assert(static_cast<int>(ETimerNodeType::InvalidOrMax) == 4, "Not all cases are handled in switch below!?");
 	switch (NodeType)
 	{
-	case ETimerNodeType::GpuScope:			return TEXT("Profiler.FiltersAndPresets.StatTypeIcon"); //im:TODO
-	case ETimerNodeType::ComputeScope:		return TEXT("Profiler.FiltersAndPresets.StatTypeIcon"); //im:TODO
-	case ETimerNodeType::CpuScope:			return TEXT("Profiler.FiltersAndPresets.StatTypeIcon"); //im:TODO
-	case ETimerNodeType::Group:				return TEXT("Profiler.Misc.GenericGroup"); //TODO
-	default: return NAME_None;
+		case ETimerNodeType::GpuScope:		return TEXT("Profiler.FiltersAndPresets.StatTypeIcon"); //TODO: "Icons.GpuTimer"
+		case ETimerNodeType::ComputeScope:	return TEXT("Profiler.FiltersAndPresets.StatTypeIcon"); //TODO: "Icons.ComputeTimer"
+		case ETimerNodeType::CpuScope:		return TEXT("Profiler.FiltersAndPresets.StatTypeIcon"); //TODO: "Icons.CpuTimer"
+		case ETimerNodeType::Group:			return TEXT("Profiler.Misc.GenericGroup"); //TODO: "Icons.GenericGroup"
+		default:							return NAME_None;
 	}
 }
 
@@ -56,21 +56,21 @@ FName TimerNodeTypeHelper::ToBrushName(const ETimerNodeType NodeType)
 
 const FSlateBrush* TimerNodeTypeHelper::GetIconForGroup()
 {
-	return FEditorStyle::GetBrush(TEXT("Profiler.Misc.GenericGroup")); //TODO
+	return FEditorStyle::GetBrush(TEXT("Profiler.Misc.GenericGroup")); //TODO: FInsightsStyle::GetBrush(TEXT("Icons.GenericGroup"));
 }
 
 ////////////////////////////////////////////////////////////////////////////////////////////////////
 
 const FSlateBrush* TimerNodeTypeHelper::GetIconForTimerNodeType(const ETimerNodeType NodeType)
 {
-	static_assert(ETimerNodeType::InvalidOrMax == 4, "Not all cases are handled in switch below!?");
+	static_assert(static_cast<int>(ETimerNodeType::InvalidOrMax) == 4, "Not all cases are handled in switch below!?");
 	switch (NodeType)
 	{
-	case ETimerNodeType::GpuScope:			return FEditorStyle::GetBrush(TEXT("Profiler.Type.NumberFloat")); //TODO
-	case ETimerNodeType::ComputeScope:		return FEditorStyle::GetBrush(TEXT("Profiler.Type.NumberFloat")); //TODO
-	case ETimerNodeType::CpuScope:			return FEditorStyle::GetBrush(TEXT("Profiler.Type.NumberFloat")); //TODO
-	case ETimerNodeType::Group:				return FEditorStyle::GetBrush(TEXT("Profiler.Misc.GenericGroup")); //TODO
-	default: return nullptr;
+		case ETimerNodeType::GpuScope:		return FEditorStyle::GetBrush(TEXT("Profiler.Type.NumberFloat")); //TODO: FInsightsStyle::GetBrush(TEXT("Icons.GpuTimer"));
+		case ETimerNodeType::ComputeScope:	return FEditorStyle::GetBrush(TEXT("Profiler.Type.NumberFloat")); //TODO: FInsightsStyle::GetBrush(TEXT("Icons.ComputeTimer"));
+		case ETimerNodeType::CpuScope:		return FEditorStyle::GetBrush(TEXT("Profiler.Type.NumberFloat")); //TODO: FInsightsStyle::GetBrush(TEXT("Icons.CpuTimer"));
+		case ETimerNodeType::Group:			return FEditorStyle::GetBrush(TEXT("Profiler.Misc.GenericGroup")); //TODO: FInsightsStyle::GetBrush(TEXT("Icons.GenericGroup"));
+		default:							return nullptr;
 	}
 }
 
@@ -80,17 +80,17 @@ const FSlateBrush* TimerNodeTypeHelper::GetIconForTimerNodeType(const ETimerNode
 
 FText TimerNodeGroupingHelper::ToName(const ETimerGroupingMode GroupingMode)
 {
+	static_assert(static_cast<int>(ETimerGroupingMode::InvalidOrMax) == 7, "Not all cases are handled in switch below!?");
 	switch (GroupingMode)
 	{
-	case ETimerGroupingMode::Flat:					return LOCTEXT("Grouping_Name_Flat",				"Flat");
-	case ETimerGroupingMode::ByName:				return LOCTEXT("Grouping_Name_ByName",				"Timer Name");
-	case ETimerGroupingMode::ByMetaGroupName:		return LOCTEXT("Grouping_Name_MetaGroupName",		"Meta Group Name");
-	case ETimerGroupingMode::ByType:				return LOCTEXT("Grouping_Name_Type",				"Timer Type");
-	case ETimerGroupingMode::ByTotalInclusiveTime:	return LOCTEXT("Grouping_Name_TotalInclusiveTime",	"Total Inclusive Time");
-	case ETimerGroupingMode::ByTotalExclusiveTime:	return LOCTEXT("Grouping_Name_TotalExclusiveTime",	"Total Exclusive Time");
-	case ETimerGroupingMode::ByInstanceCount:		return LOCTEXT("Grouping_Name_InstanceCount",		"Instance Count");
-
-	default: return LOCTEXT("InvalidOrMax", "InvalidOrMax");
+		case ETimerGroupingMode::Flat:					return LOCTEXT("Grouping_Name_Flat",				"Flat");
+		case ETimerGroupingMode::ByName:				return LOCTEXT("Grouping_Name_ByName",				"Timer Name");
+		case ETimerGroupingMode::ByMetaGroupName:		return LOCTEXT("Grouping_Name_MetaGroupName",		"Meta Group Name");
+		case ETimerGroupingMode::ByType:				return LOCTEXT("Grouping_Name_Type",				"Timer Type");
+		case ETimerGroupingMode::ByTotalInclusiveTime:	return LOCTEXT("Grouping_Name_TotalInclusiveTime",	"Total Inclusive Time");
+		case ETimerGroupingMode::ByTotalExclusiveTime:	return LOCTEXT("Grouping_Name_TotalExclusiveTime",	"Total Exclusive Time");
+		case ETimerGroupingMode::ByInstanceCount:		return LOCTEXT("Grouping_Name_InstanceCount",		"Instance Count");
+		default:										return LOCTEXT("InvalidOrMax", "InvalidOrMax");
 	}
 }
 
@@ -98,17 +98,17 @@ FText TimerNodeGroupingHelper::ToName(const ETimerGroupingMode GroupingMode)
 
 FText TimerNodeGroupingHelper::ToDescription(const ETimerGroupingMode GroupingMode)
 {
+	static_assert(static_cast<int>(ETimerGroupingMode::InvalidOrMax) == 7, "Not all cases are handled in switch below!?");
 	switch (GroupingMode)
 	{
-	case ETimerGroupingMode::Flat:					return LOCTEXT("Grouping_Desc_Flat",				"Creates a single group. Includes all timers.");
-	case ETimerGroupingMode::ByName:				return LOCTEXT("Grouping_Desc_ByName",				"Creates one group for one letter.");
-	case ETimerGroupingMode::ByMetaGroupName:		return LOCTEXT("Grouping_Desc_MetaGroupName",		"Creates groups based on timer metadata group names.");
-	case ETimerGroupingMode::ByType:				return LOCTEXT("Grouping_Desc_Type",				"Creates one group for each timer type.");
-	case ETimerGroupingMode::ByTotalInclusiveTime:	return LOCTEXT("Grouping_Desc_TotalInclusiveTime",	"Creates one group for each logarithmic range ie. 0.001 - 0.01, 0.01 - 0.1, 0.1 - 1.0, 1.0 - 10.0 etc");
-	case ETimerGroupingMode::ByTotalExclusiveTime:	return LOCTEXT("Grouping_Desc_TotalExclusiveTime",	"Creates one group for each logarithmic range ie. 0.001 - 0.01, 0.01 - 0.1, 0.1 - 1.0, 1.0 - 10.0 etc");
-	case ETimerGroupingMode::ByInstanceCount:		return LOCTEXT("Grouping_Desc_InstanceCount",		"Creates one group for each logarithmic range ie. 0, 1 - 10, 10 - 100, 100 - 1000, etc");
-
-	default: return LOCTEXT("InvalidOrMax", "InvalidOrMax");
+		case ETimerGroupingMode::Flat:					return LOCTEXT("Grouping_Desc_Flat",				"Creates a single group. Includes all timers.");
+		case ETimerGroupingMode::ByName:				return LOCTEXT("Grouping_Desc_ByName",				"Creates one group for one letter.");
+		case ETimerGroupingMode::ByMetaGroupName:		return LOCTEXT("Grouping_Desc_MetaGroupName",		"Creates groups based on metadata group names of timers.");
+		case ETimerGroupingMode::ByType:				return LOCTEXT("Grouping_Desc_Type",				"Creates one group for each timer type.");
+		case ETimerGroupingMode::ByTotalInclusiveTime:	return LOCTEXT("Grouping_Desc_TotalInclusiveTime",	"Creates one group for each logarithmic range ie. 0.001 - 0.01, 0.01 - 0.1, 0.1 - 1.0, 1.0 - 10.0 etc");
+		case ETimerGroupingMode::ByTotalExclusiveTime:	return LOCTEXT("Grouping_Desc_TotalExclusiveTime",	"Creates one group for each logarithmic range ie. 0.001 - 0.01, 0.01 - 0.1, 0.1 - 1.0, 1.0 - 10.0 etc");
+		case ETimerGroupingMode::ByInstanceCount:		return LOCTEXT("Grouping_Desc_InstanceCount",		"Creates one group for each logarithmic range ie. 0, 1 - 10, 10 - 100, 100 - 1000, etc");
+		default:										return LOCTEXT("InvalidOrMax", "InvalidOrMax");
 	}
 }
 
@@ -116,18 +116,17 @@ FText TimerNodeGroupingHelper::ToDescription(const ETimerGroupingMode GroupingMo
 
 FName TimerNodeGroupingHelper::ToBrushName(const ETimerGroupingMode GroupingMode)
 {
+	static_assert(static_cast<int>(ETimerGroupingMode::InvalidOrMax) == 7, "Not all cases are handled in switch below!?");
 	switch (GroupingMode)
 	{
-	//TODO: "UnrealProfiler.*Icon"
-	case ETimerGroupingMode::Flat:					return TEXT("Profiler.FiltersAndPresets.GroupNameIcon"); //TODO
-	case ETimerGroupingMode::ByName:				return TEXT("Profiler.FiltersAndPresets.GroupNameIcon"); //TODO
-	case ETimerGroupingMode::ByMetaGroupName	:	return TEXT("Profiler.FiltersAndPresets.StatNameIcon"); //TODO
-	case ETimerGroupingMode::ByType:				return TEXT("Profiler.FiltersAndPresets.StatTypeIcon"); //TODO
-	case ETimerGroupingMode::ByTotalInclusiveTime:	return TEXT("Profiler.FiltersAndPresets.StatValueIcon"); //TODO
-	case ETimerGroupingMode::ByTotalExclusiveTime:	return TEXT("Profiler.FiltersAndPresets.StatValueIcon"); //TODO
-	case ETimerGroupingMode::ByInstanceCount:		return TEXT("Profiler.FiltersAndPresets.StatValueIcon"); //TODO
-
-	default: return NAME_None;
+		case ETimerGroupingMode::Flat:					return TEXT("Profiler.FiltersAndPresets.GroupNameIcon"); //TODO: "Icons.Grouping.Flat"
+		case ETimerGroupingMode::ByName:				return TEXT("Profiler.FiltersAndPresets.GroupNameIcon"); //TODO: "Icons.Grouping.ByName"
+		case ETimerGroupingMode::ByMetaGroupName	:	return TEXT("Profiler.FiltersAndPresets.StatNameIcon"); //TODO
+		case ETimerGroupingMode::ByType:				return TEXT("Profiler.FiltersAndPresets.StatTypeIcon"); //TODO
+		case ETimerGroupingMode::ByTotalInclusiveTime:	return TEXT("Profiler.FiltersAndPresets.StatValueIcon"); //TODO
+		case ETimerGroupingMode::ByTotalExclusiveTime:	return TEXT("Profiler.FiltersAndPresets.StatValueIcon"); //TODO
+		case ETimerGroupingMode::ByInstanceCount:		return TEXT("Profiler.FiltersAndPresets.StatValueIcon"); //TODO
+		default:										return NAME_None;
 	}
 }
 
diff --git a/Engine/Source/Developer/TraceInsights/Private/Insights/ViewModels/TimerNodeHelper.h b/Engine/Source/Developer/TraceInsights/Private/Insights/ViewModels/TimerNodeHelper.h
--- a/Engine/Source/Developer/TraceInsights/Private/Insights/ViewModels/TimerNodeHelper.h
+++ b/Engine/Source/Developer/TraceInsights/Private/Insights/ViewModels/TimerNodeHelper.h
@@ -92,7 +92,7 @@ struct TimerNodeSortingHelper
 
 	//////////////////////////////////////////////////
 	// Sorting by timer's meta group name.
-	// If meta group name is the same then sort by name.
+	// If meta group names are the same then sort by name.
 
 	struct ByMetaGroupNameAscending
 	{
@@ -128,7 +128,7 @@ struct TimerNodeSortingHelper
 
 	//////////////////////////////////////////////////
 	// Sorting by timer's type.
-	// If type is the same then sort by name.
+	// If types are the same then sort by name.
 
 	struct ByTypeAscending
 	{
@@ -169,16 +169,16 @@ struct TimerNodeSortingHelper
 	};
 
 	//////////////////////////////////////////////////
-	// Sorting by a timer's stats value.
-	// If stats value is the same then sort by name.
+	// Sorting by an aggregated stats value.
+	// If aggregated stats values are the same then sort by name.
 
-	#define SORT_BY_STATS_ASCENDING(Type, SortName, Stats) \
+	#define SORT_BY_STATS_ASCENDING(Type, SortName, AggregatedStatsMember) \
 		struct SortName \
 		{ \
 			FORCEINLINE_DEBUGGABLE bool operator()(const FTimerNodePtr& A, const FTimerNodePtr& B) const \
 			{ \
-				const Type ValueA = A->GetStats().Stats; \
-				const Type ValueB = B->GetStats().Stats; \
+				const Type ValueA = A->GetAggregatedStats().AggregatedStatsMember; \
+				const Type ValueB = B->GetAggregatedStats().AggregatedStatsMember; \
 				\
 				if (ValueA == ValueB) \
 				{ \
@@ -192,13 +192,13 @@ struct TimerNodeSortingHelper
 			} \
 		};
 
-	#define SORT_BY_STATS_DESCENDING(Type, SortName, Stats) \
+	#define SORT_BY_STATS_DESCENDING(Type, SortName, AggregatedStatsMember) \
 		struct SortName \
 		{ \
 			FORCEINLINE_DEBUGGABLE bool operator()(const FTimerNodePtr& A, const FTimerNodePtr& B) const \
 			{ \
-				const Type ValueA = A->GetStats().Stats; \
-				const Type ValueB = B->GetStats().Stats; \
+				const Type ValueA = A->GetAggregatedStats().AggregatedStatsMember; \
+				const Type ValueB = B->GetAggregatedStats().AggregatedStatsMember; \
 				\
 				if (ValueA == ValueB) \
 				{ \
@@ -212,9 +212,9 @@ struct TimerNodeSortingHelper
 			} \
 		};
 
-	#define SORT_BY_STATS(Type, SortName, Stats) \
-		SORT_BY_STATS_ASCENDING(Type, SortName##Ascending, Stats) \
-		SORT_BY_STATS_DESCENDING(Type, SortName##Descending, Stats)
+	#define SORT_BY_STATS(Type, SortName, AggregatedStatsMember) \
+		SORT_BY_STATS_ASCENDING(Type, SortName##Ascending, AggregatedStatsMember) \
+		SORT_BY_STATS_DESCENDING(Type, SortName##Descending, AggregatedStatsMember)
 
 	SORT_BY_STATS(uint64, ByInstanceCount, InstanceCount);
 
diff --git a/Engine/Source/Developer/TraceInsights/Private/Insights/ViewModels/TimersViewColumn.h b/Engine/Source/Developer/TraceInsights/Private/Insights/ViewModels/TimersViewColumn.h
--- a/Engine/Source/Developer/TraceInsights/Private/Insights/ViewModels/TimersViewColumn.h
+++ b/Engine/Source/Developer/TraceInsights/Private/Insights/ViewModels/TimersViewColumn.h
@@ -4,40 +4,39 @@
 
 #include "CoreMinimal.h"
 #include "Types/SlateEnums.h"
+#include "Misc/EnumClassFlags.h"
 
-class FTimerNode;
+class FTimerNode; // TODO: IInsightsTreeNode; generic formatter; see also StatsViewColumn.h
 
-/// Holds information about a column in the timer view widget.
+enum class ETimersViewColumnFlags : uint32
+{
+	None = 0,
+
+	CanBeHidden = (1 << 0),
+	CanBeSorted = (1 << 1),
+	CanBeFiltered = (1 << 2),
+};
+ENUM_CLASS_FLAGS(ETimersViewColumnFlags);
+
+/** Holds information about a column in the timer view widget. */
 class FTimersViewColumn
 {
 	friend struct FTimersViewColumnFactory;
 
 public:
-	enum EFlags : uint32
-	{
-		None = 0,
-
-		CanBeHidden = (1 << 0),
-		CanBeSorted = (1 << 1),
-		CanBeFiltered = (1 << 2),
-
-		NameColumnFlags = CanBeSorted | CanBeFiltered,
-		StatsColumnFlags = CanBeHidden | CanBeSorted | CanBeFiltered,
-	};
-
 	typedef TFunction<FText(const FTimersViewColumn& Column, const FTimerNode& TimerNode)> FGetFormattedValueFn;
 
 public:
-	/// Whether this column can be hidden.
-	bool bCanBeHidden() const { return ((uint32)Flags & (uint32)EFlags::CanBeHidden) != 0; }
+	/** Whether this column can be hidden. */
+	bool bCanBeHidden() const { return EnumHasAnyFlags(Flags, ETimersViewColumnFlags::CanBeHidden); }
 
-	/// Whether this column cab be used for sorting.
-	bool bCanBeSorted() const { return ((uint32)Flags & (uint32)EFlags::CanBeSorted) != 0; }
+	/** Whether this column cab be used for sorting. */
+	bool bCanBeSorted() const { return EnumHasAnyFlags(Flags, ETimersViewColumnFlags::CanBeSorted); }
 
-	/// Where this column can be used to filtering displayed results.
-	bool bCanBeFiltered() const { return ((uint32)Flags & (uint32)EFlags::CanBeFiltered) != 0; }
+	/** Where this column can be used to filtering displayed results. */
+	bool bCanBeFiltered() const { return EnumHasAnyFlags(Flags, ETimersViewColumnFlags::CanBeFiltered); }
 
-	/// If MinColumnWidth == MaxColumnWidth, this column has fixed width and cannot be resized.
+	/** If MinColumnWidth == MaxColumnWidth, this column has fixed width and cannot be resized. */
 	bool bIsFixedColumnWidth() const { return MinColumnWidth == MaxColumnWidth; }
 
 	FText GetFormattedValue(const FTimerNode& TimerNode) const
@@ -46,10 +45,10 @@ public:
 	}
 
 protected:
-	/// No default constructor.
+	/** No default constructor. */
 	FTimersViewColumn() = delete;
 
-	/// Initialization constructor, only used in FTimersViewColumnFactory.
+	/** Initialization constructor, only used in FStatsViewColumnFactory. */
 	FTimersViewColumn
 	(
 		int32 InOrder,
@@ -59,7 +58,7 @@ protected:
 		FText InTitleName,
 		FText InDescription,
 		bool bInIsVisible,
-		const EFlags InFlags,
+		const ETimersViewColumnFlags InFlags,
 		const EHorizontalAlignment InHorizontalAlignment,
 		const float InInitialColumnWidth,
 		const float InMinColumnWidth,
@@ -83,37 +82,37 @@ protected:
 	}
 
 public:
-	/// Order value, to sort columns in the tree view.
+	/** Order value, to sort columns in the tree view. */
 	int32 Order;
 
-	/// Name of the column, name of the property.
+	/** Name of the column, name of the property. */
 	FName Id;
 
-	/// Name of the column used by the searching system.
+	/** Name of the column used by the searching system. */
 	FName SearchId;
 
-	/// Short name of the column, displayed in the column header.
+	/** Short name of the column, displayed in the column header. */
 	FText ShortName;
 
-	/// Title name of the column, displayed as title in the column tooltip.
+	/** Title name of the column, displayed as title in the column tooltip. */
 	FText TitleName;
 
-	/// Long name of the column, displayed in the column tooltip.
+	/** Long name of the column, displayed in the column tooltip. */
 	FText Description;
 
-	/// Is this column visible?
+	/** Is this column visible? */
 	bool bIsVisible;
 
-	/// On/off switches.
-	EFlags Flags;
+	/** On/off switches. */
+	ETimersViewColumnFlags Flags;
 
-	/// Horizontal alignment of the content in this column.
+	/** Horizontal alignment of the content in this column. */
 	EHorizontalAlignment HorizontalAlignment;
 
-	float InitialColumnWidth; /// Initial column width.
-	float MinColumnWidth; /// Minimum column width.
-	float MaxColumnWidth; /// Maximum column width.
+	float InitialColumnWidth; /**< Initial column width. */
+	float MinColumnWidth; /**< Minimum column width. */
+	float MaxColumnWidth; /**< Maximum column width. */
 
-	/// Custom function used to format (as an FText) the value of the timer node to be displayed by this column.
+	/** Custom function used to format (as an FText) the value of the timer node to be displayed by this column. */
 	FGetFormattedValueFn GetFormattedValueFn;
 };
diff --git a/Engine/Source/Developer/TraceInsights/Private/Insights/ViewModels/TimersViewColumnFactory.h b/Engine/Source/Developer/TraceInsights/Private/Insights/ViewModels/TimersViewColumnFactory.h
--- a/Engine/Source/Developer/TraceInsights/Private/Insights/ViewModels/TimersViewColumnFactory.h
+++ b/Engine/Source/Developer/TraceInsights/Private/Insights/ViewModels/TimersViewColumnFactory.h
@@ -4,316 +4,59 @@
 
 #include "CoreMinimal.h"
 
-// Insights
-#include "Insights/Common/TimeUtils.h"
-#include "Insights/ViewModels/TimerNodeHelper.h"
 #include "Insights/ViewModels/TimersViewColumn.h"
 
-#define LOCTEXT_NAMESPACE "STimerView"
+////////////////////////////////////////////////////////////////////////////////////////////////////
 
-struct FTimersViewColumnFactory
+struct FTimersViewColumns
 {
-	/** Default constructor. */
-	FTimersViewColumnFactory()
-	{
-		uint32 Index = 0;
-
-		// FTimerViewColumn
-		// (
-		//     Order,
-		//     Id,
-		//     SearchId,
-		//     ShortName,
-		//     TitleName,
-		//     Description,
-		//     bIsVisible,
-		//     Flags,
-		//     HorizontalAlignment,
-		//     InitialWidth, MinWidth, MaxWidth,
-		//     GetFormattedValueFn
-		// )
-
-		//////////////////////////////////////////////////
-		// Name
-
-		Collection.Add(new FTimersViewColumn
-		(
-			Index++,
-			TEXT("Name"),
-			TEXT("name"),
-			LOCTEXT("TimerNameColumnName", "Name"),
-			LOCTEXT("TimerNameColumnTitle", "Timer or Group Name"),
-			LOCTEXT("TimerNameColumnDesc", "Name of the timer or group."),
-			true,
-			FTimersViewColumn::EFlags::NameColumnFlags,
-			HAlign_Left,
-			246.0f, 10.0f, FLT_MAX,
-			[](const FTimersViewColumn& Column, const FTimerNode& TimerNode) -> FText
-			{ return FText::FromName(TimerNode.GetName()); }
-		));
-
-		//////////////////////////////////////////////////
-		// Meta Group Name
-
-		Collection.Add(new FTimersViewColumn
-		(
-			Index++,
-			TEXT("MetaGroupName"),
-			TEXT("metagroupname"),
-			LOCTEXT("TimerMetaGroupNameColumnName", "Meta Group"),
-			LOCTEXT("TimerMetaGroupNameColumnTitle", "Meta Group Name"),
-			LOCTEXT("TimerMetaGroupNameColumnDesc", "Name of the meta group."),
-			false,
-			(FTimersViewColumn::EFlags)(FTimersViewColumn::EFlags::CanBeHidden |
-										FTimersViewColumn::EFlags::CanBeSorted |
-										FTimersViewColumn::EFlags::CanBeFiltered),
-			HAlign_Left,
-			100.0f, 10.0f, FLT_MAX,
-			[](const FTimersViewColumn& Column, const FTimerNode& TimerNode) -> FText
-			{ return FText::FromName(TimerNode.GetMetaGroupName()); }
-		));
-
-		//////////////////////////////////////////////////
-		// Type
-
-		Collection.Add(new FTimersViewColumn
-		(
-			Index++,
-			TEXT("Type"),
-			TEXT("type"),
-			LOCTEXT("TimerTypeColumnName", "Type"),
-			LOCTEXT("TimerTypeColumnTitle", "Type"),
-			LOCTEXT("TimerTypeColumnDesc", "Type of timer or group."),
-			false,
-			(FTimersViewColumn::EFlags)(FTimersViewColumn::EFlags::CanBeHidden |
-										FTimersViewColumn::EFlags::CanBeSorted |
-										FTimersViewColumn::EFlags::CanBeFiltered),
-			HAlign_Left,
-			60.0f, 10.0f, FLT_MAX,
-			[](const FTimersViewColumn& Column, const FTimerNode& TimerNode) -> FText
-			{ return TimerNodeTypeHelper::ToName(TimerNode.GetType()); }
-		));
-
-		//////////////////////////////////////////////////
-		// Instance Count
-
-		Collection.Add(new FTimersViewColumn
-		(
-			Index++,
-			TEXT("InstanceCount"),
-			TEXT("count"),
-			LOCTEXT("InstanceCountName", "Count"),
-			LOCTEXT("InstanceCountTitle", "Instance Count"),
-			LOCTEXT("InstanceCountDesc", "Number of timer's instances"),
-			true,
-			FTimersViewColumn::EFlags::StatsColumnFlags,
-			HAlign_Right,
-			60.0f, 0.0f, FLT_MAX,
-			[](const FTimersViewColumn& Column, const FTimerNode& TimerNode) -> FText
-			{ return FText::AsNumber(TimerNode.GetStats().InstanceCount); }
-		));
-
-		//////////////////////////////////////////////////
-
-		constexpr float TotalTimeColumnInitialWidth = 60.0f;
-		constexpr float TotalTimeColumnMinWidth = 4.0f;
-		constexpr float TotalTimeColumnMaxWidth = FLT_MAX;
-
-		constexpr float TimeMsColumnInitialWidth = 50.0f;
-		constexpr float TimeMsColumnMinWidth = 4.0f;
-		constexpr float TimeMsColumnMaxWidth = FLT_MAX;
-
-		//////////////////////////////////////////////////
-		// Inclusive Time stats
-
-		Collection.Add(new FTimersViewColumn
-		(
-			Index++,
-			TEXT("TotalInclusiveTime"),
-			TEXT("inc"),
-			LOCTEXT("TotalInclusiveTimeName", "Incl."),
-			LOCTEXT("TotalInclusiveTimeTitle", "Total Inclusive Time"),
-			LOCTEXT("TotalInclusiveTimeDesc", "Total inclusive duration of selected timer's instances."),
-			true,
-			FTimersViewColumn::EFlags::StatsColumnFlags,
-			HAlign_Right,
-			TotalTimeColumnInitialWidth, TotalTimeColumnMinWidth, TotalTimeColumnMaxWidth,
-			[](const FTimersViewColumn& Column, const FTimerNode& TimerNode) -> FText
-			{ return FText::FromString(TimeUtils::FormatTimeAuto(TimerNode.GetStats().TotalInclusiveTime)); }
-		));
-
-		Collection.Add(new FTimersViewColumn
-		(
-			Index++,
-			TEXT("MinInclusiveTime"),
-			TEXT("mininc"),
-			LOCTEXT("MinInclusiveTimeName", "I.Min"),
-			LOCTEXT("MinInclusiveTimeTitle", "Min Inclusive Time (ms)"),
-			LOCTEXT("MinInclusiveTimeDesc", "Minimum inclusive duration of selected timer's instances, in milliseconds"),
-			false,
-			FTimersViewColumn::EFlags::StatsColumnFlags,
-			HAlign_Right,
-			TimeMsColumnInitialWidth, TimeMsColumnMinWidth, TimeMsColumnMaxWidth,
-			[](const FTimersViewColumn& Column, const FTimerNode& TimerNode) -> FText
-			{ return FText::FromString(TimeUtils::FormatTimeMs(TimerNode.GetStats().MinInclusiveTime)); }
-		));
-
-		Collection.Add(new FTimersViewColumn
-		(
-			Index++,
-			TEXT("MaxInclusiveTime"),
-			TEXT("maxinc"),
-			LOCTEXT("MaxInclusiveTimeName", "I.Max"),
-			LOCTEXT("MaxInclusiveTimeTitle", "Max Inclusive Time (ms)"),
-			LOCTEXT("MaxInclusiveTimeDesc", "Maximum inclusive duration of selected timer's instances, in milliseconds"),
-			false,
-			FTimersViewColumn::EFlags::StatsColumnFlags,
-			HAlign_Right,
-			TimeMsColumnInitialWidth, TimeMsColumnMinWidth, TimeMsColumnMaxWidth,
-			[](const FTimersViewColumn& Column, const FTimerNode& TimerNode) -> FText
-			{ return FText::FromString(TimeUtils::FormatTimeMs(TimerNode.GetStats().MaxInclusiveTime)); }
-		));
-
-		Collection.Add(new FTimersViewColumn
-		(
-			Index++,
-			TEXT("AvgInclusiveTime"),
-			TEXT("avginc"),
-			LOCTEXT("AvgInclusiveTimeName", "I.Avg."),
-			LOCTEXT("AvgInclusiveTimeTitle", "Average Inclusive Time (ms)"),
-			LOCTEXT("AvgInclusiveTimeDesc", "Average inclusive duration of selected timer's instances, in milliseconds"),
-			false,
-			FTimersViewColumn::EFlags::StatsColumnFlags,
-			HAlign_Right,
-			TimeMsColumnInitialWidth, TimeMsColumnMinWidth, TimeMsColumnMaxWidth,
-			[](const FTimersViewColumn& Column, const FTimerNode& TimerNode) -> FText
-			{ return FText::FromString(TimeUtils::FormatTimeMs(TimerNode.GetStats().AverageInclusiveTime)); }
-		));
-
-		Collection.Add(new FTimersViewColumn
-		(
-			Index++,
-			TEXT("MedInclusiveTime"),
-			TEXT("medinc"),
-			LOCTEXT("MedInclusiveTimeName", "I.Med."),
-			LOCTEXT("MedInclusiveTimeTitle", "Median Inclusive Time (ms)"),
-			LOCTEXT("MedInclusiveTimeDesc", "Median inclusive duration of selected timer's instances, in milliseconds"),
-			false,
-			FTimersViewColumn::EFlags::StatsColumnFlags,
-			HAlign_Right,
-			TimeMsColumnInitialWidth, TimeMsColumnMinWidth, TimeMsColumnMaxWidth,
-			[](const FTimersViewColumn& Column, const FTimerNode& TimerNode) -> FText
-			{ return FText::FromString(TimeUtils::FormatTimeMs(TimerNode.GetStats().MedianInclusiveTime)); }
-		));
-
-		//////////////////////////////////////////////////
-		// Exclusive Time stats
-
-		Collection.Add(new FTimersViewColumn
-		(
-			Index++,
-			TEXT("TotalExclusiveTime"),
-			TEXT("exc"),
-			LOCTEXT("TotalExclusiveTimeName", "Excl."),
-			LOCTEXT("TotalExclusiveTimeTitle", "Total Exclusive Time"),
-			LOCTEXT("TotalExclusiveTimeDesc", "Total exclusive duration of selected timer's instances."),
-			true,
-			FTimersViewColumn::EFlags::StatsColumnFlags,
-			HAlign_Right,
-			TotalTimeColumnInitialWidth, TotalTimeColumnMinWidth, TotalTimeColumnMaxWidth,
-			[](const FTimersViewColumn& Column, const FTimerNode& TimerNode) -> FText
-			{ return FText::FromString(TimeUtils::FormatTimeAuto(TimerNode.GetStats().TotalExclusiveTime)); }
-		));
-
-		Collection.Add(new FTimersViewColumn
-		(
-			Index++,
-			TEXT("MinExclusiveTime"),
-			TEXT("minexc"),
-			LOCTEXT("MinExclusiveTimeName", "E.Min"),
-			LOCTEXT("MinExclusiveTimeTitle", "Min Exclusive Time (ms)"),
-			LOCTEXT("MinExclusiveTimeDesc", "Minimum exclusive duration of selected timer's instances, in milliseconds"),
-			false,
-			FTimersViewColumn::EFlags::StatsColumnFlags,
-			HAlign_Right,
-			TimeMsColumnInitialWidth, TimeMsColumnMinWidth, TimeMsColumnMaxWidth,
-			[](const FTimersViewColumn& Column, const FTimerNode& TimerNode) -> FText
-			{ return FText::FromString(TimeUtils::FormatTimeMs(TimerNode.GetStats().MinExclusiveTime)); }
-		));
-
-		Collection.Add(new FTimersViewColumn
-		(
-			Index++,
-			TEXT("MaxExclusiveTime"),
-			TEXT("maxexc"),
-			LOCTEXT("MaxExclusiveTimeName", "E.Max"),
-			LOCTEXT("MaxExclusiveTimeTitle", "Max Exclusive Time (ms)"),
-			LOCTEXT("MaxExclusiveTimeDesc", "Maximum exclusive duration of selected timer's instances, in milliseconds"),
-			false,
-			FTimersViewColumn::EFlags::StatsColumnFlags,
-			HAlign_Right,
-			TimeMsColumnInitialWidth, TimeMsColumnMinWidth, TimeMsColumnMaxWidth,
-			[](const FTimersViewColumn& Column, const FTimerNode& TimerNode) -> FText
-			{ return FText::FromString(TimeUtils::FormatTimeMs(TimerNode.GetStats().MaxExclusiveTime)); }
-		));
-
-		Collection.Add(new FTimersViewColumn
-		(
-			Index++,
-			TEXT("AvgExclusiveTime"),
-			TEXT("avgexc"),
-			LOCTEXT("AvgExclusiveTimeName", "E.Avg."),
-			LOCTEXT("AvgExclusiveTimeTitle", "Average Exclusive Time (ms)"),
-			LOCTEXT("AvgExclusiveTimeDesc", "Average exclusive duration of selected timer's instances, in milliseconds"),
-			false,
-			FTimersViewColumn::EFlags::StatsColumnFlags,
-			HAlign_Right,
-			TimeMsColumnInitialWidth, TimeMsColumnMinWidth, TimeMsColumnMaxWidth,
-			[](const FTimersViewColumn& Column, const FTimerNode& TimerNode) -> FText
-			{ return FText::FromString(TimeUtils::FormatTimeMs(TimerNode.GetStats().AverageExclusiveTime)); }
-		));
-
-		Collection.Add(new FTimersViewColumn
-		(
-			Index++,
-			TEXT("MedExclusiveTime"),
-			TEXT("medexc"),
-			LOCTEXT("MedExclusiveTimeName", "E.Med."),
-			LOCTEXT("MedExclusiveTimeTitle", "Median Exclusive Time (ms)"),
-			LOCTEXT("MedExclusiveTimeDesc", "Median exclusive duration of timer's instances, in milliseconds"),
-			false,
-			FTimersViewColumn::EFlags::StatsColumnFlags,
-			HAlign_Right,
-			TimeMsColumnInitialWidth, TimeMsColumnMinWidth, TimeMsColumnMaxWidth,
-			[](const FTimersViewColumn& Column, const FTimerNode& TimerNode) -> FText
-			{ return FText::FromString(TimeUtils::FormatTimeMs(TimerNode.GetStats().MedianExclusiveTime)); }
-		));
+	//////////////////////////////////////////////////
+	// Column identifiers
+
+	static const FName NameColumnID;
+	static const FName MetaGroupNameColumnID;
+	static const FName TypeColumnID;
+	static const FName InstanceCountColumnID;
+
+	// Inclusive Time columns
+	static const FName TotalInclusiveTimeColumnID;
+	static const FName MaxInclusiveTimeColumnID;
+	static const FName UpperQuartileInclusiveTimeColumnID;
+	static const FName AverageInclusiveTimeColumnID;
+	static const FName MedianInclusiveTimeColumnID;
+	static const FName LowerQuartileInclusiveTimeColumnID;
+	static const FName MinInclusiveTimeColumnID;
+
+	// Exclusive Time columns
+	static const FName TotalExclusiveTimeColumnID;
+	static const FName MaxExclusiveTimeColumnID;
+	static const FName UpperQuartileExclusiveTimeColumnID;
+	static const FName AverageExclusiveTimeColumnID;
+	static const FName MedianExclusiveTimeColumnID;
+	static const FName LowerQuartileExclusiveTimeColumnID;
+	static const FName MinExclusiveTimeColumnID;
+};
 
-		//////////////////////////////////////////////////
+////////////////////////////////////////////////////////////////////////////////////////////////////
 
-		for (int32 MapIndex = 0; MapIndex < Collection.Num(); MapIndex++)
-		{
-			ColumnIdToPtrMapping.Add(Collection[MapIndex]->Id, Collection[MapIndex]);
-		}
-	}
+struct FTimersViewColumnFactory
+{
+private:
+	/** Default constructor. */
+	FTimersViewColumnFactory();
 
-	~FTimersViewColumnFactory()
-	{
-		for (const FTimersViewColumn* Column : Collection)
-			delete Column;
-		Collection.Reset();
-	}
+	/** Destructor. */
+	~FTimersViewColumnFactory();
 
-	/** Contains basic information about columns used in the timer view widget. Names should be localized. */
+public:
+	/** Contains basic information about columns used in the Timers view widget. Names should be localized. */
 	TArray<FTimersViewColumn*> Collection;
 
+	/** Mapping between column IDs and FTimersViewColumn pointers. */
 	TMap<FName, const FTimersViewColumn*> ColumnIdToPtrMapping;
 
-	static const FTimersViewColumnFactory& Get()
-	{
-		static FTimersViewColumnFactory Instance;
-		return Instance;
-	}
+	/** Singleton */
+	static const FTimersViewColumnFactory& Get();
 };
 
-#undef LOCTEXT_NAMESPACE
+////////////////////////////////////////////////////////////////////////////////////////////////////
diff --git a/Engine/Source/Developer/TraceInsights/Private/Insights/ViewModels/TimingEventsTrack.cpp b/Engine/Source/Developer/TraceInsights/Private/Insights/ViewModels/TimingEventsTrack.cpp
--- a/Engine/Source/Developer/TraceInsights/Private/Insights/ViewModels/TimingEventsTrack.cpp
+++ b/Engine/Source/Developer/TraceInsights/Private/Insights/ViewModels/TimingEventsTrack.cpp
@@ -92,24 +92,19 @@ void FTimingEventsTrack::Reset()
 
 ////////////////////////////////////////////////////////////////////////////////////////////////////
 
-void FTimingEventsTrack::UpdateHoveredState(float MX, float MY, const FTimingTrackViewport& Viewport)
+void FTimingEventsTrack::UpdateHoveredState(float MouseX, float MouseY, const FTimingTrackViewport& Viewport)
 {
-	if (MY >= Y && MY < Y + H)
+	static const float HeaderWidth = 100.0f;
+	static const float HeaderHeight = 14.0f;
+
+	if (MouseY >= GetPosY() && MouseY < GetPosY() + GetHeight())
 	{
-		bIsHovered = true;
-
-		if (MX < 100.0f && MY < Y + 14.0f)
-		{
-			bIsHeaderHovered = true;
-		}
-		else
-		{
-			bIsHeaderHovered = false;
-		}
+		SetHoveredState(true);
+		SetHeaderHoveredState(MouseX < HeaderWidth && MouseY < GetPosY() + HeaderHeight);
 	}
 	else
 	{
-		bIsHovered = false;
+		SetHoveredState(false);
 	}
 }
 
diff --git a/Engine/Source/Developer/TraceInsights/Private/Insights/ViewModels/TimingEventsTrack.h b/Engine/Source/Developer/TraceInsights/Private/Insights/ViewModels/TimingEventsTrack.h
--- a/Engine/Source/Developer/TraceInsights/Private/Insights/ViewModels/TimingEventsTrack.h
+++ b/Engine/Source/Developer/TraceInsights/Private/Insights/ViewModels/TimingEventsTrack.h
@@ -41,7 +41,7 @@ struct FTimingEventsTrackLayout
 
 ////////////////////////////////////////////////////////////////////////////////////////////////////
 
-enum ETimingEventsTrackType
+enum class ETimingEventsTrackType
 {
 	Cpu,
 	Gpu,
diff --git a/Engine/Source/Developer/TraceInsights/Private/Insights/ViewModels/TimingViewDrawHelper.cpp b/Engine/Source/Developer/TraceInsights/Private/Insights/ViewModels/TimingViewDrawHelper.cpp
--- a/Engine/Source/Developer/TraceInsights/Private/Insights/ViewModels/TimingViewDrawHelper.cpp
+++ b/Engine/Source/Developer/TraceInsights/Private/Insights/ViewModels/TimingViewDrawHelper.cpp
@@ -33,7 +33,7 @@ static const int32 NumLayers = 5;
 ////////////////////////////////////////////////////////////////////////////////////////////////////
 
 FTimingViewDrawHelper::FTimingViewDrawHelper(const FDrawContext& InDC, const FTimingTrackViewport& InViewport, const FTimingEventsTrackLayout& InLayout)
-	: DC(InDC)
+	: DrawContext(InDC)
 	, Viewport(InViewport)
 	, Layout(InLayout)
 	, WhiteBrush(FCoreStyle::Get().GetBrush("WhiteBrush"))
@@ -43,11 +43,7 @@ FTimingViewDrawHelper::FTimingViewDrawHelper(const FDrawContext& InDC, const FTi
 	, ValidAreaColor(0.07f, 0.07f, 0.07f, 1.0f)
 	, InvalidAreaColor(0.1f, 0.07f, 0.07f, 1.0f)
 	, EventFont(FCoreStyle::GetDefaultFontStyle("Regular", 8))
-	, NumEvents(0)
-	, NumDrawBoxes(0)
-	, NumMergedBoxes(0)
-	, NumDrawBorders(0)
-	, NumDrawTexts(0)
+	, Stats()
 {
 }
 
@@ -60,19 +56,6 @@ FTimingViewDrawHelper::~FTimingViewDrawHelper()
 
 ////////////////////////////////////////////////////////////////////////////////////////////////////
 
-void FTimingViewDrawHelper::Begin()
-{
-	//NumEvents = 0;
-	//NumDrawBoxes = 0;
-	//NumDrawBorders = 0;
-	//NumDrawTexts = 0;
-
-	//NumDrawTimeMarkerBoxes = 0;
-	//NumDrawTimeMarkerTexts = 0;
-}
-
-////////////////////////////////////////////////////////////////////////////////////////////////////
-
 void FTimingViewDrawHelper::DrawBackground() const
 {
 	float X0 = Viewport.TimeToSlateUnitsRounded(0.0);
@@ -86,20 +69,20 @@ void FTimingViewDrawHelper::DrawBackground() const
 		ValidX1 = W;
 
 		// Draw invalid area (entire view).
-		DC.DrawBox(0.0f, 0.0f, W, H, BackgroundAreaBrush, InvalidAreaColor);
+		DrawContext.DrawBox(0.0f, 0.0f, W, H, BackgroundAreaBrush, InvalidAreaColor);
 	}
 	else // X0 < W && X1 > 0
 	{
 		if (X0 > 0.0f)
 		{
 			// Draw invalid area (left).
-			DC.DrawBox(0.0f, 0.0f, X0, H, BackgroundAreaBrush, InvalidAreaColor);
+			DrawContext.DrawBox(0.0f, 0.0f, X0, H, BackgroundAreaBrush, InvalidAreaColor);
 		}
 
 		if (X1 < W)
 		{
 			// Draw invalid area (right).
-			DC.DrawBox(X1, 0.0f, W - X1, H, BackgroundAreaBrush, InvalidAreaColor);
+			DrawContext.DrawBox(X1, 0.0f, W - X1, H, BackgroundAreaBrush, InvalidAreaColor);
 		}
 
 		ValidX0 = FMath::Max(X0, 0.0f);
@@ -108,11 +91,11 @@ void FTimingViewDrawHelper::DrawBackground() const
 		if (ValidX1 > ValidX0)
 		{
 			// Draw valid area.
-			DC.DrawBox(ValidX0, 0.0f, ValidX1 - ValidX0, H, BackgroundAreaBrush, ValidAreaColor);
+			DrawContext.DrawBox(ValidX0, 0.0f, ValidX1 - ValidX0, H, BackgroundAreaBrush, ValidAreaColor);
 		}
 	}
 
-	DC.LayerId++;
+	DrawContext.LayerId++;
 }
 
 ////////////////////////////////////////////////////////////////////////////////////////////////////
@@ -129,7 +112,7 @@ void FTimingViewDrawHelper::BeginTimelines()
 		Y = FMath::Min(Y, Viewport.Height);
 
 		// Draw invalid area (top).
-		DC.DrawBox(DC.LayerId + TimelineHeaderLayer, ValidX0, 0.0f, ValidX1 - ValidX0, Y, BackgroundAreaBrush, InvalidAreaColor);
+		DrawContext.DrawBox(DrawContext.LayerId + TimelineHeaderLayer, ValidX0, 0.0f, ValidX1 - ValidX0, Y, BackgroundAreaBrush, InvalidAreaColor);
 	}
 }
 
@@ -139,14 +122,14 @@ bool FTimingViewDrawHelper::BeginTimeline(FTimingEventsTrack& Track)
 {
 	TimelineIndex++;
 
-	Track.Y = TimelineTopY + Viewport.ScrollPosY;
+	Track.SetPosY(TimelineTopY + Viewport.ScrollPosY);
 
-	if (Track.H < Layout.MinTimelineH)
-		Track.H = Layout.MinTimelineH;
+	if (Track.GetHeight() < Layout.MinTimelineH)
+		Track.SetHeight(Layout.MinTimelineH);
 
-	if (TimelineTopY + Track.H < -1.0f)
+	if (TimelineTopY + Track.GetHeight() < -1.0f)
 	{
-		TimelineTopY += Track.H;
+		TimelineTopY += Track.GetHeight();
 		return false;
 	}
 
@@ -197,33 +180,39 @@ void FTimingViewDrawHelper::EndTimeline(FTimingEventsTrack& Track)
 		NewH = 1.0f + Layout.TimelineDY + Layout.EventH * (MaxDepth + 1) + Layout.EventDY * MaxDepth + Layout.TimelineDY;
 
 		if (NewH < RealMinTimelineH)
+		{
 			NewH = RealMinTimelineH;
+		}
 	}
 
-	if (Track.H < NewH)
+	if (Track.GetHeight() < NewH)
 	{
+		float NewTrackHeight;
 		if (Layout.bIsCompactMode)
 		{
-			Track.H = FMath::CeilToFloat(Track.H * 0.5f + NewH * 0.5f);
+			NewTrackHeight = FMath::CeilToFloat(Track.GetHeight() * 0.5f + NewH * 0.5f);
 		}
 		else
 		{
-			Track.H = FMath::CeilToFloat(Track.H * 0.9f + NewH * 0.1f);
+			NewTrackHeight = FMath::CeilToFloat(Track.GetHeight() * 0.9f + NewH * 0.1f);
 		}
+		Track.SetHeight(NewTrackHeight);
 	}
-	else if (Track.H > NewH)
+	else if (Track.GetHeight() > NewH)
 	{
+		float NewTrackHeight;
 		if (Layout.bIsCompactMode)
 		{
-			Track.H = FMath::FloorToFloat(Track.H * 0.5f + NewH * 0.5f);
+			NewTrackHeight = FMath::FloorToFloat(Track.GetHeight() * 0.5f + NewH * 0.5f);
 		}
 		else
 		{
-			Track.H = FMath::FloorToFloat(Track.H * 0.9f + NewH * 0.1f);
+			NewTrackHeight = FMath::FloorToFloat(Track.GetHeight() * 0.9f + NewH * 0.1f);
 		}
+		Track.SetHeight(NewTrackHeight);
 	}
 
-	if (Track.H > 0)
+	if (Track.GetHeight() > 0)
 	{
 		const float Y = Viewport.TopOffset + TimelineTopY;
 
@@ -247,24 +236,24 @@ void FTimingViewDrawHelper::EndTimeline(FTimingEventsTrack& Track)
 		}
 
 		// Draw a horizontal line between timelines.
-		DC.DrawBox(DC.LayerId + TimelineHeaderLayer, 0, Y, Viewport.Width, 1.0f, WhiteBrush, Color);
+		DrawContext.DrawBox(DrawContext.LayerId + TimelineHeaderLayer, 0, Y, Viewport.Width, 1.0f, WhiteBrush, Color);
 
 		// Draw name of timeline.
 		//const FString Name = FString::Printf(TEXT("%d Y: %g, H: %g, VO: %g, VY: %g"), TimelineIndex, Track.Y, Track.H, Viewport.TopOffset, Viewport.ScrollPosY); //debug
 		const TSharedRef<FSlateFontMeasure> FontMeasureService = FSlateApplication::Get().GetRenderer()->GetFontMeasureService();
 		float NameWidth = FontMeasureService->Measure(Track.GetName(), EventFont).X;
-		DC.DrawBox(DC.LayerId + TimelineHeaderLayer, 0.0f, Y + 1.0f, NameWidth + 4.0f, 12.0f, WhiteBrush, Color);
-		DC.DrawText(DC.LayerId + TimelineTextLayer, 2.0f, Y, Track.GetName(), EventFont, TextColor);
+		DrawContext.DrawBox(DrawContext.LayerId + TimelineHeaderLayer, 0.0f, Y + 1.0f, NameWidth + 4.0f, 12.0f, WhiteBrush, Color);
+		DrawContext.DrawText(DrawContext.LayerId + TimelineTextLayer, 2.0f, Y, Track.GetName(), EventFont, TextColor);
 	}
 
-	TimelineTopY += Track.H;
+	TimelineTopY += Track.GetHeight();
 }
 
 ////////////////////////////////////////////////////////////////////////////////////////////////////
 
 void FTimingViewDrawHelper::AddEvent(double EventStartTime, double EventEndTime, uint32 EventDepth, const TCHAR* EventName, uint32 Color)
 {
-	NumEvents++;
+	Stats.NumEvents++;
 
 	float EventX1 = Viewport.TimeToSlateUnitsRounded(EventStartTime);
 	if (EventX1 > Viewport.Width)
@@ -419,8 +408,8 @@ void FTimingViewDrawHelper::AddEvent(double EventStartTime, double EventEndTime,
 	if (EventW > 2.0f && Layout.EventH > 2.0f)
 	{
 		// Fill inside of the timing event box.
-		DC.DrawBox(DC.LayerId + EventFillLayer, EventX1 + 1.0f, EventY + 1.0f, EventW - 2.0f, Layout.EventH - 2.0f, WhiteBrush, EventColorFill);
-		NumDrawBoxes++;
+		DrawContext.DrawBox(DrawContext.LayerId + EventFillLayer, EventX1 + 1.0f, EventY + 1.0f, EventW - 2.0f, Layout.EventH - 2.0f, WhiteBrush, EventColorFill);
+		Stats.NumDrawBoxes++;
 	}
 
 	// Save X2, for current depth.
@@ -437,8 +426,8 @@ void FTimingViewDrawHelper::AddEvent(double EventStartTime, double EventEndTime,
 		}
 
 		const FLinearColor EventColorBorder(EventColorFill.R * BorderColorFactor, EventColorFill.G * BorderColorFactor, EventColorFill.B * BorderColorFactor, 1.0f);
-		DC.DrawBox(DC.LayerId + EventBorderLayer, EventX1, EventY, EventW, Layout.EventH, EventsBorderBrush, EventColorBorder);
-		NumDrawBorders++;
+		DrawContext.DrawBox(DrawContext.LayerId + EventBorderLayer, EventX1, EventY, EventW, Layout.EventH, EventsBorderBrush, EventColorBorder);
+		Stats.NumDrawBorders++;
 	}
 	else // 1px or 2px boxes
 	{
@@ -451,7 +440,7 @@ void FTimingViewDrawHelper::AddEvent(double EventStartTime, double EventEndTime,
 		{
 			// Merge it with previous box.
 			Box.X2 = EventX2;
-			NumMergedBoxes++;
+			Stats.NumMergedBoxes++;
 		}
 		else
 		{
@@ -491,8 +480,8 @@ void FTimingViewDrawHelper::AddEvent(double EventStartTime, double EventEndTime,
 			// Ref: https://books.google.ro/books?id=0pVr7dhmdWYC
 			const bool bIsDarkColor = (EventColorFill.ComputeLuminance() < 0.4f);
 
-			DC.DrawText(DC.LayerId + EventTextLayer, EventX1 + 2.0f, EventY + 1.0f, Name, 0, LastWholeCharacterIndex + 1, EventFont, bIsDarkColor ? FLinearColor::White : FLinearColor::Black);
-			NumDrawTexts++;
+			DrawContext.DrawText(DrawContext.LayerId + EventTextLayer, EventX1 + 2.0f, EventY + 1.0f, Name, 0, LastWholeCharacterIndex + 1, EventFont, bIsDarkColor ? FLinearColor::White : FLinearColor::Black);
+			Stats.NumDrawTexts++;
 		}
 	}
 }
@@ -501,8 +490,8 @@ void FTimingViewDrawHelper::AddEvent(double EventStartTime, double EventEndTime,
 
 void FTimingViewDrawHelper::DrawBox(const FBoxData& Box, const float EventY, const float EventH)
 {
-	DC.DrawBox(DC.LayerId + EventBorderLayer, Box.X1, EventY, Box.X2 - Box.X1, EventH, WhiteBrush, Box.LinearColor);
-	NumDrawBoxes++;
+	DrawContext.DrawBox(DrawContext.LayerId + EventBorderLayer, Box.X1, EventY, Box.X2 - Box.X1, EventH, WhiteBrush, Box.LinearColor);
+	Stats.NumDrawBoxes++;
 }
 
 ////////////////////////////////////////////////////////////////////////////////////////////////////
@@ -512,7 +501,7 @@ void FTimingViewDrawHelper::EndTimelines()
 	float Y = Viewport.TopOffset + TimelineTopY;
 
 	// Draw a last horizontal line.
-	DC.DrawBox(DC.LayerId + TimelineHeaderLayer, 0, Y, Viewport.Width, 1.0f, WhiteBrush, FLinearColor(0.05f, 0.05f, 0.05f, 1.0f));
+	DrawContext.DrawBox(DrawContext.LayerId + TimelineHeaderLayer, 0, Y, Viewport.Width, 1.0f, WhiteBrush, FLinearColor(0.05f, 0.05f, 0.05f, 1.0f));
 	Y += 1.0f;
 
 	if (Y < Viewport.Height && ValidX1 > ValidX0)
@@ -520,10 +509,10 @@ void FTimingViewDrawHelper::EndTimelines()
 		Y = FMath::Max(Y, 0.0f);
 
 		// Draw invalid area (bottom).
-		DC.DrawBox(DC.LayerId + TimelineHeaderLayer, ValidX0, Y, ValidX1 - ValidX0, Viewport.Height - Y, BackgroundAreaBrush, InvalidAreaColor);
+		DrawContext.DrawBox(DrawContext.LayerId + TimelineHeaderLayer, ValidX0, Y, ValidX1 - ValidX0, Viewport.Height - Y, BackgroundAreaBrush, InvalidAreaColor);
 	}
 
-	DC.LayerId += NumLayers;
+	DrawContext.LayerId += NumLayers;
 }
 
 ////////////////////////////////////////////////////////////////////////////////////////////////////
@@ -544,7 +533,9 @@ void FTimingViewDrawHelper::DrawTimingEventHighlight(double StartTime, double En
 	}
 
 	if (EventX1 == EventX2)
+	{
 		EventX2 = EventX1 + 1.0f;
+	}
 
 	// Limit event width on the viewport's left side.
 	const float MinX = -2.0f; // -2 allows event border to remain outside screen
@@ -566,35 +557,26 @@ void FTimingViewDrawHelper::DrawTimingEventHighlight(double StartTime, double En
 	{
 		const FSlateBrush* HighlightBrush = BorderBrush;
 
-		const FLinearColor Color(1.0f, 1.0f, 0.0f, 1.0f);
+		const FLinearColor Color(1.0f, 1.0f, 0.0f, 1.0f); // yellow
 
 		// Draw border around the timing event box.
-		DC.DrawBox(EventX1 - 2.0f, Y - 2.0f, EventW + 4.0f, Layout.EventH + 4.0f, HighlightBrush, Color);
+		DrawContext.DrawBox(EventX1 - 2.0f, Y - 2.0f, EventW + 4.0f, Layout.EventH + 4.0f, HighlightBrush, Color);
 	}
 	else // EHighlightMode::Selected or EHighlightMode::SelectedAndHovered
 	{
 		const FSlateBrush* SelectedBrush = BorderBrush;
 
+		// Animate color from white (if selected and hovered) or yellow (if only selected) to black, using a squared sine function.
 		const double Time = static_cast<double>(FPlatformTime::Cycles64()) * FPlatformTime::GetSecondsPerCycle64();
 		float S = FMath::Sin(2.0 * Time);
-		S = S * S;
-		const float B = (Mode == EHighlightMode::SelectedAndHovered) ? 0.0f : S;
-		const FLinearColor Color(S, S, B, 1.0f);
+		S = S * S; // squared, to ensure only positive [0 - 1] values
+		const float Blue = (Mode == EHighlightMode::SelectedAndHovered) ? 0.0f : S;
+		const FLinearColor Color(S, S, Blue, 1.0f);
 
 		// Draw border around the timing event box.
-		DC.DrawBox(EventX1 - 2.0f, Y - 2.0f, EventW + 4.0f, Layout.EventH + 4.0f, SelectedBrush, Color);
-
-		//const FLinearColor Color2(1.0f - S, 1.0f - S, 1.0f - S, 1.0f);
-		//const FSlateBrush* SelectedBrush2 = FEditorStyle::GetBrush("LevelViewport.ActiveViewportBorder");
-		//DC.DrawBox(EventX1 - 2.0f, Y - 2.0f, EventW + 4.0f, Layout.EventH + 4.0f, SelectedBrush2, Color2);
+		DrawContext.DrawBox(EventX1 - 2.0f, Y - 2.0f, EventW + 4.0f, Layout.EventH + 4.0f, SelectedBrush, Color);
 	}
-	DC.LayerId++;
-}
-
-////////////////////////////////////////////////////////////////////////////////////////////////////
-
-void FTimingViewDrawHelper::End()
-{
+	DrawContext.LayerId++;
 }
 
 ////////////////////////////////////////////////////////////////////////////////////////////////////
diff --git a/Engine/Source/Developer/TraceInsights/Private/Insights/ViewModels/TimingViewDrawHelper.h b/Engine/Source/Developer/TraceInsights/Private/Insights/ViewModels/TimingViewDrawHelper.h
--- a/Engine/Source/Developer/TraceInsights/Private/Insights/ViewModels/TimingViewDrawHelper.h
+++ b/Engine/Source/Developer/TraceInsights/Private/Insights/ViewModels/TimingViewDrawHelper.h
@@ -17,10 +17,10 @@ class FTimingTrackViewport;
 
 ////////////////////////////////////////////////////////////////////////////////////////////////////
 
-class FTimingViewDrawHelper : public FNoncopyable
+class FTimingViewDrawHelper
 {
 public:
-	enum EHighlightMode
+	enum class EHighlightMode
 	{
 		Hovered,
 		Selected,
@@ -55,15 +55,51 @@ private:
 	//	bool bUseDarkTextColor; // true if text needs to be displayed in Black, otherwise will be displayed in White
 	//};
 
+	struct FStats
+	{
+		int32 NumEvents;
+		int32 NumDrawBoxes;
+		int32 NumMergedBoxes;
+		int32 NumDrawBorders;
+		int32 NumDrawTexts;
+		int32 NumDrawTimeMarkerBoxes;
+		int32 NumDrawTimeMarkerTexts;
+
+		FStats()
+			: NumEvents(0)
+			, NumDrawBoxes(0)
+			, NumMergedBoxes(0)
+			, NumDrawBorders(0)
+			, NumDrawTexts(0)
+			, NumDrawTimeMarkerBoxes(0)
+			, NumDrawTimeMarkerTexts(0)
+		{}
+	};
+
 public:
-	FTimingViewDrawHelper(const FDrawContext& DC, const FTimingTrackViewport& InViewport, const FTimingEventsTrackLayout& Layout);
+	FTimingViewDrawHelper(const FDrawContext& InDrawContext, const FTimingTrackViewport& InViewport, const FTimingEventsTrackLayout& InLayout);
 	~FTimingViewDrawHelper();
 
-	const FTimingTrackViewport& GetViewport() { return Viewport; }
+	/**
+	 * Non-copyable
+	 */
+	FTimingViewDrawHelper(const FTimingViewDrawHelper&) = delete;
+	FTimingViewDrawHelper& operator=(const FTimingViewDrawHelper&) = delete;
 
-	void Begin();
+	const FDrawContext& GetDrawContext() const { return DrawContext; }
+	const FTimingTrackViewport& GetViewport() const { return Viewport; }
+	const FTimingEventsTrackLayout& GetLayout() const { return Layout; }
+
+	int32 GetNumEvents() const              { return Stats.NumEvents; }
+	int32 GetNumDrawBoxes() const           { return Stats.NumDrawBoxes; }
+	int32 GetNumMergedBoxes() const         { return Stats.NumMergedBoxes; }
+	int32 GetNumDrawBorders() const         { return Stats.NumDrawBorders; }
+	int32 GetNumDrawTexts() const           { return Stats.NumDrawTexts; }
+	int32 GetNumDrawTimeMarkerBoxes() const { return Stats.NumDrawTimeMarkerBoxes; }
+	int32 GetNumDrawTimeMarkerTexts() const { return Stats.NumDrawTimeMarkerTexts; }
 
 	void DrawBackground() const;
+	void DrawTimingEventHighlight(double StartTime, double EndTime, float Y, EHighlightMode Mode);
 
 	//TODO: move the following in a Builder class
 	void BeginTimelines();
@@ -72,15 +108,11 @@ public:
 	void EndTimeline(FTimingEventsTrack& Track);
 	void EndTimelines();
 
-	void DrawTimingEventHighlight(double StartTime, double EndTime, float Y, EHighlightMode Mode);
-
-	void End();
-
 private:
 	void DrawBox(const FBoxData& Box, const float EventY, const float EventH);
 
-public:
-	const FDrawContext& DC;
+private:
+	const FDrawContext& DrawContext;
 	const FTimingTrackViewport& Viewport;
 	const FTimingEventsTrackLayout& Layout;
 
@@ -101,19 +133,13 @@ public:
 	float TimelineTopY;
 	float TimelineY;
 	int32 MaxDepth;
-	int TimelineIndex;
+	int32 TimelineIndex;
 
 	TArray<float> LastEventX2; // X2 value for last event on each depth, for current timeline
 	TArray<FBoxData> LastBox;
 
-	// Debug stats.
-	mutable int NumEvents;
-	mutable int NumDrawBoxes;
-	mutable int NumMergedBoxes;
-	mutable int NumDrawBorders;
-	mutable int NumDrawTexts;
-	mutable int NumDrawTimeMarkerBoxes;
-	mutable int NumDrawTimeMarkerTexts;
+	/** Debug stats */
+	mutable FStats Stats;
 };
 
 ////////////////////////////////////////////////////////////////////////////////////////////////////
