[TraceInsights]

diff --git a/Engine/Source/Developer/TraceInsights/Private/Insights/Widgets/SFrameTrack.cpp b/Engine/Source/Developer/TraceInsights/Private/Insights/Widgets/SFrameTrack.cpp
--- a/Engine/Source/Developer/TraceInsights/Private/Insights/Widgets/SFrameTrack.cpp
+++ b/Engine/Source/Developer/TraceInsights/Private/Insights/Widgets/SFrameTrack.cpp
@@ -59,6 +59,8 @@ void SFrameTrack::Reset()
 	TimelinesOrder.Add(2);
 	bIsStateDirty = true;
 
+	AnalysisSyncNextTimestamp = 0;
+
 	MousePosition = FVector2D::ZeroVector;
 
 	MousePositionOnButtonDown = FVector2D::ZeroVector;
@@ -173,35 +175,43 @@ void SFrameTrack::Tick(const FGeometry& AllottedGeometry, const double InCurrent
 		}
 	}
 
-	TSharedPtr<const Trace::IAnalysisSession> Session = FInsightsManager::Get()->GetSession();
-	if (Session.IsValid())
+	uint64 Time = FPlatformTime::Cycles64();
+	if (Time > AnalysisSyncNextTimestamp)
 	{
-		Trace::FAnalysisSessionReadScope _(*Session.Get());
-		Session->ReadFramesProvider([this](const Trace::IFrameProvider& FramesProvider)
+		const uint64 WaitTime = static_cast<uint64>(0.1 / FPlatformTime::GetSecondsPerCycle64()); // 100ms
+		AnalysisSyncNextTimestamp = Time + WaitTime;
+
+		TSharedPtr<const Trace::IAnalysisSession> Session = FInsightsManager::Get()->GetSession();
+		if (Session.IsValid())
 		{
-			for (int32 FrameType = 0; FrameType < TraceFrameType_Count; ++FrameType)
-			{
-				FFrameTrackTimeline& CachedTimeline = CachedTimelines.FindOrAdd(FrameType);
-				CachedTimeline.Id = FrameType;
+			Trace::FAnalysisSessionReadScope SessionReadScope(*Session.Get());
 
-				int32 NumFrames = FramesProvider.GetFrameCount(ETraceFrameType(FrameType));
-				if (NumFrames > Viewport.MaxIndex)
+			Session->ReadFramesProvider([this](const Trace::IFrameProvider& FramesProvider)
+			{
+				for (int32 FrameType = 0; FrameType < TraceFrameType_Count; ++FrameType)
 				{
-					Viewport.SetMinMaxIndexInterval(0, NumFrames);
-					UpdateHorizontalScrollBar();
-					bIsStateDirty = true;
-
-					// Auto zoom out.
-					float SampleW = Viewport.GetSampleWidth();
-					int32 NumSamples = FMath::CeilToInt(Viewport.Width / SampleW);
-					if (NumFrames > NumSamples * Viewport.GetNumFramesPerSample())
+					FFrameTrackTimeline& CachedTimeline = CachedTimelines.FindOrAdd(FrameType);
+					CachedTimeline.Id = FrameType;
+
+					int32 NumFrames = FramesProvider.GetFrameCount(ETraceFrameType(FrameType));
+					if (NumFrames > Viewport.MaxIndex)
 					{
-						ZoomHorizontally(-0.1f, 0.0f);
-						Viewport.ScrollAtPosX(0.0f);
+						Viewport.SetMinMaxIndexInterval(0, NumFrames);
+						UpdateHorizontalScrollBar();
+						bIsStateDirty = true;
+
+						// Auto zoom out.
+						float SampleW = Viewport.GetSampleWidth();
+						int32 NumSamples = FMath::CeilToInt(Viewport.Width / SampleW);
+						if (NumFrames > NumSamples * Viewport.GetNumFramesPerSample())
+						{
+							ZoomHorizontally(-0.1f, 0.0f);
+							Viewport.ScrollAtPosX(0.0f);
+						}
 					}
 				}
-			}
-		});
+			});
+		}
 	}
 
 	if (bIsStateDirty)
@@ -218,47 +228,41 @@ void SFrameTrack::UpdateState()
 	FStopwatch Stopwatch;
 	Stopwatch.Start();
 
-	for (auto It = CachedTimelines.CreateIterator(); It; ++It)
+	// Reset stats.
+	for (TPair<uint64, FFrameTrackTimeline>& KeyValuePair : CachedTimelines)
 	{
-		(*It).Value.NumAggregatedFrames = 0;
+		FFrameTrackTimeline& Timeline = KeyValuePair.Value;
+		Timeline.NumAggregatedFrames = 0;
 	}
-
-	FFrameTrackCacheContext CC(Viewport);
-	CC.Begin();
+	NumUpdatedFrames = 0;
 
 	TSharedPtr<const Trace::IAnalysisSession> Session = FInsightsManager::Get()->GetSession();
 	if (Session.IsValid())
 	{
-		Trace::FAnalysisSessionReadScope _(*Session.Get());
-		Session->ReadFramesProvider([this, &CC](const Trace::IFrameProvider& FramesProvider)
+		Trace::FAnalysisSessionReadScope SessionReadScope(*Session.Get());
+
+		Session->ReadFramesProvider([this](const Trace::IFrameProvider& FramesProvider)
 		{
 			for (int32 FrameType = 0; FrameType < TraceFrameType_Count; ++FrameType)
 			{
-				FFrameTrackTimeline& CachedTimeline = CachedTimelines.FindOrAdd(FrameType);
-				CachedTimeline.Id = FrameType;
+				FFrameTrackTimeline& Timeline = CachedTimelines.FindOrAdd(FrameType);
+				Timeline.Id = FrameType;
 
-				if (CC.BeginTimeline(CachedTimeline))
-				{
-					int32 StartIndex = FMath::Max(0, CC.Viewport.GetIndexAtViewportX(0.0f));
-					int32 EndIndex = CC.Viewport.GetIndexAtViewportX(Viewport.Width);
-					//int32 NumFrames = Timeline.GetEventCount();
-					//int32 EndIndex = FMath::Min(NumFrames, CC.Viewport.GetIndexAtViewportX(Viewport.Width));
+				FFrameTrackTimelineBuilder Builder(Timeline, Viewport);
 
-					FramesProvider.EnumerateFrames(ETraceFrameType(FrameType), static_cast<uint64>(StartIndex), static_cast<uint64>(EndIndex), [&CC](const Trace::FFrame& Frame)
-					{
-						CC.AddEvent(Frame);
-					});
+				int32 StartIndex = FMath::Max(0, Viewport.GetIndexAtViewportX(0.0f));
+				int32 EndIndex = Viewport.GetIndexAtViewportX(Viewport.Width);
+				FramesProvider.EnumerateFrames(ETraceFrameType(FrameType), static_cast<uint64>(StartIndex), static_cast<uint64>(EndIndex), [&Builder](const Trace::FFrame& Frame)
+				{
+					Builder.AddFrame(Frame);
+				});
 
-					CC.EndTimeline();
-				}
+				NumUpdatedFrames += Builder.GetNumAddedFrames();
 			}
 		});
 	}
 
-	CC.End();
-
 	Stopwatch.Stop();
-	NumUpdatedFrames = CC.NumFrames;
 	UpdateDurationHistory.AddValue(Stopwatch.AccumulatedTime);
 }
 
@@ -267,14 +271,14 @@ void SFrameTrack::UpdateState()
 FSampleRef SFrameTrack::GetSampleAtMousePosition(float X, float Y)
 {
 	float SampleW = Viewport.GetSampleWidth();
-	int SampleIndex = FMath::FloorToInt(X / SampleW);
+	int32 SampleIndex = FMath::FloorToInt(X / SampleW);
 	if (SampleIndex >= 0)
 	{
 		const float ViewportHeight = FMath::RoundToFloat(Viewport.Height);
 		const float SampleY2 = ViewportHeight - FMath::RoundToFloat(Viewport.GetViewportYForValue(0.0));
 
 		// Search in reverse paint order.
-		for (int TimelineIndex = TimelinesOrder.Num() - 1; TimelineIndex >= 0; --TimelineIndex)
+		for (int32 TimelineIndex = TimelinesOrder.Num() - 1; TimelineIndex >= 0; --TimelineIndex)
 		{
 			int32 TimelineId = TimelinesOrder[TimelineIndex];
 			if (CachedTimelines.Contains(TimelineId))
@@ -313,9 +317,13 @@ void SFrameTrack::SelectFrameAtMousePosition(float X, float Y)
 {
 	FSampleRef SampleRef = GetSampleAtMousePosition(X, Y);
 	if (!SampleRef.IsValid())
+	{
 		SampleRef = GetSampleAtMousePosition(X - 1.0f, Y);
+	}
 	if (!SampleRef.IsValid())
+	{
 		SampleRef = GetSampleAtMousePosition(X + 1.0f, Y);
+	}
 
 	if (SampleRef.IsValid())
 	{
@@ -335,96 +343,84 @@ void SFrameTrack::SelectFrameAtMousePosition(float X, float Y)
 
 int32 SFrameTrack::OnPaint(const FPaintArgs& Args, const FGeometry& AllottedGeometry, const FSlateRect& MyCullingRect, FSlateWindowElementList& OutDrawElements, int32 LayerId, const FWidgetStyle& InWidgetStyle, bool bParentEnabled) const
 {
-	const TSharedRef<FSlateFontMeasure> FontMeasureService = FSlateApplication::Get().GetRenderer()->GetFontMeasureService();
-
 	const bool bEnabled = ShouldBeEnabled(bParentEnabled);
 	const ESlateDrawEffect DrawEffects = bEnabled ? ESlateDrawEffect::None : ESlateDrawEffect::DisabledEffect;
+	FDrawContext DrawContext(AllottedGeometry, MyCullingRect, InWidgetStyle, DrawEffects, OutDrawElements, LayerId);
+
+	const TSharedRef<FSlateFontMeasure> FontMeasureService = FSlateApplication::Get().GetRenderer()->GetFontMeasureService();
 
 	const FSlateBrush* WhiteBrush = FCoreStyle::Get().GetBrush("WhiteBrush");
 
 	const float ViewWidth = AllottedGeometry.Size.X;
 	const float ViewHeight = AllottedGeometry.Size.Y;
 
-	FFrameTrackDrawContext DC(Viewport, AllottedGeometry, LayerId, OutDrawElements);
+	int32 NumDrawSamples = 0;
 
-	FStopwatch Stopwatch;
-	Stopwatch.Start();
+	{
+		FStopwatch Stopwatch;
+		Stopwatch.Start();
 
-	DC.DrawBackground(InWidgetStyle);
+		FFrameTrackDrawHelper Helper(DrawContext, Viewport);
 
-	// Draw frames.
-	{
-		DC.NumFrames = 0;
-		DC.NumDrawSamples = 0;
+		Helper.DrawBackground();
 
-		for (auto It = CachedTimelines.CreateConstIterator(); It; ++It)
+		// Draw frames, for each visible timeline.
+		for (const TPair<uint64, FFrameTrackTimeline>& KeyValuePair : CachedTimelines)
 		{
-			DC.DrawCached((*It).Value);
+			const FFrameTrackTimeline& Timeline = KeyValuePair.Value;
+			Helper.DrawCached(Timeline);
 		}
-		//ensure(DC.NumFrames == NumUpdatedFrames);
-	}
+		NumDrawSamples = Helper.GetNumDrawSamples();
 
-	// Highlight the mouse hovered sample.
-	if (HoveredSample.IsValid())
-	{
-		DC.DrawHoveredSample(*HoveredSample.Sample);
-
-		FString Text = FString::Format(TEXT("{0} frame {1} ({2})"),
-			{
-				(HoveredSample.Timeline->Id == 0) ? TEXT("Game") :
-				(HoveredSample.Timeline->Id == 1) ? TEXT("Render") : TEXT("Misc"),
-				FText::AsNumber(HoveredSample.Sample->LargestFrameIndex).ToString(),
-				TimeUtils::FormatTimeAuto(HoveredSample.Sample->LargestFrameDuration)
-			});
+		// Highlight the mouse hovered sample.
+		if (HoveredSample.IsValid())
+		{
+			Helper.DrawHoveredSample(*HoveredSample.Sample);
 
-		FSlateFontInfo SummaryFont = FCoreStyle::GetDefaultFontStyle("Regular", 8);
-		FVector2D TextSize = FontMeasureService->Measure(Text, SummaryFont);
+			FString Text = FString::Format(TEXT("{0} frame {1} ({2})"),
+				{
+					(HoveredSample.Timeline->Id == 0) ? TEXT("Game") :
+					(HoveredSample.Timeline->Id == 1) ? TEXT("Render") : TEXT("Misc"),
+					FText::AsNumber(HoveredSample.Sample->LargestFrameIndex).ToString(),
+					TimeUtils::FormatTimeAuto(HoveredSample.Sample->LargestFrameDuration)
+				});
 
-		const float DX = 2.0f;
-		const float W2 = TextSize.X/2 + DX;
+			FSlateFontInfo SummaryFont = FCoreStyle::GetDefaultFontStyle("Regular", 8);
+			FVector2D TextSize = FontMeasureService->Measure(Text, SummaryFont);
 
-		float X1 = Viewport.GetViewportXForIndex(HoveredSample.Sample->LargestFrameIndex);
-		float CX = X1 + FMath::RoundToFloat(Viewport.GetSampleWidth() / 2);
-		if (CX + W2 > Viewport.Width)
-			CX = FMath::RoundToFloat(Viewport.Width - W2);
-		if (CX - W2 < 0)
-			CX = W2;
+			const float DX = 2.0f;
+			const float W2 = TextSize.X / 2 + DX;
 
-		const float Y = 10.0f;
-		const float H = 14.0f;
+			float X1 = Viewport.GetViewportXForIndex(HoveredSample.Sample->LargestFrameIndex);
+			float CX = X1 + FMath::RoundToFloat(Viewport.GetSampleWidth() / 2);
+			if (CX + W2 > Viewport.Width)
+			{
+				CX = FMath::RoundToFloat(Viewport.Width - W2);
+			}
+			if (CX - W2 < 0)
+			{
+				CX = W2;
+			}
 
-		FSlateDrawElement::MakeBox(
-			OutDrawElements,
-			LayerId,
-			MAKE_PAINT_GEOMETRY_RC(AllottedGeometry, CX - W2, Y, 2 * W2, H),
-			WhiteBrush,
-			DrawEffects,
-			FLinearColor(0.7, 0.7, 0.7, 0.9)
-		);
+			const float Y = 10.0f;
+			const float H = 14.0f;
+			DrawContext.DrawBox(CX - W2, Y, 2 * W2, H, WhiteBrush, FLinearColor(0.7, 0.7, 0.7, 0.9));
+			DrawContext.DrawText(CX - W2 + DX, Y + 1.0f, Text, SummaryFont, FLinearColor(0.0, 0.0, 0.0, 0.9));
+		}
 
-		FSlateDrawElement::MakeText
-		(
-			OutDrawElements,
-			LayerId,
-			MAKE_PAINT_GEOMETRY_PT(AllottedGeometry, CX - W2 + DX, Y + 1.0f),
-			Text,
-			SummaryFont,
-			DrawEffects,
-			FLinearColor(0.0, 0.0, 0.0, 0.9)
-		);
-	}
+		TSharedPtr<STimingProfilerWindow> Window = FTimingProfilerManager::Get()->GetProfilerWindow();
+		if (Window && Window->TimingView && CachedTimelines.Num() > 0)
+		{
+			// Highlight the area corresponding to viewport of Timing View.
+			const double StartTime = Window->TimingView->GetViewport().StartTime;
+			const double EndTime = Window->TimingView->GetViewport().EndTime;
+			Helper.DrawHighlightedInterval(CachedTimelines[0], StartTime, EndTime);
+		}
 
-	TSharedPtr<STimingProfilerWindow> Window = FTimingProfilerManager::Get()->GetProfilerWindow();
-	if (Window.IsValid() && CachedTimelines.Num() > 0)
-	{
-		const double StartTime = Window->TimingView->GetViewport().StartTime;
-		const double EndTime = Window->TimingView->GetViewport().EndTime;
-		DC.DrawHighlightedInterval(CachedTimelines[0], StartTime, EndTime);
+		Stopwatch.Stop();
+		DrawDurationHistory.AddValue(Stopwatch.AccumulatedTime);
 	}
 
-	Stopwatch.Stop();
-	DrawDurationHistory.AddValue(Stopwatch.AccumulatedTime);
-
 	//////////////////////////////////////////////////
 
 	const bool bShouldDisplayDebugInfo = FInsightsManager::Get()->IsDebugInfoEnabled();
@@ -491,7 +487,7 @@ int32 SFrameTrack::OnPaint(const FPaintArgs& Args, const FGeometry& AllottedGeom
 			FString::Format(TEXT("U: {0} frames, D: {1} samples"),
 				{
 					FText::AsNumber(NumUpdatedFrames).ToString(),
-					FText::AsNumber(DC.NumDrawSamples).ToString()
+					FText::AsNumber(NumDrawSamples).ToString()
 				}),
 			SummaryFont,
 			DrawEffects,
diff --git a/Engine/Source/Developer/TraceInsights/Private/Insights/Widgets/SFrameTrack.h b/Engine/Source/Developer/TraceInsights/Private/Insights/Widgets/SFrameTrack.h
--- a/Engine/Source/Developer/TraceInsights/Private/Insights/Widgets/SFrameTrack.h
+++ b/Engine/Source/Developer/TraceInsights/Private/Insights/Widgets/SFrameTrack.h
@@ -141,6 +141,8 @@ protected:
 
 	bool bIsStateDirty;
 
+	uint64 AnalysisSyncNextTimestamp;
+
 	//////////////////////////////////////////////////
 
 	TSharedPtr<SScrollBar> HorizontalScrollBar;
diff --git a/Engine/Source/Developer/TraceInsights/Private/Insights/Widgets/SGraphTrack.cpp b/Engine/Source/Developer/TraceInsights/Private/Insights/Widgets/SGraphTrack.cpp
--- a/Engine/Source/Developer/TraceInsights/Private/Insights/Widgets/SGraphTrack.cpp
+++ b/Engine/Source/Developer/TraceInsights/Private/Insights/Widgets/SGraphTrack.cpp
@@ -62,21 +62,8 @@ void SGraphTrack::Reset()
 
 void SGraphTrack::Construct(const FArguments& InArgs)
 {
-	ChildSlot
-	[
-		SNew(SOverlay)
-		.Visibility(EVisibility::SelfHitTestInvisible)
-
-		//+SOverlay::Slot()
-		//.HAlign(HAlign_Left)
-		//.VAlign(VAlign_Top)
-		//.Padding(FMargin(48.0f, 16.0f, 48.0f, 16.0f))	// Make some space for graph labels
-		//[
-		//	SAssignNew(GraphDescriptionsVBox, SVerticalBox)
-		//]
-	];
-
-	//BindCommands();
+	//TODO: Add toggle buttons for Points/Lines/Bars.
+	//TODO: Add "show list of series" expand button + toggle buttons for each available graph series.
 }
 
 ////////////////////////////////////////////////////////////////////////////////////////////////////
@@ -112,115 +99,17 @@ void SGraphTrack::Tick(const FGeometry& AllottedGeometry, const double InCurrent
 
 ////////////////////////////////////////////////////////////////////////////////////////////////////
 
-//PRAGMA_DISABLE_OPTIMIZATION
 int32 SGraphTrack::OnPaint(const FPaintArgs& Args, const FGeometry& AllottedGeometry, const FSlateRect& MyCullingRect, FSlateWindowElementList& OutDrawElements, int32 LayerId, const FWidgetStyle& InWidgetStyle, bool bParentEnabled) const
 {
 	const bool bEnabled = ShouldBeEnabled(bParentEnabled);
 	const ESlateDrawEffect DrawEffects = bEnabled ? ESlateDrawEffect::None : ESlateDrawEffect::DisabledEffect;
-	FDrawContext DC(AllottedGeometry, MyCullingRect, InWidgetStyle, DrawEffects, OutDrawElements, LayerId);
+	FDrawContext DrawContext(AllottedGeometry, MyCullingRect, InWidgetStyle, DrawEffects, OutDrawElements, LayerId);
 
-	GraphTrack.Draw(DC, Viewport);
-	TimeRulerTrack.Draw(DC, Viewport);
-
-	/*
-	static double TotalTime = 0.0f;
-	static uint32 NumCalls = 0;
-	const double StartTime = FPlatformTime::Seconds();
-
-	const TSharedRef<FSlateFontMeasure> FontMeasureService = FSlateApplication::Get().GetRenderer()->GetFontMeasureService();
-
-	// Rendering info.
-	const bool bEnabled = ShouldBeEnabled(bParentEnabled);
-	ESlateDrawEffect DrawEffects = bEnabled ? ESlateDrawEffect::None : ESlateDrawEffect::DisabledEffect;
-	const FSlateBrush* WhiteBrush = FCoreStyle::Get().GetBrush("WhiteBrush");
-
-	const float Width = AllottedGeometry.GetLocalSize().X;
-	const float Height = AllottedGeometry.GetLocalSize().Y;
-
-	// Draw background.
-	const FSlateBrush* AreaBrush = FEditorStyle::GetBrush("Profiler.LineGraphArea");
-	FSlateDrawElement::MakeBox
-	(
-		OutDrawElements,
-		LayerId,
-		AllottedGeometry.ToPaintGeometry(FVector2D(0,0), FVector2D(Width, Height)),
-		AreaBrush,
-		DrawEffects,
-		AreaBrush->GetTint(InWidgetStyle) * InWidgetStyle.GetColorAndOpacityTint()
-	);
-	LayerId++;
-
-	// Draw a test graph...
-	static const int EventCount = 1000;
-	const FLinearColor LineColor(0.0f, 0.5f, 1.0f, 1.0f);
-	TArray<FVector2D> LinePoints;
-	FVector2D V1(0.0f, 0.0f);
-	FVector2D V2;
-	FRandomStream RandomStream(0);
-	for (int Index = 0; Index < EventCount; ++Index)
-	{
-		V2.X = FMath::TruncToFloat((Width / (float)EventCount) * (float)Index);
-		V2.Y = Height * RandomStream.GetFraction();
-		if ((V2 - V1).Size() > 1.0)
-		{
-			LinePoints.Add(V1);
-			LinePoints.Add(V2);
-			FSlateDrawElement::MakeLines(OutDrawElements, LayerId, AllottedGeometry.ToPaintGeometry(), LinePoints, DrawEffects, LineColor);
-			LinePoints.Empty();
-			V1 = V2;
-		}
-	}
-	*/
-
-	//double sec = FPlatformTime::GetSecondsPerCycle();
-
-	//FSlateFontInfo SummaryFont = FCoreStyle::GetDefaultFontStyle("Regular", 8);
-	//const float MaxFontCharHeight = FontMeasureService->Measure(TEXT("!"), SummaryFont).Y;
-
-	/*
-	// Debug draw (slate fps)
-	++LayerId;
-	FSlateDrawElement::MakeBox(
-		OutDrawElements,
-		LayerId,
-		AllottedGeometry.ToPaintGeometry(FVector2D(8.0f, 8.0f), FVector2D(200.0, MaxFontCharHeight + 4.0)),
-		EventBrush,
-		DrawEffects,
-		FLinearColor(FColor(255, 255, 255))
-	);
-	static uint64 Time = 0;
-	const uint64 CurrentTime = FPlatformTime::Cycles64();
-	const uint64 Duration = CurrentTime - Time;
-	Time = CurrentTime;
-	static const int DurationCount = 64;
-	static uint64 Durations[DurationCount];
-	static uint64 DurationOffset = 0;
-	Durations[(DurationOffset++)% DurationCount] = Duration;
-	uint64 AvgDuration = 0;
-	for (int I = 0; I < DurationCount; ++I)
-		AvgDuration += Durations[I];
-	AvgDuration = AvgDuration / DurationCount;
-	const double AvgFps = AvgDuration > 0 ? 1.0 / (FPlatformTime::GetSecondsPerCycle64() * (double)AvgDuration) : 0.0;
-	const FString DebnugStatsStr = FString::Printf(TEXT("%d ms (%d ms) %.1f fps %d events"),
-		FMath::TruncToInt(AvgDuration * 1000 * FPlatformTime::GetSecondsPerCycle64()),
-		FMath::TruncToInt(Duration * 1000 * FPlatformTime::GetSecondsPerCycle64()),
-		AvgFps,
-		EventCount);
-	FSlateDrawElement::MakeText
-	(
-		OutDrawElements,
-		LayerId,
-		AllottedGeometry.ToOffsetPaintGeometry(FVector2D(10.0f, 10.0f)),
-		DebnugStatsStr,
-		SummaryFont,
-		DrawEffects,
-		FLinearColor::Red
-	);
-	*/
+	GraphTrack.Draw(DrawContext, Viewport);
+	TimeRulerTrack.Draw(DrawContext, Viewport);
 
 	return SCompoundWidget::OnPaint(Args, AllottedGeometry, MyCullingRect, OutDrawElements, LayerId, InWidgetStyle, bParentEnabled && IsEnabled());
 }
-//PRAGMA_ENABLE_OPTIMIZATION
 
 ////////////////////////////////////////////////////////////////////////////////////////////////////
 
diff --git a/Engine/Source/Developer/TraceInsights/Private/Insights/Widgets/SInsightsSettings.cpp b/Engine/Source/Developer/TraceInsights/Private/Insights/Widgets/SInsightsSettings.cpp
--- a/Engine/Source/Developer/TraceInsights/Private/Insights/Widgets/SInsightsSettings.cpp
+++ b/Engine/Source/Developer/TraceInsights/Private/Insights/Widgets/SInsightsSettings.cpp
@@ -34,7 +34,7 @@ void SInsightsSettings::Construct(const FArguments& InArgs)
 	SettingPtr = InArgs._SettingPtr;
 
 	const TSharedRef<SGridPanel> SettingsGrid = SNew(SGridPanel);
-	int CurrentRowPos = 0;
+	int32 CurrentRowPos = 0;
 
 	AddTitle(LOCTEXT("SettingsTitle","Unreal Insights - Settings"), SettingsGrid, CurrentRowPos);
 	AddSeparator(SettingsGrid, CurrentRowPos);
diff --git a/Engine/Source/Developer/TraceInsights/Private/Insights/Widgets/SLogView.cpp b/Engine/Source/Developer/TraceInsights/Private/Insights/Widgets/SLogView.cpp
--- a/Engine/Source/Developer/TraceInsights/Private/Insights/Widgets/SLogView.cpp
+++ b/Engine/Source/Developer/TraceInsights/Private/Insights/Widgets/SLogView.cpp
@@ -476,27 +476,27 @@ void SLogView::Construct(const FArguments& InArgs)
 			SNew(SHorizontalBox)
 
 			// Icon
-			+ SHorizontalBox::Slot()
-			.VAlign(VAlign_Center)
-			.AutoWidth()
-			[
-				SNew(SImage)
-				.Image(FEditorStyle::GetBrush(TEXT("Log.TabIcon")))
-			]
-
-			// Label
-			+ SHorizontalBox::Slot()
-			.VAlign(VAlign_Center)
-			.Padding(2.0f, 0.0f)
-			.AutoWidth()
-			[
-				SNew(STextBlock)
-				.Text(LOCTEXT("LogViewLabel", "Log View"))
-			]
+			//+ SHorizontalBox::Slot()
+			//.VAlign(VAlign_Center)
+			//.AutoWidth()
+			//[
+			//	SNew(SImage)
+			//	.Image(FEditorStyle::GetBrush(TEXT("Log.TabIcon")))
+			//]
+			//
+			//// Label
+			//+ SHorizontalBox::Slot()
+			//.VAlign(VAlign_Center)
+			//.Padding(2.0f, 0.0f)
+			//.AutoWidth()
+			//[
+			//	SNew(STextBlock)
+			//	.Text(LOCTEXT("LogViewLabel", "Log View"))
+			//]
 
 			// Verbosity Threshold Filter
 			+SHorizontalBox::Slot()
-			.Padding(2.0f, 0.0f, 0.0f, 0.0f)
+			//.Padding(2.0f, 0.0f, 0.0f, 0.0f)
 			.VAlign(VAlign_Center)
 			.AutoWidth()
 			[
@@ -695,7 +695,7 @@ TSharedRef<ITableRow> SLogView::OnGenerateRow(TSharedPtr<FLogMessage> InLogMessa
 
 void SLogView::Tick(const FGeometry& AllottedGeometry, const double InCurrentTime, const float InDeltaTime)
 {
-	int NewMessageCount = 0;
+	int32 NewMessageCount = 0;
 
 	TSharedPtr<const Trace::IAnalysisSession> Session = FInsightsManager::Get()->GetSession();
 
@@ -703,17 +703,17 @@ void SLogView::Tick(const FGeometry& AllottedGeometry, const double InCurrentTim
 
 	if (Session.IsValid())
 	{
-		Trace::FAnalysisSessionReadScope _(*Session.Get());
+		Trace::FAnalysisSessionReadScope SessionReadScope(*Session.Get());
 		Session->ReadLogProvider([this, &NewMessageCount](const Trace::ILogProvider& LogProvider)
 		{
-			NewMessageCount = static_cast<int>(LogProvider.GetMessageCount());
+			NewMessageCount = static_cast<int32>(LogProvider.GetMessageCount());
 
 			//TODO: show only categories that are used in current trace
 			//TODO: cause of duplicates: a) runtime, b) case insensitive, c) stripped "Log" prefix
 			//TODO: FString vs. FText vs. FName ?
 
-			//TODO: int NumCategories = static_cast<int>(LogProvider.GetCategoriesCount());
-			int NumCategories = 0;
+			//TODO: int32 NumCategories = static_cast<int32>(LogProvider.GetCategoriesCount());
+			int32 NumCategories = 0;
 			LogProvider.EnumerateCategories([&NumCategories](const Trace::FLogCategory& Category)
 			{
 				NumCategories++;
@@ -796,13 +796,13 @@ void SLogView::Tick(const FGeometry& AllottedGeometry, const double InCurrentTim
 			{
 				FilteringStopwatch.Restart();
 
-				for (int Index = TotalNumMessages; Index < NewMessageCount; Index++)
+				for (int32 Index = TotalNumMessages; Index < NewMessageCount; Index++)
 				{
 					FLogMessage LogMessage(Index);
 					Messages.Add(MakeShared<FLogMessage>(MoveTemp(LogMessage)));
 				}
 
-				const int NumAddedMessages = NewMessageCount - TotalNumMessages;
+				const int32 NumAddedMessages = NewMessageCount - TotalNumMessages;
 				TotalNumMessages = NewMessageCount;
 				ListView->RebuildList();
 				bIsDirty = false;
@@ -852,8 +852,8 @@ void SLogView::Tick(const FGeometry& AllottedGeometry, const double InCurrentTim
 			const TArray<uint32>& FilteredMessages = Task.GetFilteredMessages();
 
 			// Add filtered messages to current Messages array.
-			const int NumFilteredMessages = FilteredMessages.Num();
-			for (int Index = 0; Index < NumFilteredMessages; Index++)
+			const int32 NumFilteredMessages = FilteredMessages.Num();
+			for (int32 Index = 0; Index < NumFilteredMessages; Index++)
 			{
 				FLogMessage LogMessage(FilteredMessages[Index]);
 				Messages.Add(MakeShared<FLogMessage>(MoveTemp(LogMessage)));
@@ -869,7 +869,7 @@ void SLogView::Tick(const FGeometry& AllottedGeometry, const double InCurrentTim
 			uint64 DurationMs = FilteringStopwatch.GetAccumulatedTimeMs();
 			if (DurationMs > 10) // avoids spams
 			{
-				int NumAsyncFilteredMessages = Task.GetEndIndex() - Task.GetStartIndex();
+				int32 NumAsyncFilteredMessages = Task.GetEndIndex() - Task.GetStartIndex();
 				double Speed = static_cast<double>(NumAsyncFilteredMessages) / FilteringStopwatch.GetAccumulatedTime();
 				UE_LOG(TimingProfiler, Log, TEXT("[LogView] Updated (%d added / %d async filtered / %d total messages) in %llu ms (%.2f messages/second)."),
 					NumFilteredMessages, NumAsyncFilteredMessages, TotalNumMessages, DurationMs, Speed);
@@ -1075,7 +1075,7 @@ void  SLogView::CreateVerbosityThresholdMenuSection(FMenuBuilder& MenuBuilder)
 		{ ELogVerbosity::Fatal,       LOCTEXT("VerbosityThresholdFatal",       "Fatal"),       LOCTEXT("VerbosityThresholdFatal_Tooltip",       "Show only Fatal log messages."), },
 	};
 
-	for (int Index = 0; Index < sizeof(VerbosityThresholds) / sizeof(FVerbosityThresholdInfo); ++Index)
+	for (int32 Index = 0; Index < sizeof(VerbosityThresholds) / sizeof(FVerbosityThresholdInfo); ++Index)
 	{
 		const FVerbosityThresholdInfo& Threshold = VerbosityThresholds[Index];
 
diff --git a/Engine/Source/Developer/TraceInsights/Private/Insights/Widgets/SLogView.h b/Engine/Source/Developer/TraceInsights/Private/Insights/Widgets/SLogView.h
--- a/Engine/Source/Developer/TraceInsights/Private/Insights/Widgets/SLogView.h
+++ b/Engine/Source/Developer/TraceInsights/Private/Insights/Widgets/SLogView.h
@@ -116,8 +116,8 @@ protected:
 	FLogFilter Filter;
 	uint64 FilterChangeNumber;
 
-	int FilteringStartIndex; // Start index (of the range of log messages to filter) currenly used by the async task
-	int FilteringEndIndex; // End index (of the range of log messages to filter) currenly used by the async task
+	int32 FilteringStartIndex; // Start index (of the range of log messages to filter) currenly used by the async task
+	int32 FilteringEndIndex; // End index (of the range of log messages to filter) currenly used by the async task
 	uint64 FilteringChangeNumber; // Change number of the filter currenly used by the async task
 	TUniquePtr<FAsyncTask<FLogFilteringAsyncTask>> FilteringAsyncTask; // The async task to filter log messages on a worker thread
 	mutable volatile bool bIsFilteringAsyncTaskCancelRequested; // true if we want the async task to finish asap
@@ -125,10 +125,10 @@ protected:
 	/** Stopwatch used to measure how long it takes to filter the message list. */
 	mutable FStopwatch FilteringStopwatch;
 
-	int TotalNumCategories;
+	int32 TotalNumCategories;
 
 	/** Total number of log messages processed, from the source Trace session. Used to detect when new log messages are added in the source Trace session. */
-	int TotalNumMessages;
+	int32 TotalNumMessages;
 
 	/** true if the list of messages is not yet updated (the filter has changed and/or the source trace messages have changed) */
 	bool bIsDirty;
diff --git a/Engine/Source/Developer/TraceInsights/Private/Insights/Widgets/SStartPageWindow.cpp b/Engine/Source/Developer/TraceInsights/Private/Insights/Widgets/SStartPageWindow.cpp
new file mode 100644
--- /dev/null
+++ b/Engine/Source/Developer/TraceInsights/Private/Insights/Widgets/SStartPageWindow.cpp
@@ -0,0 +1,493 @@
+// Copyright 1998-2019 Epic Games, Inc. All Rights Reserved.
+
+#include "SStartPageWindow.h"
+
+#include "EditorStyleSet.h"
+#include "Framework/Docking/WorkspaceItem.h"
+#include "SlateOptMacros.h"
+#include "Widgets/Docking/SDockTab.h"
+#include "Widgets/Images/SImage.h"
+#include "Widgets/Layout/SBorder.h"
+#include "Widgets/Layout/SBox.h"
+#include "Widgets/Layout/SSpacer.h"
+#include "Widgets/Notifications/SNotificationList.h"
+#include "Widgets/SBoxPanel.h"
+#include "Widgets/Text/STextBlock.h"
+#include "Widgets/Input/SButton.h"
+//#include "WorkspaceMenuStructure.h"
+//#include "WorkspaceMenuStructureModule.h"
+
+#if WITH_EDITOR
+	#include "EngineAnalytics.h"
+	#include "Runtime/Analytics/Analytics/Public/AnalyticsEventAttribute.h"
+	#include "Runtime/Analytics/Analytics/Public/Interfaces/IAnalyticsProvider.h"
+#endif // WITH_EDITOR
+
+// Insights
+#include "Insights/InsightsManager.h"
+#include "Insights/InsightsStyle.h"
+#include "Insights/Version.h"
+#include "Insights/Widgets/SInsightsSettings.h"
+
+////////////////////////////////////////////////////////////////////////////////////////////////////
+
+#define LOCTEXT_NAMESPACE "SStartPageWindow"
+
+////////////////////////////////////////////////////////////////////////////////////////////////////
+
+// TODO: move this function to InsightsManager or in TraceSession.h
+static FText GetTextForNotification(const EInsightsNotificationType NotificatonType, const ELoadingProgressState ProgressState, const FString& Filename, const float ProgressPercent = 0.0f)
+{
+	FText Result;
+
+	if (NotificatonType == EInsightsNotificationType::LoadingTraceFile)
+	{
+		if (ProgressState == ELoadingProgressState::Started)
+		{
+			FFormatNamedArguments Args;
+			Args.Add(TEXT("Filename"), FText::FromString(Filename));
+			Result = FText::Format(LOCTEXT("DescF_OfflineCapture_Started", "Started loading a file {Filename}"), Args);
+		}
+		else if (ProgressState == ELoadingProgressState::InProgress)
+		{
+			FFormatNamedArguments Args;
+			Args.Add(TEXT("Filename"), FText::FromString(Filename));
+			Args.Add(TEXT("LoadingProgressPercent"), FText::AsPercent(ProgressPercent));
+			Result = FText::Format(LOCTEXT("DescF_OfflineCapture_InProgress", "Loading a file {Filename} {LoadingProgressPercent}"), Args);
+		}
+		else if (ProgressState == ELoadingProgressState::Loaded)
+		{
+			FFormatNamedArguments Args;
+			Args.Add(TEXT("Filename"), FText::FromString(Filename));
+			Result = FText::Format(LOCTEXT("DescF_OfflineCapture_Loaded", "File {Filename} has been successfully loaded"), Args);
+		}
+		else if (ProgressState == ELoadingProgressState::Failed)
+		{
+			FFormatNamedArguments Args;
+			Args.Add(TEXT("Filename"), FText::FromString(Filename));
+			Result = FText::Format(LOCTEXT("DescF_OfflineCapture_Failed", "Failed to load file {Filename}"), Args);
+		}
+	}
+
+	return Result;
+}
+
+////////////////////////////////////////////////////////////////////////////////////////////////////
+
+SStartPageWindow::SStartPageWindow()
+	: DurationActive(0.0f)
+{
+}
+
+////////////////////////////////////////////////////////////////////////////////////////////////////
+
+SStartPageWindow::~SStartPageWindow()
+{
+#if WITH_EDITOR
+	if (DurationActive > 0.0f && FEngineAnalytics::IsAvailable())
+	{
+		FEngineAnalytics::GetProvider().RecordEvent(TEXT("Editor.Usage.Profiler"), FAnalyticsEventAttribute(TEXT("Duration"), DurationActive));
+	}
+#endif // WITH_EDITOR
+}
+
+////////////////////////////////////////////////////////////////////////////////////////////////////
+
+BEGIN_SLATE_FUNCTION_BUILD_OPTIMIZATION
+
+////////////////////////////////////////////////////////////////////////////////////////////////////
+
+void SStartPageWindow::Construct(const FArguments& InArgs)
+{
+	ChildSlot
+		[
+			SNew(SOverlay)
+
+			// Version
+			+ SOverlay::Slot()
+				.HAlign(HAlign_Right)
+				.VAlign(VAlign_Top)
+				.Padding(0.0f, -16.0f, 0.0f, 0.0f)
+				[
+					SNew(STextBlock)
+						.Clipping(EWidgetClipping::ClipToBoundsWithoutIntersecting)
+						.Text(LOCTEXT("UnrealInsightsVersion", UNREAL_INSIGHTS_VERSION_STRING_EX))
+						.ColorAndOpacity(FLinearColor(0.15f, 0.15f, 0.15f, 1.0f))
+				]
+
+			// Overlay slot for the main window area
+			+ SOverlay::Slot()
+				.HAlign(HAlign_Center)
+				.VAlign(VAlign_Center)
+				[
+					SNew(SHorizontalBox)
+
+					+ SHorizontalBox::Slot()
+						.AutoWidth()
+						[
+							SNew(SButton)
+							//.OnClicked(this, &StartPageWindow::Live_OnClicked)
+							.ToolTipText(LOCTEXT("LiveButtonTooltip", "Start a live session or try loading the last trace file."))
+							.ContentPadding(8.0f)
+							.Content()
+							[
+								SNew(SHorizontalBox)
+
+								+SHorizontalBox::Slot()
+									.AutoWidth()
+									[
+										SNew(SImage)
+										.Image(FInsightsStyle::GetBrush("Live.Icon.Large"))
+									]
+								+SHorizontalBox::Slot()
+									.AutoWidth()
+									.VAlign(VAlign_Center)
+									[
+										SNew(SBox)
+										.WidthOverride(100.0f)
+										[
+											SNew(STextBlock)
+											.Font(FCoreStyle::GetDefaultFontStyle("Regular", 18))
+											.Text(LOCTEXT("LiveButtonText", "Live"))
+										]
+									]
+							]
+						]
+
+					+ SHorizontalBox::Slot()
+						.AutoWidth()
+						[
+							SNew(SBox)
+							.WidthOverride(6.0f)
+						]
+
+					+ SHorizontalBox::Slot()
+						.AutoWidth()
+						[
+							SNew(SButton)
+							//.OnClicked(this, &StartPageWindow::Load_OnClicked)
+							.ToolTipText(LOCTEXT("LoadButtonTooltip", "Last a trace file."))
+							.ContentPadding(8.0f)
+							.Content()
+							[
+								SNew(SHorizontalBox)
+
+								+SHorizontalBox::Slot()
+									.AutoWidth()
+									[
+										SNew(SImage)
+										.Image(FInsightsStyle::GetBrush("Load.Icon.Large"))
+									]
+								+SHorizontalBox::Slot()
+									.AutoWidth()
+									.VAlign(VAlign_Center)
+									[
+										SNew(SBox)
+										.WidthOverride(100.0f)
+										[
+											SNew(STextBlock)
+											.Font(FCoreStyle::GetDefaultFontStyle("Regular", 18))
+											.Text(LOCTEXT("LoadButtonText", "Load..."))
+										]
+									]
+							]
+						]
+				]
+
+			// Session hint overlay
+			+ SOverlay::Slot()
+				.HAlign(HAlign_Center)
+				.VAlign(VAlign_Center)
+				.Padding(0.0f, 0.0f, 0.0f, 48.0f)
+				[
+					SNew(SBorder)
+						.Visibility(this, &SStartPageWindow::IsSessionOverlayVisible)
+						.BorderImage(FEditorStyle::GetBrush("NotificationList.ItemBackground"))
+						.Padding(8.0f)
+						[
+							SNew(STextBlock)
+								.Text(LOCTEXT("SelectTraceOverlayText", "Please select a trace."))
+						]
+				]
+
+			// Notification area overlay
+			+ SOverlay::Slot()
+				.HAlign(HAlign_Right)
+				.VAlign(VAlign_Bottom)
+				.Padding(16.0f)
+				[
+					SAssignNew(NotificationList, SNotificationList)
+				]
+
+			// Settings dialog overlay
+			+ SOverlay::Slot()
+				.HAlign(HAlign_Center)
+				.VAlign(VAlign_Center)
+				.Expose(OverlaySettingsSlot)
+		];
+}
+
+END_SLATE_FUNCTION_BUILD_OPTIMIZATION
+
+////////////////////////////////////////////////////////////////////////////////////////////////////
+
+EVisibility SStartPageWindow::IsSessionOverlayVisible() const
+{
+	if (FInsightsManager::Get()->GetSession().IsValid())
+	{
+		return EVisibility::Hidden;
+	}
+
+	return EVisibility::Visible;
+}
+
+////////////////////////////////////////////////////////////////////////////////////////////////////
+
+bool SStartPageWindow::IsSessionValid() const
+{
+	const bool bIsActive = FInsightsManager::Get()->GetSession().IsValid();
+	return bIsActive;
+}
+
+////////////////////////////////////////////////////////////////////////////////////////////////////
+
+void SStartPageWindow::ManageLoadingProgressNotificationState(const FString& Filename, const EInsightsNotificationType NotificatonType, const ELoadingProgressState ProgressState, const float LoadingProgress)
+{
+	const FString BaseFilename = FPaths::GetBaseFilename(Filename);
+
+	if (ProgressState == ELoadingProgressState::Started)
+	{
+		const bool bContains = ActiveNotifications.Contains(Filename);
+		if (!bContains)
+		{
+			FNotificationInfo NotificationInfo(GetTextForNotification(NotificatonType, ProgressState, BaseFilename));
+			NotificationInfo.bFireAndForget = false;
+			NotificationInfo.bUseLargeFont = false;
+
+			// Add two buttons, one for cancel, one for loading the received file.
+			if (NotificatonType == EInsightsNotificationType::LoadingTraceFile)
+			{
+				const FText CancelButtonText = LOCTEXT("CancelButton_Text", "Cancel");
+				const FText CancelButtonTT = LOCTEXT("CancelButton_TTText", "Hides this notification");
+				const FText LoadButtonText = LOCTEXT("LoadButton_Text", "Load file");
+				const FText LoadButtonTT = LOCTEXT("LoadButton_TTText", "Loads the received file and hides this notification");
+
+				NotificationInfo.ButtonDetails.Add(FNotificationButtonInfo(CancelButtonText, CancelButtonTT,
+					FSimpleDelegate::CreateSP(this, &SStartPageWindow::SendingServiceSideCapture_Cancel, Filename), SNotificationItem::CS_Success));
+				NotificationInfo.ButtonDetails.Add(FNotificationButtonInfo(LoadButtonText, LoadButtonTT,
+					FSimpleDelegate::CreateSP(this, &SStartPageWindow::SendingServiceSideCapture_Load, Filename), SNotificationItem::CS_Success));
+			}
+
+			SNotificationItemWeak NotificationItem = NotificationList->AddNotification(NotificationInfo);
+			NotificationItem.Pin()->SetCompletionState(SNotificationItem::CS_Pending);
+			ActiveNotifications.Add(Filename, NotificationItem);
+		}
+	}
+	else if (ProgressState == ELoadingProgressState::InProgress)
+	{
+		const SNotificationItemWeak* LoadingProgressPtr = ActiveNotifications.Find(Filename);
+		if (LoadingProgressPtr)
+		{
+			TSharedPtr<SNotificationItem> LoadingProcessPinned = LoadingProgressPtr->Pin();
+			LoadingProcessPinned->SetText(GetTextForNotification(NotificatonType, ProgressState, BaseFilename, LoadingProgress));
+			LoadingProcessPinned->SetCompletionState(SNotificationItem::CS_Pending);
+		}
+	}
+	else if (ProgressState == ELoadingProgressState::Loaded)
+	{
+		const SNotificationItemWeak* LoadingProgressPtr = ActiveNotifications.Find(Filename);
+		if (LoadingProgressPtr)
+		{
+			TSharedPtr<SNotificationItem> LoadingProcessPinned = LoadingProgressPtr->Pin();
+			LoadingProcessPinned->SetText(GetTextForNotification(NotificatonType, ProgressState, BaseFilename));
+			LoadingProcessPinned->SetCompletionState(SNotificationItem::CS_Success);
+
+			// Notifications for received files are handled by the user.
+			if (NotificatonType == EInsightsNotificationType::LoadingTraceFile)
+			{
+				LoadingProcessPinned->ExpireAndFadeout();
+				ActiveNotifications.Remove(Filename);
+			}
+		}
+	}
+	else if (ProgressState == ELoadingProgressState::Failed)
+	{
+		const SNotificationItemWeak* LoadingProgressPtr = ActiveNotifications.Find(Filename);
+		if (LoadingProgressPtr)
+		{
+			TSharedPtr<SNotificationItem> LoadingProcessPinned = LoadingProgressPtr->Pin();
+			LoadingProcessPinned->SetText(GetTextForNotification(NotificatonType, ProgressState, BaseFilename));
+			LoadingProcessPinned->SetCompletionState(SNotificationItem::CS_Fail);
+
+			LoadingProcessPinned->ExpireAndFadeout();
+			ActiveNotifications.Remove(Filename);
+		}
+	}
+	else if (ProgressState == ELoadingProgressState::Cancelled)
+	{
+		const SNotificationItemWeak* LoadingProgressPtr = ActiveNotifications.Find(Filename);
+		if (LoadingProgressPtr)
+		{
+			TSharedPtr<SNotificationItem> LoadingProcessPinned = LoadingProgressPtr->Pin();
+			LoadingProcessPinned->ExpireAndFadeout();
+			ActiveNotifications.Remove(Filename);
+		}
+	}
+}
+
+////////////////////////////////////////////////////////////////////////////////////////////////////
+
+void SStartPageWindow::SendingServiceSideCapture_Cancel(const FString Filename)
+{
+	const SNotificationItemWeak* LoadingProgressPtr = ActiveNotifications.Find(Filename);
+	if (LoadingProgressPtr)
+	{
+		TSharedPtr<SNotificationItem> LoadingProcessPinned = LoadingProgressPtr->Pin();
+		LoadingProcessPinned->ExpireAndFadeout();
+		ActiveNotifications.Remove(Filename);
+	}
+}
+
+////////////////////////////////////////////////////////////////////////////////////////////////////
+
+void SStartPageWindow::SendingServiceSideCapture_Load(const FString Filename)
+{
+	const SNotificationItemWeak* LoadingProgressPtr = ActiveNotifications.Find(Filename);
+	if (LoadingProgressPtr)
+	{
+		TSharedPtr<SNotificationItem> LoadingProcessPinned = LoadingProgressPtr->Pin();
+		LoadingProcessPinned->ExpireAndFadeout();
+		ActiveNotifications.Remove(Filename);
+
+		const FString PathName = FPaths::ProfilingDir() + TEXT("UnrealStats/Received/");
+		const FString TraceFilepath = PathName + Filename;
+		FInsightsManager::Get()->LoadTraceFile(TraceFilepath);
+	}
+}
+
+////////////////////////////////////////////////////////////////////////////////////////////////////
+
+EActiveTimerReturnType SStartPageWindow::UpdateActiveDuration(double InCurrentTime, float InDeltaTime)
+{
+	DurationActive += InDeltaTime;
+
+	// The window will explicitly unregister this active timer when the mouse leaves.
+	return EActiveTimerReturnType::Continue;
+}
+
+////////////////////////////////////////////////////////////////////////////////////////////////////
+
+void SStartPageWindow::OnMouseEnter(const FGeometry& MyGeometry, const FPointerEvent& MouseEvent)
+{
+	SCompoundWidget::OnMouseEnter(MyGeometry, MouseEvent);
+
+	if (!ActiveTimerHandle.IsValid())
+	{
+		ActiveTimerHandle = RegisterActiveTimer(0.f, FWidgetActiveTimerDelegate::CreateSP(this, &SStartPageWindow::UpdateActiveDuration));
+	}
+}
+
+////////////////////////////////////////////////////////////////////////////////////////////////////
+
+void SStartPageWindow::OnMouseLeave(const FPointerEvent& MouseEvent)
+{
+	SCompoundWidget::OnMouseLeave(MouseEvent);
+
+	auto PinnedActiveTimerHandle = ActiveTimerHandle.Pin();
+	if (PinnedActiveTimerHandle.IsValid())
+	{
+		UnRegisterActiveTimer(PinnedActiveTimerHandle.ToSharedRef());
+	}
+}
+
+////////////////////////////////////////////////////////////////////////////////////////////////////
+
+FReply SStartPageWindow::OnKeyDown(const FGeometry& MyGeometry, const FKeyEvent& InKeyEvent)
+{
+	return FReply::Unhandled();
+}
+
+////////////////////////////////////////////////////////////////////////////////////////////////////
+
+FReply SStartPageWindow::OnDragOver(const FGeometry& MyGeometry, const FDragDropEvent& DragDropEvent)
+{
+	TSharedPtr<FExternalDragOperation> DragDropOp = DragDropEvent.GetOperationAs<FExternalDragOperation>();
+	if (DragDropOp.IsValid())
+	{
+		if (DragDropOp->HasFiles())
+		{
+			const TArray<FString>& Files = DragDropOp->GetFiles();
+			if (Files.Num() == 1)
+			{
+				const FString DraggedFileExtension = FPaths::GetExtension(Files[0], true);
+				if (DraggedFileExtension == TEXT(".utrace"))
+				{
+					return FReply::Handled();
+				}
+			}
+		}
+	}
+
+	return SCompoundWidget::OnDragOver(MyGeometry, DragDropEvent);
+}
+
+////////////////////////////////////////////////////////////////////////////////////////////////////
+
+FReply SStartPageWindow::OnDrop(const FGeometry& MyGeometry, const FDragDropEvent& DragDropEvent)
+{
+	TSharedPtr<FExternalDragOperation> DragDropOp = DragDropEvent.GetOperationAs<FExternalDragOperation>();
+	if (DragDropOp.IsValid())
+	{
+		if (DragDropOp->HasFiles())
+		{
+			// For now, only allow a single file.
+			const TArray<FString>& Files = DragDropOp->GetFiles();
+			if (Files.Num() == 1)
+			{
+				const FString DraggedFileExtension = FPaths::GetExtension(Files[0], true);
+				if (DraggedFileExtension == TEXT(".utrace"))
+				{
+					// Enqueue load operation.
+					FInsightsManager::Get()->LoadTraceFile(Files[0]);
+					return FReply::Handled();
+				}
+			}
+		}
+	}
+
+	return SCompoundWidget::OnDrop(MyGeometry, DragDropEvent);
+}
+
+////////////////////////////////////////////////////////////////////////////////////////////////////
+
+void SStartPageWindow::OpenSettings()
+{
+	MainContentPanel->SetEnabled(false);
+	(*OverlaySettingsSlot)
+	[
+		SNew(SBorder)
+		.BorderImage(FEditorStyle::GetBrush("NotificationList.ItemBackground"))
+		.Padding(8.0f)
+		[
+			SNew(SInsightsSettings)
+			.OnClose(this, &SStartPageWindow::CloseSettings)
+			.SettingPtr(&FInsightsManager::GetSettings())
+		]
+	];
+}
+
+////////////////////////////////////////////////////////////////////////////////////////////////////
+
+void SStartPageWindow::CloseSettings()
+{
+	// Close the profiler settings by simply replacing widget with a null one.
+	(*OverlaySettingsSlot)
+	[
+		SNullWidget::NullWidget
+	];
+	MainContentPanel->SetEnabled(true);
+}
+
+////////////////////////////////////////////////////////////////////////////////////////////////////
+
+#undef LOCTEXT_NAMESPACE
diff --git a/Engine/Source/Developer/TraceInsights/Private/Insights/Widgets/SStartPageWindow.h b/Engine/Source/Developer/TraceInsights/Private/Insights/Widgets/SStartPageWindow.h
new file mode 100644
--- /dev/null
+++ b/Engine/Source/Developer/TraceInsights/Private/Insights/Widgets/SStartPageWindow.h
@@ -0,0 +1,150 @@
+// Copyright 1998-2019 Epic Games, Inc. All Rights Reserved.
+
+#pragma once
+
+#include "CoreMinimal.h"
+#include "Framework/Docking/TabManager.h"
+#include "Input/Reply.h"
+#include "Layout/Visibility.h"
+#include "Misc/Guid.h"
+#include "SlateFwd.h"
+#include "Widgets/DeclarativeSyntaxSupport.h"
+#include "Widgets/Docking/SDockTab.h"
+#include "Widgets/Layout/SSplitter.h"
+#include "Widgets/SCompoundWidget.h"
+#include "Widgets/SOverlay.h"
+
+// Insights
+#include "Insights/InsightsManager.h"
+
+////////////////////////////////////////////////////////////////////////////////////////////////////
+
+class FActiveTimerHandle;
+class SVerticalBox;
+
+////////////////////////////////////////////////////////////////////////////////////////////////////
+
+/** Type definition for shared pointers to instances of SNotificationItem. */
+typedef TSharedPtr<class SNotificationItem> SNotificationItemPtr;
+
+/** Type definition for shared references to instances of SNotificationItem. */
+typedef TSharedRef<class SNotificationItem> SNotificationItemRef;
+
+/** Type definition for weak references to instances of SNotificationItem. */
+typedef TWeakPtr<class SNotificationItem> SNotificationItemWeak;
+
+////////////////////////////////////////////////////////////////////////////////////////////////////
+
+/** Implements the Start Page window. */
+class SStartPageWindow : public SCompoundWidget
+{
+public:
+	/** Default constructor. */
+	SStartPageWindow();
+
+	/** Virtual destructor. */
+	virtual ~SStartPageWindow();
+
+	SLATE_BEGIN_ARGS(SStartPageWindow){}
+	SLATE_END_ARGS()
+
+	/** Constructs this widget. */
+	void Construct(const FArguments& InArgs);
+
+	void ManageLoadingProgressNotificationState(const FString& Filename, const EInsightsNotificationType NotificatonType, const ELoadingProgressState ProgressState, const float LoadingProgress);
+
+	void OpenSettings();
+	void CloseSettings();
+
+	void ShowHideTab(const FName& TabID, bool bIsVisible);
+
+protected:
+	/** Callback for determining the visibility of the 'Select a session' overlay. */
+	EVisibility IsSessionOverlayVisible() const;
+
+	bool IsSessionValid() const;
+
+	void SendingServiceSideCapture_Cancel(const FString Filename);
+
+	void SendingServiceSideCapture_Load(const FString Filename);
+
+private:
+
+	/**
+	 * Ticks this widget.  Override in derived classes, but always call the parent implementation.
+	 *
+	 * @param  AllottedGeometry The space allotted for this widget
+	 * @param  InCurrentTime  Current absolute real time
+	 * @param  InDeltaTime  Real time passed since last tick
+	 */
+	//virtual void Tick(const FGeometry& AllottedGeometry, const double InCurrentTime, const float InDeltaTime) override;
+
+	/**
+	 * The system will use this event to notify a widget that the cursor has entered it. This event is NOT bubbled.
+	 *
+	 * @param MyGeometry The Geometry of the widget receiving the event
+	 * @param MouseEvent Information about the input event
+	 */
+	virtual void OnMouseEnter(const FGeometry& MyGeometry, const FPointerEvent& MouseEvent) override;
+
+	/**
+	 * The system will use this event to notify a widget that the cursor has left it. This event is NOT bubbled.
+	 *
+	 * @param MouseEvent Information about the input event
+	 */
+	virtual void OnMouseLeave(const FPointerEvent& MouseEvent) override;
+
+	/**
+	 * Called after a key is pressed when this widget has focus
+	 *
+	 * @param MyGeometry The Geometry of the widget receiving the event
+	 * @param  InKeyEvent  Key event
+	 *
+	 * @return  Returns whether the event was handled, along with other possible actions
+	 */
+	virtual FReply OnKeyDown(const FGeometry& MyGeometry, const FKeyEvent& InKeyEvent) override;
+
+	/**
+	 * Called when the user is dropping something onto a widget; terminates drag and drop.
+	 *
+	 * @param MyGeometry      The geometry of the widget receiving the event.
+	 * @param DragDropEvent   The drag and drop event.
+	 *
+	 * @return A reply that indicated whether this event was handled.
+	 */
+	virtual FReply OnDrop(const FGeometry& MyGeometry, const FDragDropEvent& DragDropEvent) override;
+
+	/**
+	 * Called during drag and drop when the the mouse is being dragged over a widget.
+	 *
+	 * @param MyGeometry      The geometry of the widget receiving the event.
+	 * @param DragDropEvent   The drag and drop event.
+	 *
+	 * @return A reply that indicated whether this event was handled.
+	 */
+	virtual FReply OnDragOver(const FGeometry& MyGeometry, const FDragDropEvent& DragDropEvent)  override;
+
+private:
+	/** Updates the amount of time the profiler has been active. */
+	EActiveTimerReturnType UpdateActiveDuration(double InCurrentTime, float InDeltaTime);
+
+public:
+	/** Widget for the non-intrusive notifications. */
+	TSharedPtr<SNotificationList> NotificationList;
+
+	/** Holds all active and visible notifications, stored as FGuid -> SNotificationItemWeak. */
+	TMap<FString, SNotificationItemWeak> ActiveNotifications;
+
+	/** Overlay slot which contains the profiler settings widget. */
+	SOverlay::FOverlaySlot* OverlaySettingsSlot;
+
+	/** The number of seconds the profiler has been active */
+	float DurationActive;
+
+private:
+	/** The handle to the active update duration tick */
+	TWeakPtr<FActiveTimerHandle> ActiveTimerHandle;
+
+	/** Holds all widgets for the profiler window like menu bar, toolbar and tabs. */
+	TSharedPtr<SVerticalBox> MainContentPanel;
+};
diff --git a/Engine/Source/Developer/TraceInsights/Private/Insights/Widgets/SStatsTableCell.h b/Engine/Source/Developer/TraceInsights/Private/Insights/Widgets/SStatsTableCell.h
new file mode 100644
--- /dev/null
+++ b/Engine/Source/Developer/TraceInsights/Private/Insights/Widgets/SStatsTableCell.h
@@ -0,0 +1,140 @@
+// Copyright 1998-2019 Epic Games, Inc. All Rights Reserved.
+
+#pragma once
+
+#include "CoreMinimal.h"
+#include "Widgets/DeclarativeSyntaxSupport.h"
+#include "Widgets/SCompoundWidget.h"
+
+// Insights
+#include "Insights/ViewModels/StatsNode.h"
+
+DECLARE_DELEGATE_TwoParams(FSetHoveredTableCell, const FName /*ColumnId*/, const FStatsNodePtr /*SamplePtr*/);
+DECLARE_DELEGATE_RetVal_OneParam(bool, FIsColumnVisibleDelegate, const FName /*ColumnId*/);
+DECLARE_DELEGATE_RetVal_OneParam(EHorizontalAlignment, FGetColumnOutlineHAlignmentDelegate, const FName /*ColumnId*/);
+
+class SStatsTableCell : public SCompoundWidget
+{
+public:
+	SLATE_BEGIN_ARGS(SStatsTableCell) {}
+		SLATE_EVENT(FSetHoveredTableCell, OnSetHoveredTableCell)
+		SLATE_ATTRIBUTE(FText, HighlightText)
+		SLATE_ARGUMENT(FStatsNodePtr, StatsNodePtr)
+		SLATE_ARGUMENT(FName, ColumnId)
+		SLATE_ARGUMENT(bool, IsStatsNameColumn)
+	SLATE_END_ARGS()
+
+	void Construct(const FArguments& InArgs, const TSharedRef<class ITableRow>& TableRow);
+
+protected:
+	TSharedRef<SWidget> GenerateWidgetForColumn(const FArguments& InArgs, const TSharedRef<class ITableRow>& TableRow);
+	TSharedRef<SWidget> GenerateWidgetForNameColumn(const FArguments& InArgs, const TSharedRef<class ITableRow>& TableRow);
+	TSharedRef<SWidget> GenerateWidgetForStatsColumn(const FArguments& InArgs, const TSharedRef<class ITableRow>& TableRow);
+
+	/**
+	 * The system will use this event to notify a widget that the cursor has entered it. This event is NOT bubbled.
+	 *
+	 * @param MyGeometry The Geometry of the widget receiving the event
+	 * @param MouseEvent Information about the input event
+	 */
+	virtual void OnMouseEnter(const FGeometry& MyGeometry, const FPointerEvent& MouseEvent) override
+	{
+		SCompoundWidget::OnMouseEnter(MyGeometry, MouseEvent);
+		SetHoveredTableCellDelegate.ExecuteIfBound(ColumnId, StatsNodePtr);
+	}
+
+	/**
+	 * The system will use this event to notify a widget that the cursor has left it. This event is NOT bubbled.
+	 *
+	 * @param MouseEvent Information about the input event
+	 */
+	virtual void OnMouseLeave(const FPointerEvent& MouseEvent) override
+	{
+		SCompoundWidget::OnMouseLeave(MouseEvent);
+		SetHoveredTableCellDelegate.ExecuteIfBound(NAME_None, nullptr);
+	}
+
+	/**
+	 * Called during drag and drop when the drag enters a widget.
+	 *
+	 * Enter/Leave events in slate are meant as lightweight notifications.
+	 * So we do not want to capture mouse or set focus in response to these.
+	 * However, OnDragEnter must also support external APIs (e.g. OLE Drag/Drop)
+	 * Those require that we let them know whether we can handle the content
+	 * being dragged OnDragEnter.
+	 *
+	 * The concession is to return a can_handled/cannot_handle
+	 * boolean rather than a full FReply.
+	 *
+	 * @param MyGeometry      The geometry of the widget receiving the event.
+	 * @param DragDropEvent   The drag and drop event.
+	 *
+	 * @return A reply that indicated whether the contents of the DragDropEvent can potentially be processed by this widget.
+	 */
+	virtual void OnDragEnter(const FGeometry& MyGeometry, const FDragDropEvent& DragDropEvent) override
+	{
+		SCompoundWidget::OnDragEnter(MyGeometry, DragDropEvent);
+		SetHoveredTableCellDelegate.ExecuteIfBound(ColumnId, StatsNodePtr);
+	}
+
+	/**
+	 * Called during drag and drop when the drag leaves a widget.
+	 *
+	 * @param DragDropEvent   The drag and drop event.
+	 */
+	virtual void OnDragLeave(const FDragDropEvent& DragDropEvent)  override
+	{
+		SCompoundWidget::OnDragLeave(DragDropEvent);
+		SetHoveredTableCellDelegate.ExecuteIfBound(NAME_None, nullptr);
+	}
+
+	EVisibility GetHintIconVisibility() const
+	{
+		return IsHovered() ? EVisibility::Visible : EVisibility::Hidden;
+	}
+
+	EVisibility GetCulledEventsIconVisibility() const
+	{
+		//return StatsNodePtr->HasCulledChildren() ? EVisibility::Visible : EVisibility::Collapsed;
+		return EVisibility::Collapsed;
+	}
+
+	FText GetNameEx() const
+	{
+		return StatsNodePtr->GetNameEx();
+	}
+
+	FSlateColor GetColorAndOpacity() const
+	{
+		const FLinearColor TextColor =
+			StatsNodePtr->IsFiltered() ? FLinearColor(1.0f, 1.0f, 1.0f, 0.5f) :
+			FLinearColor::White;
+		return TextColor;
+	}
+
+	FSlateColor GetStatsColorAndOpacity() const
+	{
+		const FLinearColor TextColor =
+			StatsNodePtr->IsFiltered() ? FLinearColor(1.0f, 1.0f, 1.0f, 0.5f) :
+			StatsNodePtr->GetAggregatedStats().Count == 0 ? FLinearColor(1.0f, 1.0f, 1.0f, 0.6f) :
+			FLinearColor::White;
+		return TextColor;
+	}
+
+	FLinearColor GetShadowColorAndOpacity() const
+	{
+		const FLinearColor ShadowColor =
+			StatsNodePtr->IsFiltered() ? FLinearColor(0.f, 0.f, 0.f, 0.25f) :
+			FLinearColor(0.0f, 0.0f, 0.0f, 0.5f);
+		return ShadowColor;
+	}
+
+protected:
+	/** A shared pointer to the stats node. */
+	FStatsNodePtr StatsNodePtr;
+
+	/** The Id of the column where this stats belongs. */
+	FName ColumnId;
+
+	FSetHoveredTableCell SetHoveredTableCellDelegate;
+};
diff --git a/Engine/Source/Developer/TraceInsights/Private/Insights/Widgets/SStatsTableRow.h b/Engine/Source/Developer/TraceInsights/Private/Insights/Widgets/SStatsTableRow.h
new file mode 100644
--- /dev/null
+++ b/Engine/Source/Developer/TraceInsights/Private/Insights/Widgets/SStatsTableRow.h
@@ -0,0 +1,93 @@
+// Copyright 1998-2019 Epic Games, Inc. All Rights Reserved.
+
+#pragma once
+
+#include "CoreMinimal.h"
+#include "Widgets/DeclarativeSyntaxSupport.h"
+#include "Widgets/Views/SHeaderRow.h"
+#include "Widgets/Views/STableViewBase.h"
+#include "Widgets/Views/STableRow.h"
+#include "Widgets/Views/STreeView.h"
+
+// Insights
+#include "Insights/ViewModels/StatsNodeHelper.h"
+
+DECLARE_DELEGATE_RetVal_OneParam(bool, FShouldBeEnabledDelegate, const uint32 /*StatsId*/);
+DECLARE_DELEGATE_RetVal_OneParam(bool, FIsColumnVisibleDelegate, const FName /*ColumnId*/);
+DECLARE_DELEGATE_TwoParams(FSetHoveredTableCell, const FName /*ColumnId*/, const FStatsNodePtr /*StatsNodePtr*/);
+DECLARE_DELEGATE_RetVal_OneParam(EHorizontalAlignment, FGetColumnOutlineHAlignmentDelegate, const FName /*ColumnId*/);
+
+/** Widget that represents a table row in the stats' tree control. Generates widgets for each column on demand. */
+class SStatsTableRow : public SMultiColumnTableRow<FStatsNodePtr>
+{
+public:
+	SLATE_BEGIN_ARGS(SStatsTableRow) {}
+		SLATE_EVENT(FShouldBeEnabledDelegate, OnShouldBeEnabled)
+		SLATE_EVENT(FIsColumnVisibleDelegate, OnIsColumnVisible)
+		SLATE_EVENT(FSetHoveredTableCell, OnSetHoveredTableCell)
+		SLATE_EVENT(FGetColumnOutlineHAlignmentDelegate, OnGetColumnOutlineHAlignmentDelegate)
+		SLATE_ATTRIBUTE(FText, HighlightText)
+		SLATE_ATTRIBUTE(FName, HighlightedStatsName)
+		SLATE_ARGUMENT(FStatsNodePtr, StatsNodePtr)
+	SLATE_END_ARGS()
+
+public:
+	void Construct(const FArguments& InArgs, const TSharedRef<STableViewBase>& InOwnerTableView);
+
+	virtual TSharedRef<SWidget> GenerateWidgetForColumn(const FName& ColumnId) override;
+
+	/**
+	 * Called when Slate detects that a widget started to be dragged.
+	 * Usage:
+	 * A widget can ask Slate to detect a drag.
+	 * OnMouseDown() reply with FReply::Handled().DetectDrag(SharedThis(this)).
+	 * Slate will either send an OnDragDetected() event or do nothing.
+	 * If the user releases a mouse button or leaves the widget before
+	 * a drag is triggered (maybe user started at the very edge) then no event will be
+	 * sent.
+	 *
+	 * @param  InMyGeometry  Widget geometry
+	 * @param  InMouseEvent  MouseMove that triggered the drag
+	 *
+	 */
+	virtual FReply OnDragDetected(const FGeometry& MyGeometry, const FPointerEvent& MouseEvent) override;
+
+protected:
+	/**
+	 * @return a text which describes this table row, refers to both groups and stats
+	 */
+	FText GetText() const;
+
+	/**
+	 * @return a font style which is used to draw this table row, refers to both groups and stats
+	 */
+	FSlateFontInfo GetFont() const;
+
+	/**
+	 * @return a color and opacity value used to draw this table row, refers to both groups and stats
+	 */
+	FSlateColor GetColorAndOpacity() const;
+
+	FSlateColor GetBackgroundColorAndOpacity() const;
+	FSlateColor GetBackgroundColorAndOpacity(double Time) const;
+	FSlateColor GetOutlineColorAndOpacity() const;
+	const FSlateBrush* GetOutlineBrush(const FName ColumnId) const;
+	bool HandleShouldBeEnabled() const;
+	EVisibility IsColumnVisible(const FName ColumnId) const;
+	void OnSetHoveredTableCell(const FName ColumnId, const FStatsNodePtr SamplePtr);
+
+protected:
+	/** Data context for this table row. */
+	FStatsNodePtr StatsNodePtr;
+
+	FShouldBeEnabledDelegate OnShouldBeEnabled;
+	FIsColumnVisibleDelegate IsColumnVisibleDelegate;
+	FSetHoveredTableCell SetHoveredTableCellDelegate;
+	FGetColumnOutlineHAlignmentDelegate GetColumnOutlineHAlignmentDelegate;
+
+	/** Text to be highlighted on stats name. */
+	TAttribute<FText> HighlightText;
+
+	/** Name of the stats that should be drawn as highlighted. */
+	TAttribute<FName> HighlightedStatsName;
+};
