[TraceInsights]

diff --git a/Engine/Source/Developer/TraceInsights/Private/Insights/Widgets/STimersView.cpp b/Engine/Source/Developer/TraceInsights/Private/Insights/Widgets/STimersView.cpp
--- a/Engine/Source/Developer/TraceInsights/Private/Insights/Widgets/STimersView.cpp
+++ b/Engine/Source/Developer/TraceInsights/Private/Insights/Widgets/STimersView.cpp
@@ -18,7 +18,6 @@
 // Insights
 #include "Insights/TimingProfilerCommon.h"
 #include "Insights/TimingProfilerManager.h"
-//#include "Insights/ViewModels/TraceSession.h"
 #include "Insights/ViewModels/TimerNodeHelper.h"
 #include "Insights/ViewModels/TimersViewColumnFactory.h"
 #include "Insights/Widgets/STimersViewTooltip.h"
@@ -28,11 +27,16 @@
 
 ////////////////////////////////////////////////////////////////////////////////////////////////////
 
+const EColumnSortMode::Type STimersView::DefaultColumnSortMode(EColumnSortMode::Descending);
+const FName STimersView::DefaultColumnBeingSorted(FTimersViewColumns::TotalInclusiveTimeColumnID);
+
+////////////////////////////////////////////////////////////////////////////////////////////////////
+
 STimersView::STimersView()
 	: bExpansionSaved(false)
 	, GroupingMode(ETimerGroupingMode::ByType)
-	, ColumnSortMode(EColumnSortMode::Type::Descending)
-	, ColumnBeingSorted(FName(TEXT("TotalInclusiveTime")))
+	, ColumnSortMode(DefaultColumnSortMode)
+	, ColumnBeingSorted(DefaultColumnBeingSorted)
 {
 	FMemory::Memset(bTimerNodeIsVisible, 1);
 }
@@ -41,10 +45,10 @@ STimersView::STimersView()
 
 STimersView::~STimersView()
 {
-	// Remove ourselves from the profiler manager.
+	// Remove ourselves from the Insights manager.
 	if (FInsightsManager::Get().IsValid())
 	{
-		FInsightsManager::Get()->OnRequestTimersUpdate().RemoveAll(this);
+		FInsightsManager::Get()->GetSessionChangedEvent().RemoveAll(this);
 	}
 }
 
@@ -130,12 +134,12 @@ void STimersView::Construct(const FArguments& InArgs)
 						GetToggleButtonForTimerType(ETimerNodeType::GpuScope)
 					]
 
-					+SHorizontalBox::Slot()
-					.Padding(FMargin(1.0f,0.0f,1.0f,0.0f))
-					.FillWidth(1.0f)
-					[
-						GetToggleButtonForTimerType(ETimerNodeType::ComputeScope)
-					]
+					//+SHorizontalBox::Slot()
+					//.Padding(FMargin(1.0f,0.0f,1.0f,0.0f))
+					//.FillWidth(1.0f)
+					//[
+					//	GetToggleButtonForTimerType(ETimerNodeType::ComputeScope)
+					//]
 
 					+SHorizontalBox::Slot()
 					.Padding(FMargin(1.0f,0.0f,1.0f,0.0f))
@@ -209,8 +213,8 @@ void STimersView::Construct(const FArguments& InArgs)
 
 	CreateGroupByOptionsSources();
 
-	// Register ourselves with the profiler manager.
-	FInsightsManager::Get()->OnRequestTimersUpdate().AddSP(this, &STimersView::ProfilerManager_OnRequestTimersUpdate);
+	// Register ourselves with the Insights manager.
+	FInsightsManager::Get()->GetSessionChangedEvent().AddSP(this, &STimersView::InsightsManager_OnSessionChanged);
 }
 END_SLATE_FUNCTION_BUILD_OPTIMIZATION
 
@@ -219,7 +223,7 @@ END_SLATE_FUNCTION_BUILD_OPTIMIZATION
 TSharedPtr<SWidget> STimersView::TreeView_GetMenuContent() const
 {
 	const TArray<FTimerNodePtr> SelectedTimerNodes = TreeView->GetSelectedItems();
-	const int NumSelectedTimerNodes = SelectedTimerNodes.Num();
+	const int32 NumSelectedTimerNodes = SelectedTimerNodes.Num();
 	FTimerNodePtr SelectedTimerNode = NumSelectedTimerNodes ? SelectedTimerNodes[0] : nullptr;
 
 	const FTimersViewColumn* const * ColumnPtrPtr = FTimersViewColumnFactory::Get().ColumnIdToPtrMapping.Find(HoveredColumnId);
@@ -272,74 +276,9 @@ TSharedPtr<SWidget> STimersView::TreeView_GetMenuContent() const
 	}
 	MenuBuilder.EndSection();
 
-	/*
-	// Culling/Filtering menu
-	MenuBuilder.BeginSection("Culling/Filtering", LOCTEXT("ContextMenu_Header_Root", "Root"));
-	{
-		// Culling menu
-		FText CullingDesc;
-
-		if (!Column.bCanBeCulled)
-		{
-			CullingDesc = LOCTEXT("ContextMenu_Culling_DescErrCol", "Culling not available, please select a different column");
-		}
-		else if (NumSelectedTimerNodes == 1)
-		{
-			CullingDesc = FText::Format(LOCTEXT("ContextMenu_Culling_DescFmt", "Cull events to '{0}' based on '{1}'"), PropertyValue, PropertyName);
-		}
-		else
-		{
-			CullingDesc = LOCTEXT("ContextMenu_Culling_DescErrEve", "Culling not available, please select one event");
-		}
-
-		MenuBuilder.AddMenuEntry
-		(
-			CullingDesc,
-			LOCTEXT("ContextMenu_Culling_Desc_TT", "Culls the event graph based on the property value of the selected event"),
-			FSlateIcon(FEditorStyle::GetStyleSetName(), "Profiler.EventGraph.CullEvents"), CullByProperty_Custom(SelectedTimerNode, HoveredColumnID, false), NAME_None, EUserInterfaceActionType::Button
-		);
-
-		// Filtering menu
-		FText FilteringDesc;
-
-		if (!Column.bCanBeFiltered)
-		{
-			FilteringDesc = LOCTEXT("ContextMenu_Filtering_DescErrCol", "Filtering not available, please select a different column");
-		}
-		else if (NumSelectedTimerNodes == 1)
-		{
-			FilteringDesc = FText::Format(LOCTEXT("ContextMenu_Filtering_DescFmt", "Filter events to '{0}' based on '{1}'"), PropertyValue, PropertyName);
-		}
-		else
-		{
-			FilteringDesc = LOCTEXT("ContextMenu_Filtering_DescErrEve", "Filtering not available, please select one event");
-		}
-
-		MenuBuilder.AddMenuEntry
-		(
-			FilteringDesc,
-			LOCTEXT("ContextMenu_Filtering_Desc_TT", "Filters the event graph based on the property value of the selected event"),
-			FSlateIcon(FEditorStyle::GetStyleSetName(), "Profiler.EventGraph.FilterEvents"), FilterOutByProperty_Custom(SelectedTimerNode, HoveredColumnID, false), NAME_None, EUserInterfaceActionType::Button
-		);
-
-		MenuBuilder.AddMenuSeparator();
-		MenuBuilder.AddMenuEntry
-		(
-			LOCTEXT("ContextMenu_ClearHistory", "Reset to default"),
-			LOCTEXT("ContextMenu_ClearHistory_Reset_Desc", "For the selected event graph resets root/culling/filter to the default state and clears the history"),
-			FSlateIcon(FEditorStyle::GetStyleSetName(), "Profiler.Misc.ResetToDefault"), ClearHistory_Custom(), NAME_None, EUserInterfaceActionType::Button
-		);
-	}
-	MenuBuilder.EndSection();
-	*/
-
-	/*
-		Event graph coloring based on inclusive time as a gradient from black to red?
-	*/
-
 	MenuBuilder.BeginSection("Misc", LOCTEXT("ContextMenu_Header_Misc", "Miscellaneous"));
 	{
-		/*
+		/*TODO
 		FUIAction Action_CopySelectedToClipboard
 		(
 			FExecuteAction::CreateSP(this, &STimersView::ContextMenu_CopySelectedToClipboard_Execute),
@@ -421,7 +360,9 @@ TSharedPtr<SWidget> STimersView::TreeView_GetMenuContent() const
 void STimersView::TreeView_BuildSortByMenu(FMenuBuilder& MenuBuilder)
 {
 	// TODO: Refactor later @see TSharedPtr<SWidget> SCascadePreviewViewportToolBar::GenerateViewMenu() const
+
 	MenuBuilder.BeginSection("ColumnName", LOCTEXT("ContextMenu_Header_Misc_ColumnName", "Column Name"));
+
 	for (auto It = TreeViewHeaderColumns.CreateConstIterator(); It; ++It)
 	{
 		const FTimersViewColumn& Column = It.Value();
@@ -442,6 +383,7 @@ void STimersView::TreeView_BuildSortByMenu(FMenuBuilder& MenuBuilder)
 			);
 		}
 	}
+
 	MenuBuilder.EndSection();
 
 	//-----------------------------------------------------------------------------
@@ -508,9 +450,6 @@ void STimersView::TreeView_BuildViewColumnMenu(FMenuBuilder& MenuBuilder)
 
 void STimersView::InitializeAndShowHeaderColumns()
 {
-	//ColumnSortMode = EColumnSortMode::Ascending;
-	//ColumnBeingSorted = FTimersViewColumnFactory::Get().Collection[0]->Id;
-
 	const int32 NumColumns = FTimersViewColumnFactory::Get().Collection.Num();
 	for (int32 ColumnIndex = 0; ColumnIndex < NumColumns; ColumnIndex++)
 	{
@@ -677,13 +616,13 @@ TSharedRef<SWidget> STimersView::TreeViewHeaderRow_GenerateColumnMenu(const FTim
 
 ////////////////////////////////////////////////////////////////////////////////////////////////////
 
-void STimersView::ProfilerManager_OnRequestTimersUpdate()
+void STimersView::InsightsManager_OnSessionChanged()
 {
-	TSharedPtr<const Trace::IAnalysisSession> SessionLocal = FInsightsManager::Get()->GetSession();
+	TSharedPtr<const Trace::IAnalysisSession> NewSession = FInsightsManager::Get()->GetSession();
 
-	if (SessionLocal != Session)
+	if (NewSession != Session)
 	{
-		Session = SessionLocal;
+		Session = NewSession;
 		RebuildTree();
 	}
 	else
@@ -723,7 +662,7 @@ void STimersView::ApplyFiltering()
 		{
 			// Add a child.
 			const FTimerNodePtr& NodePtr = GroupChildren[Cx];
-			const bool bIsChildVisible = Filters->PassesAllFilters(NodePtr) && bTimerNodeIsVisible[NodePtr->GetType()];
+			const bool bIsChildVisible = Filters->PassesAllFilters(NodePtr) && bTimerNodeIsVisible[static_cast<int>(NodePtr->GetType())];
 			if (bIsChildVisible)
 			{
 				GroupPtr->AddFilteredChild(NodePtr);
@@ -822,7 +761,7 @@ TSharedRef<SWidget> STimersView::GetToggleButtonForTimerType(const ETimerNodeTyp
 
 void STimersView::FilterByTimerType_OnCheckStateChanged(ECheckBoxState NewRadioState, const ETimerNodeType InStatType)
 {
-	bTimerNodeIsVisible[InStatType] = NewRadioState == ECheckBoxState::Checked;
+	bTimerNodeIsVisible[static_cast<int>(InStatType)] = (NewRadioState == ECheckBoxState::Checked);
 	ApplyFiltering();
 }
 
@@ -830,7 +769,7 @@ void STimersView::FilterByTimerType_OnCheckStateChanged(ECheckBoxState NewRadioS
 
 ECheckBoxState STimersView::FilterByTimerType_IsChecked(const ETimerNodeType InStatType) const
 {
-	return bTimerNodeIsVisible[InStatType] ? ECheckBoxState::Checked : ECheckBoxState::Unchecked;
+	return bTimerNodeIsVisible[static_cast<int>(InStatType)] ? ECheckBoxState::Checked : ECheckBoxState::Unchecked;
 }
 
 ////////////////////////////////////////////////////////////////////////////////////////////////////
@@ -1182,20 +1121,24 @@ void STimersView::SortTimers()
 			} \
 		}
 
-		 CHECK_AND_SORT_COLUMN(FName(TEXT("Name")), ByName)
-	else CHECK_AND_SORT_COLUMN(FName(TEXT("MetaGroupName")), ByMetaGroupName)
-	else CHECK_AND_SORT_COLUMN(FName(TEXT("Type")), ByType)
-	else CHECK_AND_SORT_COLUMN(FName(TEXT("InstanceCount")), ByInstanceCount)
-	else CHECK_AND_SORT_COLUMN(FName(TEXT("TotalInclusiveTime")), ByTotalInclusiveTime)
-	else CHECK_AND_SORT_COLUMN(FName(TEXT("MinInclusiveTime")), ByMinInclusiveTime)
-	else CHECK_AND_SORT_COLUMN(FName(TEXT("MaxInclusiveTime")), ByMaxInclusiveTime)
-	else CHECK_AND_SORT_COLUMN(FName(TEXT("AvgInclusiveTime")), ByAverageInclusiveTime)
-	else CHECK_AND_SORT_COLUMN(FName(TEXT("MedInclusiveTime")), ByMedianInclusiveTime)
-	else CHECK_AND_SORT_COLUMN(FName(TEXT("TotalExclusiveTime")), ByTotalExclusiveTime)
-	else CHECK_AND_SORT_COLUMN(FName(TEXT("MinExclusiveTime")), ByMinExclusiveTime)
-	else CHECK_AND_SORT_COLUMN(FName(TEXT("MaxExclusiveTime")), ByMaxExclusiveTime)
-	else CHECK_AND_SORT_COLUMN(FName(TEXT("AvgExclusiveTime")), ByAverageExclusiveTime)
-	else CHECK_AND_SORT_COLUMN(FName(TEXT("MedExclusiveTime")), ByMedianExclusiveTime)
+		 CHECK_AND_SORT_COLUMN(FTimersViewColumns::NameColumnID,          ByName)
+	else CHECK_AND_SORT_COLUMN(FTimersViewColumns::MetaGroupNameColumnID, ByMetaGroupName)
+	else CHECK_AND_SORT_COLUMN(FTimersViewColumns::TypeColumnID,          ByType)
+	else CHECK_AND_SORT_COLUMN(FTimersViewColumns::InstanceCountColumnID, ByInstanceCount)
+	// Inclusive Time columns
+	else CHECK_AND_SORT_COLUMN(FTimersViewColumns::TotalInclusiveTimeColumnID, ByTotalInclusiveTime)
+	else CHECK_AND_SORT_COLUMN(FTimersViewColumns::MaxInclusiveTimeColumnID,     ByMaxInclusiveTime)
+	else CHECK_AND_SORT_COLUMN(FTimersViewColumns::AverageInclusiveTimeColumnID, ByAverageInclusiveTime)
+	else CHECK_AND_SORT_COLUMN(FTimersViewColumns::MedianInclusiveTimeColumnID,  ByMedianInclusiveTime)
+	else CHECK_AND_SORT_COLUMN(FTimersViewColumns::MinInclusiveTimeColumnID,     ByMinInclusiveTime)
+	// Exclusive Time columns
+	else CHECK_AND_SORT_COLUMN(FTimersViewColumns::TotalExclusiveTimeColumnID,   ByTotalExclusiveTime)
+	else CHECK_AND_SORT_COLUMN(FTimersViewColumns::MaxExclusiveTimeColumnID,     ByMaxExclusiveTime)
+	else CHECK_AND_SORT_COLUMN(FTimersViewColumns::AverageExclusiveTimeColumnID, ByAverageExclusiveTime)
+	else CHECK_AND_SORT_COLUMN(FTimersViewColumns::MedianExclusiveTimeColumnID,  ByMedianExclusiveTime)
+	else CHECK_AND_SORT_COLUMN(FTimersViewColumns::MinExclusiveTimeColumnID,     ByMinExclusiveTime)
+
+	#undef CHECK_AND_SORT_COLUMN
 }
 
 ////////////////////////////////////////////////////////////////////////////////////////////////////
@@ -1210,6 +1153,8 @@ EColumnSortMode::Type STimersView::GetSortModeForColumn(const FName ColumnId) co
 	return ColumnSortMode;
 }
 
+////////////////////////////////////////////////////////////////////////////////////////////////////
+
 void STimersView::SetSortModeForColumn(const FName& ColumnId, const EColumnSortMode::Type SortMode)
 {
 	ColumnBeingSorted = ColumnId;
@@ -1220,6 +1165,8 @@ void STimersView::SetSortModeForColumn(const FName& ColumnId, const EColumnSortM
 	ApplyFiltering();
 }
 
+////////////////////////////////////////////////////////////////////////////////////////////////////
+
 void STimersView::OnSortModeChanged(const EColumnSortPriority::Type SortPriority, const FName& ColumnId, const EColumnSortMode::Type SortMode)
 {
 	SetSortModeForColumn(ColumnId, SortMode);
@@ -1235,6 +1182,8 @@ bool STimersView::HeaderMenu_SortMode_IsChecked(const FName ColumnId, const ECol
 	return ColumnBeingSorted == ColumnId && ColumnSortMode == InSortMode;
 }
 
+////////////////////////////////////////////////////////////////////////////////////////////////////
+
 bool STimersView::HeaderMenu_SortMode_CanExecute(const FName ColumnId, const EColumnSortMode::Type InSortMode) const
 {
 	const FTimersViewColumn& Column = TreeViewHeaderColumns.FindChecked(ColumnId);
@@ -1246,6 +1195,8 @@ bool STimersView::HeaderMenu_SortMode_CanExecute(const FName ColumnId, const ECo
 	return bCanExecute;
 }
 
+////////////////////////////////////////////////////////////////////////////////////////////////////
+
 void STimersView::HeaderMenu_SortMode_Execute(const FName ColumnId, const EColumnSortMode::Type InSortMode)
 {
 	SetSortModeForColumn(ColumnId, InSortMode);
@@ -1261,11 +1212,15 @@ bool STimersView::ContextMenu_SortMode_IsChecked(const EColumnSortMode::Type InS
 	return ColumnSortMode == InSortMode;
 }
 
+////////////////////////////////////////////////////////////////////////////////////////////////////
+
 bool STimersView::ContextMenu_SortMode_CanExecute(const EColumnSortMode::Type InSortMode) const
 {
 	return ColumnSortMode != InSortMode;
 }
 
+////////////////////////////////////////////////////////////////////////////////////////////////////
+
 void STimersView::ContextMenu_SortMode_Execute(const EColumnSortMode::Type InSortMode)
 {
 	SetSortModeForColumn(ColumnBeingSorted, InSortMode);
@@ -1281,11 +1236,15 @@ bool STimersView::ContextMenu_SortByColumn_IsChecked(const FName ColumnId)
 	return ColumnId == ColumnBeingSorted;
 }
 
+////////////////////////////////////////////////////////////////////////////////////////////////////
+
 bool STimersView::ContextMenu_SortByColumn_CanExecute(const FName ColumnId) const
 {
 	return ColumnId != ColumnBeingSorted;
 }
 
+////////////////////////////////////////////////////////////////////////////////////////////////////
+
 void STimersView::ContextMenu_SortByColumn_Execute(const FName ColumnId)
 {
 	SetSortModeForColumn(ColumnId, EColumnSortMode::Descending);
@@ -1302,6 +1261,8 @@ bool STimersView::HeaderMenu_HideColumn_CanExecute(const FName ColumnId) const
 	return Column.bCanBeHidden();
 }
 
+////////////////////////////////////////////////////////////////////////////////////////////////////
+
 void STimersView::HeaderMenu_HideColumn_Execute(const FName ColumnId)
 {
 	FTimersViewColumn& Column = TreeViewHeaderColumns.FindChecked(ColumnId);
@@ -1319,12 +1280,16 @@ bool STimersView::ContextMenu_ToggleColumn_IsChecked(const FName ColumnId)
 	return Column.bIsVisible;
 }
 
+////////////////////////////////////////////////////////////////////////////////////////////////////
+
 bool STimersView::ContextMenu_ToggleColumn_CanExecute(const FName ColumnId) const
 {
 	const FTimersViewColumn& Column = TreeViewHeaderColumns.FindChecked(ColumnId);
 	return Column.bCanBeHidden();
 }
 
+////////////////////////////////////////////////////////////////////////////////////////////////////
+
 void STimersView::ContextMenu_ToggleColumn_Execute(const FName ColumnId)
 {
 	FTimersViewColumn& Column = TreeViewHeaderColumns.FindChecked(ColumnId);
@@ -1351,8 +1316,8 @@ bool STimersView::ContextMenu_ShowAllColumns_CanExecute() const
 
 void STimersView::ContextMenu_ShowAllColumns_Execute()
 {
-	ColumnSortMode = EColumnSortMode::Descending;
-	ColumnBeingSorted = FName(TEXT("TotalInclusiveTime"));
+	ColumnSortMode = DefaultColumnSortMode;
+	ColumnBeingSorted = DefaultColumnBeingSorted;
 
 	const int32 NumColumns = FTimersViewColumnFactory::Get().Collection.Num();
 	for (int32 ColumnIndex = 0; ColumnIndex < NumColumns; ColumnIndex++)
@@ -1382,20 +1347,24 @@ void STimersView::ContextMenu_ShowMinMaxMedColumns_Execute()
 {
 	TSet<FName> Preset =
 	{
-		TEXT("Name"),
-		TEXT("InstanceCount"),
-		TEXT("TotalInclusiveTime"),
-		TEXT("MinInclusiveTime"),
-		TEXT("MaxInclusiveTime"),
-		TEXT("MedInclusiveTime"),
-		TEXT("TotalExclusiveTime"),
-		TEXT("MinExclusiveTime"),
-		TEXT("MaxExclusiveTime"),
-		TEXT("MedExclusiveTime"),
+		FTimersViewColumns::NameColumnID,
+		//FTimersViewColumns::MetaGroupNameColumnID,
+		//FTimersViewColumns::TypeColumnID,
+		FTimersViewColumns::InstanceCountColumnID,
+
+		FTimersViewColumns::TotalInclusiveTimeColumnID,
+		FTimersViewColumns::MaxInclusiveTimeColumnID,
+		FTimersViewColumns::MedianInclusiveTimeColumnID,
+		FTimersViewColumns::MinInclusiveTimeColumnID,
+
+		FTimersViewColumns::TotalExclusiveTimeColumnID,
+		FTimersViewColumns::MaxExclusiveTimeColumnID,
+		FTimersViewColumns::MedianExclusiveTimeColumnID,
+		FTimersViewColumns::MinExclusiveTimeColumnID,
 	};
 
 	ColumnSortMode = EColumnSortMode::Descending;
-	ColumnBeingSorted = FName(TEXT("TotalInclusiveTime"));
+	ColumnBeingSorted = FTimersViewColumns::TotalInclusiveTimeColumnID;
 
 	const int32 NumColumns = FTimersViewColumnFactory::Get().Collection.Num();
 	for (int32 ColumnIndex = 0; ColumnIndex < NumColumns; ColumnIndex++)
@@ -1428,8 +1397,8 @@ bool STimersView::ContextMenu_ResetColumns_CanExecute() const
 
 void STimersView::ContextMenu_ResetColumns_Execute()
 {
-	ColumnSortMode = EColumnSortMode::Descending;
-	ColumnBeingSorted = FName(TEXT("TotalInclusiveTime"));
+	ColumnSortMode = DefaultColumnSortMode;
+	ColumnBeingSorted = DefaultColumnBeingSorted;
 
 	const int32 NumColumns = FTimersViewColumnFactory::Get().Collection.Num();
 	for (int32 ColumnIndex = 0; ColumnIndex < NumColumns; ColumnIndex++)
@@ -1450,32 +1419,37 @@ void STimersView::ContextMenu_ResetColumns_Execute()
 
 ////////////////////////////////////////////////////////////////////////////////////////////////////
 
-void STimersView::RebuildTree(bool bRebuildEvenIfNotChanged)
+void STimersView::RebuildTree(bool bResync)
 {
-	bool bTimerNodesListChanged = false;
+	bool bListHasChanged = false;
 
-	if (bRebuildEvenIfNotChanged)
+	if (bResync)
 	{
 		TimerNodes.Empty(TimerNodes.Num());
 		//TimerNodesMap.Empty(TimerNodesMap.Num());
 		TimerNodesIdMap.Empty(TimerNodesIdMap.Num());
-
-		bTimerNodesListChanged = true;
+		bListHasChanged = true;
 	}
 
 	if (Session.IsValid())
 	{
-		Trace::FAnalysisSessionReadScope _(*Session.Get());
+		Trace::FAnalysisSessionReadScope SessionReadScope(*Session.Get());
 
-		Session->ReadTimingProfilerProvider([this, bRebuildEvenIfNotChanged, &bTimerNodesListChanged](const Trace::ITimingProfilerProvider& TimingProfilerProvider)
+		Session->ReadTimingProfilerProvider([this, &bResync, &bListHasChanged](const Trace::ITimingProfilerProvider& TimingProfilerProvider)
 		{
-			TimingProfilerProvider.ReadTimers([this, bRebuildEvenIfNotChanged, &bTimerNodesListChanged](const Trace::FTimingProfilerTimer* Timers, uint64 TimersCount)
+			TimingProfilerProvider.ReadTimers([this, &bResync, &bListHasChanged](const Trace::FTimingProfilerTimer* Timers, uint64 TimersCount)
 			{
-				if (TimersCount != TimerNodes.Num() || bRebuildEvenIfNotChanged)
+				if (!bResync)
+				{
+					bResync = (TimersCount != TimerNodes.Num());
+				}
+
+				if (bResync)
 				{
 					TimerNodes.Empty(TimerNodes.Num());
 					//TimerNodesMap.Empty(TimerNodesMap.Num());
 					TimerNodesIdMap.Empty(TimerNodesIdMap.Num());
+					bListHasChanged = true;
 
 					for (uint64 TimerIndex = 0; TimerIndex < TimersCount; ++TimerIndex)
 					{
@@ -1489,14 +1463,12 @@ void STimersView::RebuildTree(bool bRebuildEvenIfNotChanged)
 						//TimerNodesMap.Add(Name, TimerNodePtr);
 						TimerNodesIdMap.Add(Timer.Id, TimerNodePtr);
 					}
-
-					bTimerNodesListChanged = true;
 				}
 			});
 		});
 	}
 
-	if (bTimerNodesListChanged)
+	if (bListHasChanged)
 	{
 		UpdateTree();
 		UpdateStats(StatsStartTime, StatsEndTime);
@@ -1567,11 +1539,13 @@ struct FTimeCalculationHelper
 	TMap<uint64, FTimerStatsEx>& StatsMap;
 };
 
+////////////////////////////////////////////////////////////////////////////////////////////////////
+
 void STimersView::UpdateStats(double StartTime, double EndTime)
 {
 	for (const FTimerNodePtr& TimerNodePtr : TimerNodes)
 	{
-		TimerNodePtr->ResetStats();
+		TimerNodePtr->ResetAggregatedStats();
 	}
 
 	if (Session.IsValid() &&
@@ -1592,9 +1566,10 @@ void STimersView::UpdateStats(double StartTime, double EndTime)
 		FTimeCalculationHelper CalculationHelper(Ctx.StartTime, Ctx.EndTime, Ctx.Map);
 
 		{
+			Trace::FAnalysisSessionReadScope SessionReadScope(*Session.Get());
+
 			// Compute instance count and total/min/max inclusive/exclusive times for each timer.
 			// Iterate through all timing events (all thread timelines).
-			Trace::FAnalysisSessionReadScope _(*Session.Get());
 			Session->ReadTimingProfilerProvider([&CalculationHelper](const Trace::ITimingProfilerProvider& TimingProfilerProvider)
 			{
 				TimingProfilerProvider.EnumerateTimelines([&CalculationHelper](const Trace::ITimingProfilerProvider::Timeline& Timeline)
@@ -1633,7 +1608,7 @@ void STimersView::UpdateStats(double StartTime, double EndTime)
 			FTimerNodePtr* TimerNodePtrPtr = TimerNodesIdMap.Find(KV.Key);
 			if (TimerNodePtrPtr != nullptr)
 			{
-				(*TimerNodePtrPtr)->SetStats(KV.Value.BaseStats);
+				(*TimerNodePtrPtr)->SetAggregatedStats(KV.Value.BaseStats);
 			}
 		}
 
@@ -1649,19 +1624,27 @@ void STimersView::UpdateStats(double StartTime, double EndTime)
 
 void STimersView::UpdateTotalMinMaxTimerStats(FTimerStatsEx& StatsEx, double InclTime, double ExclTime)
 {
-	FTimerStats& Stats = StatsEx.BaseStats;
+	FTimerAggregatedStats& Stats = StatsEx.BaseStats;
 
 	Stats.TotalInclusiveTime += InclTime;
 	if (InclTime < Stats.MinInclusiveTime)
+	{
 		Stats.MinInclusiveTime = InclTime;
+	}
 	if (InclTime > Stats.MaxInclusiveTime)
-	Stats.MaxInclusiveTime = InclTime;
+	{
+		Stats.MaxInclusiveTime = InclTime;
+	}
 
 	Stats.TotalExclusiveTime += ExclTime;
 	if (ExclTime < Stats.MinExclusiveTime)
+	{
 		Stats.MinExclusiveTime = ExclTime;
+	}
 	if (ExclTime > Stats.MaxExclusiveTime)
-	Stats.MaxExclusiveTime = ExclTime;
+	{
+		Stats.MaxExclusiveTime = ExclTime;
+	}
 
 	Stats.InstanceCount++;
 }
@@ -1670,54 +1653,62 @@ void STimersView::UpdateTotalMinMaxTimerStats(FTimerStatsEx& StatsEx, double Inc
 
 void STimersView::PreComputeHistogram(FTimerStatsEx& StatsEx)
 {
-	const FTimerStats& Stats = StatsEx.BaseStats;
+	const FTimerAggregatedStats& Stats = StatsEx.BaseStats;
 
 	// Each bucket (Histogram[i]) will be centered on a value.
 	// I.e. First bucket (bucket 0) is centered on Min value: [Min-DT/2, Min+DT/2)
 	// and last bucket (bucket N-1) is centered on Max value: [Max-DT/2, Max+DT/2).
 
-	constexpr double InvHistogramLen = 1.0 / static_cast<double>(HistogramLen - 1);
+	constexpr double InvHistogramLen = 1.0 / static_cast<double>(FTimerStatsEx::HistogramLen - 1);
 
 	if (Stats.MaxInclusiveTime == Stats.MinInclusiveTime)
+	{
 		StatsEx.InclDT = 1.0; // single large bucket
+	}
 	else
+	{
 		StatsEx.InclDT = (Stats.MaxInclusiveTime - Stats.MinInclusiveTime) * InvHistogramLen;
+	}
 
 	if (Stats.MaxExclusiveTime == Stats.MinExclusiveTime)
+	{
 		StatsEx.ExclDT = 1.0; // single large bucket
+	}
 	else
+	{
 		StatsEx.ExclDT = (Stats.MaxExclusiveTime - Stats.MinExclusiveTime) * InvHistogramLen;
+	}
 }
 
 ////////////////////////////////////////////////////////////////////////////////////////////////////
 
 void STimersView::UpdateHistogramForTimerStats(FTimerStatsEx& StatsEx, double InclTime, double ExclTime)
 {
-	const FTimerStats& Stats = StatsEx.BaseStats;
+	const FTimerAggregatedStats& Stats = StatsEx.BaseStats;
 
-	int InclIndex = static_cast<int>((InclTime - Stats.MinInclusiveTime + StatsEx.InclDT / 2) / StatsEx.InclDT);
+	int32 InclIndex = static_cast<int32>((InclTime - Stats.MinInclusiveTime + StatsEx.InclDT / 2) / StatsEx.InclDT);
 	ensure(InclIndex >= 0);
 	if (InclIndex < 0)
 	{
 		InclIndex = 0;
 	}
-	ensure(InclIndex < HistogramLen);
-	if (InclIndex >= HistogramLen)
+	ensure(InclIndex < FTimerStatsEx::HistogramLen);
+	if (InclIndex >= FTimerStatsEx::HistogramLen)
 	{
-		InclIndex = HistogramLen - 1;
+		InclIndex = FTimerStatsEx::HistogramLen - 1;
 	}
 	StatsEx.InclHistogram[InclIndex]++;
 
-	int ExclIndex = static_cast<int>((ExclTime - Stats.MinExclusiveTime + StatsEx.ExclDT / 2) / StatsEx.ExclDT);
+	int32 ExclIndex = static_cast<int32>((ExclTime - Stats.MinExclusiveTime + StatsEx.ExclDT / 2) / StatsEx.ExclDT);
 	ensure(ExclIndex >= 0);
 	if (ExclIndex < 0)
 	{
 		ExclIndex = 0;
 	}
-	ensure(ExclIndex < HistogramLen);
-	if (ExclIndex >= HistogramLen)
+	ensure(ExclIndex < FTimerStatsEx::HistogramLen);
+	if (ExclIndex >= FTimerStatsEx::HistogramLen)
 	{
-		ExclIndex = HistogramLen - 1;
+		ExclIndex = FTimerStatsEx::HistogramLen - 1;
 	}
 	StatsEx.ExclHistogram[ExclIndex]++;
 }
@@ -1726,7 +1717,7 @@ void STimersView::UpdateHistogramForTimerStats(FTimerStatsEx& StatsEx, double In
 
 void STimersView::PostProcessTimerStats(FTimerStatsEx& StatsEx, bool bComputeMedian)
 {
-	FTimerStats& Stats = StatsEx.BaseStats;
+	FTimerAggregatedStats& Stats = StatsEx.BaseStats;
 
 	// Compute average inclusive/exclusive times.
 	ensure(Stats.InstanceCount > 0);
@@ -1736,11 +1727,11 @@ void STimersView::PostProcessTimerStats(FTimerStatsEx& StatsEx, bool bComputeMed
 
 	if (bComputeMedian)
 	{
-		const int HalfCount = Stats.InstanceCount / 2;
+		const int32 HalfCount = Stats.InstanceCount / 2;
 
 		// Compute median inclusive time.
-		int InclCount = 0;
-		for (int HistogramIndex = 0; HistogramIndex < HistogramLen; HistogramIndex++)
+		int32 InclCount = 0;
+		for (int32 HistogramIndex = 0; HistogramIndex < FTimerStatsEx::HistogramLen; HistogramIndex++)
 		{
 			InclCount += StatsEx.InclHistogram[HistogramIndex];
 			if (InclCount > HalfCount)
@@ -1758,8 +1749,8 @@ void STimersView::PostProcessTimerStats(FTimerStatsEx& StatsEx, bool bComputeMed
 		}
 
 		// Compute median exclusive time.
-		int ExclCount = 0;
-		for (int HistogramIndex = 0; HistogramIndex < HistogramLen; HistogramIndex++)
+		int32 ExclCount = 0;
+		for (int32 HistogramIndex = 0; HistogramIndex < FTimerStatsEx::HistogramLen; HistogramIndex++)
 		{
 			ExclCount += StatsEx.InclHistogram[HistogramIndex];
 			if (ExclCount > HalfCount)
diff --git a/Engine/Source/Developer/TraceInsights/Private/Insights/Widgets/STimersView.h b/Engine/Source/Developer/TraceInsights/Private/Insights/Widgets/STimersView.h
--- a/Engine/Source/Developer/TraceInsights/Private/Insights/Widgets/STimersView.h
+++ b/Engine/Source/Developer/TraceInsights/Private/Insights/Widgets/STimersView.h
@@ -20,6 +20,8 @@
 #include "Insights/ViewModels/TimerNode.h"
 #include "Insights/ViewModels/TimersViewColumn.h"
 
+////////////////////////////////////////////////////////////////////////////////////////////////////
+
 class FMenuBuilder;
 
 namespace Trace
@@ -28,37 +30,44 @@ namespace Trace
 	struct FTimelineEvent;
 }
 
-static constexpr int HistogramLen = 100; // number of buckets per histogram
+////////////////////////////////////////////////////////////////////////////////////////////////////
 
 struct FTimerStatsEx
 {
-	FTimerStats BaseStats;
+	static constexpr int32 HistogramLen = 100; // number of buckets per histogram
+
+	FTimerAggregatedStats BaseStats;
 
 	// Histogram for computing median inclusive time.
-	int InclHistogram[HistogramLen];
+	int32 InclHistogram[HistogramLen];
 	double InclDT; // bucket size
 
 	// Histogram for computing median exclusive time.
-	int ExclHistogram[HistogramLen];
+	int32 ExclHistogram[HistogramLen];
 	double ExclDT; // bucket size
 
 	FTimerStatsEx()
 	{
-		FMemory::Memzero(InclHistogram, sizeof(int) * HistogramLen);
+		FMemory::Memzero(InclHistogram, sizeof(int32) * HistogramLen);
 		InclDT = 1.0;
 
-		FMemory::Memzero(ExclHistogram, sizeof(int) * HistogramLen);
+		FMemory::Memzero(ExclHistogram, sizeof(int32) * HistogramLen);
 		ExclDT = 1.0;
 	}
 };
 
+////////////////////////////////////////////////////////////////////////////////////////////////////
+
 /** The filter collection - used for updating the list of timer nodes. */
 typedef TFilterCollection<const FTimerNodePtr&> FTimerNodeFilterCollection;
 
 /** The text based filter - used for updating the list of timer nodes. */
 typedef TTextFilter<const FTimerNodePtr&> FTimerNodeTextFilter;
 
-/** A custom widget used to display the list of timers. */
+////////////////////////////////////////////////////////////////////////////////////////////////////
+/**
+ * A custom widget used to display the list of timers.
+ */
 class STimersView : public SCompoundWidget
 {
 public:
@@ -73,11 +82,15 @@ public:
 
 	/**
 	 * Construct this widget
-	 * @param	InArgs	- The declaration data for this widget
+	 * @param InArgs - The declaration data for this widget
 	 */
 	void Construct(const FArguments& InArgs);
 
-	void RebuildTree(bool bRebuildEvenIfNotChanged = true);
+	/**
+	 * Rebuilds the tree (if necessary).
+	 * @param bResync - If true, it forces a resync with list of timers from Analysis, even if the list did not changed since last sync.
+	 */
+	void RebuildTree(bool bResync = true);
 	void UpdateStats(double StartTime, double EndTime);
 
 	void SelectTimerNode(uint64 Id);
@@ -95,14 +108,14 @@ protected:
 	static void UpdateHistogramForTimerStats(FTimerStatsEx& StatsEx, double InclTime, double ExclTime);
 	static void PostProcessTimerStats(FTimerStatsEx& StatsEx, bool bComputeMedian);
 
-	/** Called when the filter and presets widget should be updated with the latest data. */
-	void ProfilerManager_OnRequestTimersUpdate();
+	/** Called when the analysis session has changed. */
+	void InsightsManager_OnSessionChanged();
 
 	/**
 	 * Populates OutSearchStrings with the strings that should be used in searching.
 	 *
-	 * @param GroupOrStatNodePtr	- the group and stat node to get a text description from.
-	 * @param OutSearchStrings		- an array of strings to use in searching.
+	 * @param GroupOrStatNodePtr - the group and stat node to get a text description from.
+	 * @param OutSearchStrings   - an array of strings to use in searching.
 	 *
 	 */
 	void HandleItemToStringArray(const FTimerNodePtr& GroupOrStatNodePtr, TArray<FString>& OutSearchStrings) const;
@@ -129,8 +142,8 @@ protected:
 
 	/**
 	 * Called by STreeView to retrieves the children for the specified parent item.
-	 * @param InParent		- The parent node to retrieve the children from.
-	 * @param OutChildren	- List of children for the parent node.
+	 * @param InParent    - The parent node to retrieve the children from.
+	 * @param OutChildren - List of children for the parent node.
 	 */
 	void TreeView_OnGetChildren(FTimerNodePtr InParent, TArray<FTimerNodePtr>& OutChildren);
 
@@ -236,77 +249,77 @@ protected:
 	void ContextMenu_ResetColumns_Execute();
 
 protected:
-	/// A weak pointer to the profiler session used to populate this widget.
+	/** A weak pointer to the profiler session used to populate this widget. */
 	TSharedPtr<const Trace::IAnalysisSession>/*Weak*/ Session;
 
 	//////////////////////////////////////////////////
 	// Tree View, Columns
 
-	/// The tree widget which holds the list of timer groups and timers corresponding with each group.
+	/** The tree widget which holds the list of groups and timers corresponding with each group. */
 	TSharedPtr<STreeView<FTimerNodePtr>> TreeView;
 
-	/// Column metadata used to initialize column arguments, stored as PropertyName -> FEventGraphColumn.
+	/** Column metadata used to initialize column arguments, stored as PropertyName -> FEventGraphColumn. */
 	TMap<FName, FTimersViewColumn> TreeViewHeaderColumns;
 
-	/// Column arguments used to initialize a new header column in the tree view, stored as column name to column arguments mapping.
+	/** Column arguments used to initialize a new header column in the tree view, stored as column name to column arguments mapping. */
 	TMap<FName, SHeaderRow::FColumn::FArguments> TreeViewHeaderColumnArgs;
 
-	/// Holds the tree view header row widget which display all columns in the tree view.
+	/** Holds the tree view header row widget which display all columns in the tree view. */
 	TSharedPtr<SHeaderRow> TreeViewHeaderRow;
 
-	/// External scrollbar used to synchronize tree view position.
+	/** External scrollbar used to synchronize tree view position. */
 	TSharedPtr<SScrollBar> ExternalScrollbar;
 
 	//////////////////////////////////////////////////
 	// Hovered Column, Hovered Timer Node
 
-	/// Name of the column currently being hovered by the mouse.
+	/** Name of the column currently being hovered by the mouse. */
 	FName HoveredColumnId;
 
-	/// A shared pointer to the timer currently being hovered by the mouse.
+	/** A shared pointer to the timer node currently being hovered by the mouse. */
 	FTimerNodePtr HoveredTimerNodePtr;
 
-	/// Name of the timer that should be drawn as highlighted.
+	/** Name of the timer that should be drawn as highlighted. */
 	FName HighlightedTimerName;
 
 	//////////////////////////////////////////////////
 	// Timer Nodes
 
-	/// An array of group and timer nodes generated from the metadata.
+	/** An array of group and timer nodes generated from the metadata. */
 	TArray<FTimerNodePtr> GroupNodes;
 
-	/// A filtered array of group and timer nodes to be displayed in the tree widget.
+	/** A filtered array of group and timer nodes to be displayed in the tree widget. */
 	TArray<FTimerNodePtr> FilteredGroupNodes;
 
-	/// All timer nodes.
+	/** All timer nodes. */
 	TSet<FTimerNodePtr> TimerNodes;
 
-	/// All timer nodes, stored as TimerName -> FTimerNodePtr.
+	/** All timer nodes, stored as TimerName -> FTimerNodePtr. */
 	//TMap<FName, FTimerNodePtr> TimerNodesMap;
 
-	/// All timer nodes, stored as TimerId -> FTimerNodePtr.
+	/** All timer nodes, stored as TimerId -> FTimerNodePtr. */
 	TMap<uint64, FTimerNodePtr> TimerNodesIdMap;
 
-	/// Currently expanded group nodes.
+	/** Currently expanded group nodes. */
 	TSet<FTimerNodePtr> ExpandedNodes;
 
-	/// If true, the expanded nodes have been saved before applying a text filter.
+	/** If true, the expanded nodes have been saved before applying a text filter. */
 	bool bExpansionSaved;
 
 	//////////////////////////////////////////////////
 	// Search box and filters
 
-	/// The search box widget used to filter items displayed in the stats and groups tree.
+	/** The search box widget used to filter items displayed in the stats and groups tree. */
 	TSharedPtr<SSearchBox> SearchBox;
 
-	/// The text based filter.
+	/** The text based filter. */
 	TSharedPtr<FTimerNodeTextFilter> TextFilter;
 
-	/// The filter collection.
+	/** The filter collection. */
 	TSharedPtr<FTimerNodeFilterCollection> Filters;
 
-	/// Holds the visibility of each timer type.
-	bool bTimerNodeIsVisible[ETimerNodeType::InvalidOrMax];
+	/** Holds the visibility of each timer type. */
+	bool bTimerNodeIsVisible[static_cast<int>(ETimerNodeType::InvalidOrMax)];
 
 	//////////////////////////////////////////////////
 	// Grouping
@@ -315,20 +328,25 @@ protected:
 
 	TSharedPtr<SComboBox<TSharedPtr<ETimerGroupingMode>>> GroupByComboBox;
 
-	/// How we group the timers?
+	/** How we group the timers? */
 	ETimerGroupingMode GroupingMode;
 
 	//////////////////////////////////////////////////
 	// Sorting
 
-	/// How we sort the timers?
+	/** How we sort the timers? */
 	EColumnSortMode::Type ColumnSortMode;
 
-	/// Name of the column currently being sorted, NAME_None if sorting is disabled.
+	/** Name of the column currently being sorted, NAME_None if sorting is disabled. */
 	FName ColumnBeingSorted;
 
+	static const EColumnSortMode::Type DefaultColumnSortMode;
+	static const FName DefaultColumnBeingSorted;
+
 	//////////////////////////////////////////////////
 
 	double StatsStartTime;
 	double StatsEndTime;
 };
+
+////////////////////////////////////////////////////////////////////////////////////////////////////
diff --git a/Engine/Source/Developer/TraceInsights/Private/Insights/Widgets/STimersViewTooltip.cpp b/Engine/Source/Developer/TraceInsights/Private/Insights/Widgets/STimersViewTooltip.cpp
--- a/Engine/Source/Developer/TraceInsights/Private/Insights/Widgets/STimersViewTooltip.cpp
+++ b/Engine/Source/Developer/TraceInsights/Private/Insights/Widgets/STimersViewTooltip.cpp
@@ -67,9 +67,9 @@ TSharedPtr<SToolTip> STimersViewTooltip::GetTableCellTooltip(const TSharedPtr<FT
 	const float Alpha = 0.0f;//TimerNodePtr->_FramePct * 0.01f;
 	const FLinearColor ColorAndOpacity = FMath::Lerp(DefaultColor, ThreadColor,Alpha);
 
-	const FTimerStats& Stats = TimerNodePtr->GetStats();
+	const FTimerAggregatedStats& Stats = TimerNodePtr->GetAggregatedStats();
 
-	const int NumDigits = 5;
+	const int32 NumDigits = 5;
 
 	const FText TotalInclusiveTimeText = FText::FromString(TimeUtils::FormatTimeAuto(Stats.TotalInclusiveTime));
 	const FText MinInclusiveTimeText = FText::FromString(TimeUtils::FormatTimeMs(Stats.MinInclusiveTime, NumDigits, true));
@@ -261,10 +261,10 @@ TSharedPtr<SToolTip> STimersViewTooltip::GetTableCellTooltip(const TSharedPtr<FT
 
 	int32 Row = 1;
 	AddStatsRow(GridPanel, Row, LOCTEXT("TT_TotalTime",   "Total Time:"),   TotalInclusiveTimeText, TotalExclusiveTimeText);
-	AddStatsRow(GridPanel, Row, LOCTEXT("TT_MinTime",     "Min Time:"),     MinInclusiveTimeText, MinExclusiveTimeText);
-	AddStatsRow(GridPanel, Row, LOCTEXT("TT_MaxTime",     "Max Time:"),     MaxInclusiveTimeText, MaxExclusiveTimeText);
-	AddStatsRow(GridPanel, Row, LOCTEXT("TT_AverageTime", "Average Time:"), AvgInclusiveTimeText, AvgExclusiveTimeText);
-	AddStatsRow(GridPanel, Row, LOCTEXT("TT_MedianTime",  "Median Time:"),  MedInclusiveTimeText, MedExclusiveTimeText);
+	AddStatsRow(GridPanel, Row, LOCTEXT("TT_MaxTime",     "Max Time:"),     MaxInclusiveTimeText,   MaxExclusiveTimeText);
+	AddStatsRow(GridPanel, Row, LOCTEXT("TT_AverageTime", "Average Time:"), AvgInclusiveTimeText,   AvgExclusiveTimeText);
+	AddStatsRow(GridPanel, Row, LOCTEXT("TT_MedianTime",  "Median Time:"),  MedInclusiveTimeText,   MedExclusiveTimeText);
+	AddStatsRow(GridPanel, Row, LOCTEXT("TT_MinTime",     "Min Time:"),     MinInclusiveTimeText,   MinExclusiveTimeText);
 
 	/*
 	//TODO: We need Stats hierarchy (not the grouping hierarchy)!
@@ -322,7 +322,7 @@ TSharedPtr<SToolTip> STimersViewTooltip::GetTableCellTooltip(const TSharedPtr<FT
 		MinimalChildren.Sort(FCompareByFloatDescending());
 
 		FString ChildrenNames;
-		const int NumChildrenToDisplay = FMath::Min(MinimalChildren.Num(), 3);
+		const int32 NumChildrenToDisplay = FMath::Min(MinimalChildren.Num(), 3);
 		for(int32 SortedChildIndex = 0; SortedChildIndex < NumChildrenToDisplay; SortedChildIndex++)
 		{
 			const FEventNameAndPct& MinimalChild = MinimalChildren[SortedChildIndex];
@@ -362,7 +362,7 @@ TSharedRef<SToolTip> STimersViewTooltip::GetTooltip()
 	if (Session.IsValid())
 	{
 		const TSharedRef<SGridPanel> ToolTipGrid = SNew(SGridPanel);
-		int CurrentRowPos = 0;
+		int32 CurrentRowPos = 0;
 
 		AddHeader(ToolTipGrid, CurrentRowPos);
 		AddDescription(ToolTipGrid, CurrentRowPos);
diff --git a/Engine/Source/Developer/TraceInsights/Private/Insights/Widgets/STimersViewTooltip.h b/Engine/Source/Developer/TraceInsights/Private/Insights/Widgets/STimersViewTooltip.h
--- a/Engine/Source/Developer/TraceInsights/Private/Insights/Widgets/STimersViewTooltip.h
+++ b/Engine/Source/Developer/TraceInsights/Private/Insights/Widgets/STimersViewTooltip.h
@@ -21,7 +21,7 @@ class FTimersViewColumn;
 
 #define LOCTEXT_NAMESPACE "STimersView"
 
-// Timers View Tooltip
+/** Timers View Tooltip */
 class STimersViewTooltip
 {
 	const uint32 TimerId;
diff --git a/Engine/Source/Developer/TraceInsights/Private/Insights/Widgets/STimingProfilerToolbar.cpp b/Engine/Source/Developer/TraceInsights/Private/Insights/Widgets/STimingProfilerToolbar.cpp
--- a/Engine/Source/Developer/TraceInsights/Private/Insights/Widgets/STimingProfilerToolbar.cpp
+++ b/Engine/Source/Developer/TraceInsights/Private/Insights/Widgets/STimingProfilerToolbar.cpp
@@ -47,23 +47,24 @@ void STimingProfilerToolbar::Construct(const FArguments& InArgs)
 			{
 				ToolbarBuilder.AddToolBarButton(FInsightsCommands::Get().InsightsManager_Live);
 				ToolbarBuilder.AddToolBarButton(FInsightsCommands::Get().InsightsManager_Load);
-				ToolbarBuilder.AddToolBarButton(FInsightsCommands::Get().InsightsManager_Mock);
+				//ToolbarBuilder.AddToolBarButton(FInsightsCommands::Get().InsightsManager_Mock);
 			}
 			ToolbarBuilder.EndSection();
 			ToolbarBuilder.BeginSection("View");
 			{
 				ToolbarBuilder.AddToolBarButton(FTimingProfilerCommands::Get().ToggleFramesTrackVisibility);
 				ToolbarBuilder.AddToolBarButton(FTimingProfilerCommands::Get().ToggleGraphTrackVisibility);
-				ToolbarBuilder.AddToolBarButton(FTimingProfilerCommands::Get().ToggleTimingTrackVisibility);
+				ToolbarBuilder.AddToolBarButton(FTimingProfilerCommands::Get().ToggleTimingViewVisibility);
 				ToolbarBuilder.AddToolBarButton(FTimingProfilerCommands::Get().ToggleTimersViewVisibility);
+				ToolbarBuilder.AddToolBarButton(FTimingProfilerCommands::Get().ToggleStatsCountersViewVisibility);
 				ToolbarBuilder.AddToolBarButton(FTimingProfilerCommands::Get().ToggleLogViewVisibility);
 			}
 			ToolbarBuilder.EndSection();
-			ToolbarBuilder.BeginSection("Options");
-			{
-				ToolbarBuilder.AddToolBarButton(FInsightsCommands::Get().OpenSettings);
-			}
-			ToolbarBuilder.EndSection();
+			//ToolbarBuilder.BeginSection("Options");
+			//{
+			//	ToolbarBuilder.AddToolBarButton(FInsightsCommands::Get().OpenSettings);
+			//}
+			//ToolbarBuilder.EndSection();
 		}
 
 		static void FillToolbar2(FToolBarBuilder& ToolbarBuilder)
diff --git a/Engine/Source/Developer/TraceInsights/Private/Insights/Widgets/STimingProfilerWindow.cpp b/Engine/Source/Developer/TraceInsights/Private/Insights/Widgets/STimingProfilerWindow.cpp
--- a/Engine/Source/Developer/TraceInsights/Private/Insights/Widgets/STimingProfilerWindow.cpp
+++ b/Engine/Source/Developer/TraceInsights/Private/Insights/Widgets/STimingProfilerWindow.cpp
@@ -3,7 +3,9 @@
 #include "STimingProfilerWindow.h"
 
 #include "EditorStyleSet.h"
+#include "Framework/Docking/WorkspaceItem.h"
 #include "SlateOptMacros.h"
+#include "Widgets/Docking/SDockTab.h"
 #include "Widgets/Images/SImage.h"
 #include "Widgets/Layout/SBorder.h"
 #include "Widgets/Layout/SBox.h"
@@ -11,6 +13,8 @@
 #include "Widgets/Notifications/SNotificationList.h"
 #include "Widgets/SBoxPanel.h"
 #include "Widgets/Text/STextBlock.h"
+//#include "WorkspaceMenuStructure.h"
+//#include "WorkspaceMenuStructureModule.h"
 
 #if WITH_EDITOR
 	#include "EngineAnalytics.h"
@@ -19,13 +23,15 @@
 #endif // WITH_EDITOR
 
 // Insights
-#include "Insights/Version.h"
 #include "Insights/InsightsManager.h"
+#include "Insights/InsightsStyle.h"
 #include "Insights/TimingProfilerManager.h"
+#include "Insights/Version.h"
 #include "Insights/Widgets/SFrameTrack.h"
 #include "Insights/Widgets/SGraphTrack.h"
 #include "Insights/Widgets/SInsightsSettings.h"
 #include "Insights/Widgets/SLogView.h"
+#include "Insights/Widgets/SStatsView.h"
 #include "Insights/Widgets/STimersView.h"
 #include "Insights/Widgets/STimingProfilerToolbar.h"
 #include "Insights/Widgets/STimingView.h"
@@ -36,6 +42,16 @@
 
 ////////////////////////////////////////////////////////////////////////////////////////////////////
 
+const FName FTimingProfilerTabs::ToolbarID(TEXT("Toolbar"));
+const FName FTimingProfilerTabs::FramesTrackID(TEXT("Frames"));
+const FName FTimingProfilerTabs::GraphTrackID(TEXT("Graph"));
+const FName FTimingProfilerTabs::TimingViewID(TEXT("TimingView"));
+const FName FTimingProfilerTabs::TimersID(TEXT("Timers"));
+const FName FTimingProfilerTabs::StatsCountersID(TEXT("StasCounters"));
+const FName FTimingProfilerTabs::LogViewID(TEXT("LogView"));
+
+////////////////////////////////////////////////////////////////////////////////////////////////////
+
 // TODO: move this function to InsightsManager or in TraceSession.h
 static FText GetTextForNotification(const EInsightsNotificationType NotificatonType, const ELoadingProgressState ProgressState, const FString& Filename, const float ProgressPercent = 0.0f)
 {
@@ -84,13 +100,6 @@ STimingProfilerWindow::STimingProfilerWindow()
 
 STimingProfilerWindow::~STimingProfilerWindow()
 {
-	// Remove ourselves from the profiler manager.
-	//if (FTimingProfilerManager::Get().IsValid())
-	//{
-	//	//...
-	//	FTimingProfilerManager::Get()->OnViewModeChanged().RemoveAll(this);
-	//}
-
 #if WITH_EDITOR
 	if (DurationActive > 0.0f && FEngineAnalytics::IsAvailable())
 	{
@@ -103,243 +112,216 @@ STimingProfilerWindow::~STimingProfilerWindow()
 
 BEGIN_SLATE_FUNCTION_BUILD_OPTIMIZATION
 
-TSharedRef<SWidget> STimingProfilerWindow::ConstructMultiTrackView()
+TSharedRef<SDockTab> STimingProfilerWindow::SpawnTab_Toolbar(const FSpawnTabArgs& Args)
 {
-	return SNew(SSplitter)
-		.Orientation(Orient_Vertical)
-
-	+ SSplitter::Slot()
-		.Value(0.1f)
-		[
-			ConstructFramesTrack()
-		]
-
-	+ SSplitter::Slot()
-		.Value(0.2f)
-		[
-			ConstructGraphTrack()
-		]
-
-	+ SSplitter::Slot()
-		.Value(0.7f)
+	return SNew(SDockTab)
+		.ShouldAutosize(true)
+		.TabRole(ETabRole::PanelTab)
 		[
-			ConstructTimingTrack()
-		]
-
-	; // end of Widget
+			SNew(STimingProfilerToolbar)
+		];
 }
 
 ////////////////////////////////////////////////////////////////////////////////////////////////////
 
-TSharedRef<SWidget> STimingProfilerWindow::ConstructFramesTrack()
+TSharedRef<SDockTab> STimingProfilerWindow::SpawnTab_FramesTrack(const FSpawnTabArgs& Args)
 {
-	return SNew(SVerticalBox)
-		.Visibility(this, &STimingProfilerWindow::IsFramesTrackVisible)
-		.IsEnabled(this, &STimingProfilerWindow::IsProfilerEnabled)
-
-	// Frame Track
-	+ SVerticalBox::Slot()
-		.FillHeight(1.0)
-		.Padding(0.0f, 0.0f, 0.0f, 0.0f)
+	return SNew(SDockTab)
+		.ShouldAutosize(false)
+		.TabRole(ETabRole::PanelTab)
 		[
-			SNew(SBox)
-				.HeightOverride(48.0f)
-				.IsEnabled(this, &STimingProfilerWindow::IsProfilerEnabled)
-				[
-					SAssignNew(FrameTrack, SFrameTrack)
-
-					//SNew(SHorizontalBox)
-					//
-					//+ SHorizontalBox::Slot()
-					//	.FillWidth(1.0f)
-					//	.Padding(0.0f)
-					//	.HAlign(HAlign_Fill)
-					//	.VAlign(VAlign_Fill)
-					//	[
-					//
-					//	]
-				]
-		]
-
-	; // end of Widget
+			SAssignNew(FrameTrack, SFrameTrack)
+			//.IsEnabled(this, &STimingProfilerWindow::IsProfilerEnabled)
+			//.Visibility(this, &STimingProfilerWindow::IsFramesTrackVisible)
+		];
 }
 
 ////////////////////////////////////////////////////////////////////////////////////////////////////
 
-TSharedRef<SWidget> STimingProfilerWindow::ConstructGraphTrack()
+TSharedRef<SDockTab> STimingProfilerWindow::SpawnTab_GraphTrack(const FSpawnTabArgs& Args)
 {
-	return SNew(SVerticalBox)
-		.Visibility(this, &STimingProfilerWindow::IsGraphTrackVisible)
-		.IsEnabled(this, &STimingProfilerWindow::IsProfilerEnabled)
-
-	/*
-	// Header
-	+ SVerticalBox::Slot()
-		.AutoHeight()
-		.Padding(0.0f, 0.0f, 0.0f, 0.0f)
-		[
-			SNew(SHorizontalBox)
-
-			+ SHorizontalBox::Slot()
-				.AutoWidth()
-				[
-					SNew(SImage)
-						.Image(FEditorStyle::GetBrush(TEXT("Profiler.Tab.GraphView")))
-				]
-
-			+ SHorizontalBox::Slot()
-				.AutoWidth()
-				[
-					SNew(STextBlock)
-						.Text(LOCTEXT("GraphTrackLabel", "Overview Timing Graph"))
-						.Margin(FMargin(0.0f, 2.0f, 0.0f, 0.0f))
-				]
-		]
-	*/
-
-	// Graph Track
-	+ SVerticalBox::Slot()
-		.FillHeight(1.0)
-		.Padding(0.0f, 0.0f, 0.0f, 0.0f)
+	return SNew(SDockTab)
+		.ShouldAutosize(false)
+		.TabRole(ETabRole::PanelTab)
 		[
-			SNew(SBox)
-				.HeightOverride(48.0f)
-				.IsEnabled(this, &STimingProfilerWindow::IsProfilerEnabled)
-				[
-					SNew(SHorizontalBox)
-
-					+ SHorizontalBox::Slot()
-						.FillWidth(1.0f)
-						.Padding(0.0f)
-						.HAlign(HAlign_Fill)
-						.VAlign(VAlign_Fill)
-						[
-							SAssignNew(GraphTrack, SGraphTrack)
-						]
-				]
-		]
-
-	; // end of Widget
+			SAssignNew(GraphTrack, SGraphTrack)
+			//.IsEnabled(this, &STimingProfilerWindow::IsProfilerEnabled)
+			//.Visibility(this, &STimingProfilerWindow::IsGraphTrackVisible)
+		];
 }
 
 ////////////////////////////////////////////////////////////////////////////////////////////////////
 
-TSharedRef<SWidget> STimingProfilerWindow::ConstructTimingTrack()
+TSharedRef<SDockTab> STimingProfilerWindow::SpawnTab_TimingView(const FSpawnTabArgs& Args)
 {
-	return SNew(SVerticalBox)
-		.Visibility(this, &STimingProfilerWindow::IsTimingTrackVisible)
-		.IsEnabled(this, &STimingProfilerWindow::IsProfilerEnabled)
-
-	/*
-	// Header
-	+ SVerticalBox::Slot()
-		.AutoHeight()
-		.Padding(0.0f, 0.0f, 0.0f, 0.0f)
+	return SNew(SDockTab)
+		.ShouldAutosize(false)
+		.TabRole(ETabRole::PanelTab)
 		[
-			SNew(SHorizontalBox)
-
-			+ SHorizontalBox::Slot()
-				.AutoWidth()
-				[
-					SNew(SImage)
-						.Image(FEditorStyle::GetBrush(TEXT("Profiler.Tab.GraphView")))
-				]
-
-			+ SHorizontalBox::Slot()
-				.AutoWidth()
-				[
-					SNew(STextBlock)
-						.Text(LOCTEXT("TimingLabel", "Main Timing View"))
-						.Margin(FMargin(0.0f, 2.0f, 0.0f, 0.0f))
-				]
-		]
-	*/
-
-	// Timing Track
-	+ SVerticalBox::Slot()
-		.FillHeight(1.0)
-		.Padding(0.0f, 0.0f, 0.0f, 0.0f)
-		[
-			SNew(SBox)
-				.HeightOverride(48.0f)
-				.IsEnabled(this, &STimingProfilerWindow::IsProfilerEnabled)
-				[
-					SNew(SHorizontalBox)
-
-					+ SHorizontalBox::Slot()
-						.FillWidth(1.0f)
-						.Padding(0.0f)
-						.HAlign(HAlign_Fill)
-						.VAlign(VAlign_Fill)
-						[
-							SAssignNew(TimingView, STimingView)
-						]
-				]
-		]
-
-	//// Threads Panel
-	//+ SVerticalBox::Slot()
-	//	.FillHeight(1.0)
-	//	.Padding(0.0f, 0.0f, 0.0f, 0.0f)
-	//	[
-	//		SAssignNew(ThreadsPanel, SVerticalBox)
-	//	]
-
-	; // end of Widget
+			SAssignNew(TimingView, STimingView)
+			//.IsEnabled(this, &STimingProfilerWindow::IsProfilerEnabled)
+			//.Visibility(this, &STimingProfilerWindow::IsTimersViewVisible)
+		];
 }
 
 ////////////////////////////////////////////////////////////////////////////////////////////////////
 
-TSharedRef<SWidget> STimingProfilerWindow::ConstructTimersView()
+TSharedRef<SDockTab> STimingProfilerWindow::SpawnTab_Timers(const FSpawnTabArgs& Args)
 {
-	return SNew(SVerticalBox)
-		.Visibility(this, &STimingProfilerWindow::IsTimersViewVisible)
-		.IsEnabled(this, &STimingProfilerWindow::IsProfilerEnabled)
-
-	// Header
-	+ SVerticalBox::Slot()
-		.AutoHeight()
+	return SNew(SDockTab)
+		.ShouldAutosize(false)
+		.TabRole(ETabRole::PanelTab)
 		[
-			SNew(SHorizontalBox)
-
-			+ SHorizontalBox::Slot()
-				.AutoWidth()
-				[
-					SNew(SImage)
-						.Image(FEditorStyle::GetBrush(TEXT("Profiler.Tab.FiltersAndPresets")))
-				]
+		
+			SAssignNew(TimersView, STimersView)
+			//.IsEnabled(this, &STimingProfilerWindow::IsProfilerEnabled)
+			//.Visibility(this, &STimingProfilerWindow::IsTimersViewVisible)
+		];
+}
 
-			+ SHorizontalBox::Slot()
-				.AutoWidth()
-				[
-					SNew(STextBlock)
-						.Text(LOCTEXT("TimersLabel", "Timers"))
-				]
-		]
+////////////////////////////////////////////////////////////////////////////////////////////////////
 
-	// Timers
-	+ SVerticalBox::Slot()
-		.FillHeight(1.0f)
-		.Padding(0.0f, 2.0f, 0.0f, 0.0f)
+TSharedRef<SDockTab> STimingProfilerWindow::SpawnTab_StatsCounters(const FSpawnTabArgs& Args)
+{
+	return SNew(SDockTab)
+		.ShouldAutosize(false)
+		.TabRole(ETabRole::PanelTab)
 		[
-			SAssignNew(TimersView, STimersView)
-		]
-
-	; // end of Widget
+			SAssignNew(StatsView, SStatsView)
+			//.IsEnabled(this, &STimingProfilerWindow::IsProfilerEnabled)
+			//.Visibility(this, &STimingProfilerWindow::IsStatsCountersViewVisible)
+			
+		];
 }
 
 ////////////////////////////////////////////////////////////////////////////////////////////////////
 
-TSharedRef<SWidget> STimingProfilerWindow::ConstructLogView()
+TSharedRef<SDockTab> STimingProfilerWindow::SpawnTab_LogView(const FSpawnTabArgs& Args)
 {
-	return SAssignNew(LogView, SLogView)
-		.Visibility(this, &STimingProfilerWindow::IsLogViewVisible);
+	return SNew(SDockTab)
+		.ShouldAutosize(false)
+		.TabRole(ETabRole::PanelTab)
+		[
+			SAssignNew(LogView, SLogView)
+			//.IsEnabled(this, &STimingProfilerWindow::IsProfilerEnabled)
+			//.Visibility(this, &STimingProfilerWindow::IsLogViewVisible)
+		];
 }
 
 ////////////////////////////////////////////////////////////////////////////////////////////////////
 
-void STimingProfilerWindow::Construct(const FArguments& InArgs)
+void STimingProfilerWindow::Construct(const FArguments& InArgs, const TSharedRef<SDockTab>& ConstructUnderMajorTab, const TSharedPtr<SWindow>& ConstructUnderWindow)
 {
+	// Create & initialize tab manager.
+	TabManager = FGlobalTabmanager::Get()->NewTabManager(ConstructUnderMajorTab);
+	TSharedRef<FWorkspaceItem> AppMenuGroup = TabManager->AddLocalWorkspaceMenuCategory(LOCTEXT("TimingProfilerMenuGroupName", "Timing Insights"));
+
+	TabManager->RegisterTabSpawner(FTimingProfilerTabs::ToolbarID, FOnSpawnTab::CreateRaw(this, &STimingProfilerWindow::SpawnTab_Toolbar))
+		.SetDisplayName(LOCTEXT("DeviceToolbarTabTitle", "Toolbar"))
+		.SetIcon(FSlateIcon(FInsightsStyle::GetStyleSetName(), "Toolbar.Icon.Small"))
+		.SetGroup(AppMenuGroup);
+
+	TabManager->RegisterTabSpawner(FTimingProfilerTabs::FramesTrackID, FOnSpawnTab::CreateRaw(this, &STimingProfilerWindow::SpawnTab_FramesTrack))
+		.SetDisplayName(LOCTEXT("FramesTrackTabTitle", "Frames"))
+		.SetIcon(FSlateIcon(FInsightsStyle::GetStyleSetName(), "FramesTrack.Icon.Small"))
+		.SetGroup(AppMenuGroup);
+
+	TabManager->RegisterTabSpawner(FTimingProfilerTabs::GraphTrackID, FOnSpawnTab::CreateRaw(this, &STimingProfilerWindow::SpawnTab_GraphTrack))
+		.SetDisplayName(LOCTEXT("GraphTrackTabTitle", "Graph"))
+		.SetIcon(FSlateIcon(FInsightsStyle::GetStyleSetName(), "GraphTrack.Icon.Small"))
+		.SetGroup(AppMenuGroup);
+
+	TabManager->RegisterTabSpawner(FTimingProfilerTabs::TimingViewID, FOnSpawnTab::CreateRaw(this, &STimingProfilerWindow::SpawnTab_TimingView))
+		.SetDisplayName(LOCTEXT("TimingViewTabTitle", "Timing View"))
+		.SetIcon(FSlateIcon(FInsightsStyle::GetStyleSetName(), "TimingView.Icon.Small"))
+		.SetGroup(AppMenuGroup);
+
+	TabManager->RegisterTabSpawner(FTimingProfilerTabs::TimersID, FOnSpawnTab::CreateRaw(this, &STimingProfilerWindow::SpawnTab_Timers))
+		.SetDisplayName(LOCTEXT("TimersTabTitle", "Timers"))
+		.SetIcon(FSlateIcon(FInsightsStyle::GetStyleSetName(), "TimersView.Icon.Small"))
+		.SetGroup(AppMenuGroup);
+
+	TabManager->RegisterTabSpawner(FTimingProfilerTabs::StatsCountersID, FOnSpawnTab::CreateRaw(this, &STimingProfilerWindow::SpawnTab_StatsCounters))
+		.SetDisplayName(LOCTEXT("StatsCountersTabTitle", "Stats Counters"))
+		.SetIcon(FSlateIcon(FInsightsStyle::GetStyleSetName(), "StatsCountersView.Icon.Small"))
+		.SetGroup(AppMenuGroup);
+
+	TabManager->RegisterTabSpawner(FTimingProfilerTabs::LogViewID, FOnSpawnTab::CreateRaw(this, &STimingProfilerWindow::SpawnTab_LogView))
+		.SetDisplayName(LOCTEXT("LogViewTabTitle", "Log View"))
+		.SetIcon(FSlateIcon(FInsightsStyle::GetStyleSetName(), "LogView.Icon.Small"))
+		.SetGroup(AppMenuGroup);
+
+	// Create tab layout.
+	const TSharedRef<FTabManager::FLayout> Layout = FTabManager::NewLayout("InsightsTimingProfilerLayout_v1.0")
+		->AddArea
+		(
+			FTabManager::NewPrimaryArea()
+			->SetOrientation(Orient_Vertical)
+			->Split
+			(
+				FTabManager::NewStack()
+				->AddTab(FTimingProfilerTabs::ToolbarID, ETabState::OpenedTab)
+				->SetHideTabWell(true)
+			)
+			->Split
+			(
+				FTabManager::NewSplitter()
+				->SetOrientation(Orient_Horizontal)
+				->SetSizeCoefficient(1.0f)
+				->Split
+				(
+					FTabManager::NewSplitter()
+					->SetOrientation(Orient_Vertical)
+					->SetSizeCoefficient(0.65f)
+					->Split
+					(
+						FTabManager::NewStack()
+						->SetSizeCoefficient(0.1f)
+						->SetHideTabWell(true)
+						->AddTab(FTimingProfilerTabs::FramesTrackID, FTimingProfilerManager::Get()->IsFramesTrackVisible() ? ETabState::OpenedTab : ETabState::ClosedTab)
+					)
+					->Split
+					(
+						FTabManager::NewStack()
+						->SetSizeCoefficient(0.2f)
+						->SetHideTabWell(true)
+						->AddTab(FTimingProfilerTabs::GraphTrackID, FTimingProfilerManager::Get()->IsGraphTrackVisible() ? ETabState::OpenedTab : ETabState::ClosedTab)
+					)
+					->Split
+					(
+						FTabManager::NewStack()
+						->SetSizeCoefficient(0.5f)
+						->SetHideTabWell(true)
+						->AddTab(FTimingProfilerTabs::TimingViewID, FTimingProfilerManager::Get()->IsTimingViewVisible() ? ETabState::OpenedTab : ETabState::ClosedTab)
+					)
+					->Split
+					(
+						FTabManager::NewStack()
+						->SetSizeCoefficient(0.2f)
+						->SetHideTabWell(true)
+						->AddTab(FTimingProfilerTabs::LogViewID, FTimingProfilerManager::Get()->IsLogViewVisible() ? ETabState::OpenedTab : ETabState::ClosedTab)
+					)
+				)
+				->Split
+				(
+					FTabManager::NewStack()
+					->SetSizeCoefficient(0.35f)
+					->AddTab(FTimingProfilerTabs::TimersID, FTimingProfilerManager::Get()->IsTimersViewVisible() ? ETabState::OpenedTab : ETabState::ClosedTab)
+					->AddTab(FTimingProfilerTabs::StatsCountersID, FTimingProfilerManager::Get()->IsStatsCountersViewVisible() ? ETabState::OpenedTab : ETabState::ClosedTab)
+				)
+			)
+		);
+
+	// Create & initialize main menu.
+	FMenuBarBuilder MenuBarBuilder = FMenuBarBuilder(TSharedPtr<FUICommandList>());
+	
+	MenuBarBuilder.AddPullDownMenu(
+		LOCTEXT("MenuLabel", "MENU"),
+		FText::GetEmpty(),
+		FNewMenuDelegate::CreateStatic(&STimingProfilerWindow::FillMenu, TabManager),
+		FName(TEXT("MENU"))
+	);
+
 	ChildSlot
 		[
 			SNew(SOverlay)
@@ -363,53 +345,20 @@ void STimingProfilerWindow::Construct(const FArguments& InArgs)
 				[
 					SAssignNew(MainContentPanel, SVerticalBox)
 
-					// Toolbar
 					+ SVerticalBox::Slot()
 						.AutoHeight()
 						[
-							SNew(STimingProfilerToolbar)
+							MenuBarBuilder.MakeWidget()
 						]
 
 					+ SVerticalBox::Slot()
-						.HAlign(HAlign_Fill)
-						.VAlign(VAlign_Fill)
+						.FillHeight(1.0f)
 						[
-							SNew(SSplitter)
-								.Orientation(Orient_Horizontal)
-
-							+ SSplitter::Slot()
-								.Value(0.65f)
-								[
-									SNew(SSplitter)
-										.Orientation(Orient_Vertical)
-
-									// Multi Track View
-									+ SSplitter::Slot()
-										.Value(0.75f)
-										[
-											ConstructMultiTrackView()
-										]
-
-									// Log View
-									+ SSplitter::Slot()
-										.Value(0.25f)
-										[
-											ConstructLogView()
-										]
-								]
-
-							// Timers View
-							+ SSplitter::Slot()
-								.Value(0.35f)
-								//.SizeRule(SSplitter::SizeToContent)
-								.Expose(TimersViewSlot)
-								[
-									ConstructTimersView()
-								]
+							TabManager->RestoreFrom(Layout, ConstructUnderWindow).ToSharedRef()
 						]
 				]
 
-			// session hint overlay
+			// Session hint overlay
 			+ SOverlay::Slot()
 				.HAlign(HAlign_Center)
 				.VAlign(VAlign_Center)
@@ -424,7 +373,7 @@ void STimingProfilerWindow::Construct(const FArguments& InArgs)
 						]
 				]
 
-			// notification area overlay
+			// Notification area overlay
 			+ SOverlay::Slot()
 				.HAlign(HAlign_Right)
 				.VAlign(VAlign_Bottom)
@@ -433,21 +382,37 @@ void STimingProfilerWindow::Construct(const FArguments& InArgs)
 					SAssignNew(NotificationList, SNotificationList)
 				]
 
-			// profiler settings overlay
+			// Settings dialog overlay
 			+ SOverlay::Slot()
 				.HAlign(HAlign_Center)
 				.VAlign(VAlign_Center)
 				.Expose(OverlaySettingsSlot)
 		];
 
-	//ProfilerMiniView->OnSelectionBoxChanged().AddSP(GraphPanel.ToSharedRef(), &STimingProfilerGraphPanel::MiniView_OnSelectionBoxChanged);
-	//FTimingProfilerManager::Get()->OnViewModeChanged().AddSP(this, &STimingProfilerWindow::ProfilerManager_OnViewModeChanged);
+	// Tell tab-manager about the global menu bar.
+	TabManager->SetMenuMultiBox(MenuBarBuilder.GetMultiBox());
 }
 
 END_SLATE_FUNCTION_BUILD_OPTIMIZATION
 
 ////////////////////////////////////////////////////////////////////////////////////////////////////
 
+void STimingProfilerWindow::FillMenu(FMenuBuilder& MenuBuilder, const TSharedPtr<FTabManager> TabManager)
+{
+	if (!TabManager.IsValid())
+	{
+		return;
+	}
+
+#if !WITH_EDITOR
+	//TODO: FGlobalTabmanager::Get()->PopulateTabSpawnerMenu(MenuBuilder, WorkspaceMenu::GetMenuStructure().GetStructureRoot());
+#endif //!WITH_EDITOR
+
+	TabManager->PopulateLocalTabSpawnerMenu(MenuBuilder);
+}
+
+////////////////////////////////////////////////////////////////////////////////////////////////////
+
 EVisibility STimingProfilerWindow::IsFramesTrackVisible() const
 {
 	if (FTimingProfilerManager::Get()->IsFramesTrackVisible() &&
@@ -474,9 +439,9 @@ EVisibility STimingProfilerWindow::IsGraphTrackVisible() const
 
 ////////////////////////////////////////////////////////////////////////////////////////////////////
 
-EVisibility STimingProfilerWindow::IsTimingTrackVisible() const
+EVisibility STimingProfilerWindow::IsTimingViewVisible() const
 {
-	if (FTimingProfilerManager::Get()->IsTimingTrackVisible() &&
+	if (FTimingProfilerManager::Get()->IsTimingViewVisible() &&
 		FInsightsManager::Get()->GetSession().IsValid())
 	{
 		return EVisibility::Visible;
@@ -513,6 +478,27 @@ EVisibility STimingProfilerWindow::IsLogViewVisible() const
 
 ////////////////////////////////////////////////////////////////////////////////////////////////////
 
+void STimingProfilerWindow::ShowHideTab(const FName& TabID, bool bIsVisible)
+{
+	if (!bIsVisible)
+	{
+		TSharedPtr<SDockTab> Tab = TabManager->FindExistingLiveTab(TabID);
+		if (Tab.IsValid())
+		{
+			Tab->RequestCloseTab();
+		}
+	}
+	else
+	{
+		if (TabManager->CanSpawnTab(TabID))
+		{
+			TabManager->InvokeTab(TabID);
+		}
+	}
+}
+
+////////////////////////////////////////////////////////////////////////////////////////////////////
+
 EVisibility STimingProfilerWindow::IsSessionOverlayVisible() const
 {
 	if (FInsightsManager::Get()->GetSession().IsValid())
@@ -772,30 +758,4 @@ void STimingProfilerWindow::CloseProfilerSettings()
 
 ////////////////////////////////////////////////////////////////////////////////////////////////////
 
-//void STimingProfilerWindow::ProfilerManager_OnViewModeChanged(ETimingProfilerViewMode NewViewMode)
-//{
-//	if (NewViewMode == ETimingProfilerViewMode::LineIndexBased)
-//	{
-//		EventGraphPanel->SetVisibility(EVisibility::Visible);
-//		EventGraphPanel->SetEnabled(true);
-//
-//		(*FiltersAndPresetsSlot)
-//		[
-//			FiltersAndPresets.ToSharedRef()
-//		];
-//	}
-//	else if (NewViewMode == ETimingProfilerViewMode::ThreadViewTimeBased)
-//	{
-//		EventGraphPanel->SetVisibility(EVisibility::Collapsed);
-//		EventGraphPanel->SetEnabled(false);
-//
-//		(*FiltersAndPresetsSlot)
-//		[
-//			SNullWidget::NullWidget
-//		];
-//	}
-//}
-
-////////////////////////////////////////////////////////////////////////////////////////////////////
-
 #undef LOCTEXT_NAMESPACE
diff --git a/Engine/Source/Developer/TraceInsights/Private/Insights/Widgets/STimingProfilerWindow.h b/Engine/Source/Developer/TraceInsights/Private/Insights/Widgets/STimingProfilerWindow.h
--- a/Engine/Source/Developer/TraceInsights/Private/Insights/Widgets/STimingProfilerWindow.h
+++ b/Engine/Source/Developer/TraceInsights/Private/Insights/Widgets/STimingProfilerWindow.h
@@ -3,11 +3,13 @@
 #pragma once
 
 #include "CoreMinimal.h"
+#include "Framework/Docking/TabManager.h"
 #include "Input/Reply.h"
 #include "Layout/Visibility.h"
 #include "Misc/Guid.h"
 #include "SlateFwd.h"
 #include "Widgets/DeclarativeSyntaxSupport.h"
+#include "Widgets/Docking/SDockTab.h"
 #include "Widgets/Layout/SSplitter.h"
 #include "Widgets/SCompoundWidget.h"
 #include "Widgets/SOverlay.h"
@@ -15,71 +17,109 @@
 // Insights
 #include "Insights/InsightsManager.h"
 
+////////////////////////////////////////////////////////////////////////////////////////////////////
+
 class FActiveTimerHandle;
 class SVerticalBox;
 class SFrameTrack;
 class SGraphTrack;
+class SStatsView;
 class STimersView;
 class STimingView;
 class SLogView;
 
-/// Type definition for shared pointers to instances of SNotificationItem.
+////////////////////////////////////////////////////////////////////////////////////////////////////
+
+/** Type definition for shared pointers to instances of SNotificationItem. */
 typedef TSharedPtr<class SNotificationItem> SNotificationItemPtr;
 
-/// Type definition for shared references to instances of SNotificationItem.
+/** Type definition for shared references to instances of SNotificationItem. */
 typedef TSharedRef<class SNotificationItem> SNotificationItemRef;
 
-/// Type definition for weak references to instances of SNotificationItem.
+/** Type definition for weak references to instances of SNotificationItem. */
 typedef TWeakPtr<class SNotificationItem> SNotificationItemWeak;
 
-/// Implements the timing profiler window.
+////////////////////////////////////////////////////////////////////////////////////////////////////
+
+struct FTimingProfilerTabs
+{
+	// Tab identifiers
+	static const FName ToolbarID;
+	static const FName FramesTrackID;
+	static const FName GraphTrackID;
+	static const FName TimingViewID;
+	static const FName TimersID;
+	static const FName StatsCountersID;
+	static const FName LogViewID;
+};
+
+////////////////////////////////////////////////////////////////////////////////////////////////////
+
+/** Implements the timing profiler window. */
 class STimingProfilerWindow : public SCompoundWidget
 {
 public:
-	/// Default constructor.
+	/** Default constructor. */
 	STimingProfilerWindow();
 
-	/// Virtual destructor.
+	/** Virtual destructor. */
 	virtual ~STimingProfilerWindow();
 
 	SLATE_BEGIN_ARGS(STimingProfilerWindow){}
 	SLATE_END_ARGS()
 
-	/// Constructs this widget.
-	void Construct(const FArguments& InArgs);
+	/** Constructs this widget. */
+	void Construct(const FArguments& InArgs, const TSharedRef<SDockTab>& ConstructUnderMajorTab, const TSharedPtr<SWindow>& ConstructUnderWindow);
 
 	void ManageLoadingProgressNotificationState(const FString& Filename, const EInsightsNotificationType NotificatonType, const ELoadingProgressState ProgressState, const float LoadingProgress);
 
 	void OpenProfilerSettings();
 	void CloseProfilerSettings();
 
+	void ShowHideTab(const FName& TabID, bool bIsVisible);
+
 protected:
-	TSharedRef<SWidget> ConstructMultiTrackView();
-	TSharedRef<SWidget> ConstructFramesTrack();
-	TSharedRef<SWidget> ConstructGraphTrack();
-	TSharedRef<SWidget> ConstructTimingTrack();
-	TSharedRef<SWidget> ConstructTimersView();
-	TSharedRef<SWidget> ConstructLogView();
-
-	/// Callback for determining the visibility of the Frames track.
+	TSharedRef<SDockTab> SpawnTab_Toolbar(const FSpawnTabArgs& Args);
+	TSharedRef<SDockTab> SpawnTab_FramesTrack(const FSpawnTabArgs& Args);
+	TSharedRef<SDockTab> SpawnTab_GraphTrack(const FSpawnTabArgs& Args);
+	TSharedRef<SDockTab> SpawnTab_TimingView(const FSpawnTabArgs& Args);
+	TSharedRef<SDockTab> SpawnTab_Timers(const FSpawnTabArgs& Args);
+	TSharedRef<SDockTab> SpawnTab_StatsCounters(const FSpawnTabArgs& Args);
+	TSharedRef<SDockTab> SpawnTab_LogView(const FSpawnTabArgs& Args);
+
+	//void RegisterTabSpawners(const TSharedRef<class FTabManager>& InTabManager);
+	//void UnregisterTabSpawners(const TSharedRef<class FTabManager>& InTabManager);
+
+	/**
+	 * Fill the main menu with menu items.
+	 *
+	 * @param MenuBuilder The multi-box builder that should be filled with content for this pull-down menu.
+	 * @param TabManager A Tab Manager from which to populate tab spawner menu items.
+	 */
+	static void FillMenu(FMenuBuilder& MenuBuilder, const TSharedPtr<FTabManager> TabManager);
+
+	/** Callback for determining the visibility of the Frames track. */
 	EVisibility IsFramesTrackVisible() const;
 
-	/// Callback for determining the visibility of the Graph track.
+	/** Callback for determining the visibility of the Graph track. */
 	EVisibility IsGraphTrackVisible() const;
 
-	/// Callback for determining the visibility of the Timing track.
-	EVisibility IsTimingTrackVisible() const;
+	/** Callback for determining the visibility of the Timing view. */
+	EVisibility IsTimingViewVisible() const;
 
-	/// Callback for determining the visibility of the Timers View.
+	/** Callback for determining the visibility of the Timers View. */
 	EVisibility IsTimersViewVisible() const;
 
-	/// Callback for determining the visibility of the Log View.
+	/** Callback for determining the visibility of the Stats Counters View. */
+	EVisibility IsStatsCountersVisible() const;
+
+	/** Callback for determining the visibility of the Log View. */
 	EVisibility IsLogViewVisible() const;
 
-	/// Callback for determining the visibility of the 'Select a session' overlay.
+	/** Callback for determining the visibility of the 'Select a session' overlay. */
 	EVisibility IsSessionOverlayVisible() const;
 
-	/// Callback for getting the enabled state of the profiler window.
+	/** Callback for getting the enabled state of the profiler window. */
 	bool IsProfilerEnabled() const;
 
 	void SendingServiceSideCapture_Cancel(const FString Filename);
@@ -143,45 +183,47 @@ private:
 	virtual FReply OnDragOver(const FGeometry& MyGeometry, const FDragDropEvent& DragDropEvent)  override;
 
 private:
-	/// Updates the amount of time the profiler has been active
+	/** Updates the amount of time the profiler has been active. */
 	EActiveTimerReturnType UpdateActiveDuration(double InCurrentTime, float InDeltaTime);
 
 public:
-	/// Widget for the frame track
+	/** Widget for the frame track */
 	TSharedPtr<SFrameTrack> FrameTrack;
 
-	/// Widget for the graph track
+	/** Widget for the graph track */
 	TSharedPtr<SGraphTrack> GraphTrack;
 
-	/// Holds all timing tracks (one for each thread).
-	TSharedPtr<SVerticalBox> ThreadsPanel;
-
-	/// Widget for the timing track
+	/** Widget for the timing track */
 	TSharedPtr<STimingView> TimingView;
 
-	/// Holds the Timers view widget/slot.
-	SSplitter::FSlot* TimersViewSlot;
+	/** Holds the Timers view widget. */
 	TSharedPtr<STimersView> TimersView;
 
-	/// Widget for the log view
+	/** Holds the Stats (Counters) view widget. */
+	TSharedPtr<SStatsView> StatsView;
+
+	/** Widget for the log view */
 	TSharedPtr<SLogView> LogView;
 
-	/// Widget for the non-intrusive notifications.
+	/** Widget for the non-intrusive notifications. */
 	TSharedPtr<SNotificationList> NotificationList;
 
-	/// Holds all active and visible notifications, stored as FGuid -> SNotificationItemWeak.
+	/** Holds all active and visible notifications, stored as FGuid -> SNotificationItemWeak. */
 	TMap<FString, SNotificationItemWeak> ActiveNotifications;
 
-	/// Overlay slot which contains the profiler settings widget.
+	/** Overlay slot which contains the profiler settings widget. */
 	SOverlay::FOverlaySlot* OverlaySettingsSlot;
 
-	/// The number of seconds the profiler has been active
+	/** The number of seconds the profiler has been active */
 	float DurationActive;
 
 private:
-	/// The handle to the active update duration tick
+	/** Holds the tab manager that manages the front-end's tabs. */
+	TSharedPtr<FTabManager> TabManager;
+
+	/** The handle to the active update duration tick */
 	TWeakPtr<FActiveTimerHandle> ActiveTimerHandle;
 
-	/// Holds all widgets for the profiler window like menu bar, toolbar and tabs.
+	/** Holds all widgets for the profiler window like menu bar, toolbar and tabs. */
 	TSharedPtr<SVerticalBox> MainContentPanel;
 };
