[TraceInsights]

diff --git a/Engine/Source/Developer/TraceInsights/Private/Insights/Widgets/STimingView.cpp b/Engine/Source/Developer/TraceInsights/Private/Insights/Widgets/STimingView.cpp
--- a/Engine/Source/Developer/TraceInsights/Private/Insights/Widgets/STimingView.cpp
+++ b/Engine/Source/Developer/TraceInsights/Private/Insights/Widgets/STimingView.cpp
@@ -35,6 +35,7 @@
 #include "Insights/TimingProfilerCommon.h"
 #include "Insights/TimingProfilerManager.h"
 #include "Insights/ViewModels/TimingViewDrawHelper.h"
+#include "Insights/Widgets/SStatsView.h"
 #include "Insights/Widgets/STimersView.h"
 #include "Insights/Widgets/STimingProfilerWindow.h"
 
@@ -342,11 +343,27 @@ void STimingView::Tick(const FGeometry& AllottedGeometry, const double InCurrent
 		}
 	}
 
-	// Check if TimersView needs to update the list of timers.
-	TSharedPtr<STimingProfilerWindow> Wnd = FTimingProfilerManager::Get()->GetProfilerWindow();
-	if (Wnd && Wnd->TimersView)
+	// We need to check if TimersView or StatsView needs to update their lists of timers / counters.
+	// But, ensure we do not check too often.
+	static uint64 NextTimestamp = 0;
+	uint64 Time = FPlatformTime::Cycles64();
+	if (Time > NextTimestamp)
 	{
-		Wnd->TimersView->RebuildTree(false);
+		const uint64 WaitTime = static_cast<uint64>(0.2 / FPlatformTime::GetSecondsPerCycle64()); // 200ms
+		NextTimestamp = Time + WaitTime;
+
+		TSharedPtr<STimingProfilerWindow> Wnd = FTimingProfilerManager::Get()->GetProfilerWindow();
+		if (Wnd)
+		{
+			if (Wnd->TimersView)
+			{
+				Wnd->TimersView->RebuildTree(false);
+			}
+			if (Wnd->StatsView)
+			{
+				Wnd->StatsView->RebuildTree(false);
+			}
+		}
 	}
 
 	UpdateIo();
@@ -354,7 +371,7 @@ void STimingView::Tick(const FGeometry& AllottedGeometry, const double InCurrent
 	TSharedPtr<const Trace::IAnalysisSession> Session = FInsightsManager::Get()->GetSession();
 	if (Session)
 	{
-		Trace::FAnalysisSessionReadScope _(*Session.Get());
+		Trace::FAnalysisSessionReadScope SessionReadScope(*Session.Get());
 
 		// Check if horizontal scroll area has changed.
 		double SessionTime = Session->GetDurationSeconds();
@@ -378,7 +395,9 @@ void STimingView::Tick(const FGeometry& AllottedGeometry, const double InCurrent
 
 		if (IoOverviewTrack == nullptr)
 		{
-			IoOverviewTrack = AddTimingEventsTrack((0xF << 28) | 1, ETimingEventsTrackType::Io, TEXT("I/O Overview"), -1);
+			// Note: This track is hardcoded for now, as the I/O timelines are just prototypes for now (will be removed once the functionality is moved in analyzer).
+			uint64 OverviewTrackId = (0xF << 28) | 1;
+			IoOverviewTrack = AddTimingEventsTrack(OverviewTrackId, ETimingEventsTrackType::Io, TEXT("I/O Overview"), -1);
 			IoOverviewTrack->SetVisibilityFlag(bShowHideAllIoTracks);
 			bIsTimingEventsTrackDirty = true;
 		}
@@ -401,7 +420,7 @@ void STimingView::Tick(const FGeometry& AllottedGeometry, const double InCurrent
 			// Check available CPU tracks.
 			Session->ReadThreadProvider([this, &bIsTimingEventsTrackDirty, &TimingProfilerProvider](const Trace::IThreadProvider& ThreadProvider)
 			{
-				int Order = 1;
+				int32 Order = 1;
 				ThreadProvider.EnumerateThreads([this, &Order, &bIsTimingEventsTrackDirty, &TimingProfilerProvider](const Trace::FThreadInfo& ThreadInfo)
 				{
 					uint32 CpuTimelineIndex;
@@ -430,7 +449,9 @@ void STimingView::Tick(const FGeometry& AllottedGeometry, const double InCurrent
 
 		if (IoActivityTrack == nullptr)
 		{
-			IoActivityTrack = AddTimingEventsTrack((0xF << 28) | 2, ETimingEventsTrackType::Io, TEXT("I/O Activity"), 999999);
+			// Note: This track is hardcoded for now, as the I/O timelines are just prototypes for now (will be removed once the functionality is moved in analyzer).
+			uint64 ActivityTrackId = (0xF << 28) | 2;
+			IoActivityTrack = AddTimingEventsTrack(ActivityTrackId, ETimingEventsTrackType::Io, TEXT("I/O Activity"), 999999);
 			IoActivityTrack->SetVisibilityFlag(bShowHideAllIoTracks);
 			bIsTimingEventsTrackDirty = true;
 		}
@@ -447,9 +468,10 @@ void STimingView::Tick(const FGeometry& AllottedGeometry, const double InCurrent
 
 	// Compute total height of visible tracks.
 	float TotalScrollHeight = 0.0f;
-	for (int TrackIndex = 0; TrackIndex < TimingEventsTracks.Num(); ++TrackIndex)
+	for (FTimingEventsTrack* TrackPtr : TimingEventsTracks)
 	{
-		FTimingEventsTrack& Track = *TimingEventsTracks[TrackIndex];
+		ensure(TrackPtr != nullptr);
+		FTimingEventsTrack& Track = *TrackPtr;
 
 		if (Track.IsVisible())
 		{
@@ -538,7 +560,7 @@ void STimingView::Tick(const FGeometry& AllottedGeometry, const double InCurrent
 			GraphTrack.Update(Viewport);
 		}
 
-		for (int TrackIndex = 0; TrackIndex < TimingEventsTracks.Num(); ++TrackIndex)
+		for (int32 TrackIndex = 0; TrackIndex < TimingEventsTracks.Num(); ++TrackIndex)
 		{
 			TimingEventsTracks[TrackIndex]->Update(Viewport);
 		}
@@ -554,44 +576,45 @@ int32 STimingView::OnPaint(const FPaintArgs& Args, const FGeometry& AllottedGeom
 {
 	const bool bEnabled = ShouldBeEnabled(bParentEnabled);
 	const ESlateDrawEffect DrawEffects = bEnabled ? ESlateDrawEffect::None : ESlateDrawEffect::DisabledEffect;
-	FDrawContext DC(AllottedGeometry, MyCullingRect, InWidgetStyle, DrawEffects, OutDrawElements, LayerId);
+	FDrawContext DrawContext(AllottedGeometry, MyCullingRect, InWidgetStyle, DrawEffects, OutDrawElements, LayerId);
 
-	//OutDrawElements.GetRootDrawLayer().DrawElements.Reserve(50000);
+#if 0 // Enabling this may further increase UI performance (TODO: profile if this is really needed again).
+	// Avoids multiple resizes of Slate's draw elements buffers.
+	OutDrawElements.GetRootDrawLayer().DrawElements.Reserve(50000);
+#endif
 
 	const TSharedRef<FSlateFontMeasure> FontMeasureService = FSlateApplication::Get().GetRenderer()->GetFontMeasureService();
 
 	const float ViewWidth = AllottedGeometry.GetLocalSize().X;
 	const float ViewHeight = AllottedGeometry.GetLocalSize().Y;
 
-	/*
-	// Warming up Slate vertex/index buffers to avoid initial freezes due to multiple resizes of those buffers.
-	static bool bWarmingUp = true;
-	if (bWarmingUp)
+#if 0 // Enabling this may further increase UI performance (TODO: profile if this is really needed again).
+	// Warm up Slate vertex/index buffers to avoid initial freezes due to multiple resizes of those buffers.
+	static bool bWarmingUp = false;
+	if (!bWarmingUp)
 	{
-		bWarmingUp = false;
+		bWarmingUp = true;
 
 		FRandomStream RandomStream(0);
-		const int Count = 1'000'000;
-		for (int Index = 0; Index < Count; ++Index)
+		const int32 Count = 1'000'000;
+		for (int32 Index = 0; Index < Count; ++Index)
 		{
 			float X = ViewWidth * RandomStream.GetFraction();
 			float Y = ViewHeight * RandomStream.GetFraction();
 			FLinearColor Color(RandomStream.GetFraction(), RandomStream.GetFraction(), RandomStream.GetFraction(), 1.0f);
-			OnPaintState.DrawBox(X, Y, 1.0f, 1.0f, WhiteBrush, Color);
+			DrawContext.DrawBox(X, Y, 1.0f, 1.0f, WhiteBrush, Color);
 		}
 		LayerId++;
 		LayerId++;
 	}
-	*/
+#endif
 
 	//////////////////////////////////////////////////
 
 	FStopwatch Stopwatch;
 	Stopwatch.Start();
 
-	FTimingViewDrawHelper Helper(DC, Viewport, Layout);
-
-	Helper.Begin();
+	FTimingViewDrawHelper Helper(DrawContext, Viewport, Layout);
 
 	// Draw background.
 	Helper.DrawBackground();
@@ -600,7 +623,7 @@ int32 STimingView::OnPaint(const FPaintArgs& Args, const FGeometry& AllottedGeom
 	{
 		Helper.BeginTimelines();
 
-		for (int TrackIndex = 0; TrackIndex < TimingEventsTracks.Num(); ++TrackIndex)
+		for (int32 TrackIndex = 0; TrackIndex < TimingEventsTracks.Num(); ++TrackIndex)
 		{
 			FTimingEventsTrack& Track = *TimingEventsTracks[TrackIndex];
 			if (Track.IsVisible())
@@ -633,7 +656,7 @@ int32 STimingView::OnPaint(const FPaintArgs& Args, const FGeometry& AllottedGeom
 
 	if (GraphTrack.IsVisible())
 	{
-		GraphTrack.Draw(DC, Viewport);
+		GraphTrack.Draw(DrawContext, Viewport);
 	}
 
 	if (!FTimingEvent::AreEquals(SelectedTimingEvent, HoveredTimingEvent))
@@ -665,26 +688,24 @@ int32 STimingView::OnPaint(const FPaintArgs& Args, const FGeometry& AllottedGeom
 	if (MarkersTrack.IsVisible())
 	{
 		// Draw the time markers (vertical lines + category & message texts).
-		MarkersTrack.Draw(DC, Viewport);
+		MarkersTrack.Draw(DrawContext, Viewport);
 	}
 
-	Helper.End();
-
 	//////////////////////////////////////////////////
 
 	// Draw the time ruler.
 	if (TimeRulerTrack.IsVisible())
 	{
-		TimeRulerTrack.Draw(DC, Viewport, MousePosition, bIsSelecting, SelectionStartTime, SelectionEndTime);
+		TimeRulerTrack.Draw(DrawContext, Viewport, MousePosition, bIsSelecting, SelectionStartTime, SelectionEndTime);
 	}
 
 	// Fill background for the Tracks filter combobox.
-	DC.DrawBox(0, 0, 66, 22, WhiteBrush, FLinearColor(0.05f, 0.05f, 0.05f, 1.0f));
+	DrawContext.DrawBox(0, 0, 66, 22, WhiteBrush, FLinearColor(0.05f, 0.05f, 0.05f, 1.0f));
 
 	//////////////////////////////////////////////////
 
 	// Draw the time range selection.
-	DrawTimeRangeSelection(DC);
+	DrawTimeRangeSelection(DrawContext);
 
 	//////////////////////////////////////////////////
 
@@ -693,8 +714,8 @@ int32 STimingView::OnPaint(const FPaintArgs& Args, const FGeometry& AllottedGeom
 	float TimeMarkerX = Viewport.TimeToSlateUnitsRounded(TimeMarker);
 	if (TimeMarkerX >= 0.0f && TimeMarkerX < Viewport.Width)
 	{
-		DC.DrawBox(TimeMarkerX, 0.0f, 1.0f, Viewport.Height, WhiteBrush, FLinearColor(0.85f, 0.5f, 0.03f, 0.5f));
-		DC.LayerId++;
+		DrawContext.DrawBox(TimeMarkerX, 0.0f, 1.0f, Viewport.Height, WhiteBrush, FLinearColor(0.85f, 0.5f, 0.03f, 0.5f));
+		DrawContext.LayerId++;
 	}
 
 	//////////////////////////////////////////////////
@@ -717,10 +738,10 @@ int32 STimingView::OnPaint(const FPaintArgs& Args, const FGeometry& AllottedGeom
 		const float DbgX = ViewWidth - DbgW - 20.0f;
 		float DbgY = Viewport.TopOffset + 10.0f;
 
-		DC.LayerId++;
+		DrawContext.LayerId++;
 
-		DC.DrawBox(DbgX - 2.0f, DbgY - 2.0f, DbgW, DbgH, WhiteBrush, FLinearColor(1.0f, 1.0f, 1.0f, 0.9f));
-		DC.LayerId++;
+		DrawContext.DrawBox(DbgX - 2.0f, DbgY - 2.0f, DbgW, DbgH, WhiteBrush, FLinearColor(1.0f, 1.0f, 1.0f, 0.9f));
+		DrawContext.LayerId++;
 
 		FLinearColor DbgTextColor(0.0f, 0.0f, 0.0f, 0.9f);
 
@@ -738,7 +759,7 @@ int32 STimingView::OnPaint(const FPaintArgs& Args, const FGeometry& AllottedGeom
 
 		const uint64 AvgDrawDurationMs = FStopwatch::Cycles64ToMilliseconds(DrawDurationHistory.ComputeAverage());
 
-		DC.DrawText
+		DrawContext.DrawText
 		(
 			DbgX, DbgY,
 			FString::Printf(TEXT("D: %llu ms + %llu ms = %llu ms (%d fps)"),
@@ -756,7 +777,7 @@ int32 STimingView::OnPaint(const FPaintArgs& Args, const FGeometry& AllottedGeom
 		const uint64 AvgTimelineUpdateDurationMs = FStopwatch::Cycles64ToMilliseconds(TimelineCacheUpdateDurationHistory.ComputeAverage());
 		const uint64 LastTimelineUpdateDurationMs = FStopwatch::Cycles64ToMilliseconds(TimelineCacheUpdateDurationHistory.GetValue(0));
 
-		DC.DrawText
+		DrawContext.DrawText
 		(
 			DbgX, DbgY,
 			FString::Printf(TEXT("U1 avg: %llu ms, last: %llu ms"),
@@ -772,7 +793,7 @@ int32 STimingView::OnPaint(const FPaintArgs& Args, const FGeometry& AllottedGeom
 		const uint64 AvgTimeMarkerUpdateDurationMs = FStopwatch::Cycles64ToMilliseconds(TimeMarkerCacheUpdateDurationHistory.ComputeAverage());
 		const uint64 LastTimeMarkerUpdateDurationMs = FStopwatch::Cycles64ToMilliseconds(TimeMarkerCacheUpdateDurationHistory.GetValue(0));
 
-		DC.DrawText
+		DrawContext.DrawText
 		(
 			DbgX, DbgY,
 			FString::Printf(TEXT("U2 avg: %llu ms, last: %llu ms"),
@@ -785,16 +806,16 @@ int32 STimingView::OnPaint(const FPaintArgs& Args, const FGeometry& AllottedGeom
 		//////////////////////////////////////////////////
 		// Display timing events stats.
 
-		DC.DrawText
+		DrawContext.DrawText
 		(
 			DbgX, DbgY,
 			FString::Format(TEXT("{0} events : {1} ({2}) boxes, {3} borders, {4} texts"),
 			{
-				FText::AsNumber(Helper.NumEvents).ToString(),
-				FText::AsNumber(Helper.NumDrawBoxes).ToString(),
-				FText::AsPercent((double)Helper.NumDrawBoxes / (Helper.NumDrawBoxes + Helper.NumMergedBoxes)).ToString(),
-				FText::AsNumber(Helper.NumDrawBorders).ToString(),
-				FText::AsNumber(Helper.NumDrawTexts).ToString(),
+				FText::AsNumber(Helper.GetNumEvents()).ToString(),
+				FText::AsNumber(Helper.GetNumDrawBoxes()).ToString(),
+				FText::AsPercent((double)Helper.GetNumDrawBoxes() / (Helper.GetNumDrawBoxes() + Helper.GetNumMergedBoxes())).ToString(),
+				FText::AsNumber(Helper.GetNumDrawBorders()).ToString(),
+				FText::AsNumber(Helper.GetNumDrawTexts()).ToString(),
 				//OutDrawElements.GetRootDrawLayer().GetElementCount(),
 			}),
 			SummaryFont, DbgTextColor
@@ -804,7 +825,7 @@ int32 STimingView::OnPaint(const FPaintArgs& Args, const FGeometry& AllottedGeom
 		//////////////////////////////////////////////////
 		// Display time markers stats.
 
-		DC.DrawText
+		DrawContext.DrawText
 		(
 			DbgX, DbgY,
 			FString::Format(TEXT("{0}{1} logs : {2} boxes, {3} texts"),
@@ -821,7 +842,7 @@ int32 STimingView::OnPaint(const FPaintArgs& Args, const FGeometry& AllottedGeom
 		//////////////////////////////////////////////////
 		// Display Graph track stats.
 
-		DC.DrawText
+		DrawContext.DrawText
 		(
 			DbgX, DbgY,
 			FString::Format(TEXT("{0} events : {1} points, {2} lines, {3} boxes"),
@@ -838,7 +859,7 @@ int32 STimingView::OnPaint(const FPaintArgs& Args, const FGeometry& AllottedGeom
 		//////////////////////////////////////////////////
 		// Display viewport's horizontal info.
 
-		DC.DrawText
+		DrawContext.DrawText
 		(
 			DbgX, DbgY,
 			FString::Printf(TEXT("SX: %g, ST: %g, ET: %s"), Viewport.ScaleX, Viewport.StartTime, *TimeUtils::FormatTimeAuto(Viewport.MaxValidTime)),
@@ -849,7 +870,7 @@ int32 STimingView::OnPaint(const FPaintArgs& Args, const FGeometry& AllottedGeom
 		//////////////////////////////////////////////////
 		// Display viewport's vertical info.
 
-		DC.DrawText
+		DrawContext.DrawText
 		(
 			DbgX, DbgY,
 			FString::Printf(TEXT("Y: %.2f, H: %g, VH: %g"), Viewport.ScrollPosY, Viewport.ScrollHeight, Viewport.Height),
@@ -867,7 +888,7 @@ int32 STimingView::OnPaint(const FPaintArgs& Args, const FGeometry& AllottedGeom
 		if (bIsPanning) InputStr += " Panning";
 		if (bIsSelecting) InputStr += " Selecting";
 		if (bIsDragging) InputStr += " Dragging";
-		DC.DrawText(DbgX, DbgY, InputStr, SummaryFont, DbgTextColor);
+		DrawContext.DrawText(DbgX, DbgY, InputStr, SummaryFont, DbgTextColor);
 		DbgY += DbgDY;
 	}
 
@@ -891,11 +912,11 @@ int32 STimingView::OnPaint(const FPaintArgs& Args, const FGeometry& AllottedGeom
 			const FLinearColor BackgroundColor(0.05f, 0.05f, 0.05f, 1.0f);
 			const FLinearColor TextColor(0.7f, 0.7f, 0.7f, 1.0f);
 
-			DC.DrawBox(X - 8.0f, Y - 2.0f, Size.X + 16.0f, Size.Y + 4.0f, WhiteBrush, BackgroundColor);
-			DC.LayerId++;
+			DrawContext.DrawBox(X - 8.0f, Y - 2.0f, Size.X + 16.0f, Size.Y + 4.0f, WhiteBrush, BackgroundColor);
+			DrawContext.LayerId++;
 
-			DC.DrawText(X, Y, Str, MainFont, TextColor);
-			DC.LayerId++;
+			DrawContext.DrawText(X, Y, Str, MainFont, TextColor);
+			DrawContext.LayerId++;
 		}
 
 		// Draw info about hovered event (like a tooltip at mouse position).
@@ -937,30 +958,30 @@ int32 STimingView::OnPaint(const FPaintArgs& Args, const FGeometry& AllottedGeom
 			const FLinearColor TextColor(0.6f, 0.6f, 0.6f, TooltipAlpha);
 			const FLinearColor ValueColor(1.0f, 1.0f, 1.0f, TooltipAlpha);
 
-			DC.DrawBox(X - 6.0f, Y - 3.0f, TooltipWidth + 12.0f, H + 6.0f, WhiteBrush, BackgroundColor);
-			DC.LayerId++;
+			DrawContext.DrawBox(X - 6.0f, Y - 3.0f, TooltipWidth + 12.0f, H + 6.0f, WhiteBrush, BackgroundColor);
+			DrawContext.LayerId++;
 
-			DC.DrawText(X, Y, Name, MainFont, NameColor);
+			DrawContext.DrawText(X, Y, Name, MainFont, NameColor);
 			Y += LineH;
 
 			const float ValueX = X + 58.0f;
 
-			DC.DrawText(X + 3.0f, Y, TEXT("Incl. Time:"), MainFont, TextColor);
+			DrawContext.DrawText(X + 3.0f, Y, TEXT("Incl. Time:"), MainFont, TextColor);
 			FString InclStr = TimeUtils::FormatTimeAuto(HoveredTimingEvent.Duration());
-			DC.DrawText(ValueX, Y, InclStr, MainFont, ValueColor);
+			DrawContext.DrawText(ValueX, Y, InclStr, MainFont, ValueColor);
 			Y += LineH;
 
-			DC.DrawText(X, Y, TEXT("Excl. Time:"), MainFont, TextColor);
+			DrawContext.DrawText(X, Y, TEXT("Excl. Time:"), MainFont, TextColor);
 			FString ExclStr = TimeUtils::FormatTimeAuto(HoveredTimingEvent.ExclusiveTime);
-			DC.DrawText(ValueX, Y, ExclStr, MainFont, ValueColor);
+			DrawContext.DrawText(ValueX, Y, ExclStr, MainFont, ValueColor);
 			Y += LineH;
 
-			DC.DrawText(X + 24.0f, Y, TEXT("Depth:"), MainFont, TextColor);
+			DrawContext.DrawText(X + 24.0f, Y, TEXT("Depth:"), MainFont, TextColor);
 			FString DepthStr = FString::Printf(TEXT("%d"), HoveredTimingEvent.Depth);
-			DC.DrawText(ValueX, Y, DepthStr, MainFont, ValueColor);
+			DrawContext.DrawText(ValueX, Y, DepthStr, MainFont, ValueColor);
 			Y += LineH;
 
-			DC.LayerId++;
+			DrawContext.LayerId++;
 		}
 		else
 		{
@@ -1002,7 +1023,7 @@ void STimingView::DrawCpuGpuTimelineTrack(FTimingViewDrawHelper& Helper, FTiming
 		TSharedPtr<const Trace::IAnalysisSession> Session = FInsightsManager::Get()->GetSession();
 		if (Session.IsValid())
 		{
-			Trace::FAnalysisSessionReadScope _(*Session.Get());
+			Trace::FAnalysisSessionReadScope SessionReadScope(*Session.Get());
 
 			Session->ReadTimingProfilerProvider([this, &Helper, &Track](const Trace::ITimingProfilerProvider& TimingProfilerProvider)
 			{
@@ -1034,6 +1055,13 @@ void STimingView::DrawCpuGpuTimelineTrack(FTimingViewDrawHelper& Helper, FTiming
 	}
 }
 
+////////////////////////////////////////////////////////////////////////////////////////////////////
+// "I/O - File Activity" prototype
+////////////////////////////////////////////////////////////////////////////////////////////////////
+
+// The I/O timelines are just prototypes for now.
+// Below code will be removed once the functionality is moved in analyzer.
+
 ////////////////////////////////////////////////////////////////////////////////////////////////////
 
 const TCHAR* GetFileActivityTypeName(Trace::EFileActivityType Type)
@@ -1142,7 +1170,7 @@ void STimingView::UpdateIo()
 		TSharedPtr<const Trace::IAnalysisSession> Session = FInsightsManager::Get()->GetSession();
 		if (Session.IsValid())
 		{
-			Trace::FAnalysisSessionReadScope _(*Session.Get());
+			Trace::FAnalysisSessionReadScope SessionReadScope(*Session.Get());
 
 			Session->ReadFileActivityProvider([this](const Trace::IFileActivityProvider& FileActivityProvider)
 			{
@@ -1189,7 +1217,7 @@ void STimingView::UpdateIo()
 			{
 				uint64 TimelineId = Event.Depth;
 
-				int Depth = -1;
+				int32 Depth = -1;
 
 				bool bIsCloseEvent = false;
 
@@ -1198,7 +1226,7 @@ void STimingView::UpdateIo()
 				if (ActivityType == Trace::FileActivityType_Open)
 				{
 					// Find lane (avoiding overlaps with other opened files).
-					for (int LaneIndex = 0; LaneIndex < Lanes.Num(); ++LaneIndex)
+					for (int32 LaneIndex = 0; LaneIndex < Lanes.Num(); ++LaneIndex)
 					{
 						FIoLane& Lane = Lanes[LaneIndex];
 						if (Event.StartTime >= Lane.EndTime)
@@ -1227,7 +1255,7 @@ void STimingView::UpdateIo()
 					bIsCloseEvent = true;
 
 					// Find lane with same id.
-					for (int LaneIndex = 0; LaneIndex < Lanes.Num(); ++LaneIndex)
+					for (int32 LaneIndex = 0; LaneIndex < Lanes.Num(); ++LaneIndex)
 					{
 						FIoLane& Lane = Lanes[LaneIndex];
 						if (Lane.FileId == TimelineId)
@@ -1249,7 +1277,7 @@ void STimingView::UpdateIo()
 				{
 					// All other events should be inside the virtual Open-Close event.
 					// Find lane with same id.
-					for (int LaneIndex = 0; LaneIndex < Lanes.Num(); ++LaneIndex)
+					for (int32 LaneIndex = 0; LaneIndex < Lanes.Num(); ++LaneIndex)
 					{
 						FIoLane& Lane = Lanes[LaneIndex];
 						if (Lane.FileId == TimelineId)
@@ -1276,7 +1304,7 @@ void STimingView::UpdateIo()
 				Event.Depth = Depth;
 			}
 
-			for (int LaneIndex = 0; LaneIndex < Lanes.Num(); ++LaneIndex)
+			for (int32 LaneIndex = 0; LaneIndex < Lanes.Num(); ++LaneIndex)
 			{
 				FIoLane& Lane = Lanes[LaneIndex];
 				if (Lane.EndTime == std::numeric_limits<double>::infinity())
@@ -1366,9 +1394,11 @@ void STimingView::DrawIoActivityTrack(FTimingViewDrawHelper& Helper, FTimingEven
 	}
 }
 
+////////////////////////////////////////////////////////////////////////////////////////////////////
+// end of "I/O - File Activity" prototype
 ////////////////////////////////////////////////////////////////////////////////////////////////////
 
-void STimingView::DrawTimeRangeSelection(FDrawContext& DC) const
+void STimingView::DrawTimeRangeSelection(FDrawContext& DrawContext) const
 {
 	if (SelectionEndTime > SelectionStartTime)
 	{
@@ -1392,8 +1422,8 @@ void STimingView::DrawTimeRangeSelection(FDrawContext& DC) const
 			}
 
 			// Draw selection area.
-			DC.DrawBox(SelectionX1, 0.0f, SelectionX2 - SelectionX1, Viewport.Height, WhiteBrush, FLinearColor(0.25f, 0.5f, 1.0f, 0.25f));
-			DC.LayerId++;
+			DrawContext.DrawBox(SelectionX1, 0.0f, SelectionX2 - SelectionX1, Viewport.Height, WhiteBrush, FLinearColor(0.25f, 0.5f, 1.0f, 0.25f));
+			DrawContext.LayerId++;
 
 			FColor ArrowFillColor(32, 64, 128, 255);
 			FLinearColor ArrowColor(ArrowFillColor);
@@ -1401,16 +1431,16 @@ void STimingView::DrawTimeRangeSelection(FDrawContext& DC) const
 			if (SelectionX1 > 0.0f)
 			{
 				// Draw left side (vertical line).
-				DC.DrawBox(SelectionX1 - 1.0f, 0.0f, 1.0f, Viewport.Height, WhiteBrush, ArrowColor);
+				DrawContext.DrawBox(SelectionX1 - 1.0f, 0.0f, 1.0f, Viewport.Height, WhiteBrush, ArrowColor);
 			}
 
 			if (SelectionX2 < Viewport.Width)
 			{
 				// Draw right side (vertical line).
-				DC.DrawBox(SelectionX2, 0.0f, 1.0f, Viewport.Height, WhiteBrush, ArrowColor);
+				DrawContext.DrawBox(SelectionX2, 0.0f, 1.0f, Viewport.Height, WhiteBrush, ArrowColor);
 			}
 
-			DC.LayerId++;
+			DrawContext.LayerId++;
 
 			const float ArrowSize = 6.0f;
 			const float ArrowY = 6.0f;
@@ -1428,14 +1458,14 @@ void STimingView::DrawTimeRangeSelection(FDrawContext& DC) const
 				{
 					HorizLineX2 -= 1.0f;
 				}
-				DC.DrawBox(HorizLineX1, ArrowY - 1.0f, HorizLineX2 - HorizLineX1, 3.0f, WhiteBrush, ArrowColor);
+				DrawContext.DrawBox(HorizLineX1, ArrowY - 1.0f, HorizLineX2 - HorizLineX1, 3.0f, WhiteBrush, ArrowColor);
 
 				if (ClipLeft < ArrowSize)
 				{
 					// Draw left arrow.
 					for (float AH = 0.0f; AH < ArrowSize; AH += 1.0f)
 					{
-						DC.DrawBox(SelectionX1 - ClipLeft + AH, ArrowY - AH, 1.0f, 2.0f * AH + 1.0f, WhiteBrush, ArrowColor);
+						DrawContext.DrawBox(SelectionX1 - ClipLeft + AH, ArrowY - AH, 1.0f, 2.0f * AH + 1.0f, WhiteBrush, ArrowColor);
 					}
 				}
 
@@ -1444,14 +1474,14 @@ void STimingView::DrawTimeRangeSelection(FDrawContext& DC) const
 					// Draw right arrow.
 					for (float AH = 0.0f; AH < ArrowSize; AH += 1.0f)
 					{
-						DC.DrawBox(SelectionX2 + ClipRight - AH - 1.0f, ArrowY - AH, 1.0f, 2.0f * AH + 1.0f, WhiteBrush, ArrowColor);
+						DrawContext.DrawBox(SelectionX2 + ClipRight - AH - 1.0f, ArrowY - AH, 1.0f, 2.0f * AH + 1.0f, WhiteBrush, ArrowColor);
 					}
 				}
 
-				DC.LayerId++;
+				DrawContext.LayerId++;
 
-				/*
-				//im: This should be a more efficeint way top draw the arrows, but it renders them with artifacts (missing vertical lines)!
+#if 0
+				//im: This should be a more efficeint way top draw the arrows, but it renders them with artifacts (missing vertical lines; shader bug?)!
 
 				const FSlateBrush* MyBrush = WhiteBrush;
 				FSlateShaderResourceProxy* ResourceProxy = FSlateDataPayload::ResourceManager->GetShaderResource(*MyBrush);
@@ -1497,9 +1527,11 @@ void STimingView::DrawTimeRangeSelection(FDrawContext& DC) const
 					Indices,
 					nullptr,
 					0,
-					0, ESlateDrawEffect::PreMultipliedAlpha);
-				DC.LayerId++;
-				*/
+					0,
+					ESlateDrawEffect::PreMultipliedAlpha);
+
+				DrawContext.LayerId++;
+#endif
 			}
 
 			//////////////////////////////////////////////////
@@ -1513,11 +1545,11 @@ void STimingView::DrawTimeRangeSelection(FDrawContext& DC) const
 
 			const float CenterX = (SelectionX1 + SelectionX2) / 2.0f;
 
-			DC.DrawBox(CenterX - TextWidth / 2 - 2.0, ArrowY - 6.0f, TextWidth + 4.0f, 13.0f, WhiteBrush, ArrowColor);
-			DC.LayerId++;
+			DrawContext.DrawBox(CenterX - TextWidth / 2 - 2.0, ArrowY - 6.0f, TextWidth + 4.0f, 13.0f, WhiteBrush, ArrowColor);
+			DrawContext.LayerId++;
 
-			DC.DrawText(CenterX - TextWidth / 2, ArrowY - 6.0f, Text, MainFont, FLinearColor::White);
-			DC.LayerId++;
+			DrawContext.DrawText(CenterX - TextWidth / 2, ArrowY - 6.0f, Text, MainFont, FLinearColor::White);
+			DrawContext.LayerId++;
 		}
 	}
 }
@@ -1542,7 +1574,7 @@ FReply STimingView::OnMouseButtonDown(const FGeometry& MyGeometry, const FPointe
 		{
 			bIsLMB_Pressed = true;
 
-			if (bIsSpaceBarKeyPressed || MousePositionOnButtonDown.Y > TimeRulerHeight)
+			if (bIsSpaceBarKeyPressed || MousePositionOnButtonDown.Y > TimeRulerTrack.GetHeight())
 			{
 				bStartPanning = true;
 			}
@@ -1625,9 +1657,16 @@ FReply STimingView::OnMouseButtonUp(const FGeometry& MyGeometry, const FPointerE
 			{
 				//TODO: SelectionChangedEvent.Broadcast(SelectionStartTime, SelectionEndTime);
 				TSharedPtr<STimingProfilerWindow> Wnd = FTimingProfilerManager::Get()->GetProfilerWindow();
-				if (Wnd.IsValid())
+				if (Wnd)
 				{
-					Wnd->TimersView->UpdateStats(SelectionStartTime, SelectionEndTime);
+					if (Wnd->TimersView)
+					{
+						Wnd->TimersView->UpdateStats(SelectionStartTime, SelectionEndTime);
+					}
+					if (Wnd->StatsView)
+					{
+						Wnd->StatsView->UpdateStats(SelectionStartTime, SelectionEndTime);
+					}
 				}
 
 				bIsSelecting = false;
@@ -2286,9 +2325,16 @@ void STimingView::SelectTimeInterval(double IntervalStartTime, double IntervalDu
 
 	//TODO: SelectionChangedEvent.Broadcast(SelectionStartTime, SelectionEndTime);
 	TSharedPtr<STimingProfilerWindow> Wnd = FTimingProfilerManager::Get()->GetProfilerWindow();
-	if (Wnd.IsValid())
+	if (Wnd)
 	{
-		Wnd->TimersView->UpdateStats(SelectionStartTime, SelectionEndTime);
+		if (Wnd->TimersView)
+		{
+			Wnd->TimersView->UpdateStats(SelectionStartTime, SelectionEndTime);
+		}
+		if (Wnd->StatsView)
+		{
+			Wnd->StatsView->UpdateStats(SelectionStartTime, SelectionEndTime);
+		}
 	}
 }
 
@@ -2315,11 +2361,11 @@ void STimingView::SelectToTimeMarker(double Time)
 {
 	if (TimeMarker < Time)
 	{
-		SelectTimeInterval(TimeMarker, Time - TimeMarker);// +TimeUtils::Nanosecond);
+		SelectTimeInterval(TimeMarker, Time - TimeMarker);
 	}
 	else
 	{
-		SelectTimeInterval(Time, TimeMarker - Time);// +TimeUtils::Nanosecond);
+		SelectTimeInterval(Time, TimeMarker - Time);
 	}
 
 	TimeMarker = Time;
@@ -2356,7 +2402,7 @@ void STimingView::SetDrawOnlyBookmarks(bool bIsBookmarksTrack)
 	if (MarkersTrack.IsBookmarksTrack() != bIsBookmarksTrack)
 	{
 		const float PrevHeight = MarkersTrack.GetHeight();
-		MarkersTrack.SetIsBookmarksTrackFlag(bIsBookmarksTrack);
+		MarkersTrack.SetBookmarksTrackFlag(bIsBookmarksTrack);
 
 		if (MarkersTrack.IsVisible())
 		{
@@ -2382,11 +2428,14 @@ void STimingView::UpdateHoveredTimingEvent(float MX, float MY)
 		for (const auto& KV : CachedTimelines)
 		{
 			const FTimingEventsTrack& Track = *KV.Value;
-			const float Y = Viewport.TopOffset + Track.GetPosY() - Viewport.ScrollPosY;
-			if (MY >= Y && MY < Y + Track.GetHeight())
+			if (Track.IsVisible())
 			{
-				HoveredTimingEvent.Track = &Track;
-				break;
+				const float Y = Viewport.TopOffset + Track.GetPosY() - Viewport.ScrollPosY;
+				if (MY >= Y && MY < Y + Track.GetHeight())
+				{
+					HoveredTimingEvent.Track = &Track;
+					break;
+				}
 			}
 		}
 
@@ -2420,11 +2469,11 @@ void STimingView::UpdateHoveredTimingEvent(float MX, float MY)
 ////////////////////////////////////////////////////////////////////////////////////////////////////
 
 bool STimingView::SearchTimingEvent(const double InStartTime,
-									 const double InEndTime,
-									 TFunctionRef<bool(double, double, uint32, uint32)> InPredicate,
-									 FTimingEvent& InOutTimingEvent,
-									 bool bInStopAtFirstMatch,
-									 bool bInSearchForLargestEvent)
+									const double InEndTime,
+									TFunctionRef<bool(double, double, uint32, uint32)> InPredicate,
+									FTimingEvent& InOutTimingEvent,
+									bool bInStopAtFirstMatch,
+									bool bInSearchForLargestEvent)
 {
 	struct FSearchTimingEventContext
 	{
@@ -2457,7 +2506,7 @@ bool STimingView::SearchTimingEvent(const double InStartTime,
 	TSharedPtr<const Trace::IAnalysisSession> Session = FInsightsManager::Get()->GetSession();
 	if (Session.IsValid())
 	{
-		Trace::FAnalysisSessionReadScope _(*Session.Get());
+		Trace::FAnalysisSessionReadScope SessionReadScope(*Session.Get());
 
 		//if (Ctx.TimingEvent.Track->GetType() == TrackType::Timing)
 		Session->ReadTimingProfilerProvider([&Ctx](const Trace::ITimingProfilerProvider& TimingProfilerProvider)
@@ -2555,9 +2604,12 @@ void STimingView::OnSelectedTimingEventChanged()
 	if (SelectedTimingEvent.IsValid())
 	{
 		TSharedPtr<STimingProfilerWindow> Wnd = FTimingProfilerManager::Get()->GetProfilerWindow();
-		if (Wnd.IsValid())
+		if (Wnd)
 		{
-			Wnd->TimersView->SelectTimerNode(SelectedTimingEvent.TypeId);
+			if (Wnd->TimersView)
+			{
+				Wnd->TimersView->SelectTimerNode(SelectedTimingEvent.TypeId);
+			}
 		}
 	}
 }
@@ -2757,14 +2809,14 @@ const FTimerNodePtr STimingView::GetTimerNode(uint64 TypeId) const
 {
 	const FTimerNodePtr* TimerNodePtrPtr = nullptr;
 	TSharedPtr<STimingProfilerWindow> Wnd = FTimingProfilerManager::Get()->GetProfilerWindow();
-	if (Wnd.IsValid() && Wnd->TimersView.IsValid())
+	if (Wnd && Wnd->TimersView)
 	{
 		TimerNodePtrPtr = Wnd->TimersView->GetTimerNode(TypeId);
 		if (TimerNodePtrPtr == nullptr)
 		{
 			// List of timers in TimersView not up to date?
 			// Refresh and try again.
-			Wnd->TimersView->RebuildTree();
+			Wnd->TimersView->RebuildTree(false);
 			TimerNodePtrPtr = Wnd->TimersView->GetTimerNode(TypeId);
 		}
 	}
@@ -2836,7 +2888,7 @@ TSharedRef<SWidget> STimingView::MakeTracksFilterMenu()
 
 void STimingView::CreateTracksMenu(FMenuBuilder& MenuBuilder)
 {
-	for (int TrackIndex = 0; TrackIndex < TimingEventsTracks.Num(); ++TrackIndex)
+	for (int32 TrackIndex = 0; TrackIndex < TimingEventsTracks.Num(); ++TrackIndex)
 	{
 		const FTimingEventsTrack& Track = *TimingEventsTracks[TrackIndex];
 		if (Track.GetHeight() > 0.0f || Layout.TargetMinTimelineH > 0.0f)
@@ -2868,7 +2920,7 @@ void STimingView::ShowHideAllCpuTracks_Execute()
 {
 	bShowHideAllCpuTracks = !bShowHideAllCpuTracks;
 
-	for (int TrackIndex = 0; TrackIndex < TimingEventsTracks.Num(); ++TrackIndex)
+	for (int32 TrackIndex = 0; TrackIndex < TimingEventsTracks.Num(); ++TrackIndex)
 	{
 		FTimingEventsTrack& Track = *TimingEventsTracks[TrackIndex];
 		if (Track.GetType() == ETimingEventsTrackType::Cpu)
@@ -2877,6 +2929,9 @@ void STimingView::ShowHideAllCpuTracks_Execute()
 		}
 	}
 
+	HoveredTimingEvent.Reset();
+	SelectedTimingEvent.Reset();
+
 	bAreTimingEventsTracksDirty = true;
 }
 
@@ -2893,7 +2948,7 @@ void STimingView::ShowHideAllGpuTracks_Execute()
 {
 	bShowHideAllGpuTracks = !bShowHideAllGpuTracks;
 
-	for (int TrackIndex = 0; TrackIndex < TimingEventsTracks.Num(); ++TrackIndex)
+	for (int32 TrackIndex = 0; TrackIndex < TimingEventsTracks.Num(); ++TrackIndex)
 	{
 		FTimingEventsTrack& Track = *TimingEventsTracks[TrackIndex];
 		if (Track.GetType() == ETimingEventsTrackType::Gpu)
@@ -2902,6 +2957,9 @@ void STimingView::ShowHideAllGpuTracks_Execute()
 		}
 	}
 
+	HoveredTimingEvent.Reset();
+	SelectedTimingEvent.Reset();
+
 	bAreTimingEventsTracksDirty = true;
 }
 
@@ -2918,7 +2976,7 @@ void STimingView::ShowHideAllIoTracks_Execute()
 {
 	bShowHideAllIoTracks = !bShowHideAllIoTracks;
 
-	for (int TrackIndex = 0; TrackIndex < TimingEventsTracks.Num(); ++TrackIndex)
+	for (int32 TrackIndex = 0; TrackIndex < TimingEventsTracks.Num(); ++TrackIndex)
 	{
 		FTimingEventsTrack& Track = *TimingEventsTracks[TrackIndex];
 		if (Track.GetType() == ETimingEventsTrackType::Io)
@@ -2987,11 +3045,12 @@ void STimingView::ToggleTrackVisibility_Execute(uint64 InTrackId)
 	{
 		FTimingEventsTrack* Track = CachedTimelines[InTrackId];
 		Track->ToggleVisibility();
+
+		HoveredTimingEvent.Reset();
+		SelectedTimingEvent.Reset();
 	}
 }
 
-
-
 ////////////////////////////////////////////////////////////////////////////////////////////////////
 
 #undef LOCTEXT_NAMESPACE
diff --git a/Engine/Source/Developer/TraceInsights/Public/Insights/Version.h b/Engine/Source/Developer/TraceInsights/Public/Insights/Version.h
--- a/Engine/Source/Developer/TraceInsights/Public/Insights/Version.h
+++ b/Engine/Source/Developer/TraceInsights/Public/Insights/Version.h
@@ -5,10 +5,10 @@
 // See https://semver.org/
 
 #define UNREAL_INSIGHTS_VERSION_MAJOR 0
-#define UNREAL_INSIGHTS_VERSION_MINOR 5
+#define UNREAL_INSIGHTS_VERSION_MINOR 6
 #define UNREAL_INSIGHTS_VERSION_PATCH 0
 
-#define UNREAL_INSIGHTS_VERSION_STRING "0.05"
+#define UNREAL_INSIGHTS_VERSION_STRING "0.06"
 
 /// Extra identifier string (like: "alpha", "beta.1", "beta.2", "test", "featureX", etc.)
 #define UNREAL_INSIGHTS_VERSION_ID_STRING "im"
