[TraceInsights] Fixed compile errors introduced by previous commit (symbol collisions when using unity build).

diff --git a/Engine/Source/Developer/TraceInsights/Private/Insights/Widgets/SStatsTableCell.h b/Engine/Source/Developer/TraceInsights/Private/Insights/Widgets/SStatsTableCell.h
--- a/Engine/Source/Developer/TraceInsights/Private/Insights/Widgets/SStatsTableCell.h
+++ b/Engine/Source/Developer/TraceInsights/Private/Insights/Widgets/SStatsTableCell.h
@@ -9,7 +9,7 @@
 // Insights
 #include "Insights/ViewModels/StatsNode.h"
 
-DECLARE_DELEGATE_TwoParams(FSetHoveredTableCell, const FName /*ColumnId*/, const FStatsNodePtr /*SamplePtr*/);
+DECLARE_DELEGATE_TwoParams(FSetHoveredStatsTableCell, const FName /*ColumnId*/, const FStatsNodePtr /*SamplePtr*/);
 DECLARE_DELEGATE_RetVal_OneParam(bool, FIsColumnVisibleDelegate, const FName /*ColumnId*/);
 DECLARE_DELEGATE_RetVal_OneParam(EHorizontalAlignment, FGetColumnOutlineHAlignmentDelegate, const FName /*ColumnId*/);
 
@@ -17,7 +17,7 @@ class SStatsTableCell : public SCompoundWidget
 {
 public:
 	SLATE_BEGIN_ARGS(SStatsTableCell) {}
-		SLATE_EVENT(FSetHoveredTableCell, OnSetHoveredTableCell)
+		SLATE_EVENT(FSetHoveredStatsTableCell, OnSetHoveredTableCell)
 		SLATE_ATTRIBUTE(FText, HighlightText)
 		SLATE_ARGUMENT(FStatsNodePtr, StatsNodePtr)
 		SLATE_ARGUMENT(FName, ColumnId)
@@ -136,5 +136,5 @@ protected:
 	/** The Id of the column where this stats belongs. */
 	FName ColumnId;
 
-	FSetHoveredTableCell SetHoveredTableCellDelegate;
+	FSetHoveredStatsTableCell SetHoveredTableCellDelegate;
 };
diff --git a/Engine/Source/Developer/TraceInsights/Private/Insights/Widgets/SStatsTableRow.h b/Engine/Source/Developer/TraceInsights/Private/Insights/Widgets/SStatsTableRow.h
--- a/Engine/Source/Developer/TraceInsights/Private/Insights/Widgets/SStatsTableRow.h
+++ b/Engine/Source/Developer/TraceInsights/Private/Insights/Widgets/SStatsTableRow.h
@@ -14,7 +14,7 @@
 
 DECLARE_DELEGATE_RetVal_OneParam(bool, FShouldBeEnabledDelegate, const uint32 /*StatsId*/);
 DECLARE_DELEGATE_RetVal_OneParam(bool, FIsColumnVisibleDelegate, const FName /*ColumnId*/);
-DECLARE_DELEGATE_TwoParams(FSetHoveredTableCell, const FName /*ColumnId*/, const FStatsNodePtr /*StatsNodePtr*/);
+DECLARE_DELEGATE_TwoParams(FSetHoveredStatsTableCell, const FName /*ColumnId*/, const FStatsNodePtr /*StatsNodePtr*/);
 DECLARE_DELEGATE_RetVal_OneParam(EHorizontalAlignment, FGetColumnOutlineHAlignmentDelegate, const FName /*ColumnId*/);
 
 /** Widget that represents a table row in the stats' tree control. Generates widgets for each column on demand. */
@@ -24,7 +24,7 @@ public:
 	SLATE_BEGIN_ARGS(SStatsTableRow) {}
 		SLATE_EVENT(FShouldBeEnabledDelegate, OnShouldBeEnabled)
 		SLATE_EVENT(FIsColumnVisibleDelegate, OnIsColumnVisible)
-		SLATE_EVENT(FSetHoveredTableCell, OnSetHoveredTableCell)
+		SLATE_EVENT(FSetHoveredStatsTableCell, OnSetHoveredTableCell)
 		SLATE_EVENT(FGetColumnOutlineHAlignmentDelegate, OnGetColumnOutlineHAlignmentDelegate)
 		SLATE_ATTRIBUTE(FText, HighlightText)
 		SLATE_ATTRIBUTE(FName, HighlightedStatsName)
@@ -82,7 +82,7 @@ protected:
 
 	FShouldBeEnabledDelegate OnShouldBeEnabled;
 	FIsColumnVisibleDelegate IsColumnVisibleDelegate;
-	FSetHoveredTableCell SetHoveredTableCellDelegate;
+	FSetHoveredStatsTableCell SetHoveredTableCellDelegate;
 	FGetColumnOutlineHAlignmentDelegate GetColumnOutlineHAlignmentDelegate;
 
 	/** Text to be highlighted on stats name. */
diff --git a/Engine/Source/Developer/TraceInsights/Private/Insights/Widgets/STimerTableCell.h b/Engine/Source/Developer/TraceInsights/Private/Insights/Widgets/STimerTableCell.h
--- a/Engine/Source/Developer/TraceInsights/Private/Insights/Widgets/STimerTableCell.h
+++ b/Engine/Source/Developer/TraceInsights/Private/Insights/Widgets/STimerTableCell.h
@@ -9,7 +9,7 @@
 // Insights
 #include "Insights/ViewModels/TimerNode.h"
 
-DECLARE_DELEGATE_TwoParams(FSetHoveredTableCell, const FName /*ColumnId*/, const FTimerNodePtr /*SamplePtr*/);
+DECLARE_DELEGATE_TwoParams(FSetHoveredTimerTableCell, const FName /*ColumnId*/, const FTimerNodePtr /*SamplePtr*/);
 DECLARE_DELEGATE_RetVal_OneParam(bool, FIsColumnVisibleDelegate, const FName /*ColumnId*/);
 DECLARE_DELEGATE_RetVal_OneParam(EHorizontalAlignment, FGetColumnOutlineHAlignmentDelegate, const FName /*ColumnId*/);
 
@@ -17,7 +17,7 @@ class STimerTableCell : public SCompoundWidget
 {
 public:
 	SLATE_BEGIN_ARGS(STimerTableCell) {}
-		SLATE_EVENT(FSetHoveredTableCell, OnSetHoveredTableCell)
+		SLATE_EVENT(FSetHoveredTimerTableCell, OnSetHoveredTableCell)
 		SLATE_ATTRIBUTE(FText, HighlightText)
 		SLATE_ARGUMENT(FTimerNodePtr, TimerNodePtr)
 		SLATE_ARGUMENT(FName, ColumnId)
@@ -136,5 +136,5 @@ protected:
 	/** The Id of the column where this timer belongs. */
 	FName ColumnId;
 
-	FSetHoveredTableCell SetHoveredTableCellDelegate;
+	FSetHoveredTimerTableCell SetHoveredTableCellDelegate;
 };
diff --git a/Engine/Source/Developer/TraceInsights/Private/Insights/Widgets/STimerTableRow.h b/Engine/Source/Developer/TraceInsights/Private/Insights/Widgets/STimerTableRow.h
--- a/Engine/Source/Developer/TraceInsights/Private/Insights/Widgets/STimerTableRow.h
+++ b/Engine/Source/Developer/TraceInsights/Private/Insights/Widgets/STimerTableRow.h
@@ -14,7 +14,7 @@
 
 DECLARE_DELEGATE_RetVal_OneParam(bool, FShouldBeEnabledDelegate, const uint32 /*TimerId*/);
 DECLARE_DELEGATE_RetVal_OneParam(bool, FIsColumnVisibleDelegate, const FName /*ColumnId*/);
-DECLARE_DELEGATE_TwoParams(FSetHoveredTableCell, const FName /*ColumnId*/, const FTimerNodePtr /*TimerNodePtr*/);
+DECLARE_DELEGATE_TwoParams(FSetHoveredTimerTableCell, const FName /*ColumnId*/, const FTimerNodePtr /*TimerNodePtr*/);
 DECLARE_DELEGATE_RetVal_OneParam(EHorizontalAlignment, FGetColumnOutlineHAlignmentDelegate, const FName /*ColumnId*/);
 
 /** Widget that represents a table row in the timers' tree control. Generates widgets for each column on demand. */
@@ -24,7 +24,7 @@ public:
 	SLATE_BEGIN_ARGS(STimerTableRow) {}
 		SLATE_EVENT(FShouldBeEnabledDelegate, OnShouldBeEnabled)
 		SLATE_EVENT(FIsColumnVisibleDelegate, OnIsColumnVisible)
-		SLATE_EVENT(FSetHoveredTableCell, OnSetHoveredTableCell)
+		SLATE_EVENT(FSetHoveredTimerTableCell, OnSetHoveredTableCell)
 		SLATE_EVENT(FGetColumnOutlineHAlignmentDelegate, OnGetColumnOutlineHAlignmentDelegate)
 		SLATE_ATTRIBUTE(FText, HighlightText)
 		SLATE_ATTRIBUTE(FName, HighlightedTimerName)
@@ -82,7 +82,7 @@ protected:
 
 	FShouldBeEnabledDelegate OnShouldBeEnabled;
 	FIsColumnVisibleDelegate IsColumnVisibleDelegate;
-	FSetHoveredTableCell SetHoveredTableCellDelegate;
+	FSetHoveredTimerTableCell SetHoveredTableCellDelegate;
 	FGetColumnOutlineHAlignmentDelegate GetColumnOutlineHAlignmentDelegate;
 
 	/** Text to be highlighted on timer name. */
diff --git a/Engine/Source/Developer/TraceInsights/Private/Insights/Widgets/STimingProfilerWindow.cpp b/Engine/Source/Developer/TraceInsights/Private/Insights/Widgets/STimingProfilerWindow.cpp
--- a/Engine/Source/Developer/TraceInsights/Private/Insights/Widgets/STimingProfilerWindow.cpp
+++ b/Engine/Source/Developer/TraceInsights/Private/Insights/Widgets/STimingProfilerWindow.cpp
@@ -53,7 +53,7 @@ const FName FTimingProfilerTabs::LogViewID(TEXT("LogView"));
 ////////////////////////////////////////////////////////////////////////////////////////////////////
 
 // TODO: move this function to InsightsManager or in TraceSession.h
-static FText GetTextForNotification(const EInsightsNotificationType NotificatonType, const ELoadingProgressState ProgressState, const FString& Filename, const float ProgressPercent = 0.0f)
+static FText Timing_GetTextForNotification(const EInsightsNotificationType NotificatonType, const ELoadingProgressState ProgressState, const FString& Filename, const float ProgressPercent = 0.0f)
 {
 	FText Result;
 
@@ -528,7 +528,7 @@ void STimingProfilerWindow::ManageLoadingProgressNotificationState(const FString
 		const bool bContains = ActiveNotifications.Contains(Filename);
 		if (!bContains)
 		{
-			FNotificationInfo NotificationInfo(GetTextForNotification(NotificatonType, ProgressState, BaseFilename));
+			FNotificationInfo NotificationInfo(Timing_GetTextForNotification(NotificatonType, ProgressState, BaseFilename));
 			NotificationInfo.bFireAndForget = false;
 			NotificationInfo.bUseLargeFont = false;
 
@@ -557,7 +557,7 @@ void STimingProfilerWindow::ManageLoadingProgressNotificationState(const FString
 		if (LoadingProgressPtr)
 		{
 			TSharedPtr<SNotificationItem> LoadingProcessPinned = LoadingProgressPtr->Pin();
-			LoadingProcessPinned->SetText(GetTextForNotification(NotificatonType, ProgressState, BaseFilename, LoadingProgress));
+			LoadingProcessPinned->SetText(Timing_GetTextForNotification(NotificatonType, ProgressState, BaseFilename, LoadingProgress));
 			LoadingProcessPinned->SetCompletionState(SNotificationItem::CS_Pending);
 		}
 	}
@@ -567,7 +567,7 @@ void STimingProfilerWindow::ManageLoadingProgressNotificationState(const FString
 		if (LoadingProgressPtr)
 		{
 			TSharedPtr<SNotificationItem> LoadingProcessPinned = LoadingProgressPtr->Pin();
-			LoadingProcessPinned->SetText(GetTextForNotification(NotificatonType, ProgressState, BaseFilename));
+			LoadingProcessPinned->SetText(Timing_GetTextForNotification(NotificatonType, ProgressState, BaseFilename));
 			LoadingProcessPinned->SetCompletionState(SNotificationItem::CS_Success);
 
 			// Notifications for received files are handled by the user.
@@ -584,7 +584,7 @@ void STimingProfilerWindow::ManageLoadingProgressNotificationState(const FString
 		if (LoadingProgressPtr)
 		{
 			TSharedPtr<SNotificationItem> LoadingProcessPinned = LoadingProgressPtr->Pin();
-			LoadingProcessPinned->SetText(GetTextForNotification(NotificatonType, ProgressState, BaseFilename));
+			LoadingProcessPinned->SetText(Timing_GetTextForNotification(NotificatonType, ProgressState, BaseFilename));
 			LoadingProcessPinned->SetCompletionState(SNotificationItem::CS_Fail);
 
 			LoadingProcessPinned->ExpireAndFadeout();
