[TraceInsights]
  - Made the start page (now renamed as "Unreal Insights") the default tab. The other tabs are now hidden by default.
  - Start Page: Implemented "Live" and "Load" buttons.
  - Start Page: Added info about Local Session Directory.

diff --git a/Engine/Source/Developer/TraceInsights/Private/Insights/InsightsCommands.cpp b/Engine/Source/Developer/TraceInsights/Private/Insights/InsightsCommands.cpp
--- a/Engine/Source/Developer/TraceInsights/Private/Insights/InsightsCommands.cpp
+++ b/Engine/Source/Developer/TraceInsights/Private/Insights/InsightsCommands.cpp
@@ -38,7 +38,6 @@ void FInsightsCommands::RegisterCommands()
 {
 	UI_COMMAND(InsightsManager_Live, "Live", "Loads profiler data from a live trace session", EUserInterfaceActionType::Button, FInputChord());
 	UI_COMMAND(InsightsManager_Load, "Load...", "Loads profiler data from a trace file", EUserInterfaceActionType::Button, FInputChord(EModifierKey::Control, EKeys::L));
-	UI_COMMAND(InsightsManager_Mock, "Mock", "Creates mock profiler data", EUserInterfaceActionType::Button, FInputChord());
 	UI_COMMAND(ToggleDebugInfo, "Debug", "Toggles the display of debug info", EUserInterfaceActionType::ToggleButton, FInputChord(EModifierKey::Control, EKeys::D));
 	UI_COMMAND(OpenSettings, "Settings", "Opens the Unreal Insights settings", EUserInterfaceActionType::Button, FInputChord(EModifierKey::Control, EKeys::O));
 }
@@ -161,33 +160,6 @@ bool FInsightsActionManager::InsightsManager_Load_CanExecute() const
 	return true;
 }
 
-////////////////////////////////////////////////////////////////////////////////////////////////////
-// InsightsManager_Mock
-////////////////////////////////////////////////////////////////////////////////////////////////////
-
-void FInsightsActionManager::Map_InsightsManager_Mock()
-{
-	FUIAction UIAction;
-	UIAction.ExecuteAction = FExecuteAction::CreateRaw(this, &FInsightsActionManager::InsightsManager_Mock_Execute);
-	UIAction.CanExecuteAction = FCanExecuteAction::CreateRaw(this, &FInsightsActionManager::InsightsManager_Mock_CanExecute);
-
-	This->CommandList->MapAction(This->GetCommands().InsightsManager_Mock, UIAction);
-}
-
-////////////////////////////////////////////////////////////////////////////////////////////////////
-
-void FInsightsActionManager::InsightsManager_Mock_Execute()
-{
-	This->CreateMockSession();
-}
-
-////////////////////////////////////////////////////////////////////////////////////////////////////
-
-bool FInsightsActionManager::InsightsManager_Mock_CanExecute() const
-{
-	return true;
-}
-
 ////////////////////////////////////////////////////////////////////////////////////////////////////
 // OpenSettings
 ////////////////////////////////////////////////////////////////////////////////////////////////////
diff --git a/Engine/Source/Developer/TraceInsights/Private/Insights/InsightsCommands.h b/Engine/Source/Developer/TraceInsights/Private/Insights/InsightsCommands.h
--- a/Engine/Source/Developer/TraceInsights/Private/Insights/InsightsCommands.h
+++ b/Engine/Source/Developer/TraceInsights/Private/Insights/InsightsCommands.h
@@ -7,16 +7,16 @@
 #include "Framework/Commands/Commands.h"
 
 ////////////////////////////////////////////////////////////////////////////////////////////////////
-/// Class that holds all profiler commands.
-
-class FInsightsCommands
-	: public TCommands<FInsightsCommands>
+/**
+ * Class that holds all profiler commands.
+ */
+class FInsightsCommands : public TCommands<FInsightsCommands>
 {
 public:
-	/// Default constructor.
+	/** Default constructor. */
 	FInsightsCommands();
 
-	/// Initialize commands.
+	/** Initialize commands. */
 	virtual void RegisterCommands() override;
 
 public:
@@ -27,32 +27,30 @@ public:
 	//     const FUIAction <CommandName>_Custom(...) const;
 	//////////////////////////////////////////////////
 
-	/// Load profiler data from a live trace session. Global version.
+	/** Load profiler data from a live trace session. Global version. */
 	TSharedPtr<FUICommandInfo> InsightsManager_Live;
 
-	/// Load profiler data from a trace file. Global version.
+	/** Load profiler data from a trace file. Global version. */
 	TSharedPtr<FUICommandInfo> InsightsManager_Load;
 
-	/// Create mock profiler data. Global version.
-	TSharedPtr<FUICommandInfo> InsightsManager_Mock;
-
-	/// Toggles the debug info. Global and custom command.
+	/** Toggles the debug info. Global and custom command. */
 	TSharedPtr<FUICommandInfo> ToggleDebugInfo;
 
-	/// Open settings for the profiler manager.
+	/** Open settings for the profiler manager. */
 	TSharedPtr<FUICommandInfo> OpenSettings;
 };
 
 ////////////////////////////////////////////////////////////////////////////////////////////////////
-/// Class that provides helper functions for the commands to avoid cluttering manager with many small functions.
-/// Can't contain any variables. Directly operates on the manager instance.
-
+/**
+ * Class that provides helper functions for the commands to avoid cluttering manager with many small functions.
+ * Can't contain any variables. Directly operates on the manager instance.
+ */
 class FInsightsActionManager
 {
 	friend class FInsightsManager;
 
 private:
-	/// Private constructor.
+	/** Private constructor. */
 	FInsightsActionManager(class FInsightsManager* Instance)
 		: This(Instance)
 	{}
@@ -61,53 +59,44 @@ private:
 	// InsightsManager_Live
 
 public:
-	void Map_InsightsManager_Live(); ///< Maps UI command info InsightsManager_Live with the specified UI command list.
+	void Map_InsightsManager_Live(); /**< Maps UI command info InsightsManager_Live with the specified UI command list. */
 protected:
-	void InsightsManager_Live_Execute(); ///< Handles FExecuteAction for InsightsManager_Live.
-	bool InsightsManager_Live_CanExecute() const; ///< Handles FCanExecuteAction for InsightsManager_Live.
+	void InsightsManager_Live_Execute(); /**< Handles FExecuteAction for InsightsManager_Live. */
+	bool InsightsManager_Live_CanExecute() const; /**< Handles FCanExecuteAction for InsightsManager_Live. */
 
 	//////////////////////////////////////////////////
 	// InsightsManager_Load
 
 public:
-	void Map_InsightsManager_Load(); ///< Maps UI command info InsightsManager_Load with the specified UI command list.
-protected:
-	void InsightsManager_Load_Execute(); ///< Handles FExecuteAction for InsightsManager_Load.
-	bool InsightsManager_Load_CanExecute() const; ///< Handles FCanExecuteAction for InsightsManager_Load.
-
-	//////////////////////////////////////////////////
-	// InsightsManager_Mock
-
-public:
-	void Map_InsightsManager_Mock(); ///< Maps UI command info InsightsManager_Mock with the specified UI command list.
+	void Map_InsightsManager_Load(); /**< Maps UI command info InsightsManager_Load with the specified UI command list. */
 protected:
-	void InsightsManager_Mock_Execute(); ///< Handles FExecuteAction for InsightsManager_Mock.
-	bool InsightsManager_Mock_CanExecute() const; ///< Handles FCanExecuteAction for InsightsManager_Mock.
+	void InsightsManager_Load_Execute(); /**< Handles FExecuteAction for InsightsManager_Load. */
+	bool InsightsManager_Load_CanExecute() const; /**< Handles FCanExecuteAction for InsightsManager_Load. */
 
 	//////////////////////////////////////////////////
 	// ToggleDebugInfo
 
 public:
-	void Map_ToggleDebugInfo_Global(); ///< Maps UI command info ToggleDebugInfo with the specified UI command list.
-	const FUIAction ToggleDebugInfo_Custom() const; ///< UI action for ToggleDebugInfo command.
+	void Map_ToggleDebugInfo_Global(); /**< Maps UI command info ToggleDebugInfo with the specified UI command list. */
+	const FUIAction ToggleDebugInfo_Custom() const; /**< UI action for ToggleDebugInfo command. */
 protected:
-	void ToggleDebugInfo_Execute(); ///< Handles FExecuteAction for ToggleDebugInfo.
-	bool ToggleDebugInfo_CanExecute() const; ///< Handles FCanExecuteAction for ToggleDebugInfo.
-	ECheckBoxState ToggleDebugInfo_GetCheckState() const; ///< Handles FGetActionCheckState for ToggleDebugInfo.
+	void ToggleDebugInfo_Execute(); /**< Handles FExecuteAction for ToggleDebugInfo. */
+	bool ToggleDebugInfo_CanExecute() const; /**< Handles FCanExecuteAction for ToggleDebugInfo. */
+	ECheckBoxState ToggleDebugInfo_GetCheckState() const; /**< Handles FGetActionCheckState for ToggleDebugInfo. */
 
 	//////////////////////////////////////////////////
 	// OpenSettings
 
 public:
-	void Map_OpenSettings_Global(); ///< Maps UI command info OpenSettings with the specified UI command list.
-	const FUIAction OpenSettings_Custom() const; ///< UI action for OpenSettings command.
+	void Map_OpenSettings_Global(); /**< Maps UI command info OpenSettings with the specified UI command list. */
+	const FUIAction OpenSettings_Custom() const; /**< UI action for OpenSettings command. */
 protected:
-	void OpenSettings_Execute(); ///< Handles FExecuteAction for OpenSettings.
-	bool OpenSettings_CanExecute() const; ///< Handles FCanExecuteAction for OpenSettings.
+	void OpenSettings_Execute(); /**< Handles FExecuteAction for OpenSettings. */
+	bool OpenSettings_CanExecute() const; /**< Handles FCanExecuteAction for OpenSettings. */
 
 	//////////////////////////////////////////////////
 
 protected:
-	/// Reference to the global instance of the Insights manager.
+	/** Reference to the global instance of the Insights manager. */
 	class FInsightsManager* This;
 };
diff --git a/Engine/Source/Developer/TraceInsights/Private/Insights/InsightsManager.cpp b/Engine/Source/Developer/TraceInsights/Private/Insights/InsightsManager.cpp
--- a/Engine/Source/Developer/TraceInsights/Private/Insights/InsightsManager.cpp
+++ b/Engine/Source/Developer/TraceInsights/Private/Insights/InsightsManager.cpp
@@ -10,7 +10,7 @@
 
 // Insights
 #include "Insights/TimingProfilerManager.h"
-#include "Insights/Widgets/STimingProfilerWindow.h"
+#include "Insights/Widgets/SStartPageWindow.h"
 
 ////////////////////////////////////////////////////////////////////////////////////////////////////
 
@@ -48,7 +48,6 @@ void FInsightsManager::BindCommands()
 {
 	ActionManager.Map_InsightsManager_Live();
 	ActionManager.Map_InsightsManager_Load();
-	ActionManager.Map_InsightsManager_Mock();
 	ActionManager.Map_ToggleDebugInfo_Global();
 	ActionManager.Map_OpenSettings_Global();
 }
@@ -119,20 +118,19 @@ bool FInsightsManager::Tick(float DeltaTime)
 void FInsightsManager::ResetSession()
 {
 	Session = nullptr;
-
-	FTimingProfilerManager::Get()->OnSessionChanged();
-	SessionChangedEvent.Broadcast();
+	OnSessionChanged();
 }
 
 ////////////////////////////////////////////////////////////////////////////////////////////////////
 
-void FInsightsManager::CreateMockSession()
+void FInsightsManager::OnSessionChanged()
 {
-	ResetSession();
-
-	Session = AnalysisService->MockAnalysis();
+	if (TSharedPtr<FTimingProfilerManager> TimingProfilerManager = FTimingProfilerManager::Get())
+	{
+		// FIXME: make TimingProfilerManager to register to SessionChangedEvent instead
+		TimingProfilerManager->OnSessionChanged();
+	}
 
-	FTimingProfilerManager::Get()->OnSessionChanged();
 	SessionChangedEvent.Broadcast();
 }
 
@@ -142,9 +140,9 @@ void FInsightsManager::CreateLiveSession()
 {
 	ResetSession();
 
-	Session = nullptr;
 	TArray<Trace::FSessionHandle> AvailableSessions;
 	SessionService->GetAvailableSessions(AvailableSessions);
+
 	for (Trace::FSessionHandle SessionHandle : AvailableSessions)
 	{
 		Trace::FSessionInfo SessionInfo;
@@ -157,6 +155,7 @@ void FInsightsManager::CreateLiveSession()
 			break;
 		}
 	}
+
 	if (!Session && AvailableSessions.Num())
 	{
 		Trace::FSessionHandle SessionHandle = AvailableSessions.Last();
@@ -165,8 +164,7 @@ void FInsightsManager::CreateLiveSession()
 		Session = AnalysisService->StartAnalysis(TEXT("Latest session"), MoveTemp(DataStream));
 	}
 
-	FTimingProfilerManager::Get()->OnSessionChanged();
-	SessionChangedEvent.Broadcast();
+	OnSessionChanged();
 }
 
 ////////////////////////////////////////////////////////////////////////////////////////////////////
@@ -180,15 +178,18 @@ void FInsightsManager::LoadTraceFile(const FString& TraceFilepath)
 
 	Session = AnalysisService->StartAnalysis(*TraceFilepath, MoveTemp(DataStream));
 
-	FTimingProfilerManager::Get()->OnSessionChanged();
-	SessionChangedEvent.Broadcast();
+	OnSessionChanged();
 }
 
 ////////////////////////////////////////////////////////////////////////////////////////////////////
 
 void FInsightsManager::OpenSettings()
 {
-	FTimingProfilerManager::Get()->GetProfilerWindow()->OpenProfilerSettings();
+	TSharedPtr<SStartPageWindow> Wnd = GetStartPageWindow();
+	if (Wnd.IsValid())
+	{
+		Wnd->OpenSettings();
+	}
 }
 
 ////////////////////////////////////////////////////////////////////////////////////////////////////
diff --git a/Engine/Source/Developer/TraceInsights/Private/Insights/InsightsManager.h b/Engine/Source/Developer/TraceInsights/Private/Insights/InsightsManager.h
--- a/Engine/Source/Developer/TraceInsights/Private/Insights/InsightsManager.h
+++ b/Engine/Source/Developer/TraceInsights/Private/Insights/InsightsManager.h
@@ -12,7 +12,7 @@
 
 ////////////////////////////////////////////////////////////////////////////////////////////////////
 
-class STimingProfilerWindow;
+class SStartPageWindow;
 
 namespace Trace
 {
@@ -22,7 +22,7 @@ namespace Trace
 }
 
 ////////////////////////////////////////////////////////////////////////////////////////////////////
-/// Enumerates loading a trace file progress states.
+/** Enumerates loading a trace file progress states. */
 
 enum class ELoadingProgressState
 {
@@ -32,7 +32,7 @@ enum class ELoadingProgressState
 	Failed,
 	Cancelled,
 
-	/// Invalid enum type, may be used as a number of enumerations.
+	/** Invalid enum type, may be used as a number of enumerations. */
 	InvalidOrMax,
 };
 
@@ -42,31 +42,34 @@ enum class EInsightsNotificationType
 {
 	LoadingTraceFile,
 
-	/// Invalid enum type, may be used as a number of enumerations.
+	/** Invalid enum type, may be used as a number of enumerations. */
 	InvalidOrMax,
 };
 
 ////////////////////////////////////////////////////////////////////////////////////////////////////
-/// This class manages following areas:
-///     Connecting/disconnecting to source trace
-///     Lifetime of all other managers (specific profilers)
-///     Global Unreal Insights application state and settings
-
+/**
+ * This class manages following areas:
+ *     Connecting/disconnecting to source trace
+ *     Lifetime of all other managers (specific profilers)
+ *     Global Unreal Insights application state and settings
+ */
 class FInsightsManager
 	: public TSharedFromThis<FInsightsManager>
 {
 	friend class FInsightsActionManager;
 
 public:
-	/// Creates the main manager, only one instance can exist.
+	/** Creates the main manager, only one instance can exist. */
 	FInsightsManager(TSharedRef<Trace::IAnalysisService> TraceAnalysisService, TSharedRef<Trace::ISessionService> SessionService);
 
-	/// Virtual destructor.
+	/** Virtual destructor. */
 	virtual ~FInsightsManager();
 
-	/// Creates an instance of the main manager and initializes global instance with the previously created instance of the manager.
-	/// @param TraceAnalysisService The trace analysis service
-	/// @param TraceSessionService  The trace session service
+	/**
+	 * Creates an instance of the main manager and initializes global instance with the previously created instance of the manager.
+	 * @param TraceAnalysisService The trace analysis service
+	 * @param TraceSessionService  The trace session service
+	 */
 	static TSharedPtr<FInsightsManager> Initialize(TSharedRef<Trace::IAnalysisService> TraceAnalysisService, TSharedRef<Trace::ISessionService> TraceSessionService)
 	{
 		if (FInsightsManager::Instance.IsValid())
@@ -80,104 +83,124 @@ public:
 		return FInsightsManager::Instance;
 	}
 
-	/// Shutdowns the main manager.
+	/** Shutdowns the main manager. */
 	void Shutdown()
 	{
 		FInsightsManager::Instance.Reset();
 	}
 
-protected:
-	/// Finishes initialization of the profiler manager
-	void PostConstructor();
-
-	/// Binds our UI commands to delegates.
-	void BindCommands();
-
-public:
-	/// @return the global instance of the main manager (FInsightsManager).
+	/** @return the global instance of the main manager (FInsightsManager). */
 	static TSharedPtr<FInsightsManager> Get();
 
-	/// @return an instance of the trace analysis session.
+	TSharedRef<Trace::IAnalysisService> GetAnalysisService() const { return AnalysisService; }
+	TSharedRef<Trace::ISessionService> GetSessionService() const { return SessionService; }
+
+	/** @return an instance of the trace analysis session. */
 	TSharedPtr<const Trace::IAnalysisSession> GetSession() const;
 
-	/// @returns UI command list for the main manager.
+	/** @returns UI command list for the main manager. */
 	const TSharedRef<FUICommandList> GetCommandList() const;
 
-	/// @return an instance of the main commands.
+	/** @return an instance of the main commands. */
 	static const FInsightsCommands& GetCommands();
 
-	/// @return an instance of the main action manager.
+	/** @return an instance of the main action manager. */
 	static FInsightsActionManager& GetActionManager();
 
-	/// @return an instance of the main settings.
+	/** @return an instance of the main settings. */
 	static FInsightsSettings& GetSettings();
 
+	void AssignStartPageWindow(const TSharedRef<SStartPageWindow>& InStartPageWindow)
+	{
+		StartPageWindow = InStartPageWindow;
+	}
+
+	/**
+	 * Converts profiler window weak pointer to a shared pointer and returns it.
+	 * Make sure the returned pointer is valid before trying to dereference it.
+	 */
+	TSharedPtr<class SStartPageWindow> GetStartPageWindow() const
+	{
+		return StartPageWindow.Pin();
+	}
+
 	////////////////////////////////////////////////////////////////////////////////////////////////////
 	// Getters and setters used by Toggle Commands.
 
-	/// @return true, if UI is allowed to display debug info
+	/** @return true, if UI is allowed to display debug info. */
 	const bool IsDebugInfoEnabled() const { return bIsDebugInfoEnabled; }
 	void SetDebugInfo(const bool bDebugInfoEnabledState) { bIsDebugInfoEnabled = bDebugInfoEnabledState; }
 
+	////////////////////////////////////////////////////////////////////////////////////////////////////
+
+	/** Creates a new profiler session instance and loads a live trace. */
+	void CreateLiveSession();
+
+	/**
+	 * Creates a new profiler session instance and loads a trace file from the specified location.
+	 * @param TraceFilepath - The path to the trace file
+	 */
+	void LoadTraceFile(const FString& TraceFilepath);
+
+	/** Opens the Settings dialog. */
+	void OpenSettings();
+
 	////////////////////////////////////////////////////////////////////////////////////////////////////
 	// SessionChangedEvent
 
 public:
-	/// The event to execute when the session has changed.
+	/** The event to execute when the session has changed. */
 	DECLARE_EVENT(FTimingProfilerManager, FSessionChangedEvent);
 	FSessionChangedEvent& GetSessionChangedEvent() { return SessionChangedEvent; }
-protected:
-	/// The event to execute when the session has changed.
+private:
+	/** The event to execute when the session has changed. */
 	FSessionChangedEvent SessionChangedEvent;
 
 	////////////////////////////////////////////////////////////////////////////////////////////////////
 
-public:
-	/// Creates a new profiler session instance initialized with mock data (for debugging purposes).
-	void CreateMockSession();
-
-	/// Creates a new profiler session instance and loads a live trace.
-	void CreateLiveSession();
-
-	/// Creates a new profiler session instance and loads a trace file from the specified location.
-	/// @param TraceFilepath - The path to the trace file
-	void LoadTraceFile(const FString& TraceFilepath);
+private:
+	/** Finishes initialization of the profiler manager. */
+	void PostConstructor();
 
-	/// Opens the Settings dialog.
-	void OpenSettings();
+	/** Binds our UI commands to delegates. */
+	void BindCommands();
 
-protected:
-	/// Updates this manager, done through FCoreTicker.
+	/** Updates this manager, done through FCoreTicker. */
 	bool Tick(float DeltaTime);
 
-	/// Resets (closes) current session instance.
+	/** Resets (closes) current session instance. */
 	void ResetSession();
 
-protected:
-	/// The delegate to be invoked when this manager ticks.
+	void OnSessionChanged();
+
+private:
+	/** The delegate to be invoked when this manager ticks. */
 	FTickerDelegate OnTick;
 
-	/// Handle to the registered OnTick.
+	/** Handle to the registered OnTick. */
 	FDelegateHandle OnTickHandle;
 
 	TSharedRef<Trace::IAnalysisService> AnalysisService;
 	TSharedRef<Trace::ISessionService> SessionService;
 
-	/// The trace analysis session.
+	/** The trace analysis session. */
 	TSharedPtr<const Trace::IAnalysisSession> Session;
 
-	/// List of UI commands for this manager. This will be filled by this and corresponding classes.
+	/** List of UI commands for this manager. This will be filled by this and corresponding classes. */
 	TSharedRef<FUICommandList> CommandList;
 
-	/// An instance of the main action manager.
+	/** An instance of the main action manager. */
 	FInsightsActionManager ActionManager;
 
-	/// An instance of the main settings.
+	/** An instance of the main settings. */
 	FInsightsSettings Settings;
 
-	/// If enabled, UI can display additional info for debugging purposes.
+	/** A weak pointer to the Start Page window. */
+	TWeakPtr<class SStartPageWindow> StartPageWindow;
+
+	/** If enabled, UI can display additional info for debugging purposes. */
 	bool bIsDebugInfoEnabled;
 
-	/// A shared pointer to the global instance of the main manager.
+	/** A shared pointer to the global instance of the main manager. */
 	static TSharedPtr<FInsightsManager> Instance;
 };
diff --git a/Engine/Source/Developer/TraceInsights/Private/Insights/InsightsSettings.h b/Engine/Source/Developer/TraceInsights/Private/Insights/InsightsSettings.h
--- a/Engine/Source/Developer/TraceInsights/Private/Insights/InsightsSettings.h
+++ b/Engine/Source/Developer/TraceInsights/Private/Insights/InsightsSettings.h
@@ -5,7 +5,7 @@
 #include "CoreMinimal.h"
 #include "Misc/ConfigCacheIni.h"
 
-/// Contains all settings for the Unreal Insights, accessible through the main manager.
+/** Contains all settings for the Unreal Insights, accessible through the main manager. */
 class FInsightsSettings
 {
 public:
@@ -69,21 +69,21 @@ public:
 
 //protected:
 public:
-	/// Contains default settings.
+	/** Contains default settings. */
 	static FInsightsSettings Defaults;
 
-	/// Setting filename ini.
+	/** Setting filename ini. */
 	FString SettingsIni;
 
-	/// Whether profiler settings is in edit mode.
+	/** Whether profiler settings is in edit mode. */
 	bool bIsEditing;
 
-	/// Whether this instance contains defaults.
+	/** Whether this instance contains defaults. */
 	bool bIsDefault;
 
 	//////////////////////////////////////////////////
 	// Actual settings.
 
-	/// If True, the Timing View will show empty tracks by default.
+	/** If True, the Timing View will show empty tracks by default. */
 	bool bShowEmptyTracksByDefault;
 };
diff --git a/Engine/Source/Developer/TraceInsights/Private/Insights/IoProfilerCommands.cpp b/Engine/Source/Developer/TraceInsights/Private/Insights/IoProfilerCommands.cpp
--- a/Engine/Source/Developer/TraceInsights/Private/Insights/IoProfilerCommands.cpp
+++ b/Engine/Source/Developer/TraceInsights/Private/Insights/IoProfilerCommands.cpp
@@ -39,7 +39,7 @@ void FIoProfilerMenuBuilder::AddMenuEntry(FMenuBuilder& MenuBuilder, const TShar
 FIoProfilerCommands::FIoProfilerCommands()
 	: TCommands<FIoProfilerCommands>(
 		TEXT("IoProfilerCommand"), // Context name for fast lookup
-		NSLOCTEXT("Contexts", "IoProfilerCommand", "Timing Profiler Command"), // Localized context name for displaying
+		NSLOCTEXT("Contexts", "IoProfilerCommand", "Asset Loading Insights Command"), // Localized context name for displaying
 		NAME_None, // Parent
 		FEditorStyle::GetStyleSetName() // Icon Style Set
 	)
@@ -48,15 +48,11 @@ FIoProfilerCommands::FIoProfilerCommands()
 
 ////////////////////////////////////////////////////////////////////////////////////////////////////
 
-// UI_COMMAND takes long for the compile to optimize
+// UI_COMMAND takes long for the compiler to optimize
 PRAGMA_DISABLE_OPTIMIZATION
 void FIoProfilerCommands::RegisterCommands()
 {
-	UI_COMMAND(ToggleFramesTrackVisibility, "Frames", "Toggles the visibility of the Frames track", EUserInterfaceActionType::ToggleButton, FInputChord(EModifierKey::Control, EKeys::F));
-	UI_COMMAND(ToggleGraphTrackVisibility, "Graph", "Toggles the visibility of the Overview Graph track", EUserInterfaceActionType::ToggleButton, FInputChord(EModifierKey::Control, EKeys::G));
-	UI_COMMAND(ToggleTimingTrackVisibility, "Timing", "Toggles the visibility of the main Timing view", EUserInterfaceActionType::ToggleButton, FInputChord(EModifierKey::Control, EKeys::T));
-	UI_COMMAND(ToggleTimersViewVisibility, "Timers", "Toggles the visibility of the Timers view", EUserInterfaceActionType::ToggleButton, FInputChord(EKeys::T));
-	UI_COMMAND(ToggleLogViewVisibility, "Log", "Toggles the visibility of the Log view", EUserInterfaceActionType::ToggleButton, FInputChord(EKeys::L));
+	UI_COMMAND(ToggleTimingViewVisibility, "Timing", "Toggles the visibility of the main Timing view", EUserInterfaceActionType::ToggleButton, FInputChord(EModifierKey::Control, EKeys::T));
 }
 PRAGMA_ENABLE_OPTIMIZATION
 
@@ -97,11 +93,7 @@ PRAGMA_ENABLE_OPTIMIZATION
 		return b##IsEnabled ? ECheckBoxState::Checked : ECheckBoxState::Unchecked;\
 	}
 
-IMPLEMENT_TOGGLE_COMMAND(ToggleFramesTrackVisibility, IsFramesTrackVisible, SetFramesTrackVisible)
-IMPLEMENT_TOGGLE_COMMAND(ToggleGraphTrackVisibility, IsGraphTrackVisible, SetGraphTrackVisible)
-IMPLEMENT_TOGGLE_COMMAND(ToggleTimingTrackVisibility, IsTimingTrackVisible, SetTimingTrackVisible)
-IMPLEMENT_TOGGLE_COMMAND(ToggleTimersViewVisibility, IsTimersViewVisible, SetTimersViewVisible)
-IMPLEMENT_TOGGLE_COMMAND(ToggleLogViewVisibility, IsLogViewVisible, SetLogViewVisible)
+IMPLEMENT_TOGGLE_COMMAND(ToggleTimingViewVisibility, IsTimingViewVisible, SetTimingViewVisible)
 
 #undef IMPLEMENT_TOGGLE_COMMAND
 
diff --git a/Engine/Source/Developer/TraceInsights/Private/Insights/IoProfilerCommands.h b/Engine/Source/Developer/TraceInsights/Private/Insights/IoProfilerCommands.h
--- a/Engine/Source/Developer/TraceInsights/Private/Insights/IoProfilerCommands.h
+++ b/Engine/Source/Developer/TraceInsights/Private/Insights/IoProfilerCommands.h
@@ -9,16 +9,16 @@
 class FMenuBuilder;
 
 ////////////////////////////////////////////////////////////////////////////////////////////////////
-/// Class that holds all profiler commands.
-
-class FIoProfilerCommands
-	: public TCommands<FIoProfilerCommands>
+/**
+ * Class that holds all profiler commands.
+ */
+class FIoProfilerCommands : public TCommands<FIoProfilerCommands>
 {
 public:
-	/// Default constructor.
+	/** Default constructor. */
 	FIoProfilerCommands();
 
-	/// Initialize commands.
+	/** Initialize commands. */
 	virtual void RegisterCommands() override;
 
 public:
@@ -29,51 +29,43 @@ public:
 	//     const FUIAction <CommandName>_Custom(...) const;
 	//////////////////////////////////////////////////
 
-	/// Toggles visibility for the Frames Track. Global and custom command.
-	TSharedPtr<FUICommandInfo> ToggleFramesTrackVisibility;
-
-	/// Toggles visibility for the Graph Track. Global and custom command.
-	TSharedPtr<FUICommandInfo> ToggleGraphTrackVisibility;
-
-	/// Toggles visibility for the Timing Track. Global and custom command.
-	TSharedPtr<FUICommandInfo> ToggleTimingTrackVisibility;
-
-	/// Toggles visibility for the Timers View. Global and custom command.
-	TSharedPtr<FUICommandInfo> ToggleTimersViewVisibility;
-
-	/// Toggles visibility for the Log View. Global and custom command.
-	TSharedPtr<FUICommandInfo> ToggleLogViewVisibility;
+	/** Toggles visibility for the Timing View. Global and custom command. */
+	TSharedPtr<FUICommandInfo> ToggleTimingViewVisibility;
 };
 
 ////////////////////////////////////////////////////////////////////////////////////////////////////
-
+/**
+ * Menu builder. Helper class for adding a customized menu entry using the global UI command info.
+ */
 class FIoProfilerMenuBuilder
 {
 public:
-	/// Helper method for adding a customized menu entry using the global UI command info.
-	/// FUICommandInfo cannot be executed with custom parameters, so we need to create a custom FUIAction,
-	/// but sometime we have global and local version for the UI command, so reuse data from the global UI command info.
-	/// Ex:
-	///     SessionInstance_ToggleCapture          - Global version will toggle capture process for all active session instances
-	///     SessionInstance_ToggleCapture_OneParam - Local version will toggle capture process only for the specified session instance
-	///
-	/// @param MenuBuilder The menu to add items to
-	/// @param FUICommandInfo A shared pointer to the UI command info
-	/// @param UIAction Customized version of the UI command info stored in an UI action
-	///
+	/**
+	 * Helper method for adding a customized menu entry using the global UI command info.
+	 * FUICommandInfo cannot be executed with custom parameters, so we need to create a custom FUIAction,
+	 * but sometime we have global and local version for the UI command, so reuse data from the global UI command info.
+	 * Ex:
+	 *     SessionInstance_ToggleCapture          - Global version will toggle capture process for all active session instances
+	 *     SessionInstance_ToggleCapture_OneParam - Local version will toggle capture process only for the specified session instance
+	 *
+	 * @param MenuBuilder The menu to add items to
+	 * @param FUICommandInfo A shared pointer to the UI command info
+	 * @param UIAction Customized version of the UI command info stored in an UI action
+	 */
 	static void AddMenuEntry(FMenuBuilder& MenuBuilder, const TSharedPtr<FUICommandInfo>& UICommandInfo, const FUIAction& UIAction);
 };
 
 ////////////////////////////////////////////////////////////////////////////////////////////////////
-/// Class that provides helper functions for the commands to avoid cluttering profiler manager with many small functions.
-/// Can't contain any variables. Directly operates on the profiler manager instance.
-
+/**
+ * Class that provides helper functions for the commands to avoid cluttering profiler manager with many small functions.
+ * Can't contain any variables. Directly operates on the profiler manager instance.
+ */
 class FIoProfilerActionManager
 {
 	friend class FIoProfilerManager;
 
 private:
-	/// Private constructor.
+	/** Private constructor. */
 	FIoProfilerActionManager(class FIoProfilerManager* Instance)
 		: This(Instance)
 	{}
@@ -90,26 +82,12 @@ protected:\
 	bool CmdName##_CanExecute() const; /**< Handles FCanExecuteAction for CmdName. */\
 	ECheckBoxState CmdName##_GetCheckState() const; /**< Handles FGetActionCheckState for CmdName. */
 
-	DECLARE_TOGGLE_COMMAND(ToggleFramesTrackVisibility)
-	DECLARE_TOGGLE_COMMAND(ToggleGraphTrackVisibility)
-	DECLARE_TOGGLE_COMMAND(ToggleTimingTrackVisibility)
-	DECLARE_TOGGLE_COMMAND(ToggleTimersViewVisibility)
-	DECLARE_TOGGLE_COMMAND(ToggleLogViewVisibility)
+	DECLARE_TOGGLE_COMMAND(ToggleTimingViewVisibility)
 #undef DECLARE_TOGGLE_COMMAND
 
 	//////////////////////////////////////////////////
-	// OpenSettings
-
-public:
-	void Map_OpenSettings_Global(); ///< Maps UI command info OpenSettings with the specified UI command list.
-	const FUIAction OpenSettings_Custom() const; ///< UI action for OpenSettings command.
-protected:
-	void OpenSettings_Execute(); ///< Handles FExecuteAction for OpenSettings.
-	bool OpenSettings_CanExecute() const; ///< Handles FCanExecuteAction for OpenSettings.
-
-	//////////////////////////////////////////////////
 
 protected:
-	/// Reference to the global instance of the profiler manager.
+	/** Reference to the global instance of the profiler manager. */
 	class FIoProfilerManager* This;
 };
diff --git a/Engine/Source/Developer/TraceInsights/Private/Insights/IoProfilerManager.cpp b/Engine/Source/Developer/TraceInsights/Private/Insights/IoProfilerManager.cpp
--- a/Engine/Source/Developer/TraceInsights/Private/Insights/IoProfilerManager.cpp
+++ b/Engine/Source/Developer/TraceInsights/Private/Insights/IoProfilerManager.cpp
@@ -11,14 +11,9 @@
 // Insights
 #include "Insights/InsightsManager.h"
 //#include "Insights/IoProfilerCommon.h"
-#include "Insights/Widgets/SFrameTrack.h"
-#include "Insights/Widgets/SGraphTrack.h"
-#include "Insights/Widgets/SLogView.h"
-#include "Insights/Widgets/STimersView.h"
 #include "Insights/Widgets/STimingView.h"
 #include "Insights/Widgets/SIoProfilerWindow.h"
 
-
 ////////////////////////////////////////////////////////////////////////////////////////////////////
 
 #define LOCTEXT_NAMESPACE "IoProfilerManager"
@@ -38,11 +33,7 @@ FIoProfilerManager::FIoProfilerManager(TSharedRef<FUICommandList> InCommandList)
 	: CommandList(InCommandList)
 	, ActionManager(this)
 	, ProfilerWindow(nullptr)
-	, bIsFramesTrackVisible(true)
-	, bIsGraphTrackVisible(false)
-	, bIsTimingTrackVisible(true)
-	, bIsTimersViewVisible(true)
-	, bIsLogViewVisible(true)
+	, bIsTimingViewVisible(true)
 {
 }
 
@@ -62,11 +53,7 @@ void FIoProfilerManager::PostConstructor()
 
 void FIoProfilerManager::BindCommands()
 {
-	ActionManager.Map_ToggleFramesTrackVisibility_Global();
-	ActionManager.Map_ToggleGraphTrackVisibility_Global();
-	ActionManager.Map_ToggleTimingTrackVisibility_Global();
-	ActionManager.Map_ToggleTimersViewVisibility_Global();
-	ActionManager.Map_ToggleLogViewVisibility_Global();
+	ActionManager.Map_ToggleTimingViewVisibility_Global();
 }
 
 ////////////////////////////////////////////////////////////////////////////////////////////////////
@@ -123,30 +110,28 @@ void FIoProfilerManager::OnSessionChanged()
 	TSharedPtr<SIoProfilerWindow> Wnd = GetProfilerWindow();
 	if (Wnd.IsValid())
 	{
-		Wnd->FrameTrack->Reset();
-		Wnd->GraphTrack->Reset();
-		Wnd->TimingView->Reset();
-		Wnd->TimersView->RebuildTree();
-		Wnd->LogView->Reset();
+		TSharedPtr<STimingView> TimingView = Wnd->GetTimingView();
+		if (TimingView)
+		{
+			TimingView->Reset();
+		}
 	}
 }
 
 ////////////////////////////////////////////////////////////////////////////////////////////////////
 
-void FIoProfilerManager::SetTimersViewVisible(const bool bTimersViewVisibleState)
+void FIoProfilerManager::SetTimingViewVisible(const bool bTimingViewVisibleState)
 {
-	bool bWasTimersViewVisible = bIsTimersViewVisible;
+	//bool bWasTimingViewVisible = bIsTimingViewVisible;
 
-	bIsTimersViewVisible = bTimersViewVisibleState;
+	bIsTimingViewVisible = bTimingViewVisibleState;
 
-	if (bIsTimersViewVisible && !bWasTimersViewVisible && FInsightsManager::Get()->GetSession().IsValid())
-	{
-		TSharedPtr<SIoProfilerWindow> Wnd = GetProfilerWindow();
-		if (Wnd.IsValid())
-		{
-			Wnd->TimersView->RebuildTree();
-		}
-	}
+	// TODO
+	//TSharedPtr<SIoProfilerWindow> Wnd = GetProfilerWindow();
+	//if (Wnd.IsValid())
+	//{
+	//	Wnd->ShowHideTab(FIoProfilerTabs::TimingViewID, bIsTimingViewVisible);
+	//}
 }
 
 ////////////////////////////////////////////////////////////////////////////////////////////////////
diff --git a/Engine/Source/Developer/TraceInsights/Private/Insights/IoProfilerManager.h b/Engine/Source/Developer/TraceInsights/Private/Insights/IoProfilerManager.h
--- a/Engine/Source/Developer/TraceInsights/Private/Insights/IoProfilerManager.h
+++ b/Engine/Source/Developer/TraceInsights/Private/Insights/IoProfilerManager.h
@@ -20,21 +20,21 @@ namespace Trace
 }
 
 ////////////////////////////////////////////////////////////////////////////////////////////////////
-/// This class manages the Timing Profiler state and settings.
-
-class FIoProfilerManager
-	: public TSharedFromThis<FIoProfilerManager>
+/**
+ * This class manages the Timing Profiler state and settings.
+ */
+class FIoProfilerManager : public TSharedFromThis<FIoProfilerManager>
 {
 	friend class FIoProfilerActionManager;
 
 public:
-	/// Creates the Timing Profiler manager, only one instance can exist.
+	/** Creates the Timing Profiler manager, only one instance can exist. */
 	FIoProfilerManager(TSharedRef<FUICommandList> InCommandList);
 
-	/// Virtual destructor.
+	/** Virtual destructor. */
 	virtual ~FIoProfilerManager();
 
-	/// Creates an instance of the profiler manager.
+	/** Creates an instance of the profiler manager. */
 	static TSharedPtr<FIoProfilerManager> Initialize()
 	{
 		if (FIoProfilerManager::Instance.IsValid())
@@ -48,34 +48,36 @@ public:
 		return FIoProfilerManager::Instance;
 	}
 
-	/// Shutdowns the Timing Profiler manager.
+	/** Shutdowns the Timing Profiler manager. */
 	void Shutdown()
 	{
 		FIoProfilerManager::Instance.Reset();
 	}
 
 protected:
-	/// Finishes initialization of the profiler manager
+	/** Finishes initialization of the profiler manager. */
 	void PostConstructor();
 
-	/// Binds our UI commands to delegates.
+	/** Binds our UI commands to delegates. */
 	void BindCommands();
 
 public:
-	/// @return the global instance of the Timing Profiler manager (FIoProfilerManager).
-	/// This is an internal singleton and cannot be used outside ProfilerModule.
-	/// For external use:
-	///     IProfilerModule& ProfilerModule = FModuleManager::Get().LoadModuleChecked<IProfilerModule>("Profiler");
-	///     ProfilerModule.GetProfilerManager();
+	/**
+	 * @return the global instance of the Timing Profiler manager (FIoProfilerManager).
+	 * This is an internal singleton and cannot be used outside ProfilerModule.
+	 * For external use:
+	 *     IProfilerModule& ProfilerModule = FModuleManager::Get().LoadModuleChecked<IProfilerModule>("Profiler");
+	 *     ProfilerModule.GetProfilerManager();
+	 */
 	static TSharedPtr<FIoProfilerManager> Get();
 
-	/// @returns UI command list for the Timing Profiler manager.
+	/** @returns UI command list for the Timing Profiler manager. */
 	const TSharedRef<FUICommandList> GetCommandList() const;
 
-	/// @return an instance of the Timing Profiler commands.
+	/** @return an instance of the Timing Profiler commands. */
 	static const FIoProfilerCommands& GetCommands();
 
-	/// @return an instance of the Timing Profiler action manager.
+	/** @return an instance of the Timing Profiler action manager. */
 	static FIoProfilerActionManager& GetActionManager();
 
 	void AssignProfilerWindow(const TSharedRef<SIoProfilerWindow>& InProfilerWindow)
@@ -83,8 +85,10 @@ public:
 		ProfilerWindow = InProfilerWindow;
 	}
 
-	/// Converts profiler window weak pointer to a shared pointer and returns it.
-	/// Make sure the returned pointer is valid before trying to dereference it.
+	/**
+	 * Converts profiler window weak pointer to a shared pointer and returns it.
+	 * Make sure the returned pointer is valid before trying to dereference it.
+	 */
 	TSharedPtr<class SIoProfilerWindow> GetProfilerWindow() const
 	{
 		return ProfilerWindow.Pin();
@@ -93,65 +97,37 @@ public:
 	////////////////////////////////////////////////////////////////////////////////////////////////////
 	// Getters and setters used by Toggle Commands.
 
-	/// @return true, if Frames Track is visible
-	const bool IsFramesTrackVisible() const { return bIsFramesTrackVisible; }
-	void SetFramesTrackVisible(const bool bFramesTrackVisibleState) { bIsFramesTrackVisible = bFramesTrackVisibleState; }
-
-	/// @return true, if Graph Track is visible
-	const bool IsGraphTrackVisible() const { return bIsGraphTrackVisible; }
-	void SetGraphTrackVisible(const bool bGraphTrackVisibleState) { bIsGraphTrackVisible = bGraphTrackVisibleState; }
-
-	/// @return true, if Timing Track is visible
-	const bool IsTimingTrackVisible() const { return bIsTimingTrackVisible; }
-	void SetTimingTrackVisible(const bool bTimingTrackVisibleState) { bIsTimingTrackVisible = bTimingTrackVisibleState; }
-
-	/// @return true, if Timers View is visible
-	const bool IsTimersViewVisible() const { return bIsTimersViewVisible; }
-	void SetTimersViewVisible(const bool bTimersViewVisibleState);
-
-	/// @return true, if Log View is visible
-	const bool IsLogViewVisible() const { return bIsLogViewVisible; }
-	void SetLogViewVisible(const bool bLogViewVisibleState) { bIsLogViewVisible = bLogViewVisibleState; }
+	/** @return true, if Timing View is visible */
+	const bool IsTimingViewVisible() const { return bIsTimingViewVisible; }
+	void SetTimingViewVisible(const bool bTimingViewVisibleState);
 
 	////////////////////////////////////////////////////////////////////////////////////////////////////
 
 	void OnSessionChanged();
 
 protected:
-	/// Updates this manager, done through FCoreTicker.
+	/** Updates this manager, done through FCoreTicker. */
 	bool Tick(float DeltaTime);
 
 protected:
-	/// The delegate to be invoked when this manager ticks.
+	/** The delegate to be invoked when this manager ticks. */
 	FTickerDelegate OnTick;
 
-	/// Handle to the registered OnTick.
+	/** Handle to the registered OnTick. */
 	FDelegateHandle OnTickHandle;
 
-	/// List of UI commands for this manager. This will be filled by this and corresponding classes.
+	/** List of UI commands for this manager. This will be filled by this and corresponding classes. */
 	TSharedRef<FUICommandList> CommandList;
 
-	/// An instance of the Timing Manager action manager.
+	/** An instance of the Timing Manager action manager. */
 	FIoProfilerActionManager ActionManager;
 
-	/// A weak pointer to the Timing Profiler window.
+	/** A weak pointer to the Timing Profiler window. */
 	TWeakPtr<class SIoProfilerWindow> ProfilerWindow;
 
-	/// If the Frames Track is visible or hidden.
-	bool bIsFramesTrackVisible;
-
-	/// If the Graph Track is visible or hidden.
-	bool bIsGraphTrackVisible;
-
-	/// If the Timing Track is visible or hidden.
-	bool bIsTimingTrackVisible;
-
-	/// If the Timers View is visible or hidden.
-	bool bIsTimersViewVisible;
-
-	/// If the Log View is visible or hidden.
-	bool bIsLogViewVisible;
+	/** If the Timing View is visible or hidden. */
+	bool bIsTimingViewVisible;
 
-	/// A shared pointer to the global instance of the profiler manager.
+	/** A shared pointer to the global instance of the profiler manager. */
 	static TSharedPtr<FIoProfilerManager> Instance;
 };
diff --git a/Engine/Source/Developer/TraceInsights/Private/Insights/TimingProfilerCommands.cpp b/Engine/Source/Developer/TraceInsights/Private/Insights/TimingProfilerCommands.cpp
--- a/Engine/Source/Developer/TraceInsights/Private/Insights/TimingProfilerCommands.cpp
+++ b/Engine/Source/Developer/TraceInsights/Private/Insights/TimingProfilerCommands.cpp
@@ -48,7 +48,7 @@ FTimingProfilerCommands::FTimingProfilerCommands()
 
 ////////////////////////////////////////////////////////////////////////////////////////////////////
 
-// UI_COMMAND takes long for the compile to optimize
+// UI_COMMAND takes long for the compiler to optimize
 PRAGMA_DISABLE_OPTIMIZATION
 void FTimingProfilerCommands::RegisterCommands()
 {
diff --git a/Engine/Source/Developer/TraceInsights/Private/Insights/TimingProfilerCommands.h b/Engine/Source/Developer/TraceInsights/Private/Insights/TimingProfilerCommands.h
--- a/Engine/Source/Developer/TraceInsights/Private/Insights/TimingProfilerCommands.h
+++ b/Engine/Source/Developer/TraceInsights/Private/Insights/TimingProfilerCommands.h
@@ -12,8 +12,7 @@ class FMenuBuilder;
 /**
  * Class that holds all profiler commands.
  */
-class FTimingProfilerCommands
-	: public TCommands<FTimingProfilerCommands>
+class FTimingProfilerCommands : public TCommands<FTimingProfilerCommands>
 {
 public:
 	/** Default constructor. */
@@ -81,7 +80,7 @@ class FTimingProfilerActionManager
 	friend class FTimingProfilerManager;
 
 private:
-	/// Private constructor.
+	/** Private constructor. */
 	FTimingProfilerActionManager(class FTimingProfilerManager* Instance)
 		: This(Instance)
 	{}
diff --git a/Engine/Source/Developer/TraceInsights/Private/Insights/TimingProfilerManager.h b/Engine/Source/Developer/TraceInsights/Private/Insights/TimingProfilerManager.h
--- a/Engine/Source/Developer/TraceInsights/Private/Insights/TimingProfilerManager.h
+++ b/Engine/Source/Developer/TraceInsights/Private/Insights/TimingProfilerManager.h
@@ -23,8 +23,7 @@ namespace Trace
 /**
  * This class manages the Timing Profiler state and settings.
  */
-class FTimingProfilerManager
-	: public TSharedFromThis<FTimingProfilerManager>
+class FTimingProfilerManager : public TSharedFromThis<FTimingProfilerManager>
 {
 	friend class FTimingProfilerActionManager;
 
@@ -56,7 +55,7 @@ public:
 	}
 
 protected:
-	/** Finishes initialization of the profiler manager */
+	/** Finishes initialization of the profiler manager. */
 	void PostConstructor();
 
 	/** Binds our UI commands to delegates. */
diff --git a/Engine/Source/Developer/TraceInsights/Private/Insights/TraceInsightsModule.cpp b/Engine/Source/Developer/TraceInsights/Private/Insights/TraceInsightsModule.cpp
--- a/Engine/Source/Developer/TraceInsights/Private/Insights/TraceInsightsModule.cpp
+++ b/Engine/Source/Developer/TraceInsights/Private/Insights/TraceInsightsModule.cpp
@@ -95,7 +95,7 @@ void FTraceInsightsModule::StartupModule()
 	// Register tab spawner for the Start Page.
 	auto& StartPageTabSpawnerEntry = FGlobalTabmanager::Get()->RegisterTabSpawner(StartPageTabName,
 		FOnSpawnTab::CreateRaw(this, &FTraceInsightsModule::SpawnStartPageTab))
-		.SetDisplayName(NSLOCTEXT("FTraceInsightsModule", "StartPageTabTitle", "Start Page"))
+		.SetDisplayName(NSLOCTEXT("FTraceInsightsModule", "StartPageTabTitle", "Unreal Insights"))
 		.SetTooltipText(NSLOCTEXT("FTraceInsightsModule", "StartPageTooltipText", "Open the start page for Unreal Insights."))
 		.SetIcon(FSlateIcon(FInsightsStyle::GetStyleSetName(), "StartPage.Icon.Small"));
 
@@ -155,10 +155,10 @@ void FTraceInsightsModule::OnNewLayout(TSharedRef<FTabManager::FLayout> NewLayou
 		(
 			FTabManager::NewStack()
 			->AddTab(StartPageTabName, ETabState::OpenedTab)
-			->AddTab(TimingProfilerTabName, ETabState::OpenedTab)
-			->AddTab(IoProfilerTabName, ETabState::OpenedTab)
-			//->AddTab(FName("MemoryProfiler"), ETabState::OpenedTab)
-			->SetForegroundTab(FTabId(TimingProfilerTabName))
+			->AddTab(TimingProfilerTabName, ETabState::ClosedTab)
+			->AddTab(IoProfilerTabName, ETabState::ClosedTab)
+			//->AddTab(FName("MemoryProfiler"), ETabState::ClosedTab)
+			->SetForegroundTab(FTabId(StartPageTabName))
 		)
 	);
 }
diff --git a/Engine/Source/Developer/TraceInsights/Private/Insights/ViewModels/LogFilter.cpp b/Engine/Source/Developer/TraceInsights/Private/Insights/ViewModels/LogFilter.cpp
--- a/Engine/Source/Developer/TraceInsights/Private/Insights/ViewModels/LogFilter.cpp
+++ b/Engine/Source/Developer/TraceInsights/Private/Insights/ViewModels/LogFilter.cpp
@@ -9,8 +9,9 @@
 #include "Insights/Widgets/SLogView.h"
 
 ////////////////////////////////////////////////////////////////////////////////////////////////////
-/// Expression context to test the given messages against the current text filter
-
+/**
+ * Expression context to test the given messages against the current text filter.
+ */
 class FLogFilter_TextFilterExpressionContext : public ITextFilterExpressionContext
 {
 public:
diff --git a/Engine/Source/Developer/TraceInsights/Private/Insights/Widgets/SFrameTrack.h b/Engine/Source/Developer/TraceInsights/Private/Insights/Widgets/SFrameTrack.h
--- a/Engine/Source/Developer/TraceInsights/Private/Insights/Widgets/SFrameTrack.h
+++ b/Engine/Source/Developer/TraceInsights/Private/Insights/Widgets/SFrameTrack.h
@@ -34,34 +34,32 @@ struct FSampleRef
 };
 
 ////////////////////////////////////////////////////////////////////////////////////////////////////
-/// Widget used to present frames data in a track.
-
+/**
+ * Widget used to present frames data in a track.
+ */
 class SFrameTrack : public SCompoundWidget
 {
 	enum
 	{
-		/// Number of pixels.
+		/** Number of pixels. */
 		MOUSE_SNAP_DISTANCE = 4,
 	};
 
-	struct EFrameTrackCursor
+	enum class EFrameTrackCursor
 	{
-		enum Type
-		{
-			Default,
-			Arrow,
-			Hand,
-		};
+		Default,
+		Arrow,
+		Hand,
 	};
 
 public:
-	/// Default constructor.
+	/** Default constructor. */
 	SFrameTrack();
 
-	/// Virtual destructor.
+	/** Virtual destructor. */
 	virtual ~SFrameTrack();
 
-	/// Resets internal widget's data to the default one.
+	/** Resets internal widget's data to the default one. */
 	void Reset();
 
 	SLATE_BEGIN_ARGS(SFrameTrack)
@@ -70,14 +68,18 @@ public:
 	}
 	SLATE_END_ARGS()
 
-	/// Construct this widget
-	/// @param InArgs The declaration data for this widget
+	/**
+	 * Construct this widget
+	 * @param InArgs The declaration data for this widget
+	 */
 	void Construct(const FArguments& InArgs);
 
-	/// Ticks this widget. Override in derived classes, but always call the parent implementation.
-	/// @param  AllottedGeometry The space allotted for this widget
-	/// @param  InCurrentTime  Current absolute real time
-	/// @param  InDeltaTime  Real time passed since last tick
+	/**
+	 * Ticks this widget. Override in derived classes, but always call the parent implementation.
+	 * @param  AllottedGeometry The space allotted for this widget
+	 * @param  InCurrentTime  Current absolute real time
+	 * @param  InDeltaTime  Real time passed since last tick
+	 */
 	virtual void Tick(const FGeometry& AllottedGeometry, const double InCurrentTime, const float InDeltaTime) override;
 
 	virtual int32 OnPaint(const FPaintArgs& Args, const FGeometry& AllottedGeometry, const FSlateRect& MyCullingRect, FSlateWindowElementList& OutDrawElements, int32 LayerId, const FWidgetStyle& InWidgetStyle, bool bParentEnabled) const override;
@@ -93,11 +95,6 @@ public:
 	virtual FCursorReply OnCursorQuery(const FGeometry& MyGeometry, const FPointerEvent& CursorEvent) const override;
 
 protected:
-	//virtual FVector2D ComputeDesiredSize(float) const override
-	//{
-	//	return FVector2D(16.0f, 16.0f);
-	//}
-
 	bool IsReady() { return true; }
 
 	void UpdateState();
@@ -107,11 +104,13 @@ protected:
 
 	void ShowContextMenu(const FVector2D& ScreenSpacePosition);
 
-	/// Binds our UI commands to delegates.
+	/** Binds our UI commands to delegates. */
 	void BindCommands();
 
-	/// Called when the user scrolls the horizontal scrollbar.
-	/// @param ScrollOffset Scroll offset as a fraction between 0 and 1.
+	/**
+	 * Called when the user scrolls the horizontal scrollbar.
+	 * @param ScrollOffset Scroll offset as a fraction between 0 and 1.
+	 */
 	void HorizontalScrollBar_OnUserScrolled(float ScrollOffset);
 	void UpdateHorizontalScrollBar();
 
@@ -121,21 +120,21 @@ protected:
 	// SelectionBoxChanged Event
 
 public:
-	/// The event to execute when the selection box has been changed.
+	/** The event to execute when the selection box has been changed. */
 	DECLARE_EVENT_TwoParams(SFrameTrack, FSelectionBoxChangedEvent, int32 /*FrameStart*/, int32 /*FrameEnd*/);
 	FSelectionBoxChangedEvent& OnSelectionBoxChanged() { return SelectionBoxChangedEvent; }
 protected:
-	/// The event to execute when the selection box has been changed.
+	/** The event to execute when the selection box has been changed. */
 	FSelectionBoxChangedEvent SelectionBoxChangedEvent;
 
 	//////////////////////////////////////////////////
 
 protected:
-	/// The track's viewport. Encapsulates info about position and scale.
+	/** The track's viewport. Encapsulates info about position and scale. */
 	FFrameTrackViewport Viewport;
 	bool bIsViewportDirty;
 
-	/// Cached info for timelines.
+	/** Cached info for timelines. */
 	TMap<uint64, FFrameTrackTimeline> CachedTimelines;
 	TArray<int32> TimelinesOrder;
 
@@ -150,21 +149,21 @@ protected:
 	//////////////////////////////////////////////////
 	// Panning, Zooming and Selection behaviors
 
-	/// The current mouse position.
+	/** The current mouse position. */
 	FVector2D MousePosition;
 
-	/// Mouse position during the call on mouse button down.
+	/** Mouse position during the call on mouse button down. */
 	FVector2D MousePositionOnButtonDown;
 	float ViewportPosXOnButtonDown;
 	float ViewportPosYOnButtonDown;
 
-	/// Mouse position during the call on mouse button up.
+	/** Mouse position during the call on mouse button up. */
 	FVector2D MousePositionOnButtonUp;
 
 	bool bIsLMB_Pressed;
 	bool bIsRMB_Pressed;
 
-	/// True, if the user is currently interactively scrolling the view by holding the right mouse button and dragging.
+	/** True, if the user is currently interactively scrolling the view by holding the right mouse button and dragging. */
 	bool bIsRMB_Scrolling;
 
 	//////////////////////////////////////////////////
@@ -179,8 +178,8 @@ protected:
 
 	FGeometry ThisGeometry;
 
-	/// Cursor type.
-	EFrameTrackCursor::Type CursorType;
+	/** Cursor type. */
+	EFrameTrackCursor CursorType;
 
 	// Debug stats
 	int32 NumUpdatedFrames;
diff --git a/Engine/Source/Developer/TraceInsights/Private/Insights/Widgets/SInsightsSettings.h b/Engine/Source/Developer/TraceInsights/Private/Insights/Widgets/SInsightsSettings.h
--- a/Engine/Source/Developer/TraceInsights/Private/Insights/Widgets/SInsightsSettings.h
+++ b/Engine/Source/Developer/TraceInsights/Private/Insights/Widgets/SInsightsSettings.h
@@ -12,7 +12,7 @@ class FInsightsSettings;
 class SGridPanel;
 enum class ECheckBoxState : uint8;
 
-/// Widget used to modify settings for the profiler, created on demand and destroyed on close.
+/** Widget used to modify settings for the profiler, created on demand and destroyed on close. */
 class SInsightsSettings : public SCompoundWidget
 {
 public:
@@ -21,8 +21,11 @@ public:
 		SLATE_ARGUMENT(FInsightsSettings*, SettingPtr)
 	SLATE_END_ARGS()
 
-	///  Construct this widget
-	/// @param InArgs The declaration data for this widget.
+	/**
+	 * Construct this widget
+	 *
+	 * @param InArgs The declaration data for this widget.
+	 */
 	void Construct(const FArguments& InArgs);
 
 private:
@@ -42,9 +45,9 @@ private:
 	ECheckBoxState OptionValue_IsChecked(const bool* ValuePtr) const;
 
 private:
-	/// Delegate to call when this profiler settings is closed.
+	/** Delegate to call when this profiler settings is closed. */
 	FSimpleDelegate OnClose;
 
-	/// Pointer to the profiler settings.
+	/** Pointer to the profiler settings. */
 	FInsightsSettings* SettingPtr;
 };
diff --git a/Engine/Source/Developer/TraceInsights/Private/Insights/Widgets/SIoProfilerWindow.cpp b/Engine/Source/Developer/TraceInsights/Private/Insights/Widgets/SIoProfilerWindow.cpp
--- a/Engine/Source/Developer/TraceInsights/Private/Insights/Widgets/SIoProfilerWindow.cpp
+++ b/Engine/Source/Developer/TraceInsights/Private/Insights/Widgets/SIoProfilerWindow.cpp
@@ -8,8 +8,10 @@
 #include "Widgets/Layout/SBorder.h"
 #include "Widgets/Layout/SBox.h"
 #include "Widgets/Layout/SSpacer.h"
+#include "Widgets/Layout/SSplitter.h"
 #include "Widgets/Notifications/SNotificationList.h"
 #include "Widgets/SBoxPanel.h"
+#include "Widgets/SOverlay.h"
 #include "Widgets/Text/STextBlock.h"
 
 #if WITH_EDITOR
@@ -22,12 +24,6 @@
 #include "Insights/Version.h"
 #include "Insights/InsightsManager.h"
 #include "Insights/IoProfilerManager.h"
-#include "Insights/Widgets/SFrameTrack.h"
-#include "Insights/Widgets/SGraphTrack.h"
-#include "Insights/Widgets/SInsightsSettings.h"
-#include "Insights/Widgets/SLogView.h"
-//#include "Insights/Widgets/SIoProfilerToolbar.h"
-#include "Insights/Widgets/STimersView.h"
 #include "Insights/Widgets/STimingView.h"
 
 ////////////////////////////////////////////////////////////////////////////////////////////////////
@@ -36,45 +32,6 @@
 
 ////////////////////////////////////////////////////////////////////////////////////////////////////
 
-// TODO: move this function to InsightsManager or in TraceSession.h
-static FText Io_GetTextForNotification(const EInsightsNotificationType NotificatonType, const ELoadingProgressState ProgressState, const FString& Filename, const float ProgressPercent = 0.0f)
-{
-	FText Result;
-
-	if (NotificatonType == EInsightsNotificationType::LoadingTraceFile)
-	{
-		if (ProgressState == ELoadingProgressState::Started)
-		{
-			FFormatNamedArguments Args;
-			Args.Add(TEXT("Filename"), FText::FromString(Filename));
-			Result = FText::Format(LOCTEXT("DescF_OfflineCapture_Started", "Started loading a file {Filename}"), Args);
-		}
-		else if (ProgressState == ELoadingProgressState::InProgress)
-		{
-			FFormatNamedArguments Args;
-			Args.Add(TEXT("Filename"), FText::FromString(Filename));
-			Args.Add(TEXT("LoadingProgressPercent"), FText::AsPercent(ProgressPercent));
-			Result = FText::Format(LOCTEXT("DescF_OfflineCapture_InProgress", "Loading a file {Filename} {LoadingProgressPercent}"), Args);
-		}
-		else if (ProgressState == ELoadingProgressState::Loaded)
-		{
-			FFormatNamedArguments Args;
-			Args.Add(TEXT("Filename"), FText::FromString(Filename));
-			Result = FText::Format(LOCTEXT("DescF_OfflineCapture_Loaded", "File {Filename} has been successfully loaded"), Args);
-		}
-		else if (ProgressState == ELoadingProgressState::Failed)
-		{
-			FFormatNamedArguments Args;
-			Args.Add(TEXT("Filename"), FText::FromString(Filename));
-			Result = FText::Format(LOCTEXT("DescF_OfflineCapture_Failed", "Failed to load file {Filename}"), Args);
-		}
-	}
-
-	return Result;
-}
-
-////////////////////////////////////////////////////////////////////////////////////////////////////
-
 SIoProfilerWindow::SIoProfilerWindow()
 	: DurationActive(0.0f)
 {
@@ -84,13 +41,6 @@ SIoProfilerWindow::SIoProfilerWindow()
 
 SIoProfilerWindow::~SIoProfilerWindow()
 {
-	// Remove ourselves from the profiler manager.
-	//if (FIoProfilerManager::Get().IsValid())
-	//{
-	//	//...
-	//	FIoProfilerManager::Get()->OnViewModeChanged().RemoveAll(this);
-	//}
-
 #if WITH_EDITOR
 	if (DurationActive > 0.0f && FEngineAnalytics::IsAvailable())
 	{
@@ -103,268 +53,6 @@ SIoProfilerWindow::~SIoProfilerWindow()
 
 BEGIN_SLATE_FUNCTION_BUILD_OPTIMIZATION
 
-TSharedRef<SWidget> SIoProfilerWindow::ConstructMultiTrackView()
-{
-	return SNew(SSplitter)
-		.Orientation(Orient_Vertical)
-
-	+ SSplitter::Slot()
-		.Value(0.1f)
-		[
-			ConstructFramesTrack()
-		]
-
-	+ SSplitter::Slot()
-		.Value(0.2f)
-		[
-			ConstructGraphTrack()
-		]
-
-	+ SSplitter::Slot()
-		.Value(0.7f)
-		[
-			ConstructTimingTrack()
-		]
-
-	; // end of Widget
-}
-
-////////////////////////////////////////////////////////////////////////////////////////////////////
-
-TSharedRef<SWidget> SIoProfilerWindow::ConstructFramesTrack()
-{
-	return SNew(SVerticalBox)
-		.Visibility(this, &SIoProfilerWindow::IsFramesTrackVisible)
-		.IsEnabled(this, &SIoProfilerWindow::IsProfilerEnabled)
-
-	// Frame Track
-	+ SVerticalBox::Slot()
-		.FillHeight(1.0)
-		.Padding(0.0f, 0.0f, 0.0f, 0.0f)
-		[
-			SNew(SBox)
-				.HeightOverride(48.0f)
-				.IsEnabled(this, &SIoProfilerWindow::IsProfilerEnabled)
-				[
-					SAssignNew(FrameTrack, SFrameTrack)
-
-					//SNew(SHorizontalBox)
-					//
-					//+ SHorizontalBox::Slot()
-					//	.FillWidth(1.0f)
-					//	.Padding(0.0f)
-					//	.HAlign(HAlign_Fill)
-					//	.VAlign(VAlign_Fill)
-					//	[
-					//
-					//	]
-				]
-		]
-
-	; // end of Widget
-}
-
-////////////////////////////////////////////////////////////////////////////////////////////////////
-
-TSharedRef<SWidget> SIoProfilerWindow::ConstructGraphTrack()
-{
-	return SNew(SVerticalBox)
-		.Visibility(this, &SIoProfilerWindow::IsGraphTrackVisible)
-		.IsEnabled(this, &SIoProfilerWindow::IsProfilerEnabled)
-
-	/*
-	// Header
-	+ SVerticalBox::Slot()
-		.AutoHeight()
-		.Padding(0.0f, 0.0f, 0.0f, 0.0f)
-		[
-			SNew(SHorizontalBox)
-
-			+ SHorizontalBox::Slot()
-				.AutoWidth()
-				[
-					SNew(SImage)
-						.Image(FEditorStyle::GetBrush(TEXT("Profiler.Tab.GraphView")))
-				]
-
-			+ SHorizontalBox::Slot()
-				.AutoWidth()
-				[
-					SNew(STextBlock)
-						.Text(LOCTEXT("GraphTrackLabel", "Overview Timing Graph"))
-						.Margin(FMargin(0.0f, 2.0f, 0.0f, 0.0f))
-				]
-		]
-	*/
-
-	// Graph Track
-	+ SVerticalBox::Slot()
-		.FillHeight(1.0)
-		.Padding(0.0f, 0.0f, 0.0f, 0.0f)
-		[
-			SNew(SBox)
-				.HeightOverride(48.0f)
-				.IsEnabled(this, &SIoProfilerWindow::IsProfilerEnabled)
-				[
-					SNew(SHorizontalBox)
-
-					+ SHorizontalBox::Slot()
-						.FillWidth(1.0f)
-						.Padding(0.0f)
-						.HAlign(HAlign_Fill)
-						.VAlign(VAlign_Fill)
-						[
-							SAssignNew(GraphTrack, SGraphTrack)
-						]
-				]
-		]
-
-	; // end of Widget
-}
-
-////////////////////////////////////////////////////////////////////////////////////////////////////
-
-TSharedRef<SWidget> SIoProfilerWindow::ConstructTimingTrack()
-{
-	return SNew(SVerticalBox)
-		.Visibility(this, &SIoProfilerWindow::IsTimingTrackVisible)
-		.IsEnabled(this, &SIoProfilerWindow::IsProfilerEnabled)
-
-	/*
-	// Header
-	+ SVerticalBox::Slot()
-		.AutoHeight()
-		.Padding(0.0f, 0.0f, 0.0f, 0.0f)
-		[
-			SNew(SHorizontalBox)
-
-			+ SHorizontalBox::Slot()
-				.AutoWidth()
-				[
-					SNew(SImage)
-						.Image(FEditorStyle::GetBrush(TEXT("Profiler.Tab.GraphView")))
-				]
-
-			+ SHorizontalBox::Slot()
-				.AutoWidth()
-				[
-					SNew(STextBlock)
-						.Text(LOCTEXT("TimingLabel", "Main Timing View"))
-						.Margin(FMargin(0.0f, 2.0f, 0.0f, 0.0f))
-				]
-		]
-	*/
-
-	// Timing Track
-	+ SVerticalBox::Slot()
-		.FillHeight(1.0)
-		.Padding(0.0f, 0.0f, 0.0f, 0.0f)
-		[
-			SNew(SBox)
-				.HeightOverride(48.0f)
-				.IsEnabled(this, &SIoProfilerWindow::IsProfilerEnabled)
-				[
-					SNew(SHorizontalBox)
-
-					+ SHorizontalBox::Slot()
-						.FillWidth(1.0f)
-						.Padding(0.0f)
-						.HAlign(HAlign_Fill)
-						.VAlign(VAlign_Fill)
-						[
-							SAssignNew(TimingView, STimingView)
-						]
-				]
-		]
-
-	//// Threads Panel
-	//+ SVerticalBox::Slot()
-	//	.FillHeight(1.0)
-	//	.Padding(0.0f, 0.0f, 0.0f, 0.0f)
-	//	[
-	//		SAssignNew(ThreadsPanel, SVerticalBox)
-	//	]
-
-	; // end of Widget
-}
-
-////////////////////////////////////////////////////////////////////////////////////////////////////
-
-TSharedRef<SWidget> SIoProfilerWindow::ConstructTimersView()
-{
-	return SNew(SVerticalBox)
-		.Visibility(this, &SIoProfilerWindow::IsTimersViewVisible)
-		.IsEnabled(this, &SIoProfilerWindow::IsProfilerEnabled)
-
-	// Header
-	+ SVerticalBox::Slot()
-		.AutoHeight()
-		[
-			SNew(SHorizontalBox)
-
-			+ SHorizontalBox::Slot()
-				.AutoWidth()
-				[
-					SNew(SImage)
-						.Image(FEditorStyle::GetBrush(TEXT("Profiler.Tab.FiltersAndPresets")))
-				]
-
-			+ SHorizontalBox::Slot()
-				.AutoWidth()
-				[
-					SNew(STextBlock)
-						.Text(LOCTEXT("TimersLabel", "Timers"))
-				]
-		]
-
-	// Timers
-	+ SVerticalBox::Slot()
-		.FillHeight(1.0f)
-		.Padding(0.0f, 2.0f, 0.0f, 0.0f)
-		[
-			SAssignNew(TimersView, STimersView)
-		]
-
-	; // end of Widget
-}
-
-////////////////////////////////////////////////////////////////////////////////////////////////////
-
-TSharedRef<SWidget> SIoProfilerWindow::ConstructLogView()
-{
-	return SNew(SVerticalBox)
-		.Visibility(this, &SIoProfilerWindow::IsLogViewVisible)
-
-	// Header
-	+ SVerticalBox::Slot()
-		.AutoHeight()
-		[
-			SNew(SHorizontalBox)
-			+ SHorizontalBox::Slot()
-				.AutoWidth()
-				[
-					SNew(SImage)
-						.Image(FEditorStyle::GetBrush(TEXT("Profiler.Tab.FiltersAndPresets")))
-				]
-
-			+ SHorizontalBox::Slot()
-				.AutoWidth()
-				[
-					SNew(STextBlock)
-						.Text(LOCTEXT("LogViewLabel", "Log View"))
-				]
-		]
-
-	// Log View
-	+ SVerticalBox::Slot()
-		.FillHeight(1.0f)
-		[
-			SAssignNew(LogView, SLogView)
-		]
-
-	; // end of Widget
-}
-
 ////////////////////////////////////////////////////////////////////////////////////////////////////
 
 void SIoProfilerWindow::Construct(const FArguments& InArgs)
@@ -390,56 +78,15 @@ void SIoProfilerWindow::Construct(const FArguments& InArgs)
 				.HAlign(HAlign_Fill)
 				.VAlign(VAlign_Fill)
 				[
-					SAssignNew(MainContentPanel, SVerticalBox)
-
-					// Toolbar
-					//+ SVerticalBox::Slot()
-					//	.AutoHeight()
-					//	[
-					//TODO:		SNew(SIoProfilerToolbar)
-					//	]
-
-					+ SVerticalBox::Slot()
-						.HAlign(HAlign_Fill)
-						.VAlign(VAlign_Fill)
+					SNew(SBox)
+						.Visibility(this, &SIoProfilerWindow::IsTimingViewVisible)
+						.IsEnabled(this, &SIoProfilerWindow::IsProfilerEnabled)
 						[
-							SNew(SSplitter)
-								.Orientation(Orient_Horizontal)
-
-							+ SSplitter::Slot()
-								.Value(0.65f)
-								[
-									SNew(SSplitter)
-										.Orientation(Orient_Vertical)
-
-									// Multi Track View
-									+ SSplitter::Slot()
-										.Value(0.75f)
-										[
-											//ConstructMultiTrackView()
-											ConstructTimingTrack()
-										]
-
-									// Log View
-									//+ SSplitter::Slot()
-									//	.Value(0.25f)
-									//	[
-									//		ConstructLogView()
-									//	]
-								]
-
-							// Timers View
-							//+ SSplitter::Slot()
-							//	.Value(0.35f)
-							//	//.SizeRule(SSplitter::SizeToContent)
-							//	.Expose(TimersViewSlot)
-							//	[
-							//		ConstructTimersView()
-							//	]
+							SAssignNew(TimingView, STimingView)
 						]
 				]
 
-			// session hint overlay
+			// Session hint overlay
 			+ SOverlay::Slot()
 				.HAlign(HAlign_Center)
 				.VAlign(VAlign_Center)
@@ -453,94 +100,25 @@ void SIoProfilerWindow::Construct(const FArguments& InArgs)
 								.Text(LOCTEXT("SelectTraceOverlayText", "Please select a trace."))
 						]
 				]
-
-			// notification area overlay
-			+ SOverlay::Slot()
-				.HAlign(HAlign_Right)
-				.VAlign(VAlign_Bottom)
-				.Padding(16.0f)
-				[
-					SAssignNew(NotificationList, SNotificationList)
-				]
-
-			// profiler settings overlay
-			+ SOverlay::Slot()
-				.HAlign(HAlign_Center)
-				.VAlign(VAlign_Center)
-				.Expose(OverlaySettingsSlot)
 		];
 
 	TimingView->SetGraphTrackVisible(true);
-
-	//ProfilerMiniView->OnSelectionBoxChanged().AddSP(GraphPanel.ToSharedRef(), &SIoProfilerGraphPanel::MiniView_OnSelectionBoxChanged);
-	//FIoProfilerManager::Get()->OnViewModeChanged().AddSP(this, &SIoProfilerWindow::ProfilerManager_OnViewModeChanged);
 }
 
 END_SLATE_FUNCTION_BUILD_OPTIMIZATION
 
 ////////////////////////////////////////////////////////////////////////////////////////////////////
 
-EVisibility SIoProfilerWindow::IsFramesTrackVisible() const
+EVisibility SIoProfilerWindow::IsTimingViewVisible() const
 {
-	if (FIoProfilerManager::Get()->IsFramesTrackVisible() &&
-		FInsightsManager::Get()->GetSession().IsValid())
-	{
-		return EVisibility::Visible;
-	}
-
-	return EVisibility::Collapsed;
-}
-
-////////////////////////////////////////////////////////////////////////////////////////////////////
-
-EVisibility SIoProfilerWindow::IsGraphTrackVisible() const
-{
-	if (FIoProfilerManager::Get()->IsGraphTrackVisible() &&
-		FInsightsManager::Get()->GetSession().IsValid())
-	{
-		return EVisibility::Visible;
-	}
-
-	return EVisibility::Collapsed;
-}
-
-////////////////////////////////////////////////////////////////////////////////////////////////////
-
-EVisibility SIoProfilerWindow::IsTimingTrackVisible() const
-{
-	if (FIoProfilerManager::Get()->IsTimingTrackVisible() &&
-		FInsightsManager::Get()->GetSession().IsValid())
-	{
-		return EVisibility::Visible;
-	}
-
-	return EVisibility::Collapsed;
-}
-
-////////////////////////////////////////////////////////////////////////////////////////////////////
-
-EVisibility SIoProfilerWindow::IsTimersViewVisible() const
-{
-	if (FIoProfilerManager::Get()->IsTimersViewVisible() &&
-		FInsightsManager::Get()->GetSession().IsValid())
+	if (FInsightsManager::Get()->GetSession().IsValid())
 	{
 		return EVisibility::Visible;
 	}
-
-	return EVisibility::Collapsed;
-}
-
-////////////////////////////////////////////////////////////////////////////////////////////////////
-
-EVisibility SIoProfilerWindow::IsLogViewVisible() const
-{
-	if (FIoProfilerManager::Get()->IsLogViewVisible() &&
-		FInsightsManager::Get()->GetSession().IsValid())
+	else
 	{
-		return EVisibility::Visible;
+		return EVisibility::Collapsed;
 	}
-
-	return EVisibility::Collapsed;
 }
 
 ////////////////////////////////////////////////////////////////////////////////////////////////////
@@ -551,132 +129,17 @@ EVisibility SIoProfilerWindow::IsSessionOverlayVisible() const
 	{
 		return EVisibility::Hidden;
 	}
-
-	return EVisibility::Visible;
-}
-
-////////////////////////////////////////////////////////////////////////////////////////////////////
-
-bool SIoProfilerWindow::IsProfilerEnabled() const
-{
-	const bool bIsActive = FInsightsManager::Get()->GetSession().IsValid();
-	return bIsActive;
-}
-
-////////////////////////////////////////////////////////////////////////////////////////////////////
-
-void SIoProfilerWindow::ManageLoadingProgressNotificationState(const FString& Filename, const EInsightsNotificationType NotificatonType, const ELoadingProgressState ProgressState, const float LoadingProgress)
-{
-	const FString BaseFilename = FPaths::GetBaseFilename(Filename);
-
-	if (ProgressState == ELoadingProgressState::Started)
-	{
-		const bool bContains = ActiveNotifications.Contains(Filename);
-		if (!bContains)
-		{
-			FNotificationInfo NotificationInfo(Io_GetTextForNotification(NotificatonType, ProgressState, BaseFilename));
-			NotificationInfo.bFireAndForget = false;
-			NotificationInfo.bUseLargeFont = false;
-
-			// Add two buttons, one for cancel, one for loading the received file.
-			if (NotificatonType == EInsightsNotificationType::LoadingTraceFile)
-			{
-				const FText CancelButtonText = LOCTEXT("CancelButton_Text", "Cancel");
-				const FText CancelButtonTT = LOCTEXT("CancelButton_TTText", "Hides this notification");
-				const FText LoadButtonText = LOCTEXT("LoadButton_Text", "Load file");
-				const FText LoadButtonTT = LOCTEXT("LoadButton_TTText", "Loads the received file and hides this notification");
-
-				NotificationInfo.ButtonDetails.Add(FNotificationButtonInfo(CancelButtonText, CancelButtonTT,
-					FSimpleDelegate::CreateSP(this, &SIoProfilerWindow::SendingServiceSideCapture_Cancel, Filename), SNotificationItem::CS_Success));
-				NotificationInfo.ButtonDetails.Add(FNotificationButtonInfo(LoadButtonText, LoadButtonTT,
-					FSimpleDelegate::CreateSP(this, &SIoProfilerWindow::SendingServiceSideCapture_Load, Filename), SNotificationItem::CS_Success));
-			}
-
-			SNotificationItemWeak NotificationItem = NotificationList->AddNotification(NotificationInfo);
-			NotificationItem.Pin()->SetCompletionState(SNotificationItem::CS_Pending);
-			ActiveNotifications.Add(Filename, NotificationItem);
-		}
-	}
-	else if (ProgressState == ELoadingProgressState::InProgress)
-	{
-		const SNotificationItemWeak* LoadingProgressPtr = ActiveNotifications.Find(Filename);
-		if (LoadingProgressPtr)
-		{
-			TSharedPtr<SNotificationItem> LoadingProcessPinned = LoadingProgressPtr->Pin();
-			LoadingProcessPinned->SetText(Io_GetTextForNotification(NotificatonType, ProgressState, BaseFilename, LoadingProgress));
-			LoadingProcessPinned->SetCompletionState(SNotificationItem::CS_Pending);
-		}
-	}
-	else if (ProgressState == ELoadingProgressState::Loaded)
-	{
-		const SNotificationItemWeak* LoadingProgressPtr = ActiveNotifications.Find(Filename);
-		if (LoadingProgressPtr)
-		{
-			TSharedPtr<SNotificationItem> LoadingProcessPinned = LoadingProgressPtr->Pin();
-			LoadingProcessPinned->SetText(Io_GetTextForNotification(NotificatonType, ProgressState, BaseFilename));
-			LoadingProcessPinned->SetCompletionState(SNotificationItem::CS_Success);
-
-			// Notifications for received files are handled by the user.
-			if (NotificatonType == EInsightsNotificationType::LoadingTraceFile)
-			{
-				LoadingProcessPinned->ExpireAndFadeout();
-				ActiveNotifications.Remove(Filename);
-			}
-		}
-	}
-	else if (ProgressState == ELoadingProgressState::Failed)
-	{
-		const SNotificationItemWeak* LoadingProgressPtr = ActiveNotifications.Find(Filename);
-		if (LoadingProgressPtr)
-		{
-			TSharedPtr<SNotificationItem> LoadingProcessPinned = LoadingProgressPtr->Pin();
-			LoadingProcessPinned->SetText(Io_GetTextForNotification(NotificatonType, ProgressState, BaseFilename));
-			LoadingProcessPinned->SetCompletionState(SNotificationItem::CS_Fail);
-
-			LoadingProcessPinned->ExpireAndFadeout();
-			ActiveNotifications.Remove(Filename);
-		}
-	}
-	else if (ProgressState == ELoadingProgressState::Cancelled)
+	else
 	{
-		const SNotificationItemWeak* LoadingProgressPtr = ActiveNotifications.Find(Filename);
-		if (LoadingProgressPtr)
-		{
-			TSharedPtr<SNotificationItem> LoadingProcessPinned = LoadingProgressPtr->Pin();
-			LoadingProcessPinned->ExpireAndFadeout();
-			ActiveNotifications.Remove(Filename);
-		}
-	}
-}
-
-////////////////////////////////////////////////////////////////////////////////////////////////////
-
-void SIoProfilerWindow::SendingServiceSideCapture_Cancel(const FString Filename)
-{
-	const SNotificationItemWeak* LoadingProgressPtr = ActiveNotifications.Find(Filename);
-	if (LoadingProgressPtr)
-	{
-		TSharedPtr<SNotificationItem> LoadingProcessPinned = LoadingProgressPtr->Pin();
-		LoadingProcessPinned->ExpireAndFadeout();
-		ActiveNotifications.Remove(Filename);
+		return EVisibility::Visible;
 	}
 }
 
 ////////////////////////////////////////////////////////////////////////////////////////////////////
 
-void SIoProfilerWindow::SendingServiceSideCapture_Load(const FString Filename)
+bool SIoProfilerWindow::IsProfilerEnabled() const
 {
-	const SNotificationItemWeak* LoadingProgressPtr = ActiveNotifications.Find(Filename);
-	if (LoadingProgressPtr)
-	{
-		TSharedPtr<SNotificationItem> LoadingProcessPinned = LoadingProgressPtr->Pin();
-		LoadingProcessPinned->ExpireAndFadeout();
-		ActiveNotifications.Remove(Filename);
-
-		const FString PathName = FPaths::ProfilingDir() + TEXT("UnrealStats/Received/");
-		const FString TraceFilepath = PathName + Filename;
-		FInsightsManager::Get()->LoadTraceFile(TraceFilepath);
-	}
+	return FInsightsManager::Get()->GetSession().IsValid();
 }
 
 ////////////////////////////////////////////////////////////////////////////////////////////////////
@@ -774,60 +237,4 @@ FReply SIoProfilerWindow::OnDrop(const FGeometry& MyGeometry, const FDragDropEve
 
 ////////////////////////////////////////////////////////////////////////////////////////////////////
 
-void SIoProfilerWindow::OpenProfilerSettings()
-{
-	MainContentPanel->SetEnabled(false);
-	(*OverlaySettingsSlot)
-	[
-		SNew(SBorder)
-		.BorderImage(FEditorStyle::GetBrush("NotificationList.ItemBackground"))
-		.Padding(8.0f)
-		[
-			SNew(SInsightsSettings)
-			.OnClose(this, &SIoProfilerWindow::CloseProfilerSettings)
-			.SettingPtr(&FInsightsManager::GetSettings())
-		]
-	];
-}
-
-////////////////////////////////////////////////////////////////////////////////////////////////////
-
-void SIoProfilerWindow::CloseProfilerSettings()
-{
-	// Close the profiler settings by simply replacing widget with a null one.
-	(*OverlaySettingsSlot)
-	[
-		SNullWidget::NullWidget
-	];
-	MainContentPanel->SetEnabled(true);
-}
-
-////////////////////////////////////////////////////////////////////////////////////////////////////
-
-//void SIoProfilerWindow::ProfilerManager_OnViewModeChanged(EIoProfilerViewMode NewViewMode)
-//{
-//	if (NewViewMode == EIoProfilerViewMode::LineIndexBased)
-//	{
-//		EventGraphPanel->SetVisibility(EVisibility::Visible);
-//		EventGraphPanel->SetEnabled(true);
-//
-//		(*FiltersAndPresetsSlot)
-//		[
-//			FiltersAndPresets.ToSharedRef()
-//		];
-//	}
-//	else if (NewViewMode == EIoProfilerViewMode::ThreadViewTimeBased)
-//	{
-//		EventGraphPanel->SetVisibility(EVisibility::Collapsed);
-//		EventGraphPanel->SetEnabled(false);
-//
-//		(*FiltersAndPresetsSlot)
-//		[
-//			SNullWidget::NullWidget
-//		];
-//	}
-//}
-
-////////////////////////////////////////////////////////////////////////////////////////////////////
-
 #undef LOCTEXT_NAMESPACE
diff --git a/Engine/Source/Developer/TraceInsights/Private/Insights/Widgets/SIoProfilerWindow.h b/Engine/Source/Developer/TraceInsights/Private/Insights/Widgets/SIoProfilerWindow.h
--- a/Engine/Source/Developer/TraceInsights/Private/Insights/Widgets/SIoProfilerWindow.h
+++ b/Engine/Source/Developer/TraceInsights/Private/Insights/Widgets/SIoProfilerWindow.h
@@ -8,88 +8,47 @@
 #include "Misc/Guid.h"
 #include "SlateFwd.h"
 #include "Widgets/DeclarativeSyntaxSupport.h"
-#include "Widgets/Layout/SSplitter.h"
 #include "Widgets/SCompoundWidget.h"
-#include "Widgets/SOverlay.h"
 
 // Insights
 #include "Insights/InsightsManager.h"
 
 class FActiveTimerHandle;
-class SVerticalBox;
-class SFrameTrack;
-class SGraphTrack;
-class STimersView;
 class STimingView;
-class SLogView;
 
-/// Type definition for shared pointers to instances of SNotificationItem.
-typedef TSharedPtr<class SNotificationItem> SNotificationItemPtr;
-
-/// Type definition for shared references to instances of SNotificationItem.
-typedef TSharedRef<class SNotificationItem> SNotificationItemRef;
-
-/// Type definition for weak references to instances of SNotificationItem.
-typedef TWeakPtr<class SNotificationItem> SNotificationItemWeak;
-
-/// Implements the timing profiler window.
+/** Implements the timing profiler window. */
 class SIoProfilerWindow : public SCompoundWidget
 {
 public:
-	/// Default constructor.
+	/** Default constructor. */
 	SIoProfilerWindow();
 
-	/// Virtual destructor.
+	/** Virtual destructor. */
 	virtual ~SIoProfilerWindow();
 
 	SLATE_BEGIN_ARGS(SIoProfilerWindow){}
 	SLATE_END_ARGS()
 
-	/// Constructs this widget.
+	/** Constructs this widget. */
 	void Construct(const FArguments& InArgs);
 
-	void ManageLoadingProgressNotificationState(const FString& Filename, const EInsightsNotificationType NotificatonType, const ELoadingProgressState ProgressState, const float LoadingProgress);
-
-	void OpenProfilerSettings();
-	void CloseProfilerSettings();
-
-protected:
-	TSharedRef<SWidget> ConstructMultiTrackView();
-	TSharedRef<SWidget> ConstructFramesTrack();
-	TSharedRef<SWidget> ConstructGraphTrack();
-	TSharedRef<SWidget> ConstructTimingTrack();
-	TSharedRef<SWidget> ConstructTimersView();
-	TSharedRef<SWidget> ConstructLogView();
-
-	/// Callback for determining the visibility of the Frames track.
-	EVisibility IsFramesTrackVisible() const;
-
-	/// Callback for determining the visibility of the Graph track.
-	EVisibility IsGraphTrackVisible() const;
+	TSharedPtr<STimingView> GetTimingView() const { return TimingView; }
 
-	/// Callback for determining the visibility of the Timing track.
-	EVisibility IsTimingTrackVisible() const;
-
-	/// Callback for determining the visibility of the Timers View.
-	EVisibility IsTimersViewVisible() const;
-
-	/// Callback for determining the visibility of the Log View.
-	EVisibility IsLogViewVisible() const;
+private:
+	/** Callback for determining the visibility of the Timing view. */
+	EVisibility IsTimingViewVisible() const;
 
-	/// Callback for determining the visibility of the 'Select a session' overlay.
+	/** Callback for determining the visibility of the 'Select a session' overlay. */
 	EVisibility IsSessionOverlayVisible() const;
 
-	/// Callback for getting the enabled state of the profiler window.
+	/** Callback for getting the enabled state of the profiler window. */
 	bool IsProfilerEnabled() const;
 
-	void SendingServiceSideCapture_Cancel(const FString Filename);
-
-	void SendingServiceSideCapture_Load(const FString Filename);
-
-private:
+	/** Updates the amount of time the profiler has been active. */
+	EActiveTimerReturnType UpdateActiveDuration(double InCurrentTime, float InDeltaTime);
 
 	/**
-	 * Ticks this widget.  Override in derived classes, but always call the parent implementation.
+	 * Ticks this widget. Override in derived classes, but always call the parent implementation.
 	 *
 	 * @param  AllottedGeometry The space allotted for this widget
 	 * @param  InCurrentTime  Current absolute real time
@@ -143,45 +102,12 @@ private:
 	virtual FReply OnDragOver(const FGeometry& MyGeometry, const FDragDropEvent& DragDropEvent)  override;
 
 private:
-	/// Updates the amount of time the profiler has been active
-	EActiveTimerReturnType UpdateActiveDuration(double InCurrentTime, float InDeltaTime);
-
-public:
-	/// Widget for the frame track
-	TSharedPtr<SFrameTrack> FrameTrack;
-
-	/// Widget for the graph track
-	TSharedPtr<SGraphTrack> GraphTrack;
-
-	/// Holds all timing tracks (one for each thread).
-	TSharedPtr<SVerticalBox> ThreadsPanel;
-
-	/// Widget for the timing track
+	/** Widget for the timing track */
 	TSharedPtr<STimingView> TimingView;
 
-	/// Holds the Timers view widget/slot.
-	SSplitter::FSlot* TimersViewSlot;
-	TSharedPtr<STimersView> TimersView;
-
-	/// Widget for the log view
-	TSharedPtr<SLogView> LogView;
-
-	/// Widget for the non-intrusive notifications.
-	TSharedPtr<SNotificationList> NotificationList;
-
-	/// Holds all active and visible notifications, stored as FGuid -> SNotificationItemWeak.
-	TMap<FString, SNotificationItemWeak> ActiveNotifications;
-
-	/// Overlay slot which contains the profiler settings widget.
-	SOverlay::FOverlaySlot* OverlaySettingsSlot;
-
-	/// The number of seconds the profiler has been active
+	/** The number of seconds the profiler has been active */
 	float DurationActive;
 
-private:
-	/// The handle to the active update duration tick
+	/** The handle to the active update duration tick */
 	TWeakPtr<FActiveTimerHandle> ActiveTimerHandle;
-
-	/// Holds all widgets for the profiler window like menu bar, toolbar and tabs.
-	TSharedPtr<SVerticalBox> MainContentPanel;
 };
diff --git a/Engine/Source/Developer/TraceInsights/Private/Insights/Widgets/SStartPageWindow.cpp b/Engine/Source/Developer/TraceInsights/Private/Insights/Widgets/SStartPageWindow.cpp
--- a/Engine/Source/Developer/TraceInsights/Private/Insights/Widgets/SStartPageWindow.cpp
+++ b/Engine/Source/Developer/TraceInsights/Private/Insights/Widgets/SStartPageWindow.cpp
@@ -3,8 +3,10 @@
 #include "SStartPageWindow.h"
 
 #include "EditorStyleSet.h"
+#include "Framework/Application/SlateApplication.h"
 #include "Framework/Docking/WorkspaceItem.h"
 #include "SlateOptMacros.h"
+#include "TraceServices/SessionService.h"
 #include "Widgets/Docking/SDockTab.h"
 #include "Widgets/Images/SImage.h"
 #include "Widgets/Layout/SBorder.h"
@@ -25,9 +27,11 @@
 
 // Insights
 #include "Insights/InsightsManager.h"
+#include "Insights/TimingProfilerManager.h"
 #include "Insights/InsightsStyle.h"
 #include "Insights/Version.h"
 #include "Insights/Widgets/SInsightsSettings.h"
+#include "Insights/Widgets/STimingProfilerWindow.h"
 
 ////////////////////////////////////////////////////////////////////////////////////////////////////
 
@@ -35,7 +39,7 @@
 
 ////////////////////////////////////////////////////////////////////////////////////////////////////
 
-// TODO: move this function to InsightsManager or in TraceSession.h
+
 static FText GetTextForNotification(const EInsightsNotificationType NotificatonType, const ELoadingProgressState ProgressState, const FString& Filename, const float ProgressPercent = 0.0f)
 {
 	FText Result;
@@ -117,97 +121,151 @@ void SStartPageWindow::Construct(const FArguments& InArgs)
 
 			// Overlay slot for the main window area
 			+ SOverlay::Slot()
-				.HAlign(HAlign_Center)
+				.HAlign(HAlign_Fill)
 				.VAlign(VAlign_Center)
 				[
-					SNew(SHorizontalBox)
+					SNew(SVerticalBox)
 
-					+ SHorizontalBox::Slot()
-						.AutoWidth()
+					+ SVerticalBox::Slot()
+						.AutoHeight()
+						.HAlign(HAlign_Center)
 						[
-							SNew(SButton)
-							//.OnClicked(this, &StartPageWindow::Live_OnClicked)
-							.ToolTipText(LOCTEXT("LiveButtonTooltip", "Start a live session or try loading the last trace file."))
-							.ContentPadding(8.0f)
-							.Content()
+							SNew(SBorder)
+							.Visibility(this, &SStartPageWindow::IsSessionOverlayVisible)
+							.BorderImage(FEditorStyle::GetBrush("NotificationList.ItemBackground"))
+							.Padding(8.0f)
 							[
-								SNew(SHorizontalBox)
+								SNew(STextBlock)
+								.Text(LOCTEXT("SelectTraceOverlayText", "Please select a trace."))
+							]
+						]
+
+					+ SVerticalBox::Slot()
+						.AutoHeight()
+						[
+							SNew(SBox)
+							.HeightOverride(6.0f)
+						]
 
-								+SHorizontalBox::Slot()
-									.AutoWidth()
+					+ SVerticalBox::Slot()
+						.AutoHeight()
+						.HAlign(HAlign_Center)
+						[
+							SNew(SHorizontalBox)
+
+							+ SHorizontalBox::Slot()
+								.AutoWidth()
+								[
+									SNew(SButton)
+									.OnClicked(this, &SStartPageWindow::Live_OnClicked)
+									.ToolTipText(LOCTEXT("LiveButtonTooltip", "Start a live session or load the last trace file from the Local Session Directory."))
+									.ContentPadding(8.0f)
+									.Content()
 									[
-										SNew(SImage)
-										.Image(FInsightsStyle::GetBrush("Live.Icon.Large"))
+										SNew(SHorizontalBox)
+
+										+SHorizontalBox::Slot()
+											.AutoWidth()
+											[
+												SNew(SImage)
+												.Image(FInsightsStyle::GetBrush("Live.Icon.Large"))
+											]
+
+										+SHorizontalBox::Slot()
+											.AutoWidth()
+											.VAlign(VAlign_Center)
+											[
+												SNew(SBox)
+												.WidthOverride(100.0f)
+												[
+													SNew(STextBlock)
+													.Font(FCoreStyle::GetDefaultFontStyle("Regular", 18))
+													.Text(LOCTEXT("LiveButtonText", "Live"))
+												]
+											]
 									]
-								+SHorizontalBox::Slot()
-									.AutoWidth()
-									.VAlign(VAlign_Center)
+								]
+
+							+ SHorizontalBox::Slot()
+								.AutoWidth()
+								[
+									SNew(SBox)
+									.WidthOverride(6.0f)
+								]
+
+							+ SHorizontalBox::Slot()
+								.AutoWidth()
+								[
+									SNew(SButton)
+									.OnClicked(this, &SStartPageWindow::Load_OnClicked)
+									.ToolTipText(LOCTEXT("LoadButtonTooltip", "Load a trace file."))
+									.ContentPadding(8.0f)
+									.Content()
 									[
-										SNew(SBox)
-										.WidthOverride(100.0f)
-										[
-											SNew(STextBlock)
-											.Font(FCoreStyle::GetDefaultFontStyle("Regular", 18))
-											.Text(LOCTEXT("LiveButtonText", "Live"))
-										]
+										SNew(SHorizontalBox)
+
+										+SHorizontalBox::Slot()
+											.AutoWidth()
+											[
+												SNew(SImage)
+												.Image(FInsightsStyle::GetBrush("Load.Icon.Large"))
+											]
+
+										+SHorizontalBox::Slot()
+											.AutoWidth()
+											.VAlign(VAlign_Center)
+											[
+												SNew(SBox)
+												.WidthOverride(100.0f)
+												[
+													SNew(STextBlock)
+													.Font(FCoreStyle::GetDefaultFontStyle("Regular", 18))
+													.Text(LOCTEXT("LoadButtonText", "Load..."))
+												]
+											]
 									]
-							]
+								]
 						]
 
-					+ SHorizontalBox::Slot()
-						.AutoWidth()
+					+ SVerticalBox::Slot()
+						.AutoHeight()
 						[
 							SNew(SBox)
-							.WidthOverride(6.0f)
+							.HeightOverride(6.0f)
 						]
 
-					+ SHorizontalBox::Slot()
-						.AutoWidth()
+					+ SVerticalBox::Slot()
+						.AutoHeight()
+						.HAlign(HAlign_Center)
 						[
-							SNew(SButton)
-							//.OnClicked(this, &StartPageWindow::Load_OnClicked)
-							.ToolTipText(LOCTEXT("LoadButtonTooltip", "Last a trace file."))
-							.ContentPadding(8.0f)
-							.Content()
+							SNew(SBorder)
+							.Visibility(this, &SStartPageWindow::IsSessionOverlayVisible)
+							.BorderImage(FEditorStyle::GetBrush("NotificationList.ItemBackground"))
+							.Padding(8.0f)
 							[
-								SNew(SHorizontalBox)
+								SNew(SVerticalBox)
 
-								+SHorizontalBox::Slot()
-									.AutoWidth()
+								+ SVerticalBox::Slot()
+									.AutoHeight()
+									.HAlign(HAlign_Center)
+									.Padding(0.0, 2.0f)
 									[
-										SNew(SImage)
-										.Image(FInsightsStyle::GetBrush("Load.Icon.Large"))
+										SNew(STextBlock)
+										.Text(LOCTEXT("LocalSessionDirectoryText", "Local Session Directory:"))
+										.ColorAndOpacity(FLinearColor::Gray)
 									]
-								+SHorizontalBox::Slot()
-									.AutoWidth()
-									.VAlign(VAlign_Center)
+
+								+SVerticalBox::Slot()
+									.AutoHeight()
+									.HAlign(HAlign_Center)
+									.Padding(0.0, 2.0f)
 									[
-										SNew(SBox)
-										.WidthOverride(100.0f)
-										[
-											SNew(STextBlock)
-											.Font(FCoreStyle::GetDefaultFontStyle("Regular", 18))
-											.Text(LOCTEXT("LoadButtonText", "Load..."))
-										]
+										SNew(STextBlock)
+										.Text(this, &SStartPageWindow::GetLocalSessionDirectory)
 									]
 							]
 						]
-				]
 
-			// Session hint overlay
-			+ SOverlay::Slot()
-				.HAlign(HAlign_Center)
-				.VAlign(VAlign_Center)
-				.Padding(0.0f, 0.0f, 0.0f, 48.0f)
-				[
-					SNew(SBorder)
-						.Visibility(this, &SStartPageWindow::IsSessionOverlayVisible)
-						.BorderImage(FEditorStyle::GetBrush("NotificationList.ItemBackground"))
-						.Padding(8.0f)
-						[
-							SNew(STextBlock)
-								.Text(LOCTEXT("SelectTraceOverlayText", "Please select a trace."))
-						]
 				]
 
 			// Notification area overlay
@@ -460,6 +518,71 @@ FReply SStartPageWindow::OnDrop(const FGeometry& MyGeometry, const FDragDropEven
 
 ////////////////////////////////////////////////////////////////////////////////////////////////////
 
+void SStartPageWindow::SpawnAndActivateTabs()
+{
+	//TODO: expose those in FInsightsManager; see also TraceInsightsModule.cpp
+	static const FName StartPageTabName(TEXT("StartPage"));
+	static const FName TimingProfilerTabName(TEXT("TimingProfiler"));
+	static const FName IoProfilerTabName(TEXT("IoProfiler"));
+
+	// Open tabs for profilers.
+	if (FGlobalTabmanager::Get()->CanSpawnTab(TimingProfilerTabName))
+	{
+		FGlobalTabmanager::Get()->InvokeTab(TimingProfilerTabName);
+	}
+	if (FGlobalTabmanager::Get()->CanSpawnTab(IoProfilerTabName))
+	{
+		FGlobalTabmanager::Get()->InvokeTab(IoProfilerTabName);
+	}
+
+	// Ensure Timing Insights / Timing View is the active tab.
+	if (TSharedPtr<SDockTab> TimingInsightsTab = FGlobalTabmanager::Get()->FindExistingLiveTab(TimingProfilerTabName))
+	{
+		//FGlobalTabmanager::Get()->SetActiveTab(TimingProfilerTabName); // doesn't really work :(
+		FGlobalTabmanager::Get()->InvokeTab(TimingProfilerTabName); // this activates the tab (brings the tab in foreground)
+
+		TSharedPtr<class STimingProfilerWindow> Wnd = FTimingProfilerManager::Get()->GetProfilerWindow();
+		if (Wnd)
+		{
+			TSharedPtr<FTabManager> TabManager = Wnd->GetTabManager();
+			if (TSharedPtr<SDockTab> TimingViewTab = TabManager->FindExistingLiveTab(FTimingProfilerTabs::TimingViewID))
+			{
+				TabManager->DrawAttention(TimingViewTab.ToSharedRef());
+				TimingViewTab->ActivateInParent(ETabActivationCause::SetDirectly);
+				FSlateApplication::Get().SetKeyboardFocus(TimingViewTab->GetContent());
+			}
+		}
+	}
+}
+
+////////////////////////////////////////////////////////////////////////////////////////////////////
+
+FText SStartPageWindow::GetLocalSessionDirectory() const
+{
+	TSharedRef<Trace::ISessionService> SessionService = FInsightsManager::Get()->GetSessionService();
+	return FText::FromString(FPaths::ConvertRelativePathToFull(SessionService->GetLocalSessionDirectory()));
+}
+
+////////////////////////////////////////////////////////////////////////////////////////////////////
+
+FReply SStartPageWindow::Live_OnClicked()
+{
+	SpawnAndActivateTabs();
+	FInsightsManager::Get()->GetCommandList()->ExecuteAction(FInsightsManager::GetCommands().InsightsManager_Live.ToSharedRef());
+	return FReply::Handled();
+}
+
+////////////////////////////////////////////////////////////////////////////////////////////////////
+
+FReply SStartPageWindow::Load_OnClicked()
+{
+	SpawnAndActivateTabs();
+	FInsightsManager::Get()->GetCommandList()->ExecuteAction(FInsightsManager::GetCommands().InsightsManager_Load.ToSharedRef());
+	return FReply::Handled();
+}
+
+////////////////////////////////////////////////////////////////////////////////////////////////////
+
 void SStartPageWindow::OpenSettings()
 {
 	MainContentPanel->SetEnabled(false);
diff --git a/Engine/Source/Developer/TraceInsights/Private/Insights/Widgets/SStartPageWindow.h b/Engine/Source/Developer/TraceInsights/Private/Insights/Widgets/SStartPageWindow.h
--- a/Engine/Source/Developer/TraceInsights/Private/Insights/Widgets/SStartPageWindow.h
+++ b/Engine/Source/Developer/TraceInsights/Private/Insights/Widgets/SStartPageWindow.h
@@ -58,17 +58,24 @@ public:
 
 	void ShowHideTab(const FName& TabID, bool bIsVisible);
 
-protected:
+private:
 	/** Callback for determining the visibility of the 'Select a session' overlay. */
 	EVisibility IsSessionOverlayVisible() const;
 
 	bool IsSessionValid() const;
 
 	void SendingServiceSideCapture_Cancel(const FString Filename);
-
 	void SendingServiceSideCapture_Load(const FString Filename);
 
-private:
+	FText GetLocalSessionDirectory() const;
+
+	void SpawnAndActivateTabs();
+
+	FReply Live_OnClicked();
+	FReply Load_OnClicked();
+
+	/** Updates the amount of time the profiler has been active. */
+	EActiveTimerReturnType UpdateActiveDuration(double InCurrentTime, float InDeltaTime);
 
 	/**
 	 * Ticks this widget.  Override in derived classes, but always call the parent implementation.
@@ -124,10 +131,6 @@ private:
 	 */
 	virtual FReply OnDragOver(const FGeometry& MyGeometry, const FDragDropEvent& DragDropEvent)  override;
 
-private:
-	/** Updates the amount of time the profiler has been active. */
-	EActiveTimerReturnType UpdateActiveDuration(double InCurrentTime, float InDeltaTime);
-
 public:
 	/** Widget for the non-intrusive notifications. */
 	TSharedPtr<SNotificationList> NotificationList;
diff --git a/Engine/Source/Developer/TraceInsights/Private/Insights/Widgets/STimingProfilerToolbar.cpp b/Engine/Source/Developer/TraceInsights/Private/Insights/Widgets/STimingProfilerToolbar.cpp
--- a/Engine/Source/Developer/TraceInsights/Private/Insights/Widgets/STimingProfilerToolbar.cpp
+++ b/Engine/Source/Developer/TraceInsights/Private/Insights/Widgets/STimingProfilerToolbar.cpp
@@ -47,7 +47,6 @@ void STimingProfilerToolbar::Construct(const FArguments& InArgs)
 			{
 				ToolbarBuilder.AddToolBarButton(FInsightsCommands::Get().InsightsManager_Live);
 				ToolbarBuilder.AddToolBarButton(FInsightsCommands::Get().InsightsManager_Load);
-				//ToolbarBuilder.AddToolBarButton(FInsightsCommands::Get().InsightsManager_Mock);
 			}
 			ToolbarBuilder.EndSection();
 			ToolbarBuilder.BeginSection("View");
diff --git a/Engine/Source/Developer/TraceInsights/Private/Insights/Widgets/STimingProfilerWindow.cpp b/Engine/Source/Developer/TraceInsights/Private/Insights/Widgets/STimingProfilerWindow.cpp
--- a/Engine/Source/Developer/TraceInsights/Private/Insights/Widgets/STimingProfilerWindow.cpp
+++ b/Engine/Source/Developer/TraceInsights/Private/Insights/Widgets/STimingProfilerWindow.cpp
@@ -52,45 +52,6 @@ const FName FTimingProfilerTabs::LogViewID(TEXT("LogView"));
 
 ////////////////////////////////////////////////////////////////////////////////////////////////////
 
-// TODO: move this function to InsightsManager or in TraceSession.h
-static FText Timing_GetTextForNotification(const EInsightsNotificationType NotificatonType, const ELoadingProgressState ProgressState, const FString& Filename, const float ProgressPercent = 0.0f)
-{
-	FText Result;
-
-	if (NotificatonType == EInsightsNotificationType::LoadingTraceFile)
-	{
-		if (ProgressState == ELoadingProgressState::Started)
-		{
-			FFormatNamedArguments Args;
-			Args.Add(TEXT("Filename"), FText::FromString(Filename));
-			Result = FText::Format(LOCTEXT("DescF_OfflineCapture_Started", "Started loading a file {Filename}"), Args);
-		}
-		else if (ProgressState == ELoadingProgressState::InProgress)
-		{
-			FFormatNamedArguments Args;
-			Args.Add(TEXT("Filename"), FText::FromString(Filename));
-			Args.Add(TEXT("LoadingProgressPercent"), FText::AsPercent(ProgressPercent));
-			Result = FText::Format(LOCTEXT("DescF_OfflineCapture_InProgress", "Loading a file {Filename} {LoadingProgressPercent}"), Args);
-		}
-		else if (ProgressState == ELoadingProgressState::Loaded)
-		{
-			FFormatNamedArguments Args;
-			Args.Add(TEXT("Filename"), FText::FromString(Filename));
-			Result = FText::Format(LOCTEXT("DescF_OfflineCapture_Loaded", "File {Filename} has been successfully loaded"), Args);
-		}
-		else if (ProgressState == ELoadingProgressState::Failed)
-		{
-			FFormatNamedArguments Args;
-			Args.Add(TEXT("Filename"), FText::FromString(Filename));
-			Result = FText::Format(LOCTEXT("DescF_OfflineCapture_Failed", "Failed to load file {Filename}"), Args);
-		}
-	}
-
-	return Result;
-}
-
-////////////////////////////////////////////////////////////////////////////////////////////////////
-
 STimingProfilerWindow::STimingProfilerWindow()
 	: DurationActive(0.0f)
 {
@@ -117,6 +78,7 @@ TSharedRef<SDockTab> STimingProfilerWindow::SpawnTab_Toolbar(const FSpawnTabArgs
 	return SNew(SDockTab)
 		.ShouldAutosize(true)
 		.TabRole(ETabRole::PanelTab)
+		//.IsEnabled(this, &STimingProfilerWindow::IsProfilerEnabled)
 		[
 			SNew(STimingProfilerToolbar)
 		];
@@ -129,10 +91,9 @@ TSharedRef<SDockTab> STimingProfilerWindow::SpawnTab_FramesTrack(const FSpawnTab
 	return SNew(SDockTab)
 		.ShouldAutosize(false)
 		.TabRole(ETabRole::PanelTab)
+		//.IsEnabled(this, &STimingProfilerWindow::IsProfilerEnabled)
 		[
 			SAssignNew(FrameTrack, SFrameTrack)
-			//.IsEnabled(this, &STimingProfilerWindow::IsProfilerEnabled)
-			//.Visibility(this, &STimingProfilerWindow::IsFramesTrackVisible)
 		];
 }
 
@@ -143,10 +104,9 @@ TSharedRef<SDockTab> STimingProfilerWindow::SpawnTab_GraphTrack(const FSpawnTabA
 	return SNew(SDockTab)
 		.ShouldAutosize(false)
 		.TabRole(ETabRole::PanelTab)
+		//.IsEnabled(this, &STimingProfilerWindow::IsProfilerEnabled)
 		[
 			SAssignNew(GraphTrack, SGraphTrack)
-			//.IsEnabled(this, &STimingProfilerWindow::IsProfilerEnabled)
-			//.Visibility(this, &STimingProfilerWindow::IsGraphTrackVisible)
 		];
 }
 
@@ -157,10 +117,9 @@ TSharedRef<SDockTab> STimingProfilerWindow::SpawnTab_TimingView(const FSpawnTabA
 	return SNew(SDockTab)
 		.ShouldAutosize(false)
 		.TabRole(ETabRole::PanelTab)
+		//.IsEnabled(this, &STimingProfilerWindow::IsProfilerEnabled)
 		[
 			SAssignNew(TimingView, STimingView)
-			//.IsEnabled(this, &STimingProfilerWindow::IsProfilerEnabled)
-			//.Visibility(this, &STimingProfilerWindow::IsTimersViewVisible)
 		];
 }
 
@@ -171,11 +130,9 @@ TSharedRef<SDockTab> STimingProfilerWindow::SpawnTab_Timers(const FSpawnTabArgs&
 	return SNew(SDockTab)
 		.ShouldAutosize(false)
 		.TabRole(ETabRole::PanelTab)
+		//.IsEnabled(this, &STimingProfilerWindow::IsProfilerEnabled)
 		[
-		
 			SAssignNew(TimersView, STimersView)
-			//.IsEnabled(this, &STimingProfilerWindow::IsProfilerEnabled)
-			//.Visibility(this, &STimingProfilerWindow::IsTimersViewVisible)
 		];
 }
 
@@ -186,11 +143,9 @@ TSharedRef<SDockTab> STimingProfilerWindow::SpawnTab_StatsCounters(const FSpawnT
 	return SNew(SDockTab)
 		.ShouldAutosize(false)
 		.TabRole(ETabRole::PanelTab)
+		//.IsEnabled(this, &STimingProfilerWindow::IsProfilerEnabled)
 		[
 			SAssignNew(StatsView, SStatsView)
-			//.IsEnabled(this, &STimingProfilerWindow::IsProfilerEnabled)
-			//.Visibility(this, &STimingProfilerWindow::IsStatsCountersViewVisible)
-			
 		];
 }
 
@@ -201,10 +156,9 @@ TSharedRef<SDockTab> STimingProfilerWindow::SpawnTab_LogView(const FSpawnTabArgs
 	return SNew(SDockTab)
 		.ShouldAutosize(false)
 		.TabRole(ETabRole::PanelTab)
+		//.IsEnabled(this, &STimingProfilerWindow::IsProfilerEnabled)
 		[
 			SAssignNew(LogView, SLogView)
-			//.IsEnabled(this, &STimingProfilerWindow::IsProfilerEnabled)
-			//.Visibility(this, &STimingProfilerWindow::IsLogViewVisible)
 		];
 }
 
@@ -308,6 +262,7 @@ void STimingProfilerWindow::Construct(const FArguments& InArgs, const TSharedRef
 					->SetSizeCoefficient(0.35f)
 					->AddTab(FTimingProfilerTabs::TimersID, FTimingProfilerManager::Get()->IsTimersViewVisible() ? ETabState::OpenedTab : ETabState::ClosedTab)
 					->AddTab(FTimingProfilerTabs::StatsCountersID, FTimingProfilerManager::Get()->IsStatsCountersViewVisible() ? ETabState::OpenedTab : ETabState::ClosedTab)
+					->SetForegroundTab(FTimingProfilerTabs::TimersID)
 				)
 			)
 		);
@@ -343,7 +298,7 @@ void STimingProfilerWindow::Construct(const FArguments& InArgs, const TSharedRef
 				.HAlign(HAlign_Fill)
 				.VAlign(VAlign_Fill)
 				[
-					SAssignNew(MainContentPanel, SVerticalBox)
+					SNew(SVerticalBox)
 
 					+ SVerticalBox::Slot()
 						.AutoHeight()
@@ -372,21 +327,6 @@ void STimingProfilerWindow::Construct(const FArguments& InArgs, const TSharedRef
 								.Text(LOCTEXT("SelectTraceOverlayText", "Please select a trace."))
 						]
 				]
-
-			// Notification area overlay
-			+ SOverlay::Slot()
-				.HAlign(HAlign_Right)
-				.VAlign(VAlign_Bottom)
-				.Padding(16.0f)
-				[
-					SAssignNew(NotificationList, SNotificationList)
-				]
-
-			// Settings dialog overlay
-			+ SOverlay::Slot()
-				.HAlign(HAlign_Center)
-				.VAlign(VAlign_Center)
-				.Expose(OverlaySettingsSlot)
 		];
 
 	// Tell tab-manager about the global menu bar.
@@ -413,87 +353,22 @@ void STimingProfilerWindow::FillMenu(FMenuBuilder& MenuBuilder, const TSharedPtr
 
 ////////////////////////////////////////////////////////////////////////////////////////////////////
 
-EVisibility STimingProfilerWindow::IsFramesTrackVisible() const
+void STimingProfilerWindow::ShowTab(const FName& TabID)
 {
-	if (FTimingProfilerManager::Get()->IsFramesTrackVisible() &&
-		FInsightsManager::Get()->GetSession().IsValid())
+	if (TabManager->CanSpawnTab(TabID))
 	{
-		return EVisibility::Visible;
+		TabManager->InvokeTab(TabID);
 	}
-
-	return EVisibility::Collapsed;
 }
 
 ////////////////////////////////////////////////////////////////////////////////////////////////////
 
-EVisibility STimingProfilerWindow::IsGraphTrackVisible() const
+void STimingProfilerWindow::HideTab(const FName& TabID)
 {
-	if (FTimingProfilerManager::Get()->IsGraphTrackVisible() &&
-		FInsightsManager::Get()->GetSession().IsValid())
-	{
-		return EVisibility::Visible;
-	}
-
-	return EVisibility::Collapsed;
-}
-
-////////////////////////////////////////////////////////////////////////////////////////////////////
-
-EVisibility STimingProfilerWindow::IsTimingViewVisible() const
-{
-	if (FTimingProfilerManager::Get()->IsTimingViewVisible() &&
-		FInsightsManager::Get()->GetSession().IsValid())
-	{
-		return EVisibility::Visible;
-	}
-
-	return EVisibility::Collapsed;
-}
-
-////////////////////////////////////////////////////////////////////////////////////////////////////
-
-EVisibility STimingProfilerWindow::IsTimersViewVisible() const
-{
-	if (FTimingProfilerManager::Get()->IsTimersViewVisible() &&
-		FInsightsManager::Get()->GetSession().IsValid())
-	{
-		return EVisibility::Visible;
-	}
-
-	return EVisibility::Collapsed;
-}
-
-////////////////////////////////////////////////////////////////////////////////////////////////////
-
-EVisibility STimingProfilerWindow::IsLogViewVisible() const
-{
-	if (FTimingProfilerManager::Get()->IsLogViewVisible() &&
-		FInsightsManager::Get()->GetSession().IsValid())
-	{
-		return EVisibility::Visible;
-	}
-
-	return EVisibility::Collapsed;
-}
-
-////////////////////////////////////////////////////////////////////////////////////////////////////
-
-void STimingProfilerWindow::ShowHideTab(const FName& TabID, bool bIsVisible)
-{
-	if (!bIsVisible)
-	{
-		TSharedPtr<SDockTab> Tab = TabManager->FindExistingLiveTab(TabID);
-		if (Tab.IsValid())
-		{
-			Tab->RequestCloseTab();
-		}
-	}
-	else
+	TSharedPtr<SDockTab> Tab = TabManager->FindExistingLiveTab(TabID);
+	if (Tab.IsValid())
 	{
-		if (TabManager->CanSpawnTab(TabID))
-		{
-			TabManager->InvokeTab(TabID);
-		}
+		Tab->RequestCloseTab();
 	}
 }
 
@@ -505,132 +380,17 @@ EVisibility STimingProfilerWindow::IsSessionOverlayVisible() const
 	{
 		return EVisibility::Hidden;
 	}
-
-	return EVisibility::Visible;
-}
-
-////////////////////////////////////////////////////////////////////////////////////////////////////
-
-bool STimingProfilerWindow::IsProfilerEnabled() const
-{
-	const bool bIsActive = FInsightsManager::Get()->GetSession().IsValid();
-	return bIsActive;
-}
-
-////////////////////////////////////////////////////////////////////////////////////////////////////
-
-void STimingProfilerWindow::ManageLoadingProgressNotificationState(const FString& Filename, const EInsightsNotificationType NotificatonType, const ELoadingProgressState ProgressState, const float LoadingProgress)
-{
-	const FString BaseFilename = FPaths::GetBaseFilename(Filename);
-
-	if (ProgressState == ELoadingProgressState::Started)
-	{
-		const bool bContains = ActiveNotifications.Contains(Filename);
-		if (!bContains)
-		{
-			FNotificationInfo NotificationInfo(Timing_GetTextForNotification(NotificatonType, ProgressState, BaseFilename));
-			NotificationInfo.bFireAndForget = false;
-			NotificationInfo.bUseLargeFont = false;
-
-			// Add two buttons, one for cancel, one for loading the received file.
-			if (NotificatonType == EInsightsNotificationType::LoadingTraceFile)
-			{
-				const FText CancelButtonText = LOCTEXT("CancelButton_Text", "Cancel");
-				const FText CancelButtonTT = LOCTEXT("CancelButton_TTText", "Hides this notification");
-				const FText LoadButtonText = LOCTEXT("LoadButton_Text", "Load file");
-				const FText LoadButtonTT = LOCTEXT("LoadButton_TTText", "Loads the received file and hides this notification");
-
-				NotificationInfo.ButtonDetails.Add(FNotificationButtonInfo(CancelButtonText, CancelButtonTT,
-					FSimpleDelegate::CreateSP(this, &STimingProfilerWindow::SendingServiceSideCapture_Cancel, Filename), SNotificationItem::CS_Success));
-				NotificationInfo.ButtonDetails.Add(FNotificationButtonInfo(LoadButtonText, LoadButtonTT,
-					FSimpleDelegate::CreateSP(this, &STimingProfilerWindow::SendingServiceSideCapture_Load, Filename), SNotificationItem::CS_Success));
-			}
-
-			SNotificationItemWeak NotificationItem = NotificationList->AddNotification(NotificationInfo);
-			NotificationItem.Pin()->SetCompletionState(SNotificationItem::CS_Pending);
-			ActiveNotifications.Add(Filename, NotificationItem);
-		}
-	}
-	else if (ProgressState == ELoadingProgressState::InProgress)
-	{
-		const SNotificationItemWeak* LoadingProgressPtr = ActiveNotifications.Find(Filename);
-		if (LoadingProgressPtr)
-		{
-			TSharedPtr<SNotificationItem> LoadingProcessPinned = LoadingProgressPtr->Pin();
-			LoadingProcessPinned->SetText(Timing_GetTextForNotification(NotificatonType, ProgressState, BaseFilename, LoadingProgress));
-			LoadingProcessPinned->SetCompletionState(SNotificationItem::CS_Pending);
-		}
-	}
-	else if (ProgressState == ELoadingProgressState::Loaded)
-	{
-		const SNotificationItemWeak* LoadingProgressPtr = ActiveNotifications.Find(Filename);
-		if (LoadingProgressPtr)
-		{
-			TSharedPtr<SNotificationItem> LoadingProcessPinned = LoadingProgressPtr->Pin();
-			LoadingProcessPinned->SetText(Timing_GetTextForNotification(NotificatonType, ProgressState, BaseFilename));
-			LoadingProcessPinned->SetCompletionState(SNotificationItem::CS_Success);
-
-			// Notifications for received files are handled by the user.
-			if (NotificatonType == EInsightsNotificationType::LoadingTraceFile)
-			{
-				LoadingProcessPinned->ExpireAndFadeout();
-				ActiveNotifications.Remove(Filename);
-			}
-		}
-	}
-	else if (ProgressState == ELoadingProgressState::Failed)
-	{
-		const SNotificationItemWeak* LoadingProgressPtr = ActiveNotifications.Find(Filename);
-		if (LoadingProgressPtr)
-		{
-			TSharedPtr<SNotificationItem> LoadingProcessPinned = LoadingProgressPtr->Pin();
-			LoadingProcessPinned->SetText(Timing_GetTextForNotification(NotificatonType, ProgressState, BaseFilename));
-			LoadingProcessPinned->SetCompletionState(SNotificationItem::CS_Fail);
-
-			LoadingProcessPinned->ExpireAndFadeout();
-			ActiveNotifications.Remove(Filename);
-		}
-	}
-	else if (ProgressState == ELoadingProgressState::Cancelled)
-	{
-		const SNotificationItemWeak* LoadingProgressPtr = ActiveNotifications.Find(Filename);
-		if (LoadingProgressPtr)
-		{
-			TSharedPtr<SNotificationItem> LoadingProcessPinned = LoadingProgressPtr->Pin();
-			LoadingProcessPinned->ExpireAndFadeout();
-			ActiveNotifications.Remove(Filename);
-		}
-	}
-}
-
-////////////////////////////////////////////////////////////////////////////////////////////////////
-
-void STimingProfilerWindow::SendingServiceSideCapture_Cancel(const FString Filename)
-{
-	const SNotificationItemWeak* LoadingProgressPtr = ActiveNotifications.Find(Filename);
-	if (LoadingProgressPtr)
+	else
 	{
-		TSharedPtr<SNotificationItem> LoadingProcessPinned = LoadingProgressPtr->Pin();
-		LoadingProcessPinned->ExpireAndFadeout();
-		ActiveNotifications.Remove(Filename);
+		return EVisibility::Visible;
 	}
 }
 
 ////////////////////////////////////////////////////////////////////////////////////////////////////
 
-void STimingProfilerWindow::SendingServiceSideCapture_Load(const FString Filename)
+bool STimingProfilerWindow::IsProfilerEnabled() const
 {
-	const SNotificationItemWeak* LoadingProgressPtr = ActiveNotifications.Find(Filename);
-	if (LoadingProgressPtr)
-	{
-		TSharedPtr<SNotificationItem> LoadingProcessPinned = LoadingProgressPtr->Pin();
-		LoadingProcessPinned->ExpireAndFadeout();
-		ActiveNotifications.Remove(Filename);
-
-		const FString PathName = FPaths::ProfilingDir() + TEXT("UnrealStats/Received/");
-		const FString TraceFilepath = PathName + Filename;
-		FInsightsManager::Get()->LoadTraceFile(TraceFilepath);
-	}
+	return FInsightsManager::Get()->GetSession().IsValid();
 }
 
 ////////////////////////////////////////////////////////////////////////////////////////////////////
@@ -728,34 +488,4 @@ FReply STimingProfilerWindow::OnDrop(const FGeometry& MyGeometry, const FDragDro
 
 ////////////////////////////////////////////////////////////////////////////////////////////////////
 
-void STimingProfilerWindow::OpenProfilerSettings()
-{
-	MainContentPanel->SetEnabled(false);
-	(*OverlaySettingsSlot)
-	[
-		SNew(SBorder)
-		.BorderImage(FEditorStyle::GetBrush("NotificationList.ItemBackground"))
-		.Padding(8.0f)
-		[
-			SNew(SInsightsSettings)
-			.OnClose(this, &STimingProfilerWindow::CloseProfilerSettings)
-			.SettingPtr(&FInsightsManager::GetSettings())
-		]
-	];
-}
-
-////////////////////////////////////////////////////////////////////////////////////////////////////
-
-void STimingProfilerWindow::CloseProfilerSettings()
-{
-	// Close the profiler settings by simply replacing widget with a null one.
-	(*OverlaySettingsSlot)
-	[
-		SNullWidget::NullWidget
-	];
-	MainContentPanel->SetEnabled(true);
-}
-
-////////////////////////////////////////////////////////////////////////////////////////////////////
-
 #undef LOCTEXT_NAMESPACE
diff --git a/Engine/Source/Developer/TraceInsights/Private/Insights/Widgets/STimingProfilerWindow.h b/Engine/Source/Developer/TraceInsights/Private/Insights/Widgets/STimingProfilerWindow.h
--- a/Engine/Source/Developer/TraceInsights/Private/Insights/Widgets/STimingProfilerWindow.h
+++ b/Engine/Source/Developer/TraceInsights/Private/Insights/Widgets/STimingProfilerWindow.h
@@ -10,9 +10,7 @@
 #include "SlateFwd.h"
 #include "Widgets/DeclarativeSyntaxSupport.h"
 #include "Widgets/Docking/SDockTab.h"
-#include "Widgets/Layout/SSplitter.h"
 #include "Widgets/SCompoundWidget.h"
-#include "Widgets/SOverlay.h"
 
 // Insights
 #include "Insights/InsightsManager.h"
@@ -30,17 +28,6 @@ class SLogView;
 
 ////////////////////////////////////////////////////////////////////////////////////////////////////
 
-/** Type definition for shared pointers to instances of SNotificationItem. */
-typedef TSharedPtr<class SNotificationItem> SNotificationItemPtr;
-
-/** Type definition for shared references to instances of SNotificationItem. */
-typedef TSharedRef<class SNotificationItem> SNotificationItemRef;
-
-/** Type definition for weak references to instances of SNotificationItem. */
-typedef TWeakPtr<class SNotificationItem> SNotificationItemWeak;
-
-////////////////////////////////////////////////////////////////////////////////////////////////////
-
 struct FTimingProfilerTabs
 {
 	// Tab identifiers
@@ -71,14 +58,13 @@ public:
 	/** Constructs this widget. */
 	void Construct(const FArguments& InArgs, const TSharedRef<SDockTab>& ConstructUnderMajorTab, const TSharedPtr<SWindow>& ConstructUnderWindow);
 
-	void ManageLoadingProgressNotificationState(const FString& Filename, const EInsightsNotificationType NotificatonType, const ELoadingProgressState ProgressState, const float LoadingProgress);
-
-	void OpenProfilerSettings();
-	void CloseProfilerSettings();
+	void ShowTab(const FName& TabID);
+	void HideTab(const FName& TabID);
+	void ShowHideTab(const FName& TabID, bool bShow) { bShow ? ShowTab(TabID) : HideTab(TabID); }
 
-	void ShowHideTab(const FName& TabID, bool bIsVisible);
+	TSharedPtr<FTabManager> GetTabManager() const { return TabManager; }
 
-protected:
+private:
 	TSharedRef<SDockTab> SpawnTab_Toolbar(const FSpawnTabArgs& Args);
 	TSharedRef<SDockTab> SpawnTab_FramesTrack(const FSpawnTabArgs& Args);
 	TSharedRef<SDockTab> SpawnTab_GraphTrack(const FSpawnTabArgs& Args);
@@ -87,9 +73,6 @@ protected:
 	TSharedRef<SDockTab> SpawnTab_StatsCounters(const FSpawnTabArgs& Args);
 	TSharedRef<SDockTab> SpawnTab_LogView(const FSpawnTabArgs& Args);
 
-	//void RegisterTabSpawners(const TSharedRef<class FTabManager>& InTabManager);
-	//void UnregisterTabSpawners(const TSharedRef<class FTabManager>& InTabManager);
-
 	/**
 	 * Fill the main menu with menu items.
 	 *
@@ -98,38 +81,17 @@ protected:
 	 */
 	static void FillMenu(FMenuBuilder& MenuBuilder, const TSharedPtr<FTabManager> TabManager);
 
-	/** Callback for determining the visibility of the Frames track. */
-	EVisibility IsFramesTrackVisible() const;
-
-	/** Callback for determining the visibility of the Graph track. */
-	EVisibility IsGraphTrackVisible() const;
-
-	/** Callback for determining the visibility of the Timing view. */
-	EVisibility IsTimingViewVisible() const;
-
-	/** Callback for determining the visibility of the Timers View. */
-	EVisibility IsTimersViewVisible() const;
-
-	/** Callback for determining the visibility of the Stats Counters View. */
-	EVisibility IsStatsCountersVisible() const;
-
-	/** Callback for determining the visibility of the Log View. */
-	EVisibility IsLogViewVisible() const;
-
 	/** Callback for determining the visibility of the 'Select a session' overlay. */
 	EVisibility IsSessionOverlayVisible() const;
 
 	/** Callback for getting the enabled state of the profiler window. */
 	bool IsProfilerEnabled() const;
 
-	void SendingServiceSideCapture_Cancel(const FString Filename);
-
-	void SendingServiceSideCapture_Load(const FString Filename);
-
-private:
+	/** Updates the amount of time the profiler has been active. */
+	EActiveTimerReturnType UpdateActiveDuration(double InCurrentTime, float InDeltaTime);
 
 	/**
-	 * Ticks this widget.  Override in derived classes, but always call the parent implementation.
+	 * Ticks this widget. Override in derived classes, but always call the parent implementation.
 	 *
 	 * @param  AllottedGeometry The space allotted for this widget
 	 * @param  InCurrentTime  Current absolute real time
@@ -182,10 +144,6 @@ private:
 	 */
 	virtual FReply OnDragOver(const FGeometry& MyGeometry, const FDragDropEvent& DragDropEvent)  override;
 
-private:
-	/** Updates the amount of time the profiler has been active. */
-	EActiveTimerReturnType UpdateActiveDuration(double InCurrentTime, float InDeltaTime);
-
 public:
 	/** Widget for the frame track */
 	TSharedPtr<SFrameTrack> FrameTrack;
@@ -205,15 +163,6 @@ public:
 	/** Widget for the log view */
 	TSharedPtr<SLogView> LogView;
 
-	/** Widget for the non-intrusive notifications. */
-	TSharedPtr<SNotificationList> NotificationList;
-
-	/** Holds all active and visible notifications, stored as FGuid -> SNotificationItemWeak. */
-	TMap<FString, SNotificationItemWeak> ActiveNotifications;
-
-	/** Overlay slot which contains the profiler settings widget. */
-	SOverlay::FOverlaySlot* OverlaySettingsSlot;
-
 	/** The number of seconds the profiler has been active */
 	float DurationActive;
 
@@ -223,7 +172,4 @@ private:
 
 	/** The handle to the active update duration tick */
 	TWeakPtr<FActiveTimerHandle> ActiveTimerHandle;
-
-	/** Holds all widgets for the profiler window like menu bar, toolbar and tabs. */
-	TSharedPtr<SVerticalBox> MainContentPanel;
 };
diff --git a/Engine/Source/Developer/TraceInsights/Private/Insights/Widgets/STimingView.h b/Engine/Source/Developer/TraceInsights/Private/Insights/Widgets/STimingView.h
--- a/Engine/Source/Developer/TraceInsights/Private/Insights/Widgets/STimingView.h
+++ b/Engine/Source/Developer/TraceInsights/Private/Insights/Widgets/STimingView.h
@@ -330,7 +330,7 @@ protected:
 
 	////////////////////////////////////////////////////////////
 
-	/// The track's viewport. Encapsulates info about position and scale.
+	/** The track's viewport. Encapsulates info about position and scale. */
 	FTimingTrackViewport Viewport;
 
 	bool bIsViewportDirty;
@@ -349,17 +349,17 @@ protected:
 	//	Other
 	//};
 
-	mutable TMap<uint32, FTimingEventsTrack*> CachedTimelines; /// all tracks
+	mutable TMap<uint32, FTimingEventsTrack*> CachedTimelines; /**< all tracks */
 
-	//TODO: TArray<FBaseTimingTrack*> TopTracks; /// tracks docked on top, in order to be displayed (top to bottom)
-	//TODO: TArray<FBaseTimingTrack*> ScrollableTracks; /// tracks in scrollable area, in order to be displayed (top to bottom)
-	//TODO: TArray<FBaseTimingTrack*> BottomTracks; /// tracks docked on bottom, in order to be displayed (top to bottom)
-	//TODO: TArray<FBaseTimingTrack*> ForegorundTracks; /// tracks to draw over top/scrollable/bottom tracks (can use entire area), in order to be displayed (back to front)
+	//TODO: TArray<FBaseTimingTrack*> TopTracks; /**< tracks docked on top, in order to be displayed (top to bottom) */
+	//TODO: TArray<FBaseTimingTrack*> ScrollableTracks; /**< tracks in scrollable area, in order to be displayed (top to bottom) */
+	//TODO: TArray<FBaseTimingTrack*> BottomTracks; /**< tracks docked on bottom, in order to be displayed (top to bottom) */
+	//TODO: TArray<FBaseTimingTrack*> ForegorundTracks; /**< tracks to draw over top/scrollable/bottom tracks (can use entire area), in order to be displayed (back to front) */
 
 	////////////////////////////////////////////////////////////
 	// Timing Events Tracks (GPU, CPU, IO)
 
-	TArray<FTimingEventsTrack*> TimingEventsTracks; /// timing events tracks in order to be displayed
+	TArray<FTimingEventsTrack*> TimingEventsTracks; /**< timing events tracks in order to be displayed */
 
 	bool bAreTimingEventsTracksDirty;
 
@@ -384,18 +384,18 @@ protected:
 		const TCHAR* Path;
 	};
 
-	/// All IO events, cached.
+	/** All IO events, cached. */
 	mutable TArray<FIoTimingEvent> AllIoEvents;
 
 	////////////////////////////////////////////////////////////
 
-	/// The time ruler track.
+	/** The time ruler track. */
 	FTimeRulerTrack TimeRulerTrack;
 
-	/// The time markers track.
+	/** The time markers track. */
 	FMarkersTimingTrack MarkersTrack;
 
-	/// A graph track for frame times.
+	/** A graph track for frame times. */
 	FFramesGraphTrack GraphTrack;
 
 	////////////////////////////////////////////////////////////
@@ -408,15 +408,15 @@ protected:
 
 	////////////////////////////////////////////////////////////
 
-	/// The current mouse position.
+	/** The current mouse position. */
 	FVector2D MousePosition;
 
-	/// Mouse position during the call on mouse button down.
+	/** Mouse position during the call on mouse button down. */
 	FVector2D MousePositionOnButtonDown;
 	double ViewportStartTimeOnButtonDown;
 	float ViewportScrollPosYOnButtonDown;
 
-	/// Mouse position during the call on mouse button up.
+	/** Mouse position during the call on mouse button up. */
 	FVector2D MousePositionOnButtonUp;
 
 	bool bIsLMB_Pressed;
@@ -428,12 +428,14 @@ protected:
 	////////////////////////////////////////////////////////////
 	// Panning
 
-	/// True, if the user is currently interactively panning the view (horizontally and/or vertically),
-	/// either by holding the right mouse button and dragging
-	/// or by holding spacebar pressed and dragging with left mouse button.
+	/**
+	 * True, if the user is currently interactively panning the view (horizontally and/or vertically),
+	 * either by holding the right mouse button and dragging
+	 * or by holding spacebar pressed and dragging with left mouse button.
+	 */
 	bool bIsPanning;
 
-	/// How to pan.
+	/** How to pan. */
 	enum class EPanningMode : uint8
 	{
 		None = 0,
@@ -446,7 +448,7 @@ protected:
 	////////////////////////////////////////////////////////////
 	// Selection
 
-	/// True, if the user is currently changing the selection (by holding the left mouse button and dragging).
+	/** True, if the user is currently changing the selection (by holding the left mouse button and dragging). */
 	bool bIsSelecting;
 
 	double SelectionStartTime;
diff --git a/Engine/Source/Developer/TraceInsights/Public/Insights/IUnrealInsightsModule.h b/Engine/Source/Developer/TraceInsights/Public/Insights/IUnrealInsightsModule.h
--- a/Engine/Source/Developer/TraceInsights/Public/Insights/IUnrealInsightsModule.h
+++ b/Engine/Source/Developer/TraceInsights/Public/Insights/IUnrealInsightsModule.h
@@ -6,15 +6,21 @@
 #include "Modules/ModuleInterface.h"
 #include "Framework/Docking/TabManager.h"
 
-/// Interface for an Unreal Insights module.
+/** Interface for an Unreal Insights module. */
 class IUnrealInsightsModule : public IModuleInterface
 {
 public:
-	/// Called when application starts and a new layout is created. It allows the module to create its own areas.
-	/// @param NewLayout The newly created layout.
+	/**
+	 * Called when application starts and a new layout is created. It allows the module to create its own areas.
+	 *
+	 * @param NewLayout The newly created layout.
+	 */
 	virtual void OnNewLayout(TSharedRef<FTabManager::FLayout> NewLayout) = 0;
 
-	/// Called after application layout was restored.
-	/// @param TabManager The tab manager used by Unreal Insights tools.
+	/**
+	 * Called after application layout was restored.
+	 *
+	 * @param TabManager The tab manager used by Unreal Insights tools.
+	 */
 	virtual void OnLayoutRestored(TSharedPtr<FTabManager> TabManager) = 0;
 };
diff --git a/Engine/Source/Developer/TraceInsights/Public/Insights/TimingProfilerCommon.h b/Engine/Source/Developer/TraceInsights/Public/Insights/TimingProfilerCommon.h
--- a/Engine/Source/Developer/TraceInsights/Public/Insights/TimingProfilerCommon.h
+++ b/Engine/Source/Developer/TraceInsights/Public/Insights/TimingProfilerCommon.h
@@ -19,14 +19,14 @@ DECLARE_LOG_CATEGORY_EXTERN(TimingProfiler, Log, All);
 
 DECLARE_STATS_GROUP(TEXT("TimingProfiler"), STATGROUP_TimingProfiler, STATCAT_Advanced);
 
-/// Time spent on Frame Track drawing.
+/** Time spent on Frame Track drawing. */
 DECLARE_CYCLE_STAT_EXTERN(TEXT("FrameTrackOnPaint"), STAT_FT_OnPaint, STATGROUP_TimingProfiler, );
 
-/// Time spent on Graph Track drawing.
+/** Time spent on Graph Track drawing. */
 DECLARE_CYCLE_STAT_EXTERN(TEXT("GraphTrackOnPaint"), STAT_GT_OnPaint, STATGROUP_TimingProfiler, );
 
-/// Time spent on Timing Track drawing.
+/** Time spent on Timing Track drawing. */
 DECLARE_CYCLE_STAT_EXTERN(TEXT("TimingTrackOnPaint"), STAT_TT_OnPaint, STATGROUP_TimingProfiler, );
 
-/// Time spent on ticking profiler manager.
+/** Time spent on ticking profiler manager. */
 DECLARE_CYCLE_STAT_EXTERN(TEXT("TimingProfilerTick"), STAT_TPM_Tick, STATGROUP_TimingProfiler, );
diff --git a/Engine/Source/Developer/TraceInsights/Public/Insights/Version.h b/Engine/Source/Developer/TraceInsights/Public/Insights/Version.h
--- a/Engine/Source/Developer/TraceInsights/Public/Insights/Version.h
+++ b/Engine/Source/Developer/TraceInsights/Public/Insights/Version.h
@@ -10,8 +10,8 @@
 
 #define UNREAL_INSIGHTS_VERSION_STRING "0.06"
 
-/// Extra identifier string (like: "alpha", "beta.1", "beta.2", "test", "featureX", etc.)
-#define UNREAL_INSIGHTS_VERSION_ID_STRING "im"
+/** Extra identifier string (like: "alpha", "beta.1", "beta.2", "test", "featureX", etc.). */
+//#define UNREAL_INSIGHTS_VERSION_ID_STRING "alpha"
 
 #ifdef UNREAL_INSIGHTS_VERSION_ID_STRING
 #define UNREAL_INSIGHTS_VERSION_STRING_EX "v" UNREAL_INSIGHTS_VERSION_STRING "-" UNREAL_INSIGHTS_VERSION_ID_STRING
