[TraceInsights]
  - Start Page: Added button to Explore Local Session Directory.
 -  Start Page: Added text hint about possibilities of starting an analysis session.
 -  Start Page: Added status for Trace Recorder server with button to Stop / Start it.

diff --git a/Engine/Source/Developer/TraceInsights/Private/Insights/InsightsManager.cpp b/Engine/Source/Developer/TraceInsights/Private/Insights/InsightsManager.cpp
--- a/Engine/Source/Developer/TraceInsights/Private/Insights/InsightsManager.cpp
+++ b/Engine/Source/Developer/TraceInsights/Private/Insights/InsightsManager.cpp
@@ -16,6 +16,10 @@
 
 #define LOCTEXT_NAMESPACE "InsightsManager"
 
+const FName FInsightsManagerTabs::StartPageTabId(TEXT("StartPage"));
+const FName FInsightsManagerTabs::TimingProfilerTabId(TEXT("TimingProfiler"));
+const FName FInsightsManagerTabs::IoProfilerTabId(TEXT("IoProfiler"));
+
 TSharedPtr<FInsightsManager> FInsightsManager::Instance = nullptr;
 
 ////////////////////////////////////////////////////////////////////////////////////////////////////
diff --git a/Engine/Source/Developer/TraceInsights/Private/Insights/InsightsManager.h b/Engine/Source/Developer/TraceInsights/Private/Insights/InsightsManager.h
--- a/Engine/Source/Developer/TraceInsights/Private/Insights/InsightsManager.h
+++ b/Engine/Source/Developer/TraceInsights/Private/Insights/InsightsManager.h
@@ -21,6 +21,15 @@ namespace Trace
 	class ISessionService;
 }
 
+////////////////////////////////////////////////////////////////////////////////////////////////////
+
+struct FInsightsManagerTabs
+{
+	static const FName StartPageTabId;
+	static const FName TimingProfilerTabId;
+	static const FName IoProfilerTabId;
+};
+
 ////////////////////////////////////////////////////////////////////////////////////////////////////
 /** Enumerates loading a trace file progress states. */
 
diff --git a/Engine/Source/Developer/TraceInsights/Private/Insights/TraceInsightsModule.cpp b/Engine/Source/Developer/TraceInsights/Private/Insights/TraceInsightsModule.cpp
--- a/Engine/Source/Developer/TraceInsights/Private/Insights/TraceInsightsModule.cpp
+++ b/Engine/Source/Developer/TraceInsights/Private/Insights/TraceInsightsModule.cpp
@@ -23,12 +23,6 @@
 #include "Insights/Widgets/SStartPageWindow.h"
 #include "Insights/Widgets/STimingProfilerWindow.h"
 
-////////////////////////////////////////////////////////////////////////////////////////////////////
-
-static const FName StartPageTabName(TEXT("StartPage"));
-static const FName TimingProfilerTabName(TEXT("TimingProfiler"));
-static const FName IoProfilerTabName(TEXT("IoProfiler"));
-
 ////////////////////////////////////////////////////////////////////////////////////////////////////
 /**
  * Implements the Trace Insights module (TimingProfiler, IoProfiler, etc.).
@@ -93,14 +87,14 @@ void FTraceInsightsModule::StartupModule()
 	FIoProfilerManager::Initialize();
 
 	// Register tab spawner for the Start Page.
-	auto& StartPageTabSpawnerEntry = FGlobalTabmanager::Get()->RegisterTabSpawner(StartPageTabName,
+	auto& StartPageTabSpawnerEntry = FGlobalTabmanager::Get()->RegisterTabSpawner(FInsightsManagerTabs::StartPageTabId,
 		FOnSpawnTab::CreateRaw(this, &FTraceInsightsModule::SpawnStartPageTab))
 		.SetDisplayName(NSLOCTEXT("FTraceInsightsModule", "StartPageTabTitle", "Unreal Insights"))
 		.SetTooltipText(NSLOCTEXT("FTraceInsightsModule", "StartPageTooltipText", "Open the start page for Unreal Insights."))
 		.SetIcon(FSlateIcon(FInsightsStyle::GetStyleSetName(), "StartPage.Icon.Small"));
 
 	// Register tab spawner for the Timing Insights.
-	auto& TimingProfilerTabSpawnerEntry = FGlobalTabmanager::Get()->RegisterTabSpawner(TimingProfilerTabName,
+	auto& TimingProfilerTabSpawnerEntry = FGlobalTabmanager::Get()->RegisterTabSpawner(FInsightsManagerTabs::TimingProfilerTabId,
 		FOnSpawnTab::CreateRaw(this, &FTraceInsightsModule::SpawnTimingProfilerTab))
 		.SetDisplayName(NSLOCTEXT("FTraceInsightsModule", "TimingProfilerTabTitle", "Timing Insights"))
 		.SetTooltipText(NSLOCTEXT("FTraceInsightsModule", "TimingProfilerTooltipText", "Open the Timing Insights tab."))
@@ -113,7 +107,7 @@ void FTraceInsightsModule::StartupModule()
 //#endif
 
 	// Register tab spawner for the Asset Loading Insights.
-	auto& IoProfilerTabSpawnerEntry = FGlobalTabmanager::Get()->RegisterTabSpawner(IoProfilerTabName,
+	auto& IoProfilerTabSpawnerEntry = FGlobalTabmanager::Get()->RegisterTabSpawner(FInsightsManagerTabs::IoProfilerTabId,
 		FOnSpawnTab::CreateRaw(this, &FTraceInsightsModule::SpawnIoProfilerTab))
 		.SetDisplayName(NSLOCTEXT("FTraceInsightsModule", "IoProfilerTabTitle", "Asset Loading Insights"))
 		.SetTooltipText(NSLOCTEXT("FTraceInsightsModule", "IoProfilerTooltipText", "Open the Asset Loading Insights tab."))
@@ -132,9 +126,9 @@ void FTraceInsightsModule::ShutdownModule()
 {
 	TraceSessionService->StopRecorderServer();
 
-	FGlobalTabmanager::Get()->UnregisterNomadTabSpawner(StartPageTabName);
-	FGlobalTabmanager::Get()->UnregisterNomadTabSpawner(TimingProfilerTabName);
-	FGlobalTabmanager::Get()->UnregisterNomadTabSpawner(IoProfilerTabName);
+	FGlobalTabmanager::Get()->UnregisterNomadTabSpawner(FInsightsManagerTabs::StartPageTabId);
+	FGlobalTabmanager::Get()->UnregisterNomadTabSpawner(FInsightsManagerTabs::TimingProfilerTabId);
+	FGlobalTabmanager::Get()->UnregisterNomadTabSpawner(FInsightsManagerTabs::IoProfilerTabId);
 
 	if (FInsightsManager::Get().IsValid())
 	{
@@ -154,11 +148,11 @@ void FTraceInsightsModule::OnNewLayout(TSharedRef<FTabManager::FLayout> NewLayou
 		->Split
 		(
 			FTabManager::NewStack()
-			->AddTab(StartPageTabName, ETabState::OpenedTab)
-			->AddTab(TimingProfilerTabName, ETabState::ClosedTab)
-			->AddTab(IoProfilerTabName, ETabState::ClosedTab)
+			->AddTab(FInsightsManagerTabs::StartPageTabId, ETabState::OpenedTab)
+			->AddTab(FInsightsManagerTabs::TimingProfilerTabId, ETabState::ClosedTab)
+			->AddTab(FInsightsManagerTabs::IoProfilerTabId, ETabState::ClosedTab)
 			//->AddTab(FName("MemoryProfiler"), ETabState::ClosedTab)
-			->SetForegroundTab(FTabId(StartPageTabName))
+			->SetForegroundTab(FTabId(FInsightsManagerTabs::StartPageTabId))
 		)
 	);
 }
diff --git a/Engine/Source/Developer/TraceInsights/Private/Insights/Widgets/SStartPageWindow.cpp b/Engine/Source/Developer/TraceInsights/Private/Insights/Widgets/SStartPageWindow.cpp
--- a/Engine/Source/Developer/TraceInsights/Private/Insights/Widgets/SStartPageWindow.cpp
+++ b/Engine/Source/Developer/TraceInsights/Private/Insights/Widgets/SStartPageWindow.cpp
@@ -129,6 +129,7 @@ void SStartPageWindow::Construct(const FArguments& InArgs)
 					+ SVerticalBox::Slot()
 						.AutoHeight()
 						.HAlign(HAlign_Center)
+						.Padding(3.0f, 3.0f)
 						[
 							SNew(SBorder)
 							.Visibility(this, &SStartPageWindow::IsSessionOverlayVisible)
@@ -136,25 +137,39 @@ void SStartPageWindow::Construct(const FArguments& InArgs)
 							.Padding(8.0f)
 							[
 								SNew(STextBlock)
-								.Text(LOCTEXT("SelectTraceOverlayText", "Please select a trace."))
+								.Text(LOCTEXT("SelectTraceOverlayText", "Please select a trace..."))
 							]
 						]
 
 					+ SVerticalBox::Slot()
 						.AutoHeight()
+						.HAlign(HAlign_Center)
+						.Padding(3.0f, 3.0f)
 						[
-							SNew(SBox)
-							.HeightOverride(6.0f)
+							SNew(SBorder)
+							.BorderImage(FEditorStyle::GetBrush("NotificationList.ItemBackground"))
+							.Padding(8.0f)
+							[
+								SNew(STextBlock)
+								.Text(LOCTEXT("SelectTraceHint", "\
+- Use \"Live\" button to start a live session.\n\
+   If no live session, it will load the last trace file from the Local Session Directory.\n\
+- Use \"Load\" button to load a trace file (*.utrace).\n\
+- Drag and drop a *.utrace file over this window."))
+								//.Justification(ETextJustify::Center)
+							]
 						]
 
 					+ SVerticalBox::Slot()
 						.AutoHeight()
 						.HAlign(HAlign_Center)
+						.Padding(3.0f, 3.0f)
 						[
 							SNew(SHorizontalBox)
 
 							+ SHorizontalBox::Slot()
 								.AutoWidth()
+								.Padding(3.0f, 0.0f)
 								[
 									SNew(SButton)
 									.OnClicked(this, &SStartPageWindow::Live_OnClicked)
@@ -188,13 +203,7 @@ void SStartPageWindow::Construct(const FArguments& InArgs)
 
 							+ SHorizontalBox::Slot()
 								.AutoWidth()
-								[
-									SNew(SBox)
-									.WidthOverride(6.0f)
-								]
-
-							+ SHorizontalBox::Slot()
-								.AutoWidth()
+								.Padding(3.0f, 0.0f)
 								[
 									SNew(SButton)
 									.OnClicked(this, &SStartPageWindow::Load_OnClicked)
@@ -204,14 +213,14 @@ void SStartPageWindow::Construct(const FArguments& InArgs)
 									[
 										SNew(SHorizontalBox)
 
-										+SHorizontalBox::Slot()
+										+ SHorizontalBox::Slot()
 											.AutoWidth()
 											[
 												SNew(SImage)
 												.Image(FInsightsStyle::GetBrush("Load.Icon.Large"))
 											]
 
-										+SHorizontalBox::Slot()
+										+ SHorizontalBox::Slot()
 											.AutoWidth()
 											.VAlign(VAlign_Center)
 											[
@@ -229,17 +238,61 @@ void SStartPageWindow::Construct(const FArguments& InArgs)
 
 					+ SVerticalBox::Slot()
 						.AutoHeight()
+						.HAlign(HAlign_Center)
+						.Padding(3.0f, 3.0f)
 						[
-							SNew(SBox)
-							.HeightOverride(6.0f)
+							SNew(SBorder)
+							.BorderImage(FEditorStyle::GetBrush("NotificationList.ItemBackground"))
+							.Padding(8.0f)
+							[
+								SNew(SVerticalBox)
+
+								+ SVerticalBox::Slot()
+									.AutoHeight()
+									.HAlign(HAlign_Center)
+									.Padding(0.0, 2.0f)
+									[
+										SNew(STextBlock)
+										.Text(LOCTEXT("LocalSessionDirectoryText", "Local Session Directory:"))
+										.ColorAndOpacity(FLinearColor::Gray)
+									]
+
+								+ SVerticalBox::Slot()
+									.AutoHeight()
+									.HAlign(HAlign_Center)
+									.Padding(0.0, 2.0f)
+									[
+										SNew(SHorizontalBox)
+
+										+ SHorizontalBox::Slot()
+											.AutoWidth()
+											.Padding(2.0f, 0.0f)
+											.VAlign(VAlign_Center)
+											[
+												SNew(STextBlock)
+												.Text(this, &SStartPageWindow::GetLocalSessionDirectory)
+											]
+
+										+ SHorizontalBox::Slot()
+											.AutoWidth()
+											.Padding(2.0f, 0.0f)
+											.VAlign(VAlign_Center)
+											[
+												SNew(SButton)
+												.Text(LOCTEXT("ExploreLocalSessionDirButton", "..."))
+												.ToolTipText(LOCTEXT("ExploreLocalSessionDirButtonToolTip", "Explore the Local Session Directory"))
+												.OnClicked(this, &SStartPageWindow::ExploreLocalSessionDirectory_OnClicked)
+											]
+									]
+							]
 						]
 
 					+ SVerticalBox::Slot()
 						.AutoHeight()
 						.HAlign(HAlign_Center)
+						.Padding(3.0f, 3.0f)
 						[
 							SNew(SBorder)
-							.Visibility(this, &SStartPageWindow::IsSessionOverlayVisible)
 							.BorderImage(FEditorStyle::GetBrush("NotificationList.ItemBackground"))
 							.Padding(8.0f)
 							[
@@ -251,21 +304,52 @@ void SStartPageWindow::Construct(const FArguments& InArgs)
 									.Padding(0.0, 2.0f)
 									[
 										SNew(STextBlock)
-										.Text(LOCTEXT("LocalSessionDirectoryText", "Local Session Directory:"))
+										.Text(LOCTEXT("RecorderStatusTitle", "Trace Recorder Status:"))
 										.ColorAndOpacity(FLinearColor::Gray)
 									]
 
-								+SVerticalBox::Slot()
+								+ SVerticalBox::Slot()
 									.AutoHeight()
 									.HAlign(HAlign_Center)
 									.Padding(0.0, 2.0f)
 									[
-										SNew(STextBlock)
-										.Text(this, &SStartPageWindow::GetLocalSessionDirectory)
+										SNew(SHorizontalBox)
+
+										+ SHorizontalBox::Slot()
+											.AutoWidth()
+											.Padding(2.0f, 0.0f)
+											.VAlign(VAlign_Center)
+											[
+												SNew(STextBlock)
+												.Text(this, &SStartPageWindow::GetRecorderStatusText)
+											]
+
+										+ SHorizontalBox::Slot()
+											.AutoWidth()
+											.Padding(2.0f, 0.0f)
+											.VAlign(VAlign_Center)
+											[
+												SNew(SButton)
+												.Text(LOCTEXT("StartRecorder", "Start"))
+												.ToolTipText(LOCTEXT("StartRecorderToolTip", "Start the Trace Recorder"))
+												.OnClicked(this, &SStartPageWindow::StartTraceRecorder_OnClicked)
+												.Visibility(this, &SStartPageWindow::StartTraceRecorder_Visibility)
+											]
+
+										+ SHorizontalBox::Slot()
+											.AutoWidth()
+											.Padding(2.0f, 0.0f)
+											.VAlign(VAlign_Center)
+											[
+												SNew(SButton)
+												.Text(LOCTEXT("StopRecorder", "Stop"))
+												.ToolTipText(LOCTEXT("StopRecorderToolTip", "Stop the Trace Recorder"))
+												.OnClicked(this, &SStartPageWindow::StopTraceRecorder_OnClicked)
+												.Visibility(this, &SStartPageWindow::StopTraceRecorder_Visibility)
+											]
 									]
 							]
 						]
-
 				]
 
 			// Notification area overlay
@@ -505,8 +589,10 @@ FReply SStartPageWindow::OnDrop(const FGeometry& MyGeometry, const FDragDropEven
 				const FString DraggedFileExtension = FPaths::GetExtension(Files[0], true);
 				if (DraggedFileExtension == TEXT(".utrace"))
 				{
-					// Enqueue load operation.
+					SpawnAndActivateTabs();
+
 					FInsightsManager::Get()->LoadTraceFile(Files[0]);
+
 					return FReply::Handled();
 				}
 			}
@@ -520,34 +606,28 @@ FReply SStartPageWindow::OnDrop(const FGeometry& MyGeometry, const FDragDropEven
 
 void SStartPageWindow::SpawnAndActivateTabs()
 {
-	//TODO: expose those in FInsightsManager; see also TraceInsightsModule.cpp
-	static const FName StartPageTabName(TEXT("StartPage"));
-	static const FName TimingProfilerTabName(TEXT("TimingProfiler"));
-	static const FName IoProfilerTabName(TEXT("IoProfiler"));
-
 	// Open tabs for profilers.
-	if (FGlobalTabmanager::Get()->CanSpawnTab(TimingProfilerTabName))
+	if (FGlobalTabmanager::Get()->CanSpawnTab(FInsightsManagerTabs::TimingProfilerTabId))
 	{
-		FGlobalTabmanager::Get()->InvokeTab(TimingProfilerTabName);
+		FGlobalTabmanager::Get()->InvokeTab(FInsightsManagerTabs::TimingProfilerTabId);
 	}
-	if (FGlobalTabmanager::Get()->CanSpawnTab(IoProfilerTabName))
+	if (FGlobalTabmanager::Get()->CanSpawnTab(FInsightsManagerTabs::IoProfilerTabId))
 	{
-		FGlobalTabmanager::Get()->InvokeTab(IoProfilerTabName);
+		FGlobalTabmanager::Get()->InvokeTab(FInsightsManagerTabs::IoProfilerTabId);
 	}
 
-	// Ensure Timing Insights / Timing View is the active tab.
-	if (TSharedPtr<SDockTab> TimingInsightsTab = FGlobalTabmanager::Get()->FindExistingLiveTab(TimingProfilerTabName))
+	// Ensure Timing Insights / Timing View is the active tab / view.
+	if (TSharedPtr<SDockTab> TimingInsightsTab = FGlobalTabmanager::Get()->FindExistingLiveTab(FInsightsManagerTabs::TimingProfilerTabId))
 	{
-		//FGlobalTabmanager::Get()->SetActiveTab(TimingProfilerTabName); // doesn't really work :(
-		FGlobalTabmanager::Get()->InvokeTab(TimingProfilerTabName); // this activates the tab (brings the tab in foreground)
+		TimingInsightsTab->ActivateInParent(ETabActivationCause::SetDirectly);
 
 		TSharedPtr<class STimingProfilerWindow> Wnd = FTimingProfilerManager::Get()->GetProfilerWindow();
 		if (Wnd)
 		{
 			TSharedPtr<FTabManager> TabManager = Wnd->GetTabManager();
+
 			if (TSharedPtr<SDockTab> TimingViewTab = TabManager->FindExistingLiveTab(FTimingProfilerTabs::TimingViewID))
 			{
-				TabManager->DrawAttention(TimingViewTab.ToSharedRef());
 				TimingViewTab->ActivateInParent(ETabActivationCause::SetDirectly);
 				FSlateApplication::Get().SetKeyboardFocus(TimingViewTab->GetContent());
 			}
@@ -557,6 +637,24 @@ void SStartPageWindow::SpawnAndActivateTabs()
 
 ////////////////////////////////////////////////////////////////////////////////////////////////////
 
+FReply SStartPageWindow::Live_OnClicked()
+{
+	SpawnAndActivateTabs();
+	FInsightsManager::Get()->GetCommandList()->ExecuteAction(FInsightsManager::GetCommands().InsightsManager_Live.ToSharedRef());
+	return FReply::Handled();
+}
+
+////////////////////////////////////////////////////////////////////////////////////////////////////
+
+FReply SStartPageWindow::Load_OnClicked()
+{
+	SpawnAndActivateTabs();
+	FInsightsManager::Get()->GetCommandList()->ExecuteAction(FInsightsManager::GetCommands().InsightsManager_Load.ToSharedRef());
+	return FReply::Handled();
+}
+
+////////////////////////////////////////////////////////////////////////////////////////////////////
+
 FText SStartPageWindow::GetLocalSessionDirectory() const
 {
 	TSharedRef<Trace::ISessionService> SessionService = FInsightsManager::Get()->GetSessionService();
@@ -565,19 +663,60 @@ FText SStartPageWindow::GetLocalSessionDirectory() const
 
 ////////////////////////////////////////////////////////////////////////////////////////////////////
 
-FReply SStartPageWindow::Live_OnClicked()
+FReply SStartPageWindow::ExploreLocalSessionDirectory_OnClicked()
 {
-	SpawnAndActivateTabs();
-	FInsightsManager::Get()->GetCommandList()->ExecuteAction(FInsightsManager::GetCommands().InsightsManager_Live.ToSharedRef());
+	TSharedRef<Trace::ISessionService> SessionService = FInsightsManager::Get()->GetSessionService();
+	FString FullPath(FPaths::ConvertRelativePathToFull(SessionService->GetLocalSessionDirectory()));
+	FPlatformProcess::ExploreFolder(*FullPath);
 	return FReply::Handled();
 }
 
 ////////////////////////////////////////////////////////////////////////////////////////////////////
 
-FReply SStartPageWindow::Load_OnClicked()
+FText SStartPageWindow::GetRecorderStatusText() const
 {
-	SpawnAndActivateTabs();
-	FInsightsManager::Get()->GetCommandList()->ExecuteAction(FInsightsManager::GetCommands().InsightsManager_Load.ToSharedRef());
+	TSharedRef<Trace::ISessionService> SessionService = FInsightsManager::Get()->GetSessionService();
+	if (SessionService->IsRecorderServerRunning())
+	{
+		return FText(LOCTEXT("RecorderServerRunning", "Running"));
+	}
+	else
+	{
+		return FText(LOCTEXT("RecorderServerRunning", "Stopped"));
+	}
+}
+
+////////////////////////////////////////////////////////////////////////////////////////////////////
+
+EVisibility SStartPageWindow::StartTraceRecorder_Visibility() const
+{
+	TSharedRef<Trace::ISessionService> SessionService = FInsightsManager::Get()->GetSessionService();
+	return SessionService->IsRecorderServerRunning() ? EVisibility::Collapsed : EVisibility::Visible;
+}
+
+////////////////////////////////////////////////////////////////////////////////////////////////////
+
+EVisibility SStartPageWindow::StopTraceRecorder_Visibility() const
+{
+	TSharedRef<Trace::ISessionService> SessionService = FInsightsManager::Get()->GetSessionService();
+	return SessionService->IsRecorderServerRunning() ? EVisibility::Visible : EVisibility::Collapsed;
+}
+
+////////////////////////////////////////////////////////////////////////////////////////////////////
+
+FReply SStartPageWindow::StartTraceRecorder_OnClicked()
+{
+	TSharedRef<Trace::ISessionService> SessionService = FInsightsManager::Get()->GetSessionService();
+	SessionService->StartRecorderServer();
+	return FReply::Handled();
+}
+
+////////////////////////////////////////////////////////////////////////////////////////////////////
+
+FReply SStartPageWindow::StopTraceRecorder_OnClicked()
+{
+	TSharedRef<Trace::ISessionService> SessionService = FInsightsManager::Get()->GetSessionService();
+	SessionService->StopRecorderServer();
 	return FReply::Handled();
 }
 
diff --git a/Engine/Source/Developer/TraceInsights/Private/Insights/Widgets/SStartPageWindow.h b/Engine/Source/Developer/TraceInsights/Private/Insights/Widgets/SStartPageWindow.h
--- a/Engine/Source/Developer/TraceInsights/Private/Insights/Widgets/SStartPageWindow.h
+++ b/Engine/Source/Developer/TraceInsights/Private/Insights/Widgets/SStartPageWindow.h
@@ -67,13 +67,20 @@ private:
 	void SendingServiceSideCapture_Cancel(const FString Filename);
 	void SendingServiceSideCapture_Load(const FString Filename);
 
-	FText GetLocalSessionDirectory() const;
-
 	void SpawnAndActivateTabs();
 
 	FReply Live_OnClicked();
 	FReply Load_OnClicked();
 
+	FText GetLocalSessionDirectory() const;
+	FReply ExploreLocalSessionDirectory_OnClicked();
+
+	FText GetRecorderStatusText() const;
+	EVisibility StartTraceRecorder_Visibility() const;
+	EVisibility StopTraceRecorder_Visibility() const;
+	FReply StartTraceRecorder_OnClicked();
+	FReply StopTraceRecorder_OnClicked();
+
 	/** Updates the amount of time the profiler has been active. */
 	EActiveTimerReturnType UpdateActiveDuration(double InCurrentTime, float InDeltaTime);
 
