[TraceInsights] Fixed selecting/unselecting a timing event by clicking if mouse had a small movement. Now the +/- 2px tolerance is allowed as intended.

diff --git a/Engine/Source/Developer/TraceInsights/Private/Insights/Widgets/STimingView.cpp b/Engine/Source/Developer/TraceInsights/Private/Insights/Widgets/STimingView.cpp
--- a/Engine/Source/Developer/TraceInsights/Private/Insights/Widgets/STimingView.cpp
+++ b/Engine/Source/Developer/TraceInsights/Private/Insights/Widgets/STimingView.cpp
@@ -1672,7 +1672,7 @@ FReply STimingView::OnMouseButtonUp(const FGeometry& MyGeometry, const FPointerE
 				bIsSelecting = false;
 			}
 
-			if (!bIsDragging && bIsValidForMouseClick)
+			if (bIsValidForMouseClick)
 			{
 				// Select the hovered timing event (if any).
 				UpdateHoveredTimingEvent(MousePositionOnButtonUp.X, MousePositionOnButtonUp.Y);
