[TraceInsights] Fixed crash when re-starting a Live session after Timing Insights tab was closed.

diff --git a/Engine/Source/Developer/TraceInsights/Private/Insights/TraceInsightsModule.cpp b/Engine/Source/Developer/TraceInsights/Private/Insights/TraceInsightsModule.cpp
--- a/Engine/Source/Developer/TraceInsights/Private/Insights/TraceInsightsModule.cpp
+++ b/Engine/Source/Developer/TraceInsights/Private/Insights/TraceInsightsModule.cpp
@@ -130,6 +130,18 @@ void FTraceInsightsModule::ShutdownModule()
 	FGlobalTabmanager::Get()->UnregisterNomadTabSpawner(FInsightsManagerTabs::TimingProfilerTabId);
 	FGlobalTabmanager::Get()->UnregisterNomadTabSpawner(FInsightsManagerTabs::IoProfilerTabId);
 
+	if (FTimingProfilerManager::Get().IsValid())
+	{
+		// Shutdown the Timing Profiler manager.
+		FTimingProfilerManager::Get()->Shutdown();
+	}
+
+	if (FIoProfilerManager::Get().IsValid())
+	{
+		// Shutdown the I/O Profiler manager.
+		FIoProfilerManager::Get()->Shutdown();
+	}
+
 	if (FInsightsManager::Get().IsValid())
 	{
 		// Shutdown the main manager.
@@ -151,7 +163,6 @@ void FTraceInsightsModule::OnNewLayout(TSharedRef<FTabManager::FLayout> NewLayou
 			->AddTab(FInsightsManagerTabs::StartPageTabId, ETabState::OpenedTab)
 			->AddTab(FInsightsManagerTabs::TimingProfilerTabId, ETabState::ClosedTab)
 			->AddTab(FInsightsManagerTabs::IoProfilerTabId, ETabState::ClosedTab)
-			//->AddTab(FName("MemoryProfiler"), ETabState::ClosedTab)
 			->SetForegroundTab(FTabId(FInsightsManagerTabs::StartPageTabId))
 		)
 	);
@@ -161,19 +172,6 @@ void FTraceInsightsModule::OnNewLayout(TSharedRef<FTabManager::FLayout> NewLayou
 
 void FTraceInsightsModule::OnLayoutRestored(TSharedPtr<FTabManager> TabManager)
 {
-	// Set the Start Page as the main tab.
-	//TSharedPtr<SDockTab> StartPageDockTab = TabManager->FindExistingLiveTab(FTabId(FName("StartPage")));
-	//if (StartPageDockTab.IsValid())
-	//{
-	//	TabManager->SetMainTab(StaticCastSharedRef<SDockTab>(StartPageDockTab->AsShared()));
-	//}
-
-	// Activate the Timing Profiler tab.
-	//TSharedPtr<SDockTab> TimingProfilerDockTab = TabManager->FindExistingLiveTab(FTabId(TimingProfilerTabName));
-	//if (TimingProfilerDockTab.IsValid())
-	//{
-	//	TabManager->SetActiveTab(TimingProfilerDockTab);
-	//}
 }
 
 ////////////////////////////////////////////////////////////////////////////////////////////////////
@@ -214,9 +212,6 @@ void FTraceInsightsModule::OnTimingProfilerTabBeingClosed(TSharedRef<SDockTab> T
 {
 	// Disable TabClosed delegate.
 	TabBeingClosed->SetOnTabClosed(SDockTab::FOnTabClosedCallback());
-
-	// Shutdown the Timing Profiler manager.
-	FTimingProfilerManager::Get()->Shutdown();
 }
 
 ////////////////////////////////////////////////////////////////////////////////////////////////////
@@ -243,9 +238,6 @@ void FTraceInsightsModule::OnIoProfilerTabBeingClosed(TSharedRef<SDockTab> TabBe
 {
 	// Disable TabClosed delegate.
 	TabBeingClosed->SetOnTabClosed(SDockTab::FOnTabClosedCallback());
-
-	// Shutdown the I/O manager.
-	FIoProfilerManager::Get()->Shutdown();
 }
 
 ////////////////////////////////////////////////////////////////////////////////////////////////////
