[TraceInsights] Timing view: Added posibility to click on an empty space to deselect the time selection (i.e. not limited to clicking the time ruller anymore; but making a selection is still limited to start LMB+drag on the timer ruler).

diff --git a/Engine/Source/Developer/TraceInsights/Private/Insights/Widgets/STimingProfilerWindow.cpp b/Engine/Source/Developer/TraceInsights/Private/Insights/Widgets/STimingProfilerWindow.cpp
--- a/Engine/Source/Developer/TraceInsights/Private/Insights/Widgets/STimingProfilerWindow.cpp
+++ b/Engine/Source/Developer/TraceInsights/Private/Insights/Widgets/STimingProfilerWindow.cpp
@@ -205,6 +205,9 @@ void STimingProfilerWindow::Construct(const FArguments& InArgs, const TSharedRef
 		.SetIcon(FSlateIcon(FInsightsStyle::GetStyleSetName(), "LogView.Icon.Small"))
 		.SetGroup(AppMenuGroup);
 
+	TSharedPtr<FTimingProfilerManager> TimingProfilerManager = FTimingProfilerManager::Get();
+	ensure(TimingProfilerManager.IsValid());
+
 	// Create tab layout.
 	const TSharedRef<FTabManager::FLayout> Layout = FTabManager::NewLayout("InsightsTimingProfilerLayout_v1.0")
 		->AddArea
@@ -232,36 +235,36 @@ void STimingProfilerWindow::Construct(const FArguments& InArgs, const TSharedRef
 						FTabManager::NewStack()
 						->SetSizeCoefficient(0.1f)
 						->SetHideTabWell(true)
-						->AddTab(FTimingProfilerTabs::FramesTrackID, FTimingProfilerManager::Get()->IsFramesTrackVisible() ? ETabState::OpenedTab : ETabState::ClosedTab)
+						->AddTab(FTimingProfilerTabs::FramesTrackID, TimingProfilerManager->IsFramesTrackVisible() ? ETabState::OpenedTab : ETabState::ClosedTab)
 					)
 					->Split
 					(
 						FTabManager::NewStack()
 						->SetSizeCoefficient(0.2f)
 						->SetHideTabWell(true)
-						->AddTab(FTimingProfilerTabs::GraphTrackID, FTimingProfilerManager::Get()->IsGraphTrackVisible() ? ETabState::OpenedTab : ETabState::ClosedTab)
+						->AddTab(FTimingProfilerTabs::GraphTrackID, TimingProfilerManager->IsGraphTrackVisible() ? ETabState::OpenedTab : ETabState::ClosedTab)
 					)
 					->Split
 					(
 						FTabManager::NewStack()
 						->SetSizeCoefficient(0.5f)
 						->SetHideTabWell(true)
-						->AddTab(FTimingProfilerTabs::TimingViewID, FTimingProfilerManager::Get()->IsTimingViewVisible() ? ETabState::OpenedTab : ETabState::ClosedTab)
+						->AddTab(FTimingProfilerTabs::TimingViewID, TimingProfilerManager->IsTimingViewVisible() ? ETabState::OpenedTab : ETabState::ClosedTab)
 					)
 					->Split
 					(
 						FTabManager::NewStack()
 						->SetSizeCoefficient(0.2f)
 						->SetHideTabWell(true)
-						->AddTab(FTimingProfilerTabs::LogViewID, FTimingProfilerManager::Get()->IsLogViewVisible() ? ETabState::OpenedTab : ETabState::ClosedTab)
+						->AddTab(FTimingProfilerTabs::LogViewID, TimingProfilerManager->IsLogViewVisible() ? ETabState::OpenedTab : ETabState::ClosedTab)
 					)
 				)
 				->Split
 				(
 					FTabManager::NewStack()
 					->SetSizeCoefficient(0.35f)
-					->AddTab(FTimingProfilerTabs::TimersID, FTimingProfilerManager::Get()->IsTimersViewVisible() ? ETabState::OpenedTab : ETabState::ClosedTab)
-					->AddTab(FTimingProfilerTabs::StatsCountersID, FTimingProfilerManager::Get()->IsStatsCountersViewVisible() ? ETabState::OpenedTab : ETabState::ClosedTab)
+					->AddTab(FTimingProfilerTabs::TimersID, TimingProfilerManager->IsTimersViewVisible() ? ETabState::OpenedTab : ETabState::ClosedTab)
+					->AddTab(FTimingProfilerTabs::StatsCountersID, TimingProfilerManager->IsStatsCountersViewVisible() ? ETabState::OpenedTab : ETabState::ClosedTab)
 					->SetForegroundTab(FTimingProfilerTabs::TimersID)
 				)
 			)
diff --git a/Engine/Source/Developer/TraceInsights/Private/Insights/Widgets/STimingView.cpp b/Engine/Source/Developer/TraceInsights/Private/Insights/Widgets/STimingView.cpp
--- a/Engine/Source/Developer/TraceInsights/Private/Insights/Widgets/STimingView.cpp
+++ b/Engine/Source/Developer/TraceInsights/Private/Insights/Widgets/STimingView.cpp
@@ -1656,6 +1656,7 @@ FReply STimingView::OnMouseButtonUp(const FGeometry& MyGeometry, const FPointerE
 			else if (bIsSelecting)
 			{
 				//TODO: SelectionChangedEvent.Broadcast(SelectionStartTime, SelectionEndTime);
+
 				TSharedPtr<STimingProfilerWindow> Wnd = FTimingProfilerManager::Get()->GetProfilerWindow();
 				if (Wnd)
 				{
@@ -1677,6 +1678,15 @@ FReply STimingView::OnMouseButtonUp(const FGeometry& MyGeometry, const FPointerE
 				// Select the hovered timing event (if any).
 				UpdateHoveredTimingEvent(MousePositionOnButtonUp.X, MousePositionOnButtonUp.Y);
 				SelectHoveredTimingEvent();
+
+				// When clicking on an empty space...
+				if (!SelectedTimingEvent.IsValid())
+				{
+					// ...reset selection.
+					SelectionEndTime = SelectionStartTime = 0.0;
+					LastSelectionType = ESelectionType::None;
+					//TODO: SelectionChangedEvent.Broadcast(SelectionStartTime, SelectionEndTime);
+				}
 			}
 
 			bIsDragging = false;
