UnrealInsights:
API cleanup

diff --git a/Engine/Source/Developer/TraceInsights/Private/Insights/ViewModels/GraphTrack.cpp b/Engine/Source/Developer/TraceInsights/Private/Insights/ViewModels/GraphTrack.cpp
--- a/Engine/Source/Developer/TraceInsights/Private/Insights/ViewModels/GraphTrack.cpp
+++ b/Engine/Source/Developer/TraceInsights/Private/Insights/ViewModels/GraphTrack.cpp
@@ -373,14 +373,12 @@ void FFramesGraphTrack::UpdateSeries(FGraphTrackSeries& Series, const FTimingTra
 	if (Session.IsValid())
 	{
 		Trace::FAnalysisSessionReadScope SessionReadScope(*Session.Get());
-		Session->ReadFramesProvider([this, &Builder, FrameType](const Trace::IFrameProvider& FramesProvider)
+		const Trace::IFrameProvider& FramesProvider = Session->ReadFrameProvider();
+		uint64 FrameCount = FramesProvider.GetFrameCount(FrameType);
+		FramesProvider.EnumerateFrames(FrameType, 0, FrameCount - 1, [&Builder](const Trace::FFrame& Frame)
 		{
-			uint64 FrameCount = FramesProvider.GetFrameCount(FrameType);
-			FramesProvider.EnumerateFrames(FrameType, 0, FrameCount - 1, [&Builder](const Trace::FFrame& Frame)
-			{
-				const double Duration = Frame.EndTime - Frame.StartTime;
-				Builder.AddEvent(Frame.StartTime, Duration, Duration);
-			});
+			const double Duration = Frame.EndTime - Frame.StartTime;
+			Builder.AddEvent(Frame.StartTime, Duration, Duration);
 		});
 	}
 }
diff --git a/Engine/Source/Developer/TraceInsights/Private/Insights/ViewModels/LogMessage.cpp b/Engine/Source/Developer/TraceInsights/Private/Insights/ViewModels/LogMessage.cpp
--- a/Engine/Source/Developer/TraceInsights/Private/Insights/ViewModels/LogMessage.cpp
+++ b/Engine/Source/Developer/TraceInsights/Private/Insights/ViewModels/LogMessage.cpp
@@ -145,18 +145,15 @@ FLogMessageRecord& FLogMessageCache::Get(uint64 Index)
 	if (Session.IsValid())
 	{
 		Trace::FAnalysisSessionReadScope SessionReadScope(*Session.Get());
-		Session->ReadLogProvider(
-			[this, Index](const Trace::ILogProvider& LogProvider)
-			{
-				LogProvider.ReadMessage(
-					Index,
-					[this, Index](const Trace::FLogMessage& Message)
-					{
-						FScopeLock Lock(&CriticalSection);
-						FLogMessageRecord Entry(Message);
-						Map.Add(Index, Entry);
-					});
-			});
+		const Trace::ILogProvider& LogProvider = Session->ReadLogProvider();
+		LogProvider.ReadMessage(
+			Index,
+			[this, Index](const Trace::FLogMessage& Message)
+		{
+				FScopeLock Lock(&CriticalSection);
+				FLogMessageRecord Entry(Message);
+				Map.Add(Index, Entry);
+		});
 	}
 
 	{
@@ -179,16 +176,13 @@ TSharedPtr<FLogMessageRecord> FLogMessageCache::GetUncached(uint64 Index) const
 	if (Session.IsValid())
 	{
 		Trace::FAnalysisSessionReadScope SessionReadScope(*Session.Get());
-		Session->ReadLogProvider(
-			[Index, &EntryPtr](const Trace::ILogProvider& LogProvider)
-			{
-				LogProvider.ReadMessage(
-					Index,
-					[&EntryPtr](const Trace::FLogMessage& Message)
-					{
-						EntryPtr = MakeShareable(new FLogMessageRecord(Message));
-					});
-			});
+		const Trace::ILogProvider& LogProvider = Session->ReadLogProvider();
+		LogProvider.ReadMessage(
+			Index,
+			[&EntryPtr](const Trace::FLogMessage& Message)
+		{
+				EntryPtr = MakeShareable(new FLogMessageRecord(Message));
+		});
 	}
 
 	return EntryPtr;
diff --git a/Engine/Source/Developer/TraceInsights/Private/Insights/ViewModels/MarkersTimingTrack.cpp b/Engine/Source/Developer/TraceInsights/Private/Insights/ViewModels/MarkersTimingTrack.cpp
--- a/Engine/Source/Developer/TraceInsights/Private/Insights/ViewModels/MarkersTimingTrack.cpp
+++ b/Engine/Source/Developer/TraceInsights/Private/Insights/ViewModels/MarkersTimingTrack.cpp
@@ -143,17 +143,15 @@ void FMarkersTimingTrack::Update(const FTimingTrackViewport& InViewport)
 	{
 		Trace::FAnalysisSessionReadScope SessionReadScope(*Session.Get());
 
-		Session->ReadLogProvider([&Builder](const Trace::ILogProvider& LogProvider)
-		{
-			Builder.BeginLog(LogProvider);
+		const Trace::ILogProvider& LogProvider = Session->ReadLogProvider();
+		Builder.BeginLog(LogProvider);
 
-			LogProvider.EnumerateMessages(
-				Builder.GetViewport().StartTime,
-				Builder.GetViewport().EndTime,
-				[&Builder](const Trace::FLogMessage& Message) { Builder.AddLogMessage(Message); });
+		LogProvider.EnumerateMessages(
+			Builder.GetViewport().StartTime,
+			Builder.GetViewport().EndTime,
+			[&Builder](const Trace::FLogMessage& Message) { Builder.AddLogMessage(Message); });
 
-			Builder.EndLog();
-		});
+		Builder.EndLog();
 	}
 }
 
diff --git a/Engine/Source/Developer/TraceInsights/Private/Insights/Widgets/SFrameTrack.cpp b/Engine/Source/Developer/TraceInsights/Private/Insights/Widgets/SFrameTrack.cpp
--- a/Engine/Source/Developer/TraceInsights/Private/Insights/Widgets/SFrameTrack.cpp
+++ b/Engine/Source/Developer/TraceInsights/Private/Insights/Widgets/SFrameTrack.cpp
@@ -186,31 +186,29 @@ void SFrameTrack::Tick(const FGeometry& AllottedGeometry, const double InCurrent
 		{
 			Trace::FAnalysisSessionReadScope SessionReadScope(*Session.Get());
 
-			Session->ReadFramesProvider([this](const Trace::IFrameProvider& FramesProvider)
+			const Trace::IFrameProvider& FramesProvider = Session->ReadFrameProvider();
+			for (int32 FrameType = 0; FrameType < TraceFrameType_Count; ++FrameType)
 			{
-				for (int32 FrameType = 0; FrameType < TraceFrameType_Count; ++FrameType)
-				{
-					FFrameTrackTimeline& CachedTimeline = CachedTimelines.FindOrAdd(FrameType);
-					CachedTimeline.Id = FrameType;
+				FFrameTrackTimeline& CachedTimeline = CachedTimelines.FindOrAdd(FrameType);
+				CachedTimeline.Id = FrameType;
 
-					int32 NumFrames = FramesProvider.GetFrameCount(ETraceFrameType(FrameType));
-					if (NumFrames > Viewport.MaxIndex)
+				int32 NumFrames = FramesProvider.GetFrameCount(ETraceFrameType(FrameType));
+				if (NumFrames > Viewport.MaxIndex)
+				{
+					Viewport.SetMinMaxIndexInterval(0, NumFrames);
+					UpdateHorizontalScrollBar();
+					bIsStateDirty = true;
+
+					// Auto zoom out.
+					float SampleW = Viewport.GetSampleWidth();
+					int32 NumSamples = FMath::CeilToInt(Viewport.Width / SampleW);
+					if (NumFrames > NumSamples * Viewport.GetNumFramesPerSample())
 					{
-						Viewport.SetMinMaxIndexInterval(0, NumFrames);
-						UpdateHorizontalScrollBar();
-						bIsStateDirty = true;
-
-						// Auto zoom out.
-						float SampleW = Viewport.GetSampleWidth();
-						int32 NumSamples = FMath::CeilToInt(Viewport.Width / SampleW);
-						if (NumFrames > NumSamples * Viewport.GetNumFramesPerSample())
-						{
-							ZoomHorizontally(-0.1f, 0.0f);
-							Viewport.ScrollAtPosX(0.0f);
-						}
+						ZoomHorizontally(-0.1f, 0.0f);
+						Viewport.ScrollAtPosX(0.0f);
 					}
 				}
-			});
+			}
 		}
 	}
 
@@ -241,25 +239,23 @@ void SFrameTrack::UpdateState()
 	{
 		Trace::FAnalysisSessionReadScope SessionReadScope(*Session.Get());
 
-		Session->ReadFramesProvider([this](const Trace::IFrameProvider& FramesProvider)
+		const Trace::IFrameProvider& FramesProvider = Session->ReadFrameProvider();
+		for (int32 FrameType = 0; FrameType < TraceFrameType_Count; ++FrameType)
 		{
-			for (int32 FrameType = 0; FrameType < TraceFrameType_Count; ++FrameType)
-			{
-				FFrameTrackTimeline& Timeline = CachedTimelines.FindOrAdd(FrameType);
-				Timeline.Id = FrameType;
+			FFrameTrackTimeline& Timeline = CachedTimelines.FindOrAdd(FrameType);
+			Timeline.Id = FrameType;
 
-				FFrameTrackTimelineBuilder Builder(Timeline, Viewport);
+			FFrameTrackTimelineBuilder Builder(Timeline, Viewport);
 
-				int32 StartIndex = FMath::Max(0, Viewport.GetIndexAtViewportX(0.0f));
-				int32 EndIndex = Viewport.GetIndexAtViewportX(Viewport.Width);
-				FramesProvider.EnumerateFrames(ETraceFrameType(FrameType), static_cast<uint64>(StartIndex), static_cast<uint64>(EndIndex), [&Builder](const Trace::FFrame& Frame)
-				{
-					Builder.AddFrame(Frame);
-				});
+			int32 StartIndex = FMath::Max(0, Viewport.GetIndexAtViewportX(0.0f));
+			int32 EndIndex = Viewport.GetIndexAtViewportX(Viewport.Width);
+			FramesProvider.EnumerateFrames(ETraceFrameType(FrameType), static_cast<uint64>(StartIndex), static_cast<uint64>(EndIndex), [&Builder](const Trace::FFrame& Frame)
+			{
+				Builder.AddFrame(Frame);
+			});
 
-				NumUpdatedFrames += Builder.GetNumAddedFrames();
-			}
-		});
+			NumUpdatedFrames += Builder.GetNumAddedFrames();
+		}
 	}
 
 	Stopwatch.Stop();
diff --git a/Engine/Source/Developer/TraceInsights/Private/Insights/Widgets/SLogView.cpp b/Engine/Source/Developer/TraceInsights/Private/Insights/Widgets/SLogView.cpp
--- a/Engine/Source/Developer/TraceInsights/Private/Insights/Widgets/SLogView.cpp
+++ b/Engine/Source/Developer/TraceInsights/Private/Insights/Widgets/SLogView.cpp
@@ -704,49 +704,47 @@ void SLogView::Tick(const FGeometry& AllottedGeometry, const double InCurrentTim
 	if (Session.IsValid())
 	{
 		Trace::FAnalysisSessionReadScope SessionReadScope(*Session.Get());
-		Session->ReadLogProvider([this, &NewMessageCount](const Trace::ILogProvider& LogProvider)
+		const Trace::ILogProvider& LogProvider = Session->ReadLogProvider();
+		NewMessageCount = static_cast<int32>(LogProvider.GetMessageCount());
+
+		//TODO: show only categories that are used in current trace
+		//TODO: cause of duplicates: a) runtime, b) case insensitive, c) stripped "Log" prefix
+		//TODO: FString vs. FText vs. FName ?
+
+		//TODO: int32 NumCategories = static_cast<int32>(LogProvider.GetCategoriesCount());
+		int32 NumCategories = 0;
+		LogProvider.EnumerateCategories([&NumCategories](const Trace::FLogCategory& Category)
 		{
-			NewMessageCount = static_cast<int32>(LogProvider.GetMessageCount());
+			NumCategories++;
+		});
 
-			//TODO: show only categories that are used in current trace
-			//TODO: cause of duplicates: a) runtime, b) case insensitive, c) stripped "Log" prefix
-			//TODO: FString vs. FText vs. FName ?
+		if (NumCategories != TotalNumCategories)
+		{
+			TotalNumCategories = NumCategories;
+			UE_LOG(TimingProfiler, Log, TEXT("[LogView] Total Log Categories: %d"), TotalNumCategories);
 
-			//TODO: int32 NumCategories = static_cast<int32>(LogProvider.GetCategoriesCount());
-			int32 NumCategories = 0;
-			LogProvider.EnumerateCategories([&NumCategories](const Trace::FLogCategory& Category)
+			TSet<FName> Categories;
+			LogProvider.EnumerateCategories([&Categories](const Trace::FLogCategory& Category)
 			{
-				NumCategories++;
+				FString CategoryStr(Category.Name);
+				if (CategoryStr.StartsWith(TEXT("Log")))
+					CategoryStr = CategoryStr.RightChop(3);
+				if (Categories.Contains(FName(*CategoryStr)))
+				{
+					UE_LOG(TimingProfiler, Log, TEXT("[LogView] Duplicated Log Category: \"%s\""), Category.Name);
+				}
+				Categories.Add(FName(*CategoryStr));
 			});
+			Filter.SyncAvailableCategories(Categories);
+			UE_LOG(TimingProfiler, Log, TEXT("[LogView] Unique Log Categories: %d"), Filter.GetAvailableLogCategories().Num());
 
-			if (NumCategories != TotalNumCategories)
-			{
-				TotalNumCategories = NumCategories;
-				UE_LOG(TimingProfiler, Log, TEXT("[LogView] Total Log Categories: %d"), TotalNumCategories);
-
-				TSet<FName> Categories;
-				LogProvider.EnumerateCategories([&Categories](const Trace::FLogCategory& Category)
-				{
-					FString CategoryStr(Category.Name);
-					if (CategoryStr.StartsWith(TEXT("Log")))
-						CategoryStr = CategoryStr.RightChop(3);
-					if (Categories.Contains(FName(*CategoryStr)))
-					{
-						UE_LOG(TimingProfiler, Log, TEXT("[LogView] Duplicated Log Category: \"%s\""), Category.Name);
-					}
-					Categories.Add(FName(*CategoryStr));
-				});
-				Filter.SyncAvailableCategories(Categories);
-				UE_LOG(TimingProfiler, Log, TEXT("[LogView] Unique Log Categories: %d"), Filter.GetAvailableLogCategories().Num());
-
-				//Cache.Reset();
-				Messages.Reset();
-				TotalNumMessages = 0;
-				//ListView->RebuildList();
-				bIsDirty = true;
-				DirtyStopwatch.Start();
-			}
-		});
+			//Cache.Reset();
+			Messages.Reset();
+			TotalNumMessages = 0;
+			//ListView->RebuildList();
+			bIsDirty = true;
+			DirtyStopwatch.Start();
+		}
 	}
 
 	if (NewMessageCount != TotalNumMessages)
diff --git a/Engine/Source/Developer/TraceInsights/Private/Insights/Widgets/SStatsView.cpp b/Engine/Source/Developer/TraceInsights/Private/Insights/Widgets/SStatsView.cpp
--- a/Engine/Source/Developer/TraceInsights/Private/Insights/Widgets/SStatsView.cpp
+++ b/Engine/Source/Developer/TraceInsights/Private/Insights/Widgets/SStatsView.cpp
@@ -1405,33 +1405,31 @@ void SStatsView::RebuildTree(bool bResync)
 	{
 		Trace::FAnalysisSessionReadScope SessionReadScope(*Session.Get());
 
-		Session->ReadCountersProvider([this, &bResync, &bListHasChanged](const Trace::ICountersProvider& CountersProvider)
+		const Trace::ICounterProvider& CountersProvider = Session->ReadCounterProvider();
+		if (!bResync)
 		{
-			if (!bResync)
-			{
-				bResync = (CountersProvider.GetCounterCount() != StatsNodes.Num());
-			}
+			bResync = (CountersProvider.GetCounterCount() != StatsNodes.Num());
+		}
 
-			if (bResync)
-			{
-				StatsNodes.Empty(StatsNodes.Num());
-				//StatsNodesMap.Empty(StatsNodesMap.Num());
-				StatsNodesIdMap.Empty(StatsNodesIdMap.Num());
-				bListHasChanged = true;
+		if (bResync)
+		{
+			StatsNodes.Empty(StatsNodes.Num());
+			//StatsNodesMap.Empty(StatsNodesMap.Num());
+			StatsNodesIdMap.Empty(StatsNodesIdMap.Num());
+			bListHasChanged = true;
 
-				CountersProvider.EnumerateCounters([this](const Trace::ICounter& Counter)
-				{
-					FName Name(Counter.GetName());
-					FName Group(Counter.GetDisplayHint() == Trace::CounterDisplayHint_Memory ? TEXT("Memory") :
-								Counter.GetDisplayHint() == Trace::CounterDisplayHint_FloatingPoint ? TEXT("float") : TEXT("int64"));
-					EStatsNodeType Type = Counter.GetDisplayHint() == Trace::CounterDisplayHint_FloatingPoint ? EStatsNodeType::Float : EStatsNodeType::Int64;
-					FStatsNodePtr StatsNodePtr = MakeShareable(new FStatsNode(Counter.GetId(), Name, Group, Type));
-					StatsNodes.Add(StatsNodePtr);
-					//StatsNodesMap.Add(Name, StatsNodePtr);
-					StatsNodesIdMap.Add(Counter.GetId(), StatsNodePtr);
-				});
-			}
-		});
+			CountersProvider.EnumerateCounters([this](const Trace::ICounter& Counter)
+			{
+				FName Name(Counter.GetName());
+				FName Group(Counter.GetDisplayHint() == Trace::CounterDisplayHint_Memory ? TEXT("Memory") :
+							Counter.GetDisplayHint() == Trace::CounterDisplayHint_FloatingPoint ? TEXT("float") : TEXT("int64"));
+				EStatsNodeType Type = Counter.GetDisplayHint() == Trace::CounterDisplayHint_FloatingPoint ? EStatsNodeType::Float : EStatsNodeType::Int64;
+				FStatsNodePtr StatsNodePtr = MakeShareable(new FStatsNode(Counter.GetId(), Name, Group, Type));
+				StatsNodes.Add(StatsNodePtr);
+				//StatsNodesMap.Add(Name, StatsNodePtr);
+				StatsNodesIdMap.Add(Counter.GetId(), StatsNodePtr);
+			});
+		}
 	}
 
 	if (bListHasChanged)
@@ -1728,19 +1726,17 @@ void SStatsView::UpdateStats(double StartTime, double EndTime)
 
 			// Compute instance count and total/min/max inclusive/exclusive times for each counter.
 			// Iterate through all counters.
-			Session->ReadCountersProvider([&CalculationHelperDbl, &CalculationHelperInt](const Trace::ICountersProvider& CountersProvider)
+			const Trace::ICounterProvider& CountersProvider = Session->ReadCounterProvider();
+			CountersProvider.EnumerateCounters([&CalculationHelperDbl, &CalculationHelperInt](const Trace::ICounter& Counter)
 			{
-				CountersProvider.EnumerateCounters([&CalculationHelperDbl, &CalculationHelperInt](const Trace::ICounter& Counter)
+				if (Counter.GetDisplayHint() == Trace::CounterDisplayHint_FloatingPoint)
 				{
-					if (Counter.GetDisplayHint() == Trace::CounterDisplayHint_FloatingPoint)
-					{
-						CalculationHelperDbl.Update(Counter);
-					}
-					else
-					{
-						CalculationHelperInt.Update(Counter);
-					}
-				});
+					CalculationHelperDbl.Update(Counter);
+				}
+				else
+				{
+					CalculationHelperInt.Update(Counter);
+				}
 			});
 
 			// Now, as we know min/max inclusive/exclusive times for counter, we can compute histogram and median values.
@@ -1752,19 +1748,17 @@ void SStatsView::UpdateStats(double StartTime, double EndTime)
 
 				// Compute histogram.
 				// Iterate again through all counters.
-				Session->ReadCountersProvider([&CalculationHelperDbl, &CalculationHelperInt](const Trace::ICountersProvider& CountersProvider)
+				const Trace::ICounterProvider& CountersProvider = Session->ReadCounterProvider();
+				CountersProvider.EnumerateCounters([&CalculationHelperDbl, &CalculationHelperInt](const Trace::ICounter& Counter)
 				{
-					CountersProvider.EnumerateCounters([&CalculationHelperDbl, &CalculationHelperInt](const Trace::ICounter& Counter)
+					if (Counter.GetDisplayHint() == Trace::CounterDisplayHint_FloatingPoint)
 					{
-						if (Counter.GetDisplayHint() == Trace::CounterDisplayHint_FloatingPoint)
-						{
-							CalculationHelperDbl.UpdateHistograms(Counter);
-						}
-						else
-						{
-							CalculationHelperInt.UpdateHistograms(Counter);
-						}
-					});
+						CalculationHelperDbl.UpdateHistograms(Counter);
+					}
+					else
+					{
+						CalculationHelperInt.UpdateHistograms(Counter);
+					}
 				});
 			}
 		}
diff --git a/Engine/Source/Developer/TraceInsights/Private/Insights/Widgets/STimersView.cpp b/Engine/Source/Developer/TraceInsights/Private/Insights/Widgets/STimersView.cpp
--- a/Engine/Source/Developer/TraceInsights/Private/Insights/Widgets/STimersView.cpp
+++ b/Engine/Source/Developer/TraceInsights/Private/Insights/Widgets/STimersView.cpp
@@ -1435,36 +1435,34 @@ void STimersView::RebuildTree(bool bResync)
 	{
 		Trace::FAnalysisSessionReadScope SessionReadScope(*Session.Get());
 
-		Session->ReadTimingProfilerProvider([this, &bResync, &bListHasChanged](const Trace::ITimingProfilerProvider& TimingProfilerProvider)
+		const Trace::ITimingProfilerProvider& TimingProfilerProvider = Session->ReadTimingProfilerProvider();
+		TimingProfilerProvider.ReadTimers([this, &bResync, &bListHasChanged](const Trace::FTimingProfilerTimer* Timers, uint64 TimersCount)
 		{
-			TimingProfilerProvider.ReadTimers([this, &bResync, &bListHasChanged](const Trace::FTimingProfilerTimer* Timers, uint64 TimersCount)
+			if (!bResync)
 			{
-				if (!bResync)
-				{
-					bResync = (TimersCount != TimerNodes.Num());
-				}
+				bResync = (TimersCount != TimerNodes.Num());
+			}
 
-				if (bResync)
+			if (bResync)
+			{
+				TimerNodes.Empty(TimerNodes.Num());
+				//TimerNodesMap.Empty(TimerNodesMap.Num());
+				TimerNodesIdMap.Empty(TimerNodesIdMap.Num());
+				bListHasChanged = true;
+
+				for (uint64 TimerIndex = 0; TimerIndex < TimersCount; ++TimerIndex)
 				{
-					TimerNodes.Empty(TimerNodes.Num());
-					//TimerNodesMap.Empty(TimerNodesMap.Num());
-					TimerNodesIdMap.Empty(TimerNodesIdMap.Num());
-					bListHasChanged = true;
-
-					for (uint64 TimerIndex = 0; TimerIndex < TimersCount; ++TimerIndex)
-					{
-						const Trace::FTimingProfilerTimer& Timer = Timers[TimerIndex];
-						FName Name(Timer.Name);// +TEXT(" [GPU]")));
-						FName Group(Timer.IsGpuTimer ? TEXT("GPU") : TEXT("CPU"));
-						ETimerNodeType Type = Timer.IsGpuTimer ? ETimerNodeType::GpuScope : ETimerNodeType::CpuScope;
-						FTimerNode* TimerPtr = new FTimerNode(Timer.Id, Name, Group, Type);
-						FTimerNodePtr TimerNodePtr = MakeShareable(TimerPtr);
-						TimerNodes.Add(TimerNodePtr);
-						//TimerNodesMap.Add(Name, TimerNodePtr);
-						TimerNodesIdMap.Add(Timer.Id, TimerNodePtr);
-					}
+					const Trace::FTimingProfilerTimer& Timer = Timers[TimerIndex];
+					FName Name(Timer.Name);// +TEXT(" [GPU]")));
+					FName Group(Timer.IsGpuTimer ? TEXT("GPU") : TEXT("CPU"));
+					ETimerNodeType Type = Timer.IsGpuTimer ? ETimerNodeType::GpuScope : ETimerNodeType::CpuScope;
+					FTimerNode* TimerPtr = new FTimerNode(Timer.Id, Name, Group, Type);
+					FTimerNodePtr TimerNodePtr = MakeShareable(TimerPtr);
+					TimerNodes.Add(TimerNodePtr);
+					//TimerNodesMap.Add(Name, TimerNodePtr);
+					TimerNodesIdMap.Add(Timer.Id, TimerNodePtr);
 				}
-			});
+			}
 		});
 	}
 
@@ -1487,17 +1485,17 @@ void STimersView::UpdateStats(double StartTime, double EndTime)
 	if (Session.IsValid() &&
 		StartTime < EndTime)
 	{
+		auto ThreadFilter = [](uint32 ThreadId)
+		{
+			return true;
+		};
+
 		TUniquePtr<Trace::ITable<Trace::FTimingProfilerAggregatedStats>> AggregationResultTable;
-		Trace::FAnalysisSessionReadScope SessionReadScope(*Session.Get());
-		Session->ReadTimingProfilerProvider([StartTime, EndTime, &AggregationResultTable](const Trace::ITimingProfilerProvider& TimingProfilerProvider)
 		{
-			auto ThreadFilter = [](uint32 ThreadId)
-			{
-				return true;
-			};
+			Trace::FAnalysisSessionReadScope SessionReadScope(*Session.Get());
+			const Trace::ITimingProfilerProvider& TimingProfilerProvider = Session->ReadTimingProfilerProvider();
 			AggregationResultTable.Reset(TimingProfilerProvider.CreateAggregation(StartTime, EndTime, ThreadFilter, true));
-		});
-
+		}
 		TUniquePtr<Trace::ITableReader<Trace::FTimingProfilerAggregatedStats>> AggregationResultTableReader(AggregationResultTable->CreateReader());
 		while (AggregationResultTableReader->IsValid())
 		{
diff --git a/Engine/Source/Developer/TraceInsights/Private/Insights/Widgets/STimingView.cpp b/Engine/Source/Developer/TraceInsights/Private/Insights/Widgets/STimingView.cpp
--- a/Engine/Source/Developer/TraceInsights/Private/Insights/Widgets/STimingView.cpp
+++ b/Engine/Source/Developer/TraceInsights/Private/Insights/Widgets/STimingView.cpp
@@ -404,49 +404,45 @@ void STimingView::Tick(const FGeometry& AllottedGeometry, const double InCurrent
 
 		// Sync CachedTimelines and TimingEventsTracks with available timelines on analysis side.
 		// TODO: can we make this more efficient?
-		Session->ReadTimingProfilerProvider([this, &bIsTimingEventsTrackDirty, &Session](const Trace::ITimingProfilerProvider& TimingProfilerProvider)
+		const Trace::ITimingProfilerProvider& TimingProfilerProvider = Session->ReadTimingProfilerProvider();
+		// Check if we have a GPU track.
+		uint32 GpuTimelineIndex;
+		if (TimingProfilerProvider.GetGpuTimelineIndex(GpuTimelineIndex))
 		{
-			// Check if we have a GPU track.
-			uint32 GpuTimelineIndex;
-			if (TimingProfilerProvider.GetGpuTimelineIndex(GpuTimelineIndex))
+			if (!CachedTimelines.Contains(GpuTimelineIndex))
 			{
-				if (!CachedTimelines.Contains(GpuTimelineIndex))
-				{
-					AddTimingEventsTrack(GpuTimelineIndex, ETimingEventsTrackType::Gpu, TEXT("GPU"), 0);
-					bIsTimingEventsTrackDirty = true;
-				}
+				AddTimingEventsTrack(GpuTimelineIndex, ETimingEventsTrackType::Gpu, TEXT("GPU"), 0);
+				bIsTimingEventsTrackDirty = true;
 			}
+		}
 
-			// Check available CPU tracks.
-			Session->ReadThreadProvider([this, &bIsTimingEventsTrackDirty, &TimingProfilerProvider](const Trace::IThreadProvider& ThreadProvider)
+		// Check available CPU tracks.
+		const Trace::IThreadProvider& ThreadProvider = Session->ReadThreadProvider();
+		int32 Order = 1;
+		ThreadProvider.EnumerateThreads([this, &Order, &bIsTimingEventsTrackDirty, &TimingProfilerProvider](const Trace::FThreadInfo& ThreadInfo)
+		{
+			uint32 CpuTimelineIndex;
+			if (TimingProfilerProvider.GetCpuThreadTimelineIndex(ThreadInfo.Id, CpuTimelineIndex))
 			{
-				int32 Order = 1;
-				ThreadProvider.EnumerateThreads([this, &Order, &bIsTimingEventsTrackDirty, &TimingProfilerProvider](const Trace::FThreadInfo& ThreadInfo)
+				if (!CachedTimelines.Contains(CpuTimelineIndex))
 				{
-					uint32 CpuTimelineIndex;
-					if (TimingProfilerProvider.GetCpuThreadTimelineIndex(ThreadInfo.Id, CpuTimelineIndex))
+					FString TrackName(ThreadInfo.Name && *ThreadInfo.Name ? ThreadInfo.Name : FString::Printf(TEXT("Thread %u")));
+					AddTimingEventsTrack(CpuTimelineIndex, ETimingEventsTrackType::Cpu, TrackName, Order);
+					bIsTimingEventsTrackDirty = true;
+				}
+				else
+				{
+					FTimingEventsTrack* Track = CachedTimelines[CpuTimelineIndex];
+					if (Track->GetOrder() != Order)
 					{
-						if (!CachedTimelines.Contains(CpuTimelineIndex))
-						{
-							FString TrackName(ThreadInfo.Name && *ThreadInfo.Name ? ThreadInfo.Name : FString::Printf(TEXT("Thread %u")));
-							AddTimingEventsTrack(CpuTimelineIndex, ETimingEventsTrackType::Cpu, TrackName, Order);
-							bIsTimingEventsTrackDirty = true;
-						}
-						else
-						{
-							FTimingEventsTrack* Track = CachedTimelines[CpuTimelineIndex];
-							if (Track->GetOrder() != Order)
-							{
-								Track->SetOrder(Order);
-								bIsTimingEventsTrackDirty = true;
-							}
-						}
-						Order++;
+						Track->SetOrder(Order);
+						bIsTimingEventsTrackDirty = true;
 					}
-				});
-			});
+				}
+				Order++;
+			}
 		});
-
+		
 		if (IoActivityTrack == nullptr)
 		{
 			// Note: This track is hardcoded for now, as the I/O timelines are just prototypes for now (will be removed once the functionality is moved in analyzer).
@@ -1025,31 +1021,29 @@ void STimingView::DrawCpuGpuTimelineTrack(FTimingViewDrawHelper& Helper, FTiming
 		{
 			Trace::FAnalysisSessionReadScope SessionReadScope(*Session.Get());
 
-			Session->ReadTimingProfilerProvider([this, &Helper, &Track](const Trace::ITimingProfilerProvider& TimingProfilerProvider)
+			const Trace::ITimingProfilerProvider& TimingProfilerProvider = Session->ReadTimingProfilerProvider();
+			TimingProfilerProvider.ReadTimers([this, &Helper, &Track, &TimingProfilerProvider](const Trace::FTimingProfilerTimer* Timers, uint64 TimersCount)
 			{
-				TimingProfilerProvider.ReadTimers([this, &Helper, &Track, &TimingProfilerProvider](const Trace::FTimingProfilerTimer* Timers, uint64 TimersCount)
+				TimingProfilerProvider.ReadTimeline(Track.GetId(), [this, &Helper, &Track, Timers](const Trace::ITimingProfilerProvider::Timeline& Timeline)
 				{
-					TimingProfilerProvider.ReadTimeline(Track.GetId(), [this, &Helper, &Track, Timers](const Trace::ITimingProfilerProvider::Timeline& Timeline)
+					if (bUseDownSampling)
 					{
-						if (bUseDownSampling)
+						const double SecondsPerPixel = 1.0 / Helper.GetViewport().ScaleX;
+						Timeline.EnumerateEventsDownSampled(Helper.GetViewport().StartTime, Helper.GetViewport().EndTime, SecondsPerPixel, [this, &Helper, Timers](double StartTime, double EndTime, uint32 Depth, const Trace::FTimingProfilerEvent& Event)
 						{
-							const double SecondsPerPixel = 1.0 / Helper.GetViewport().ScaleX;
-							Timeline.EnumerateEventsDownSampled(Helper.GetViewport().StartTime, Helper.GetViewport().EndTime, SecondsPerPixel, [this, &Helper, Timers](double StartTime, double EndTime, uint32 Depth, const Trace::FTimingProfilerEvent& Event)
-							{
-								Helper.AddEvent(StartTime, EndTime, Depth, Timers[Event.TimerIndex].Name);
-							});
-						}
-						else
+							Helper.AddEvent(StartTime, EndTime, Depth, Timers[Event.TimerIndex].Name);
+						});
+					}
+					else
+					{
+						Timeline.EnumerateEvents(Helper.GetViewport().StartTime, Helper.GetViewport().EndTime, [this, &Helper, Timers](double StartTime, double EndTime, uint32 Depth, const Trace::FTimingProfilerEvent& Event)
 						{
-							Timeline.EnumerateEvents(Helper.GetViewport().StartTime, Helper.GetViewport().EndTime, [this, &Helper, Timers](double StartTime, double EndTime, uint32 Depth, const Trace::FTimingProfilerEvent& Event)
-							{
-								Helper.AddEvent(StartTime, EndTime, Depth, Timers[Event.TimerIndex].Name);
-							});
-						}
-					});
+							Helper.AddEvent(StartTime, EndTime, Depth, Timers[Event.TimerIndex].Name);
+						});
+					}
 				});
 			});
-
+		
 			Helper.EndTimeline(Track);
 		}
 	}
@@ -1172,28 +1166,26 @@ void STimingView::UpdateIo()
 		{
 			Trace::FAnalysisSessionReadScope SessionReadScope(*Session.Get());
 
-			Session->ReadFileActivityProvider([this](const Trace::IFileActivityProvider& FileActivityProvider)
+			const Trace::IFileActivityProvider& FileActivityProvider = Session->ReadFileActivityProvider();
+			// Enumerate all IO events and cache them.
+			FileActivityProvider.EnumerateFileActivity([this](const Trace::FFileInfo& FileInfo, const Trace::IFileActivityProvider::Timeline& Timeline)
 			{
-				// Enumerate all IO events and cache them.
-				FileActivityProvider.EnumerateFileActivity([this](const Trace::FFileInfo& FileInfo, const Trace::IFileActivityProvider::Timeline& Timeline)
+				Timeline.EnumerateEvents(-std::numeric_limits<double>::infinity(), +std::numeric_limits<double>::infinity(),
+					[this, &FileInfo, &Timeline](double EventStartTime, double EventEndTime, uint32 EventDepth, const Trace::FFileActivity& FileActivity)
 				{
-					Timeline.EnumerateEvents(-std::numeric_limits<double>::infinity(), +std::numeric_limits<double>::infinity(),
-						[this, &FileInfo, &Timeline](double EventStartTime, double EventEndTime, uint32 EventDepth, const Trace::FFileActivity& FileActivity)
+					if (bMergeIoLanes)
 					{
-						if (bMergeIoLanes)
-						{
-							EventDepth = static_cast<uint32>(FileInfo.Id); // used by MergeLanes algorithm
-						}
-						else
-						{
-							EventDepth = FileInfo.Id % 32; // simple layout
-						}
-						uint32 Type = ((uint32)FileActivity.ActivityType & 0x0F) | (FileActivity.Failed ? 0x80 : 0);
-						AllIoEvents.Add(FIoTimingEvent{ EventStartTime, EventEndTime, EventDepth, Type, FileInfo.Path });
-					});
-
-					return true;
+						EventDepth = static_cast<uint32>(FileInfo.Id); // used by MergeLanes algorithm
+					}
+					else
+					{
+						EventDepth = FileInfo.Id % 32; // simple layout
+					}
+					uint32 Type = ((uint32)FileActivity.ActivityType & 0x0F) | (FileActivity.Failed ? 0x80 : 0);
+					AllIoEvents.Add(FIoTimingEvent{ EventStartTime, EventEndTime, EventDepth, Type, FileInfo.Path });
 				});
+
+				return true;
 			});
 		}
 
@@ -2519,86 +2511,82 @@ bool STimingView::SearchTimingEvent(const double InStartTime,
 		Trace::FAnalysisSessionReadScope SessionReadScope(*Session.Get());
 
 		//if (Ctx.TimingEvent.Track->GetType() == TrackType::Timing)
-		Session->ReadTimingProfilerProvider([&Ctx](const Trace::ITimingProfilerProvider& TimingProfilerProvider)
+		const Trace::ITimingProfilerProvider& TimingProfilerProvider = Session->ReadTimingProfilerProvider();
+		TimingProfilerProvider.ReadTimeline(Ctx.TimingEvent.Track->GetId(), [&Ctx](const Trace::ITimingProfilerProvider::Timeline& Timeline)
 		{
-			TimingProfilerProvider.ReadTimeline(Ctx.TimingEvent.Track->GetId(), [&Ctx](const Trace::ITimingProfilerProvider::Timeline& Timeline)
+			Timeline.EnumerateEvents(Ctx.StartTime, Ctx.EndTime, [&Ctx](double EventStartTime, double EventEndTime, uint32 EventDepth, const Trace::FTimingProfilerEvent& Event)
 			{
-				Timeline.EnumerateEvents(Ctx.StartTime, Ctx.EndTime, [&Ctx](double EventStartTime, double EventEndTime, uint32 EventDepth, const Trace::FTimingProfilerEvent& Event)
+				if (Ctx.bContinueSearching && Ctx.Predicate(EventStartTime, EventEndTime, EventDepth, Event.TimerIndex))
 				{
-					if (Ctx.bContinueSearching && Ctx.Predicate(EventStartTime, EventEndTime, EventDepth, Event.TimerIndex))
+					if (!Ctx.bSearchForLargestEvent || EventEndTime - EventStartTime > Ctx.LargestDuration)
 					{
-						if (!Ctx.bSearchForLargestEvent || EventEndTime - EventStartTime > Ctx.LargestDuration)
-						{
-							Ctx.LargestDuration = EventEndTime - EventStartTime;
+						Ctx.LargestDuration = EventEndTime - EventStartTime;
 
-							Ctx.TimingEvent.TypeId = Event.TimerIndex;
-							Ctx.TimingEvent.Depth = EventDepth;
-							Ctx.TimingEvent.StartTime = EventStartTime;
-							Ctx.TimingEvent.EndTime = EventEndTime;
+						Ctx.TimingEvent.TypeId = Event.TimerIndex;
+						Ctx.TimingEvent.Depth = EventDepth;
+						Ctx.TimingEvent.StartTime = EventStartTime;
+						Ctx.TimingEvent.EndTime = EventEndTime;
 
-							Ctx.bFound = true;
-							Ctx.bContinueSearching = !Ctx.bStopAtFirstMatch || Ctx.bSearchForLargestEvent;
-						}
+						Ctx.bFound = true;
+						Ctx.bContinueSearching = !Ctx.bStopAtFirstMatch || Ctx.bSearchForLargestEvent;
 					}
-				});
+				}
 			});
 		});
-
+	
 		if (Ctx.bFound)
 		{
-			Session->ReadTimingProfilerProvider([&Ctx](const Trace::ITimingProfilerProvider& TimingProfilerProvider)
+			const Trace::ITimingProfilerProvider& TimingProfilerProvider = Session->ReadTimingProfilerProvider();
+			TimingProfilerProvider.ReadTimeline(Ctx.TimingEvent.Track->GetId(), [&Ctx](const Trace::ITimingProfilerProvider::Timeline& Timeline)
 			{
-				TimingProfilerProvider.ReadTimeline(Ctx.TimingEvent.Track->GetId(), [&Ctx](const Trace::ITimingProfilerProvider::Timeline& Timeline)
+				struct FEnumerationState
 				{
-					struct FEnumerationState
-					{
-						double EventStartTime;
-						double EventEndTime;
-						uint64 EventDepth;
-
-						uint64 CurrentDepth;
-						double LastTime;
-						double ExclusiveTime;
-						bool IsInEventScope;
-					};
-					FEnumerationState State;
-
-					State.EventStartTime = Ctx.TimingEvent.StartTime;
-					State.EventEndTime = Ctx.TimingEvent.EndTime;
-					State.EventDepth = Ctx.TimingEvent.Depth;
-
-					State.CurrentDepth = 0;
-					State.LastTime = 0.0;
-					State.ExclusiveTime = 0.0;
-					State.IsInEventScope = false;
-
-					Timeline.EnumerateEvents(Ctx.TimingEvent.StartTime, Ctx.TimingEvent.EndTime, [&State](bool IsEnter, double Time, const Trace::FTimingProfilerEvent& Event)
+					double EventStartTime;
+					double EventEndTime;
+					uint64 EventDepth;
+
+					uint64 CurrentDepth;
+					double LastTime;
+					double ExclusiveTime;
+					bool IsInEventScope;
+				};
+				FEnumerationState State;
+
+				State.EventStartTime = Ctx.TimingEvent.StartTime;
+				State.EventEndTime = Ctx.TimingEvent.EndTime;
+				State.EventDepth = Ctx.TimingEvent.Depth;
+
+				State.CurrentDepth = 0;
+				State.LastTime = 0.0;
+				State.ExclusiveTime = 0.0;
+				State.IsInEventScope = false;
+
+				Timeline.EnumerateEvents(Ctx.TimingEvent.StartTime, Ctx.TimingEvent.EndTime, [&State](bool IsEnter, double Time, const Trace::FTimingProfilerEvent& Event)
+				{
+					if (IsEnter)
 					{
-						if (IsEnter)
+						if (State.IsInEventScope && State.CurrentDepth == State.EventDepth + 1)
 						{
-							if (State.IsInEventScope && State.CurrentDepth == State.EventDepth + 1)
-							{
-								State.ExclusiveTime += Time - State.LastTime;
-							}
-							if (State.CurrentDepth == State.EventDepth && Time == State.EventStartTime)
-							{
-								State.IsInEventScope = true;
-							}
-							++State.CurrentDepth;
+							State.ExclusiveTime += Time - State.LastTime;
 						}
-						else
+						if (State.CurrentDepth == State.EventDepth && Time == State.EventStartTime)
 						{
-							--State.CurrentDepth;
-							if (State.CurrentDepth == State.EventDepth && Time == State.EventEndTime)
-							{
-								State.IsInEventScope = false;
-								State.ExclusiveTime += Time - State.LastTime;
-							}
+							State.IsInEventScope = true;
 						}
-						State.LastTime = Time;
-					});
-					Ctx.TimingEvent.ExclusiveTime = State.ExclusiveTime;
+						++State.CurrentDepth;
+					}
+					else
+					{
+						--State.CurrentDepth;
+						if (State.CurrentDepth == State.EventDepth && Time == State.EventEndTime)
+						{
+							State.IsInEventScope = false;
+							State.ExclusiveTime += Time - State.LastTime;
+						}
+					}
+					State.LastTime = Time;
 				});
+				Ctx.TimingEvent.ExclusiveTime = State.ExclusiveTime;
 			});
 		}
 	}
diff --git a/Engine/Source/Developer/TraceServices/Private/AnalysisService.cpp b/Engine/Source/Developer/TraceServices/Private/AnalysisService.cpp
--- a/Engine/Source/Developer/TraceServices/Private/AnalysisService.cpp
+++ b/Engine/Source/Developer/TraceServices/Private/AnalysisService.cpp
@@ -14,13 +14,6 @@
 #include "Analyzers/PlatformFileTraceAnalysis.h"
 #include "Analyzers/StatsTraceAnalysis.h"
 #include "Math/RandomStream.h"
-#include "Model/Bookmarks.h"
-#include "Model/Log.h"
-#include "Model/Threads.h"
-#include "Model/FileActivity.h"
-#include "Model/Frames.h"
-#include "Model/LoadTimeProfiler.h"
-#include "Model/Counters.h"
 
 namespace Trace
 {
@@ -71,66 +64,18 @@ FAnalysisSession::FAnalysisSession(const TCHAR* SessionName)
 	: DurationSeconds(0.0)
 	, Allocator(32 << 20)
 	, StringStore(Allocator)
-	, BookmarkProvider(MakeShared<FBookmarkProvider>(Lock, StringStore))
-	, LogProvider(MakeShared<FLogProvider>(Allocator, Lock, StringStore))
-	, ThreadProvider(MakeShared<FThreadProvider>(Lock, StringStore))
-	, FramesProvider(MakeShared<FFrameProvider>(Allocator, Lock))
-	, TimingProfilerProvider(MakeShared<FTimingProfilerProvider>(Allocator, Lock, StringStore))
-	, FileActivityProvider(MakeShared<FFileActivityProvider>(Allocator, Lock))
-	, LoadTimeProfilerProvider(MakeShared<FLoadTimeProfilerProvider>(Allocator, Lock, StringStore))
-	, CountersProvider(MakeShared<FCountersProvider>(Allocator, Lock))
+	, BookmarkProvider(Lock, StringStore)
+	, LogProvider(Allocator, Lock, StringStore)
+	, ThreadProvider(Lock, StringStore)
+	, FrameProvider(Allocator, Lock)
+	, TimingProfilerProvider(Allocator, Lock, StringStore)
+	, FileActivityProvider(Allocator, Lock)
+	, LoadTimeProfilerProvider(Allocator, Lock, StringStore)
+	, CounterProvider(Allocator, Lock)
 {
 
 }
 
-void FAnalysisSession::ReadBookmarkProvider(TFunctionRef<void(const IBookmarkProvider&)> Callback) const
-{
-	Lock.ReadAccessCheck();
-	Callback(*BookmarkProvider);
-}
-
-void FAnalysisSession::ReadLogProvider(TFunctionRef<void(const ILogProvider&)> Callback) const
-{
-	Lock.ReadAccessCheck();
-	Callback(*LogProvider);
-}
-
-void FAnalysisSession::ReadThreadProvider(TFunctionRef<void(const IThreadProvider&)> Callback) const
-{
-	Lock.ReadAccessCheck();
-	Callback(*ThreadProvider);
-}
-
-void FAnalysisSession::ReadFramesProvider(TFunctionRef<void(const IFrameProvider&)> Callback) const
-{
-	Lock.ReadAccessCheck();
-	Callback(*FramesProvider);
-}
-
-void FAnalysisSession::ReadTimingProfilerProvider(TFunctionRef<void(const ITimingProfilerProvider&)> Callback) const
-{
-	Lock.ReadAccessCheck();
-	Callback(*TimingProfilerProvider);
-}
-
-void FAnalysisSession::ReadFileActivityProvider(TFunctionRef<void(const IFileActivityProvider&)> Callback) const
-{
-	Lock.ReadAccessCheck();
-	Callback(*FileActivityProvider);
-}
-
-void FAnalysisSession::ReadLoadTimeProfilerProvider(TFunctionRef<void(const ILoadTimeProfilerProvider&)> Callback) const
-{
-	Lock.ReadAccessCheck();
-	Callback(*LoadTimeProfilerProvider);
-}
-
-void FAnalysisSession::ReadCountersProvider(TFunctionRef<void(const ICountersProvider &)> Callback) const
-{
-	Lock.ReadAccessCheck();
-	Callback(*CountersProvider);
-}
-
 FAnalysisService::FAnalysisService()
 {
 
@@ -148,27 +93,31 @@ void FAnalysisService::AnalyzeInternal(TSharedRef<FAnalysisSession> AnalysisSess
 {
 	OnAnalysisStarted().Broadcast(AnalysisSession);
 
-	TArray<TSharedRef<IAnalyzer>> Analyzers;
+	TArray<IAnalyzer*> Analyzers;
 	{
 		Trace::FAnalysisSessionEditScope _(AnalysisSession.Get());
-		Analyzers.Add(MakeShared<FMiscTraceAnalyzer>(AnalysisSession));
-		Analyzers.Add(MakeShared<FAsyncLoadingTraceAnalyzer>(AnalysisSession));
-		Analyzers.Add(MakeShared<FCpuProfilerAnalyzer>(AnalysisSession));
-		Analyzers.Add(MakeShared<FLogTraceAnalyzer>(AnalysisSession));
-		Analyzers.Add(MakeShared<FGpuProfilerAnalyzer>(AnalysisSession));
-		Analyzers.Add(MakeShared<FPlatformFileTraceAnalyzer>(AnalysisSession));
-		Analyzers.Add(MakeShared<FStatsAnalyzer>(AnalysisSession));
+		Analyzers.Add(new FMiscTraceAnalyzer(AnalysisSession.Get()));
+		Analyzers.Add(new FAsyncLoadingTraceAnalyzer(AnalysisSession.Get()));
+		Analyzers.Add(new FCpuProfilerAnalyzer(AnalysisSession.Get()));
+		Analyzers.Add(new FLogTraceAnalyzer(AnalysisSession.Get()));
+		Analyzers.Add(new FGpuProfilerAnalyzer(AnalysisSession.Get()));
+		Analyzers.Add(new FPlatformFileTraceAnalyzer(AnalysisSession.Get()));
+		Analyzers.Add(new FStatsAnalyzer(AnalysisSession.Get()));
 	}
 
 	FAnalysisContext Context;
-	for (TSharedRef<Trace::IAnalyzer> Analyzer : Analyzers)
+	for (Trace::IAnalyzer* Analyzer : Analyzers)
 	{
-		Context.AddAnalyzer(Analyzer.Get());
+		Context.AddAnalyzer(*Analyzer);
 	}
 	Trace::FAnalysisProcessor Processor = Context.Process();
 
 	Processor.Start(DataStream);
 
+	for (Trace::IAnalyzer* Analyzer : Analyzers)
+	{
+		delete Analyzer;
+	}
 	Analyzers.Empty();
 
 	AnalysisSession->SetComplete();
@@ -182,128 +131,6 @@ void FAnalysisService::FAnalysisWorker::DoWork()
 	AnalysisSession = nullptr;
 }
 
-FAnalysisService::FMockGenerator::FMockGenerator(TSharedRef<FAnalysisSession> InAnalysisSession)
-	: AnalysisSession(InAnalysisSession)
-{
-
-}
-
-void FAnalysisService::FMockGenerator::CreateMockData()
-{
-	AnalysisSession->BeginEdit();
-
-	TSharedRef<FTimingProfilerProvider> TimingProfilerProvider = AnalysisSession->EditTimingProfilerProvider();
-	TSharedRef<FThreadProvider> ThreadProvider = AnalysisSession->EditThreadProvider();
-	// Create Threads.
-	{
-		ThreadProvider->AddThread(1, TEXT("Main"), TPri_Normal);
-		Timelines.Add(TimingProfilerProvider->EditCpuThreadTimeline(1));
-		ThreadProvider->AddThread(2, TEXT("Rendering"), TPri_Normal);
-		Timelines.Add(TimingProfilerProvider->EditCpuThreadTimeline(2));
-		ThreadProvider->AddThread(3, TEXT("Game"), TPri_Normal);
-		Timelines.Add(TimingProfilerProvider->EditCpuThreadTimeline(3));
-		int32 ThreadId = 4;
-		while (ThreadId <= 10)
-		{
-			ThreadProvider->AddThread(ThreadId, *FString::Printf(TEXT("Thread %d"), ThreadId), TPri_SlightlyBelowNormal);
-			ThreadProvider->SetThreadGroup(ThreadId, TraceThreadGroup_ThreadPool);
-			Timelines.Add(TimingProfilerProvider->EditCpuThreadTimeline(ThreadId));
-			++ThreadId;
-		}
-	}
-
-	// Create Timers.
-	{
-		EventTypes.Add(TimingProfilerProvider->AddCpuTimer(TEXT("Update")));
-		EventTypes.Add(TimingProfilerProvider->AddCpuTimer(TEXT("Collide")));
-		EventTypes.Add(TimingProfilerProvider->AddCpuTimer(TEXT("Draw")));
-		EventTypes.Add(TimingProfilerProvider->AddCpuTimer(TEXT("PostDraw")));
-		EventTypes.Add(TimingProfilerProvider->AddCpuTimer(TEXT("Present")));
-		EventTypes.Add(TimingProfilerProvider->AddCpuTimer(TEXT("OtherTimer1")));
-		EventTypes.Add(TimingProfilerProvider->AddCpuTimer(TEXT("OtherTimer2")));
-
-		/*AddTimer(TimerId++, TEXT("SomeGpuOp1"), TEXT("Render"), GpuScope, FLinearColor(1.0, 1.0, 1.0, 1.0));
-		AddTimer(TimerId++, TEXT("SomeGpuOp2"), TEXT("Render"), GpuScope, FLinearColor(1.0, 1.0, 1.0, 1.0));
-		AddTimer(TimerId++, TEXT("SomeGpuOp3"), TEXT("Render"), GpuScope, FLinearColor(1.0, 1.0, 1.0, 1.0));
-		AddTimer(TimerId++, TEXT("ComputeSomething1"), TEXT("Render"), ComputeScope, FLinearColor(1.0, 1.0, 1.0, 1.0));
-		AddTimer(TimerId++, TEXT("ComputeSomething2"), TEXT("Misc"), ComputeScope, FLinearColor(1.0, 1.0, 1.0, 1.0));*/
-	}
-
-	// Create a Timeline for each Thread.
-	{
-		static const int32 MinEventsPerTimeline = 1'000'000;
-		static const int32 MaxEventsPerTimeline = 2'000'000;
-
-		FRandomStream RandomStream(0);
-
-		for (TSharedRef<FTimingProfilerProvider::TimelineInternal> Timeline : Timelines)
-		{
-			const int32 NumEvents = RandomStream.RandRange(MinEventsPerTimeline, MaxEventsPerTimeline);
-			
-			double Time = 0;
-
-			int EventCount = NumEvents;
-			while (EventCount > 0)
-			{
-				const double Duration = RandomStream.FRandRange(0.001, 0.01); // [1ms .. 10ms]
-				InitTimingEventHierarchy(RandomStream, Timeline, EventCount, Time, Duration, 0);
-
-				static const double MaxTimeBetweenEvents = 0.001; // 1ms
-				Time += Duration + RandomStream.FRandRange(0, MaxTimeBetweenEvents);
-			}
-		}
-	}
-
-	AnalysisSession->EndEdit();
-}
-
-void FAnalysisService::FMockGenerator::InitTimingEventHierarchy(FRandomStream& InRandomStream, TSharedRef<FTimingProfilerProvider::TimelineInternal> InTimeline, int& InOutEventCount, const double InStartTime, const double InDuration, const int32 InLevel)
-{
-	static const int32 MaxLevels = 5;
-	if (InLevel >= MaxLevels)
-	{
-		return;
-	}
-
-	const double EndTime = InStartTime + InDuration;
-
-	if (InOutEventCount <= 0)
-	{
-		return;
-	}
-	
-	FTimingProfilerEvent Event;
-	Event.TimerIndex = EventTypes[InRandomStream.RandRange(0, EventTypes.Num() - 1)];
-	InTimeline->AppendBeginEvent(InStartTime, Event);
-	--InOutEventCount;
-	
-	int32 NumSubEvents = InRandomStream.RandRange(1, InRandomStream.RandRange(1, InDuration / 0.001));
-
-	double Time = InStartTime;
-	double RemainingDuration = InDuration;
-	for (int SubEventIndex = 0; SubEventIndex < NumSubEvents; SubEventIndex++)
-	{
-		double TimeAdvance = InRandomStream.FRandRange(0, FMath::Min(0.01, RemainingDuration));
-		Time += TimeAdvance;
-		if (Time >= EndTime)
-			break;
-		RemainingDuration = EndTime - Time;
-
-		double Duration = InRandomStream.FRandRange(0, RemainingDuration);
-		InitTimingEventHierarchy(InRandomStream, InTimeline, InOutEventCount, Time, Duration, InLevel + 1);
-		
-		Time += Duration;
-		RemainingDuration = EndTime - Time;
-
-		if (InOutEventCount <= 0)
-		{
-			break;
-		}
-	}
-
-	InTimeline->AppendEndEvent(EndTime);
-}
-
 FAnalysisService::~FAnalysisService()
 {
 	for (TSharedPtr<FAsyncTask<FAnalysisWorker>> Task : Tasks)
@@ -329,13 +156,4 @@ TSharedPtr<const IAnalysisSession> FAnalysisService::StartAnalysis(const TCHAR*
 	return AnalysisSession;
 }
 
-TSharedPtr<const IAnalysisSession> FAnalysisService::MockAnalysis()
-{
-	TSharedRef<FAnalysisSession> AnalysisSession = MakeShared<FAnalysisSession>(TEXT("MockSession"));
-	FMockGenerator MockGenerator(AnalysisSession);
-	MockGenerator.CreateMockData();
-
-	return AnalysisSession;
-}
-
 }
diff --git a/Engine/Source/Developer/TraceServices/Private/AnalysisServicePrivate.h b/Engine/Source/Developer/TraceServices/Private/AnalysisServicePrivate.h
--- a/Engine/Source/Developer/TraceServices/Private/AnalysisServicePrivate.h
+++ b/Engine/Source/Developer/TraceServices/Private/AnalysisServicePrivate.h
@@ -8,20 +8,17 @@
 #include "Misc/ScopeLock.h"
 #include "Model/TimingProfiler.h"
 #include "Common/StringStore.h"
-
-struct FRandomStream;
+#include "Model/Bookmarks.h"
+#include "Model/Log.h"
+#include "Model/Threads.h"
+#include "Model/FileActivity.h"
+#include "Model/Frames.h"
+#include "Model/LoadTimeProfiler.h"
+#include "Model/Counters.h"
 
 namespace Trace
 {
 
-class FBookmarkProvider;
-class FLogProvider;
-class FThreadProvider;
-class FFileActivityProvider;
-class FFrameProvider;
-class FLoadTimeProfilerProvider;
-class FCountersProvider;
-
 class FAnalysisSessionLock
 {
 public:
@@ -53,29 +50,29 @@ public:
 	void UpdateDuration(double Seconds) { Lock.WriteAccessCheck(); DurationSeconds = Seconds > DurationSeconds ? Seconds : DurationSeconds; }
 	void SetComplete() { IsComplete = true; }
 
-	TSharedRef<FBookmarkProvider> EditBookmarkProvider() { return BookmarkProvider; }
-	virtual void ReadBookmarkProvider(TFunctionRef<void(const IBookmarkProvider&)> Callback) const override;
+	FBookmarkProvider& EditBookmarkProvider() { return BookmarkProvider; }
+	virtual const IBookmarkProvider& ReadBookmarkProvider() const override { return BookmarkProvider; }
 
-	TSharedRef<FLogProvider> EditLogProvider() { return LogProvider; }
-	virtual void ReadLogProvider(TFunctionRef<void(const ILogProvider&)> Callback) const override;
+	FLogProvider& EditLogProvider() { return LogProvider; }
+	virtual const ILogProvider& ReadLogProvider() const override { return LogProvider; }
 
-	TSharedRef<FThreadProvider> EditThreadProvider() { return ThreadProvider; }
-	virtual void ReadThreadProvider(TFunctionRef<void(const IThreadProvider&)> Callback) const override;
+	FThreadProvider& EditThreadProvider() { return ThreadProvider; }
+	virtual const IThreadProvider& ReadThreadProvider() const override { return ThreadProvider; }
 
-	TSharedRef<FFrameProvider> EditFramesProvider() { return FramesProvider; }
-	virtual void ReadFramesProvider(TFunctionRef<void(const IFrameProvider&)> Callback) const override;
+	FFrameProvider& EditFrameProvider() { return FrameProvider; }
+	virtual const IFrameProvider& ReadFrameProvider() const override { return FrameProvider; }
 
-	TSharedRef<FTimingProfilerProvider> EditTimingProfilerProvider() { return TimingProfilerProvider; }
-	virtual void ReadTimingProfilerProvider(TFunctionRef<void(const ITimingProfilerProvider&)> Callback) const override;
+	FTimingProfilerProvider& EditTimingProfilerProvider() { return TimingProfilerProvider; }
+	virtual const FTimingProfilerProvider& ReadTimingProfilerProvider() const override { return TimingProfilerProvider; }
 
-	TSharedRef<FFileActivityProvider> EditFileActivityProvider() { return FileActivityProvider; }
-	virtual void ReadFileActivityProvider(TFunctionRef<void(const IFileActivityProvider &)> Callback) const override;
+	FFileActivityProvider& EditFileActivityProvider() { return FileActivityProvider; }
+	virtual const FFileActivityProvider& ReadFileActivityProvider() const override { return FileActivityProvider; }
 
-	TSharedRef<FLoadTimeProfilerProvider> EditLoadTimeProfilerProvider() { return LoadTimeProfilerProvider; }
-	virtual void ReadLoadTimeProfilerProvider(TFunctionRef<void(const ILoadTimeProfilerProvider &)> Callback) const override;
+	FLoadTimeProfilerProvider& EditLoadTimeProfilerProvider() { return LoadTimeProfilerProvider; }
+	virtual const ILoadTimeProfilerProvider& ReadLoadTimeProfilerProvider() const override { return LoadTimeProfilerProvider; }
 
-	TSharedRef<FCountersProvider> EditCountersProvider() { return CountersProvider; }
-	virtual void ReadCountersProvider(TFunctionRef<void(const ICountersProvider &)> Callback) const override;
+	FCounterProvider& EditCounterProvider() { return CounterProvider; }
+	virtual const ICounterProvider& ReadCounterProvider() const override { return CounterProvider; }
 
 	const TCHAR* StoreString(const TCHAR* String) { return StringStore.Store(String); }
 
@@ -93,14 +90,14 @@ private:
 	double DurationSeconds = 0.0;
 	FSlabAllocator Allocator;
 	FStringStore StringStore;
-	TSharedRef<FBookmarkProvider> BookmarkProvider;
-	TSharedRef<FLogProvider> LogProvider;
-	TSharedRef<FThreadProvider> ThreadProvider;
-	TSharedRef<FFrameProvider> FramesProvider;
-	TSharedRef<FTimingProfilerProvider> TimingProfilerProvider;
-	TSharedRef<FFileActivityProvider> FileActivityProvider;
-	TSharedRef<FLoadTimeProfilerProvider> LoadTimeProfilerProvider;
-	TSharedRef<FCountersProvider> CountersProvider;
+	FBookmarkProvider BookmarkProvider;
+	FLogProvider LogProvider;
+	FThreadProvider ThreadProvider;
+	FFrameProvider FrameProvider;
+	FTimingProfilerProvider TimingProfilerProvider;
+	FFileActivityProvider FileActivityProvider;
+	FLoadTimeProfilerProvider LoadTimeProfilerProvider;
+	FCounterProvider CounterProvider;
 };
 
 struct FAnalysisSessionEditScope
@@ -127,7 +124,6 @@ public:
 	virtual ~FAnalysisService();
 	virtual TSharedPtr<const IAnalysisSession> Analyze(const TCHAR* SessionName, TUniquePtr<Trace::IInDataStream>&& DataStream) override;
 	virtual TSharedPtr<const IAnalysisSession> StartAnalysis(const TCHAR* SessionName, TUniquePtr<Trace::IInDataStream>&& DataStream) override;
-	virtual TSharedPtr<const IAnalysisSession> MockAnalysis() override;
 	virtual FAnalysisStartedEvent& OnAnalysisStarted() override { return AnalysisStartedEvent; }
 	virtual FAnalysisFinishedEvent& OnAnalysisFinished() override { return AnalysisFinishedEvent; }
 
@@ -148,20 +144,6 @@ private:
 		TSharedPtr<FAnalysisSession> AnalysisSession;
 	};
 
-	class FMockGenerator
-	{
-	public:
-		FMockGenerator(TSharedRef<FAnalysisSession> InAnalysisSession);
-		void CreateMockData();
-
-	private:
-		void InitTimingEventHierarchy(FRandomStream& InRandomStream, TSharedRef<FTimingProfilerProvider::TimelineInternal> InTimeline, int& InOutEventCount, const double InStartTime, const double InDuration, const int32 InLevel);
-		
-		TSharedPtr<FAnalysisSession> AnalysisSession;
-		TArray<TSharedRef<FTimingProfilerProvider::TimelineInternal>> Timelines;
-		TArray<uint32> EventTypes;
-	};
-
 	void AnalyzeInternal(TSharedRef<FAnalysisSession> AnalysisSession, Trace::IInDataStream& DataStream);
 
 	FAnalysisStartedEvent AnalysisStartedEvent;
diff --git a/Engine/Source/Developer/TraceServices/Private/Analyzers/CpuProfilerTraceAnalysis.cpp b/Engine/Source/Developer/TraceServices/Private/Analyzers/CpuProfilerTraceAnalysis.cpp
--- a/Engine/Source/Developer/TraceServices/Private/Analyzers/CpuProfilerTraceAnalysis.cpp
+++ b/Engine/Source/Developer/TraceServices/Private/Analyzers/CpuProfilerTraceAnalysis.cpp
@@ -4,10 +4,10 @@
 #include "Common/Utils.h"
 #include "Model/Threads.h"
 
-FCpuProfilerAnalyzer::FCpuProfilerAnalyzer(TSharedRef<Trace::FAnalysisSession> InSession)
+FCpuProfilerAnalyzer::FCpuProfilerAnalyzer(Trace::FAnalysisSession& InSession)
 	: Session(InSession)
-	, ThreadProvider(InSession->EditThreadProvider())
-	, TimingProfilerProvider(InSession->EditTimingProfilerProvider())
+	, ThreadProvider(InSession.EditThreadProvider())
+	, TimingProfilerProvider(InSession.EditTimingProfilerProvider())
 {
 
 }
@@ -22,7 +22,7 @@ void FCpuProfilerAnalyzer::OnAnalysisBegin(const FOnAnalysisContext& Context)
 
 void FCpuProfilerAnalyzer::OnEvent(uint16 RouteId, const FOnEventContext& Context)
 {
-	Trace::FAnalysisSessionEditScope _(Session.Get());
+	Trace::FAnalysisSessionEditScope _(Session);
 
 	const auto& EventData = Context.EventData;
 	switch (RouteId)
@@ -30,7 +30,7 @@ void FCpuProfilerAnalyzer::OnEvent(uint16 RouteId, const FOnEventContext& Contex
 	case RouteId_EventSpec:
 	{
 		uint16 Id = EventData.GetValue("Id").As<uint16>();
-		ScopeIdToEventIdMap.Add(Id, TimingProfilerProvider->AddCpuTimer(reinterpret_cast<const TCHAR*>(EventData.GetAttachment())));
+		ScopeIdToEventIdMap.Add(Id, TimingProfilerProvider.AddCpuTimer(reinterpret_cast<const TCHAR*>(EventData.GetAttachment())));
 		break;
 	}
 	case RouteId_EventBatch:
@@ -66,7 +66,7 @@ void FCpuProfilerAnalyzer::OnEvent(uint16 RouteId, const FOnEventContext& Contex
 		check(BufferPtr == BufferEnd);
 		if (LastCycle)
 		{
-			Session->UpdateDuration(Context.SessionContext.TimestampFromCycle(LastCycle));
+			Session.UpdateDuration(Context.SessionContext.TimestampFromCycle(LastCycle));
 		}
 		ThreadState->LastCycle = LastCycle;
 		break;
@@ -79,9 +79,9 @@ TSharedRef<FCpuProfilerAnalyzer::FThreadState> FCpuProfilerAnalyzer::GetThreadSt
 	if (!ThreadStatesMap.Contains(ThreadId))
 	{
 		TSharedRef<FThreadState> ThreadState = MakeShared<FThreadState>();
-		ThreadState->Timeline = TimingProfilerProvider->EditCpuThreadTimeline(ThreadId);
+		ThreadState->Timeline = &TimingProfilerProvider.EditCpuThreadTimeline(ThreadId);
 		ThreadStatesMap.Add(ThreadId, ThreadState);
-		ThreadProvider->EnsureThreadExists(ThreadId);
+		ThreadProvider.EnsureThreadExists(ThreadId);
 		return ThreadState;
 	}
 	else
diff --git a/Engine/Source/Developer/TraceServices/Private/Analyzers/CpuProfilerTraceAnalysis.h b/Engine/Source/Developer/TraceServices/Private/Analyzers/CpuProfilerTraceAnalysis.h
--- a/Engine/Source/Developer/TraceServices/Private/Analyzers/CpuProfilerTraceAnalysis.h
+++ b/Engine/Source/Developer/TraceServices/Private/Analyzers/CpuProfilerTraceAnalysis.h
@@ -5,7 +5,6 @@
 #include "Trace/Trace.h"
 #include "Trace/Analyzer.h"
 #include "Containers/UnrealString.h"
-#include "Templates/SharedPointer.h"
 #include "Model/TimingProfiler.h"
 
 namespace Trace
@@ -18,7 +17,7 @@ class FCpuProfilerAnalyzer
 	: public Trace::IAnalyzer
 {
 public:
-	FCpuProfilerAnalyzer(TSharedRef<Trace::FAnalysisSession> Session);
+	FCpuProfilerAnalyzer(Trace::FAnalysisSession& Session);
 	virtual void OnAnalysisBegin(const FOnAnalysisContext& Context) override;
 	virtual void OnEvent(uint16 RouteId, const FOnEventContext& Context) override;
 	virtual void OnAnalysisEnd() override {};
@@ -33,7 +32,7 @@ private:
 	struct FThreadState
 	{
 		TArray<EventScopeState> ScopeStack;
-		TSharedPtr<Trace::FTimingProfilerProvider::TimelineInternal> Timeline;
+		Trace::FTimingProfilerProvider::TimelineInternal* Timeline;
 		double LastCycle = 0.0;
 	};
 
@@ -45,9 +44,9 @@ private:
 		RouteId_EventBatch,
 	};
 
-	TSharedRef<Trace::FAnalysisSession> Session;
-	TSharedRef<Trace::FThreadProvider> ThreadProvider;
-	TSharedRef<Trace::FTimingProfilerProvider> TimingProfilerProvider;
+	Trace::FAnalysisSession& Session;
+	Trace::FThreadProvider& ThreadProvider;
+	Trace::FTimingProfilerProvider& TimingProfilerProvider;
 	TMap<uint32, TSharedRef<FThreadState>> ThreadStatesMap;
 	TMap<uint16, uint32> ScopeIdToEventIdMap;
 };
diff --git a/Engine/Source/Developer/TraceServices/Private/Analyzers/GpuProfilerTraceAnalysis.cpp b/Engine/Source/Developer/TraceServices/Private/Analyzers/GpuProfilerTraceAnalysis.cpp
--- a/Engine/Source/Developer/TraceServices/Private/Analyzers/GpuProfilerTraceAnalysis.cpp
+++ b/Engine/Source/Developer/TraceServices/Private/Analyzers/GpuProfilerTraceAnalysis.cpp
@@ -3,10 +3,10 @@
 #include "AnalysisServicePrivate.h"
 #include "Common/Utils.h"
 
-FGpuProfilerAnalyzer::FGpuProfilerAnalyzer(TSharedRef<Trace::FAnalysisSession> InSession)
+FGpuProfilerAnalyzer::FGpuProfilerAnalyzer(Trace::FAnalysisSession& InSession)
 	: Session(InSession)
-	, TimingProfilerProvider(InSession->EditTimingProfilerProvider())
-	, Timeline(TimingProfilerProvider->EditGpuTimeline())
+	, TimingProfilerProvider(InSession.EditTimingProfilerProvider())
+	, Timeline(TimingProfilerProvider.EditGpuTimeline())
 	, Calibrated(false)
 {
 
@@ -22,7 +22,7 @@ void FGpuProfilerAnalyzer::OnAnalysisBegin(const FOnAnalysisContext& Context)
 
 void FGpuProfilerAnalyzer::OnEvent(uint16 RouteId, const FOnEventContext& Context)
 {
-	Trace::FAnalysisSessionEditScope _(Session.Get());
+	Trace::FAnalysisSessionEditScope _(Session);
 
 	const auto& EventData = Context.EventData;
 
@@ -32,7 +32,7 @@ void FGpuProfilerAnalyzer::OnEvent(uint16 RouteId, const FOnEventContext& Contex
 	{
 		uint64 EventType = EventData.GetValue("EventType").As<uint64>();
 		FString EventName(reinterpret_cast<const TCHAR*>(EventData.GetAttachment()), EventData.GetValue("NameLength").As<uint16>());
-		EventTypeMap.Add(EventType, TimingProfilerProvider->AddGpuTimer(*EventName));
+		EventTypeMap.Add(EventType, TimingProfilerProvider.AddGpuTimer(*EventName));
 		break;
 	}
 	case RouteId_Frame:
@@ -56,14 +56,14 @@ void FGpuProfilerAnalyzer::OnEvent(uint16 RouteId, const FOnEventContext& Contex
 				check(EventTypeMap.Contains(EventType));
 				Trace::FTimingProfilerEvent Event;
 				Event.TimerIndex = EventTypeMap[EventType];
-				Timeline->AppendBeginEvent(GpuTimestampToSessionTime(ActualTimestamp), Event);
+				Timeline.AppendBeginEvent(GpuTimestampToSessionTime(ActualTimestamp), Event);
 				++CurrentDepth;
 			}
 			else
 			{
 				check(CurrentDepth > 0);
 				--CurrentDepth;
-				Timeline->AppendEndEvent(GpuTimestampToSessionTime(ActualTimestamp));
+				Timeline.AppendEndEvent(GpuTimestampToSessionTime(ActualTimestamp));
 			}
 		}
 		check(BufferPtr == BufferEnd);
@@ -78,7 +78,7 @@ double FGpuProfilerAnalyzer::GpuTimestampToSessionTime(uint64 GpuMicroseconds)
 {
 	if (!Calibrated)
 	{
-		uint64 SessionTimeMicroseconds = uint64(Session->GetDurationSeconds() * 1000000.0);
+		uint64 SessionTimeMicroseconds = uint64(Session.GetDurationSeconds() * 1000000.0);
 		GpuTimeOffset = GpuMicroseconds - SessionTimeMicroseconds;
 		Calibrated = true;
 	}
diff --git a/Engine/Source/Developer/TraceServices/Private/Analyzers/GpuProfilerTraceAnalysis.h b/Engine/Source/Developer/TraceServices/Private/Analyzers/GpuProfilerTraceAnalysis.h
--- a/Engine/Source/Developer/TraceServices/Private/Analyzers/GpuProfilerTraceAnalysis.h
+++ b/Engine/Source/Developer/TraceServices/Private/Analyzers/GpuProfilerTraceAnalysis.h
@@ -4,7 +4,6 @@
 
 #include "Trace/Trace.h"
 #include "Trace/Analyzer.h"
-#include "Templates/SharedPointer.h"
 #include "Model/TimingProfiler.h"
 
 namespace Trace
@@ -16,7 +15,7 @@ class FGpuProfilerAnalyzer
 	: public Trace::IAnalyzer
 {
 public:
-	FGpuProfilerAnalyzer(TSharedRef<Trace::FAnalysisSession> Session);
+	FGpuProfilerAnalyzer(Trace::FAnalysisSession& Session);
 	virtual void OnAnalysisBegin(const FOnAnalysisContext& Context) override;
 	virtual void OnEvent(uint16 RouteId, const FOnEventContext& Context) override;
 	virtual void OnAnalysisEnd() override {};
@@ -30,9 +29,9 @@ private:
 
 	double GpuTimestampToSessionTime(uint64 GpuMicroseconds);
 
-	TSharedRef<Trace::FAnalysisSession> Session;
-	TSharedRef<Trace::FTimingProfilerProvider> TimingProfilerProvider;
-	TSharedRef<Trace::FTimingProfilerProvider::TimelineInternal> Timeline;
+	Trace::FAnalysisSession& Session;
+	Trace::FTimingProfilerProvider& TimingProfilerProvider;
+	Trace::FTimingProfilerProvider::TimelineInternal& Timeline;
 	TMap<uint64, uint32> EventTypeMap;
 	bool Calibrated;
 	uint64 GpuTimeOffset;
diff --git a/Engine/Source/Developer/TraceServices/Private/Analyzers/LoadTimeTraceAnalysis.cpp b/Engine/Source/Developer/TraceServices/Private/Analyzers/LoadTimeTraceAnalysis.cpp
--- a/Engine/Source/Developer/TraceServices/Private/Analyzers/LoadTimeTraceAnalysis.cpp
+++ b/Engine/Source/Developer/TraceServices/Private/Analyzers/LoadTimeTraceAnalysis.cpp
@@ -7,9 +7,9 @@
 #include "Model/LoadTimeProfiler.h"
 #include "Analyzers/MiscTraceAnalysis.h"
 
-FAsyncLoadingTraceAnalyzer::FAsyncLoadingTraceAnalyzer(TSharedRef<Trace::FAnalysisSession> InSession)
+FAsyncLoadingTraceAnalyzer::FAsyncLoadingTraceAnalyzer(Trace::FAnalysisSession& InSession)
 	: Session(InSession)
-	, LoadTimeProfilerProvider(Session->EditLoadTimeProfilerProvider())
+	, LoadTimeProfilerProvider(Session.EditLoadTimeProfilerProvider())
 {
 }
 
@@ -27,11 +27,11 @@ TSharedRef<FAsyncLoadingTraceAnalyzer::FThreadState> FAsyncLoadingTraceAnalyzer:
 		if (MainThreadId == -1)
 		{
 			MainThreadId = ThreadId;
-			ThreadState->CpuTimeline = LoadTimeProfilerProvider->EditMainThreadCpuTimeline();
+			ThreadState->CpuTimeline = LoadTimeProfilerProvider.EditMainThreadCpuTimeline();
 		}
 		else if (ThreadId == AsyncLoadingThreadId)
 		{
-			ThreadState->CpuTimeline = LoadTimeProfilerProvider->EditAsyncLoadingThreadCpuTimeline();
+			ThreadState->CpuTimeline = LoadTimeProfilerProvider.EditAsyncLoadingThreadCpuTimeline();
 		}
 		return ThreadState;
 	}
@@ -128,7 +128,7 @@ void FAsyncLoadingTraceAnalyzer::OnEvent(uint16 RouteId, const FOnEventContext&
 		TSharedRef<FLinkerState> LinkerState = ActiveLinkersMap[LinkerPtr];
 		if (LinkerState->PackageInfo)
 		{
-			Trace::FAnalysisSessionEditScope _(Session.Get());
+			Trace::FAnalysisSessionEditScope _(Session);
 			Trace::FPackageSummaryInfo& Summary = LinkerState->PackageInfo->Summary;
 			Summary.TotalHeaderSize = EventData.GetValue("TotalHeaderSize").As<uint32>();
 			Summary.NameCount = EventData.GetValue("NameCount").As<uint32>();
@@ -139,14 +139,14 @@ void FAsyncLoadingTraceAnalyzer::OnEvent(uint16 RouteId, const FOnEventContext&
 	}
 	case RouteId_BeginCreateExport:
 	{
-		Trace::FAnalysisSessionEditScope _(Session.Get());
+		Trace::FAnalysisSessionEditScope _(Session);
 
 		uint64 LinkerPtr = EventData.GetValue("Linker").As<uint64>();
 		//check(ActiveLinkersMap.Contains(LinkerPtr));
 		TSharedRef<FLinkerState> LinkerState = ActiveLinkersMap[LinkerPtr];
 		uint32 ThreadId = EventData.GetValue("ThreadId").As<uint32>();
 		TSharedRef<FThreadState> ThreadState = GetThreadState(ThreadId);
-		Trace::FPackageExportInfo& Export = LoadTimeProfilerProvider->CreateExport();
+		Trace::FPackageExportInfo& Export = LoadTimeProfilerProvider.CreateExport();
 		Export.SerialOffset = EventData.GetValue("SerialOffset").As<uint64>();
 		Export.SerialSize = EventData.GetValue("SerialSize").As<uint64>();
 		Export.IsAsset = EventData.GetValue("IsAsset").As<bool>();
@@ -166,7 +166,7 @@ void FAsyncLoadingTraceAnalyzer::OnEvent(uint16 RouteId, const FOnEventContext&
 		Trace::FPackageExportInfo* Export = ThreadState->GetCurrentExportScope();
 		ExportsMap.Add(ObjectPtr, Export);
 		{
-			Trace::FAnalysisSessionEditScope _(Session.Get());
+			Trace::FAnalysisSessionEditScope _(Session);
 			Export->Class = ObjectClass;
 			ThreadState->LeaveScope(Context.SessionContext.TimestampFromCycle(EventData.GetValue("Cycle").As<uint64>()));
 		}
@@ -180,7 +180,7 @@ void FAsyncLoadingTraceAnalyzer::OnEvent(uint16 RouteId, const FOnEventContext&
 		TSharedRef<FThreadState> ThreadState = GetThreadState(ThreadId);
 		ELoadTimeProfilerObjectEventType EventType = static_cast<ELoadTimeProfilerObjectEventType>(EventData.GetValue("EventType").As<uint8>());
 		{
-			Trace::FAnalysisSessionEditScope _(Session.Get());
+			Trace::FAnalysisSessionEditScope _(Session);
 			ThreadState->EnterExportScope(Context.SessionContext.TimestampFromCycle(EventData.GetValue("Cycle").As<uint64>()), Export ? *Export : nullptr, EventType);
 		}
 		break;
@@ -190,7 +190,7 @@ void FAsyncLoadingTraceAnalyzer::OnEvent(uint16 RouteId, const FOnEventContext&
 		uint32 ThreadId = EventData.GetValue("ThreadId").As<uint32>();
 		TSharedRef<FThreadState> ThreadState = GetThreadState(ThreadId);
 		{
-			Trace::FAnalysisSessionEditScope _(Session.Get());
+			Trace::FAnalysisSessionEditScope _(Session);
 			ThreadState->LeaveScope(Context.SessionContext.TimestampFromCycle(EventData.GetValue("Cycle").As<uint64>()));
 		}
 		break;
@@ -327,12 +327,12 @@ void FAsyncLoadingTraceAnalyzer::OnEvent(uint16 RouteId, const FOnEventContext&
 	}
 	case RouteId_NewAsyncPackage:
 	{
-		Trace::FAnalysisSessionEditScope _(Session.Get());
+		Trace::FAnalysisSessionEditScope _(Session);
 
 		uint64 AsyncPackagePtr = EventData.GetValue("AsyncPackage").As<uint64>();
 		//check(!ActivePackagesMap.Contains(AsyncPackagePtr));
 		TSharedRef<FAsyncPackageState> AsyncPackageState = MakeShared<FAsyncPackageState>();
-		AsyncPackageState->PackageInfo = &LoadTimeProfilerProvider->CreatePackage(reinterpret_cast<const TCHAR*>(EventData.GetAttachment()));
+		AsyncPackageState->PackageInfo = &LoadTimeProfilerProvider.CreatePackage(reinterpret_cast<const TCHAR*>(EventData.GetAttachment()));
 		ActiveAsyncPackagesMap.Add(AsyncPackagePtr, AsyncPackageState);
 		break;
 	}
@@ -386,7 +386,7 @@ void FAsyncLoadingTraceAnalyzer::OnEvent(uint16 RouteId, const FOnEventContext&
 	}
 	case RouteId_BeginAsyncPackageScope:
 	{
-		Trace::FAnalysisSessionEditScope _(Session.Get());
+		Trace::FAnalysisSessionEditScope _(Session);
 
 		uint64 AsyncPackagePtr = EventData.GetValue("AsyncPackage").As<uint64>();
 		//check(ActiveAsyncPackagesMap.Contains(AsyncPackagePtr));
@@ -399,7 +399,7 @@ void FAsyncLoadingTraceAnalyzer::OnEvent(uint16 RouteId, const FOnEventContext&
 	}
 	case RouteId_EndAsyncPackageScope:
 	{
-		Trace::FAnalysisSessionEditScope _(Session.Get());
+		Trace::FAnalysisSessionEditScope _(Session);
 
 		uint32 ThreadId = EventData.GetValue("ThreadId").As<uint32>();
 		TSharedRef<FThreadState> ThreadState = GetThreadState(ThreadId);
@@ -409,7 +409,7 @@ void FAsyncLoadingTraceAnalyzer::OnEvent(uint16 RouteId, const FOnEventContext&
 	}
 	case RouteId_BeginFlushAsyncLoading:
 	{
-		Trace::FAnalysisSessionEditScope _(Session.Get());
+		Trace::FAnalysisSessionEditScope _(Session);
 
 		FlushAsyncLoadingRequestId = EventData.GetValue("RequestId").As<uint64>();
 		FlushAsyncLoadingStartCycle = EventData.GetValue("Cycle").As<uint64>();
@@ -433,10 +433,10 @@ void FAsyncLoadingTraceAnalyzer::OnEvent(uint16 RouteId, const FOnEventContext&
 	}
 	case RouteId_ClassInfo:
 	{
-		Trace::FAnalysisSessionEditScope _(Session.Get());
+		Trace::FAnalysisSessionEditScope _(Session);
 
 		uint64 ClassPtr = EventData.GetValue("Class").As<uint64>();
-		const Trace::FClassInfo& ClassInfo = LoadTimeProfilerProvider->AddClassInfo(reinterpret_cast<const TCHAR*>(EventData.GetAttachment()));
+		const Trace::FClassInfo& ClassInfo = LoadTimeProfilerProvider.AddClassInfo(reinterpret_cast<const TCHAR*>(EventData.GetAttachment()));
 		ClassInfosMap.Add(ClassPtr, &ClassInfo);
 		break;
 	}
diff --git a/Engine/Source/Developer/TraceServices/Private/Analyzers/LoadTimeTraceAnalysis.h b/Engine/Source/Developer/TraceServices/Private/Analyzers/LoadTimeTraceAnalysis.h
--- a/Engine/Source/Developer/TraceServices/Private/Analyzers/LoadTimeTraceAnalysis.h
+++ b/Engine/Source/Developer/TraceServices/Private/Analyzers/LoadTimeTraceAnalysis.h
@@ -23,7 +23,7 @@ inline bool operator!=(const Trace::FLoadTimeProfilerCpuEvent& Lhs, const Trace:
 class FAsyncLoadingTraceAnalyzer : public Trace::IAnalyzer
 {
 public:
-	FAsyncLoadingTraceAnalyzer(TSharedRef<Trace::FAnalysisSession> Session);
+	FAsyncLoadingTraceAnalyzer(Trace::FAnalysisSession& Session);
 	virtual ~FAsyncLoadingTraceAnalyzer();
 
 	virtual void OnAnalysisBegin(const FOnAnalysisContext& Context) override;
@@ -198,8 +198,8 @@ private:
 		RouteId_ClassInfo,
 	};
 
-	TSharedRef<Trace::FAnalysisSession> Session;
-	TSharedRef<Trace::FLoadTimeProfilerProvider> LoadTimeProfilerProvider;
+	Trace::FAnalysisSession& Session;
+	Trace::FLoadTimeProfilerProvider& LoadTimeProfilerProvider;
 	//TSharedPtr<Trace::FMonotonicTimeline<uint64>> BlockingRequestsTimeline;
 
 	TArray<TSharedRef<FLoadMapState>> Maps;
diff --git a/Engine/Source/Developer/TraceServices/Private/Analyzers/LogTraceAnalysis.cpp b/Engine/Source/Developer/TraceServices/Private/Analyzers/LogTraceAnalysis.cpp
--- a/Engine/Source/Developer/TraceServices/Private/Analyzers/LogTraceAnalysis.cpp
+++ b/Engine/Source/Developer/TraceServices/Private/Analyzers/LogTraceAnalysis.cpp
@@ -4,9 +4,9 @@
 #include "Logging/LogTrace.h"
 #include "Model/Log.h"
 
-FLogTraceAnalyzer::FLogTraceAnalyzer(TSharedRef<Trace::FAnalysisSession> InSession)
+FLogTraceAnalyzer::FLogTraceAnalyzer(Trace::FAnalysisSession& InSession)
 	: Session(InSession)
-	, LogProvider(InSession->EditLogProvider())
+	, LogProvider(InSession.EditLogProvider())
 {
 
 }
@@ -22,7 +22,7 @@ void FLogTraceAnalyzer::OnAnalysisBegin(const FOnAnalysisContext& Context)
 
 void FLogTraceAnalyzer::OnEvent(uint16 RouteId, const FOnEventContext& Context)
 {
-	Trace::FAnalysisSessionEditScope _(Session.Get());
+	Trace::FAnalysisSessionEditScope _(Session);
 
 	const auto& EventData = Context.EventData;
 	switch (RouteId)
@@ -30,30 +30,30 @@ void FLogTraceAnalyzer::OnEvent(uint16 RouteId, const FOnEventContext& Context)
 	case RouteId_LogCategory:
 	{
 		uint64 CategoryPointer = EventData.GetValue("CategoryPointer").As<uint64>();
-		Trace::FLogCategory& Category = LogProvider->GetCategory(CategoryPointer);
-		Category.Name = Session->StoreString(reinterpret_cast<const TCHAR*>(EventData.GetAttachment()));
+		Trace::FLogCategory& Category = LogProvider.GetCategory(CategoryPointer);
+		Category.Name = Session.StoreString(reinterpret_cast<const TCHAR*>(EventData.GetAttachment()));
 		Category.DefaultVerbosity = static_cast<ELogVerbosity::Type>(EventData.GetValue("DefaultVerbosity").As<uint8>());
 		break;
 	}
 	case RouteId_LogMessageSpec:
 	{
 		uint64 LogPoint = EventData.GetValue("LogPoint").As<uint64>();
-		Trace::FLogMessageSpec& Spec = LogProvider->GetMessageSpec(LogPoint);
+		Trace::FLogMessageSpec& Spec = LogProvider.GetMessageSpec(LogPoint);
 		uint64 CategoryPointer = EventData.GetValue("CategoryPointer").As<uint64>();
-		Trace::FLogCategory& Category = LogProvider->GetCategory(CategoryPointer);
+		Trace::FLogCategory& Category = LogProvider.GetCategory(CategoryPointer);
 		Spec.Category = &Category;
 		Spec.Line = EventData.GetValue("Line").As<int32>();
 		Spec.Verbosity = static_cast<ELogVerbosity::Type>(EventData.GetValue("Verbosity").As<uint8>());
 		const ANSICHAR* File = reinterpret_cast<const ANSICHAR*>(EventData.GetAttachment());
-		Spec.File = Session->StoreString(ANSI_TO_TCHAR(File));
-		Spec.FormatString = Session->StoreString(reinterpret_cast<const TCHAR*>(EventData.GetAttachment() + strlen(File) + 1));
+		Spec.File = Session.StoreString(ANSI_TO_TCHAR(File));
+		Spec.FormatString = Session.StoreString(reinterpret_cast<const TCHAR*>(EventData.GetAttachment() + strlen(File) + 1));
 		break;
 	}
 	case RouteId_LogMessage:
 	{
 		uint64 LogPoint = EventData.GetValue("LogPoint").As<uint64>();
 		uint64 Cycle = EventData.GetValue("Cycle").As<uint64>();
-		LogProvider->AppendMessage(LogPoint, Context.SessionContext.TimestampFromCycle(Cycle), EventData.GetAttachment());
+		LogProvider.AppendMessage(LogPoint, Context.SessionContext.TimestampFromCycle(Cycle), EventData.GetAttachment());
 		break;
 	}
 	}
diff --git a/Engine/Source/Developer/TraceServices/Private/Analyzers/LogTraceAnalysis.h b/Engine/Source/Developer/TraceServices/Private/Analyzers/LogTraceAnalysis.h
--- a/Engine/Source/Developer/TraceServices/Private/Analyzers/LogTraceAnalysis.h
+++ b/Engine/Source/Developer/TraceServices/Private/Analyzers/LogTraceAnalysis.h
@@ -4,7 +4,6 @@
 
 #include "Trace/Trace.h"
 #include "Trace/Analyzer.h"
-#include "Templates/SharedPointer.h"
 
 namespace Trace
 {
@@ -16,7 +15,7 @@ class FLogTraceAnalyzer
 	: public Trace::IAnalyzer
 {
 public:
-	FLogTraceAnalyzer(TSharedRef<Trace::FAnalysisSession> Session);
+	FLogTraceAnalyzer(Trace::FAnalysisSession& Session);
 	virtual void OnAnalysisBegin(const FOnAnalysisContext& Context) override;
 	virtual void OnEvent(uint16 RouteId, const FOnEventContext& Context) override;
 	virtual void OnAnalysisEnd() override {};
@@ -29,6 +28,6 @@ private:
 		RouteId_LogMessage,
 	};
 
-	TSharedRef<Trace::FAnalysisSession> Session;
-	TSharedRef<Trace::FLogProvider> LogProvider;
+	Trace::FAnalysisSession& Session;
+	Trace::FLogProvider& LogProvider;
 };
diff --git a/Engine/Source/Developer/TraceServices/Private/Analyzers/MiscTraceAnalysis.cpp b/Engine/Source/Developer/TraceServices/Private/Analyzers/MiscTraceAnalysis.cpp
--- a/Engine/Source/Developer/TraceServices/Private/Analyzers/MiscTraceAnalysis.cpp
+++ b/Engine/Source/Developer/TraceServices/Private/Analyzers/MiscTraceAnalysis.cpp
@@ -6,14 +6,14 @@
 #include "Model/Bookmarks.h"
 #include "Model/Frames.h"
 
-FMiscTraceAnalyzer::FMiscTraceAnalyzer(TSharedRef<Trace::FAnalysisSession> InSession)
+FMiscTraceAnalyzer::FMiscTraceAnalyzer(Trace::FAnalysisSession& InSession)
 	: Session(InSession)
-	, ThreadProvider(InSession->EditThreadProvider())
-	, BookmarkProvider(InSession->EditBookmarkProvider())
-	, LogProvider(InSession->EditLogProvider())
-	, FrameProvider(InSession->EditFramesProvider())
+	, ThreadProvider(InSession.EditThreadProvider())
+	, BookmarkProvider(InSession.EditBookmarkProvider())
+	, LogProvider(InSession.EditLogProvider())
+	, FrameProvider(InSession.EditFrameProvider())
 {
-	Trace::FLogCategory& BookmarkLogCategory = LogProvider->GetCategory(Trace::FLogProvider::ReservedLogCategory_Bookmark);
+	Trace::FLogCategory& BookmarkLogCategory = LogProvider.GetCategory(Trace::FLogProvider::ReservedLogCategory_Bookmark);
 	BookmarkLogCategory.Name = TEXT("LogBookmark");
 	BookmarkLogCategory.DefaultVerbosity = ELogVerbosity::All;
 }
@@ -80,7 +80,7 @@ static ETraceThreadGroup GetThreadGroupFromName(const char* GroupName)
 
 void FMiscTraceAnalyzer::OnEvent(uint16 RouteId, const FOnEventContext& Context)
 {
-	Trace::FAnalysisSessionEditScope _(Session.Get());
+	Trace::FAnalysisSessionEditScope _(Session);
 
 	const auto& EventData = Context.EventData;
 	switch (RouteId)
@@ -88,19 +88,19 @@ void FMiscTraceAnalyzer::OnEvent(uint16 RouteId, const FOnEventContext& Context)
 	case RouteId_RegisterGameThread:
 	{
 		uint32 ThreadId = EventData.GetValue("ThreadId").As<uint32>();
-		ThreadProvider->AddGameThread(ThreadId);
+		ThreadProvider.AddGameThread(ThreadId);
 		break;
 	}
 	case RouteId_CreateThread:
 	{
 		uint32 CreatedThreadId = EventData.GetValue("CreatedThreadId").As<uint32>();
 		EThreadPriority Priority = static_cast<EThreadPriority>(EventData.GetValue("Priority").As<uint32>());
-		ThreadProvider->AddThread(CreatedThreadId, reinterpret_cast<const TCHAR*>(EventData.GetAttachment()), Priority);
+		ThreadProvider.AddThread(CreatedThreadId, reinterpret_cast<const TCHAR*>(EventData.GetAttachment()), Priority);
 		uint32 CurrentThreadId = EventData.GetValue("CurrentThreadId").As<uint32>();
 		FThreadState* ThreadState = GetThreadState(CurrentThreadId);
 		if (ThreadState->ThreadGroupStack.Num())
 		{
-			ThreadProvider->SetThreadGroup(CreatedThreadId, ThreadState->ThreadGroupStack.Top());
+			ThreadProvider.SetThreadGroup(CreatedThreadId, ThreadState->ThreadGroupStack.Top());
 		}
 		break;
 	}
@@ -109,7 +109,7 @@ void FMiscTraceAnalyzer::OnEvent(uint16 RouteId, const FOnEventContext& Context)
 		const char* GroupName = reinterpret_cast<const char*>(EventData.GetAttachment());
 		ETraceThreadGroup ThreadGroup = GetThreadGroupFromName(GroupName);
 		uint32 ThreadId = EventData.GetValue("ThreadId").As<uint32>();
-		ThreadProvider->SetThreadGroup(ThreadId, ThreadGroup);
+		ThreadProvider.SetThreadGroup(ThreadId, ThreadGroup);
 		break;
 	}
 	case RouteId_BeginThreadGroupScope:
@@ -131,14 +131,14 @@ void FMiscTraceAnalyzer::OnEvent(uint16 RouteId, const FOnEventContext& Context)
 	case RouteId_BookmarkSpec:
 	{
 		uint64 BookmarkPoint = EventData.GetValue("BookmarkPoint").As<uint64>();
-		Trace::FBookmarkSpec& Spec = BookmarkProvider->GetSpec(BookmarkPoint);
+		Trace::FBookmarkSpec& Spec = BookmarkProvider.GetSpec(BookmarkPoint);
 		Spec.Line = EventData.GetValue("Line").As<int32>();
 		const ANSICHAR* File = reinterpret_cast<const ANSICHAR*>(EventData.GetAttachment());
-		Spec.File = Session->StoreString(ANSI_TO_TCHAR(File));
-		Spec.FormatString = Session->StoreString(reinterpret_cast<const TCHAR*>(EventData.GetAttachment() + strlen(File) + 1));
+		Spec.File = Session.StoreString(ANSI_TO_TCHAR(File));
+		Spec.FormatString = Session.StoreString(reinterpret_cast<const TCHAR*>(EventData.GetAttachment() + strlen(File) + 1));
 
-		Trace::FLogMessageSpec& LogMessageSpec = LogProvider->GetMessageSpec(BookmarkPoint);
-		LogMessageSpec.Category = &LogProvider->GetCategory(Trace::FLogProvider::ReservedLogCategory_Bookmark);
+		Trace::FLogMessageSpec& LogMessageSpec = LogProvider.GetMessageSpec(BookmarkPoint);
+		LogMessageSpec.Category = &LogProvider.GetCategory(Trace::FLogProvider::ReservedLogCategory_Bookmark);
 		LogMessageSpec.Line = Spec.Line;
 		LogMessageSpec.File = Spec.File;
 		LogMessageSpec.FormatString = Spec.FormatString;
@@ -150,8 +150,8 @@ void FMiscTraceAnalyzer::OnEvent(uint16 RouteId, const FOnEventContext& Context)
 		uint64 BookmarkPoint = EventData.GetValue("BookmarkPoint").As<uint64>();
 		uint64 Cycle = EventData.GetValue("Cycle").As<uint64>();
 		double Timestamp = Context.SessionContext.TimestampFromCycle(Cycle);
-		BookmarkProvider->AppendBookmark(Timestamp, BookmarkPoint, EventData.GetAttachment());
-		LogProvider->AppendMessage(BookmarkPoint, Timestamp, EventData.GetAttachment());
+		BookmarkProvider.AppendBookmark(Timestamp, BookmarkPoint, EventData.GetAttachment());
+		LogProvider.AppendMessage(BookmarkPoint, Timestamp, EventData.GetAttachment());
 		break;
 	}
 	case RouteId_BeginFrame:
@@ -159,7 +159,7 @@ void FMiscTraceAnalyzer::OnEvent(uint16 RouteId, const FOnEventContext& Context)
 		uint64 Cycle = EventData.GetValue("Cycle").As<uint64>();
 		uint8 FrameType = EventData.GetValue("FrameType").As<uint8>();
 		check(FrameType < TraceFrameType_Count);
-		FrameProvider->BeginFrame(ETraceFrameType(FrameType), Context.SessionContext.TimestampFromCycle(Cycle));
+		FrameProvider.BeginFrame(ETraceFrameType(FrameType), Context.SessionContext.TimestampFromCycle(Cycle));
 		break;
 	}
 	case RouteId_EndFrame:
@@ -167,7 +167,7 @@ void FMiscTraceAnalyzer::OnEvent(uint16 RouteId, const FOnEventContext& Context)
 		uint64 Cycle = EventData.GetValue("Cycle").As<uint64>();
 		uint8 FrameType = EventData.GetValue("FrameType").As<uint8>();
 		check(FrameType < TraceFrameType_Count);
-		FrameProvider->EndFrame(ETraceFrameType(FrameType), Context.SessionContext.TimestampFromCycle(Cycle));
+		FrameProvider.EndFrame(ETraceFrameType(FrameType), Context.SessionContext.TimestampFromCycle(Cycle));
 		break;
 	}
 	}
diff --git a/Engine/Source/Developer/TraceServices/Private/Analyzers/MiscTraceAnalysis.h b/Engine/Source/Developer/TraceServices/Private/Analyzers/MiscTraceAnalysis.h
--- a/Engine/Source/Developer/TraceServices/Private/Analyzers/MiscTraceAnalysis.h
+++ b/Engine/Source/Developer/TraceServices/Private/Analyzers/MiscTraceAnalysis.h
@@ -21,7 +21,7 @@ class FMiscTraceAnalyzer
 	: public Trace::IAnalyzer
 {
 public:
-	FMiscTraceAnalyzer(TSharedRef<Trace::FAnalysisSession> Session);
+	FMiscTraceAnalyzer(Trace::FAnalysisSession& Session);
 	virtual void OnAnalysisBegin(const FOnAnalysisContext& Context) override;
 	virtual void OnEvent(uint16 RouteId, const FOnEventContext& Context) override;
 	virtual void OnAnalysisEnd() override {};
@@ -47,11 +47,11 @@ private:
 
 	FThreadState* GetThreadState(uint32 ThreadId);
 
-	TSharedRef<Trace::FAnalysisSession> Session;
-	TSharedRef<Trace::FThreadProvider> ThreadProvider;
-	TSharedRef<Trace::FBookmarkProvider> BookmarkProvider;
-	TSharedRef<Trace::FLogProvider> LogProvider;
-	TSharedRef<Trace::FFrameProvider> FrameProvider;
+	Trace::FAnalysisSession& Session;
+	Trace::FThreadProvider& ThreadProvider;
+	Trace::FBookmarkProvider& BookmarkProvider;
+	Trace::FLogProvider& LogProvider;
+	Trace::FFrameProvider& FrameProvider;
 	TMap<uint32, TSharedRef<FThreadState>> ThreadStateMap;
 };
 
diff --git a/Engine/Source/Developer/TraceServices/Private/Analyzers/PlatformFileTraceAnalysis.cpp b/Engine/Source/Developer/TraceServices/Private/Analyzers/PlatformFileTraceAnalysis.cpp
--- a/Engine/Source/Developer/TraceServices/Private/Analyzers/PlatformFileTraceAnalysis.cpp
+++ b/Engine/Source/Developer/TraceServices/Private/Analyzers/PlatformFileTraceAnalysis.cpp
@@ -3,9 +3,9 @@
 #include "AnalysisServicePrivate.h"
 #include "Model/FileActivity.h"
 
-FPlatformFileTraceAnalyzer::FPlatformFileTraceAnalyzer(TSharedRef<Trace::FAnalysisSession> InSession)
+FPlatformFileTraceAnalyzer::FPlatformFileTraceAnalyzer(Trace::FAnalysisSession& InSession)
 	: Session(InSession)
-	, FileActivityProvider(Session->EditFileActivityProvider())
+	, FileActivityProvider(Session.EditFileActivityProvider())
 {
 
 }
@@ -26,7 +26,7 @@ void FPlatformFileTraceAnalyzer::OnAnalysisBegin(const FOnAnalysisContext& Conte
 
 void FPlatformFileTraceAnalyzer::OnEvent(uint16 RouteId, const FOnEventContext& Context)
 {
-	Trace::FAnalysisSessionEditScope _(Session.Get());
+	Trace::FAnalysisSessionEditScope _(Session);
 
 	const auto& EventData = Context.EventData;
 	switch (RouteId)
@@ -37,9 +37,9 @@ void FPlatformFileTraceAnalyzer::OnEvent(uint16 RouteId, const FOnEventContext&
 		uint64 TempHandle = EventData.GetValue("TempHandle").As<uint64>();
 		uint32 ThreadId = EventData.GetValue("ThreadId").As<uint32>();
 		check(!PendingOpenMap.Contains(TempHandle));
-		uint32 FileIndex = FileActivityProvider->GetFileIndex(reinterpret_cast<const TCHAR*>(EventData.GetAttachment()));
+		uint32 FileIndex = FileActivityProvider.GetFileIndex(reinterpret_cast<const TCHAR*>(EventData.GetAttachment()));
 		FPendingActivity& Open = PendingOpenMap.Add(TempHandle);
-		Open.ActivityIndex = FileActivityProvider->BeginActivity(FileIndex, Trace::FileActivityType_Open, 0, 0, Time);
+		Open.ActivityIndex = FileActivityProvider.BeginActivity(FileIndex, Trace::FileActivityType_Open, 0, 0, Time);
 		Open.FileIndex = FileIndex;
 		break;
 	}
@@ -52,13 +52,13 @@ void FPlatformFileTraceAnalyzer::OnEvent(uint16 RouteId, const FOnEventContext&
 		FPendingActivity& Open = PendingOpenMap[TempHandle];
 		if (FileHandle == uint64(-1)) // TODO: Portable invalid file handle?
 		{
-			FileActivityProvider->EndActivity(Open.FileIndex, Open.ActivityIndex, Time, true);
+			FileActivityProvider.EndActivity(Open.FileIndex, Open.ActivityIndex, Time, true);
 		}
 		else
 		{
 			check(!OpenFilesMap.Contains(FileHandle));
 			OpenFilesMap.Add(FileHandle, Open.FileIndex);
-			FileActivityProvider->EndActivity(Open.FileIndex, Open.ActivityIndex, Time, false);
+			FileActivityProvider.EndActivity(Open.FileIndex, Open.ActivityIndex, Time, false);
 		}
 		PendingOpenMap.Remove(TempHandle);
 		break;
@@ -74,7 +74,7 @@ void FPlatformFileTraceAnalyzer::OnEvent(uint16 RouteId, const FOnEventContext&
 		uint64 FileIndex = OpenFilesMap[FileHandle];
 		OpenFilesMap.Remove(FileHandle);
 		FPendingActivity& Close = PendingCloseMap.Add(TempHandle);
-		Close.ActivityIndex = FileActivityProvider->BeginActivity(FileIndex, Trace::FileActivityType_Close, 0, 0, Time);
+		Close.ActivityIndex = FileActivityProvider.BeginActivity(FileIndex, Trace::FileActivityType_Close, 0, 0, Time);
 		Close.FileIndex = FileIndex;
 		break;
 	}
@@ -84,7 +84,7 @@ void FPlatformFileTraceAnalyzer::OnEvent(uint16 RouteId, const FOnEventContext&
 		uint64 TempHandle = EventData.GetValue("TempHandle").As<uint64>();
 		check(PendingCloseMap.Contains(TempHandle));
 		FPendingActivity& Close = PendingCloseMap[TempHandle];
-		FileActivityProvider->EndActivity(Close.FileIndex, Close.ActivityIndex, Time, false);
+		FileActivityProvider.EndActivity(Close.FileIndex, Close.ActivityIndex, Time, false);
 		PendingCloseMap.Remove(TempHandle);
 		break;
 	}
@@ -98,7 +98,7 @@ void FPlatformFileTraceAnalyzer::OnEvent(uint16 RouteId, const FOnEventContext&
 		uint32 ThreadId = EventData.GetValue("ThreadId").As<uint32>();
 		check(OpenFilesMap.Contains(FileHandle));
 		uint32 FileIndex = OpenFilesMap[FileHandle];
-		uint64 ReadIndex = FileActivityProvider->BeginActivity(FileIndex, Trace::FileActivityType_Read, Offset, Size, Time);
+		uint64 ReadIndex = FileActivityProvider.BeginActivity(FileIndex, Trace::FileActivityType_Read, Offset, Size, Time);
 		FPendingActivity& Read = ActiveReadsMap.Add(ReadHandle);
 		Read.FileIndex = FileIndex;
 		Read.ActivityIndex = ReadIndex;
@@ -112,7 +112,7 @@ void FPlatformFileTraceAnalyzer::OnEvent(uint16 RouteId, const FOnEventContext&
 		uint32 ThreadId = EventData.GetValue("ThreadId").As<uint32>();
 		check(ActiveReadsMap.Contains(ReadHandle));
 		const FPendingActivity& Read = ActiveReadsMap[ReadHandle];
-		FileActivityProvider->EndActivity(Read.FileIndex, Read.ActivityIndex, Time, false);
+		FileActivityProvider.EndActivity(Read.FileIndex, Read.ActivityIndex, Time, false);
 		ActiveReadsMap.Remove(ReadHandle);
 		break;
 	}
@@ -126,7 +126,7 @@ void FPlatformFileTraceAnalyzer::OnEvent(uint16 RouteId, const FOnEventContext&
 		uint32 ThreadId = EventData.GetValue("ThreadId").As<uint32>();
 		check(OpenFilesMap.Contains(FileHandle));
 		uint32 FileIndex = OpenFilesMap[FileHandle];
-		uint64 WriteIndex = FileActivityProvider->BeginActivity(FileIndex, Trace::FileActivityType_Write, Offset, Size, Time);
+		uint64 WriteIndex = FileActivityProvider.BeginActivity(FileIndex, Trace::FileActivityType_Write, Offset, Size, Time);
 		FPendingActivity& Write = ActiveWritesMap.Add(WriteHandle);
 		Write.FileIndex = FileIndex;
 		Write.ActivityIndex = WriteIndex;
@@ -140,7 +140,7 @@ void FPlatformFileTraceAnalyzer::OnEvent(uint16 RouteId, const FOnEventContext&
 		uint32 ThreadId = EventData.GetValue("ThreadId").As<uint32>();
 		check(ActiveWritesMap.Contains(WriteHandle));
 		const FPendingActivity& Write = ActiveWritesMap[WriteHandle];
-		FileActivityProvider->EndActivity(Write.FileIndex, Write.ActivityIndex, Time, false);
+		FileActivityProvider.EndActivity(Write.FileIndex, Write.ActivityIndex, Time, false);
 		ActiveWritesMap.Remove(WriteHandle);
 		break;
 	}
diff --git a/Engine/Source/Developer/TraceServices/Private/Analyzers/PlatformFileTraceAnalysis.h b/Engine/Source/Developer/TraceServices/Private/Analyzers/PlatformFileTraceAnalysis.h
--- a/Engine/Source/Developer/TraceServices/Private/Analyzers/PlatformFileTraceAnalysis.h
+++ b/Engine/Source/Developer/TraceServices/Private/Analyzers/PlatformFileTraceAnalysis.h
@@ -4,7 +4,6 @@
 
 #include "Trace/Trace.h"
 #include "Trace/Analyzer.h"
-#include "Templates/SharedPointer.h"
 
 namespace Trace
 {
@@ -17,7 +16,7 @@ class FPlatformFileTraceAnalyzer
 	: public Trace::IAnalyzer
 {
 public:
-	FPlatformFileTraceAnalyzer(TSharedRef<Trace::FAnalysisSession> Session);
+	FPlatformFileTraceAnalyzer(Trace::FAnalysisSession& Session);
 	virtual void OnAnalysisBegin(const FOnAnalysisContext& Context) override;
 	virtual void OnEvent(uint16 RouteId, const FOnEventContext& Context) override;
 	virtual void OnAnalysisEnd() override;
@@ -43,8 +42,8 @@ private:
 		uint32 FileIndex;
 	};
 
-	TSharedRef<Trace::FAnalysisSession> Session;
-	TSharedRef<Trace::FFileActivityProvider> FileActivityProvider;
+	Trace::FAnalysisSession& Session;
+	Trace::FFileActivityProvider& FileActivityProvider;
 	TMap<uint64, uint32> OpenFilesMap;
 	TMap<uint64, FPendingActivity> PendingOpenMap;
 	TMap<uint64, FPendingActivity> PendingCloseMap;
diff --git a/Engine/Source/Developer/TraceServices/Private/Analyzers/StatsTraceAnalysis.cpp b/Engine/Source/Developer/TraceServices/Private/Analyzers/StatsTraceAnalysis.cpp
--- a/Engine/Source/Developer/TraceServices/Private/Analyzers/StatsTraceAnalysis.cpp
+++ b/Engine/Source/Developer/TraceServices/Private/Analyzers/StatsTraceAnalysis.cpp
@@ -4,9 +4,9 @@
 #include "Common/Utils.h"
 #include "Model/Counters.h"
 
-FStatsAnalyzer::FStatsAnalyzer(TSharedRef<Trace::FAnalysisSession> InSession)
+FStatsAnalyzer::FStatsAnalyzer(Trace::FAnalysisSession& InSession)
 	: Session(InSession)
-	, CountersProvider(Session->EditCountersProvider())
+	, CounterProvider(Session.EditCounterProvider())
 {
 
 }
@@ -30,7 +30,7 @@ void FStatsAnalyzer::OnEvent(uint16 RouteId, const FOnEventContext& Context)
 	{
 	case RouteId_Spec:
 	{
-		Trace::FAnalysisSessionEditScope _(Session.Get());
+		Trace::FAnalysisSessionEditScope _(Session);
 		uint32 StatId = EventData.GetValue("Id").As<uint32>();
 		check(!CountersMap.Contains(StatId));
 		const ANSICHAR* Name = reinterpret_cast<const ANSICHAR*>(EventData.GetAttachment());
@@ -44,13 +44,13 @@ void FStatsAnalyzer::OnEvent(uint16 RouteId, const FOnEventContext& Context)
 		{
 			DisplayHint = Trace::CounterDisplayHint_Memory;
 		}
-		Trace::FCounterInternal* Counter = CountersProvider->CreateCounter(ANSI_TO_TCHAR(Name), Description, DisplayHint);
+		Trace::FCounterInternal* Counter = CounterProvider.CreateCounter(ANSI_TO_TCHAR(Name), Description, DisplayHint);
 		CountersMap.Add(StatId, Counter);
 		break;
 	}
 	case RouteId_EventBatch:
 	{
-		Trace::FAnalysisSessionEditScope _(Session.Get());
+		Trace::FAnalysisSessionEditScope _(Session);
 		uint32 ThreadId = EventData.GetValue("ThreadId").As<uint32>();
 		TSharedRef<FThreadState> ThreadState = GetThreadState(ThreadId);
 		uint64 BufferSize = EventData.GetAttachmentSize();
@@ -81,24 +81,24 @@ void FStatsAnalyzer::OnEvent(uint16 RouteId, const FOnEventContext& Context)
 			{
 			case Increment:
 			{
-				CountersProvider->Add(Counter, Time, int64(1));
+				CounterProvider.Add(Counter, Time, int64(1));
 				break;
 			}
 			case Decrement:
 			{
-				CountersProvider->Add(Counter, Time, int64(-1));
+				CounterProvider.Add(Counter, Time, int64(-1));
 				break;
 			}
 			case AddInteger:
 			{
 				int64 Amount = FTraceAnalyzerUtils::DecodeZigZag(BufferPtr);
-				CountersProvider->Add(Counter, Time, Amount);
+				CounterProvider.Add(Counter, Time, Amount);
 				break;
 			}
 			case SetInteger:
 			{
 				int64 Value = FTraceAnalyzerUtils::DecodeZigZag(BufferPtr);
-				CountersProvider->Set(Counter, Time, Value);
+				CounterProvider.Set(Counter, Time, Value);
 				break;
 			}
 			case AddFloat:
@@ -106,7 +106,7 @@ void FStatsAnalyzer::OnEvent(uint16 RouteId, const FOnEventContext& Context)
 				double Amount;
 				memcpy(&Amount, BufferPtr, sizeof(double));
 				BufferPtr += sizeof(double);
-				CountersProvider->Add(Counter, Time, Amount);
+				CounterProvider.Add(Counter, Time, Amount);
 				break;
 			}
 			case SetFloat:
diff --git a/Engine/Source/Developer/TraceServices/Private/Analyzers/StatsTraceAnalysis.h b/Engine/Source/Developer/TraceServices/Private/Analyzers/StatsTraceAnalysis.h
--- a/Engine/Source/Developer/TraceServices/Private/Analyzers/StatsTraceAnalysis.h
+++ b/Engine/Source/Developer/TraceServices/Private/Analyzers/StatsTraceAnalysis.h
@@ -4,20 +4,19 @@
 
 #include "Trace/Trace.h"
 #include "Trace/Analyzer.h"
-#include "Templates/SharedPointer.h"
 
 namespace Trace
 {
 	class FAnalysisSession;
 	class FCounterInternal;
-	class FCountersProvider;
+	class FCounterProvider;
 }
 
 class FStatsAnalyzer
 	: public Trace::IAnalyzer
 {
 public:
-	FStatsAnalyzer(TSharedRef<Trace::FAnalysisSession> Session);
+	FStatsAnalyzer(Trace::FAnalysisSession& Session);
 	virtual void OnAnalysisBegin(const FOnAnalysisContext& Context) override;
 	virtual void OnEvent(uint16 RouteId, const FOnEventContext& Context) override;
 	virtual void OnAnalysisEnd() override;
@@ -36,8 +35,8 @@ private:
 
 	TSharedRef<FThreadState> GetThreadState(uint32 ThreadId);
 
-	TSharedRef<Trace::FAnalysisSession> Session;
-	TSharedRef<Trace::FCountersProvider> CountersProvider;
+	Trace::FAnalysisSession& Session;
+	Trace::FCounterProvider& CounterProvider;
 	TMap<uint32, Trace::FCounterInternal*> CountersMap;
 	TMap<uint32, TSharedRef<FThreadState>> ThreadStatesMap;
 };
diff --git a/Engine/Source/Developer/TraceServices/Private/Model/Counters.cpp b/Engine/Source/Developer/TraceServices/Private/Model/Counters.cpp
--- a/Engine/Source/Developer/TraceServices/Private/Model/Counters.cpp
+++ b/Engine/Source/Developer/TraceServices/Private/Model/Counters.cpp
@@ -94,14 +94,14 @@ void FCounterInternal::EnumerateFloatValues(double IntervalStart, double Interva
 	}
 }
 
-FCountersProvider::FCountersProvider(FSlabAllocator& InAllocator, FAnalysisSessionLock& InSessionLock)
+FCounterProvider::FCounterProvider(FSlabAllocator& InAllocator, FAnalysisSessionLock& InSessionLock)
 	: Allocator(InAllocator)
 	, SessionLock(InSessionLock)
 {
 
 }
 
-FCountersProvider::~FCountersProvider()
+FCounterProvider::~FCounterProvider()
 {
 	for (FCounterInternal* Counter : Counters)
 	{
@@ -109,7 +109,7 @@ FCountersProvider::~FCountersProvider()
 	}
 }
 
-void FCountersProvider::EnumerateCounters(TFunctionRef<void(const ICounter &)> Callback) const
+void FCounterProvider::EnumerateCounters(TFunctionRef<void(const ICounter &)> Callback) const
 {
 	uint32 Id = 0;
 	for (FCounterInternal* Counter : Counters)
@@ -118,7 +118,7 @@ void FCountersProvider::EnumerateCounters(TFunctionRef<void(const ICounter &)> C
 	}
 }
 
-FCounterInternal* FCountersProvider::CreateCounter(const TCHAR* Name, const TCHAR* Description, ECounterDisplayHint DisplayHint)
+FCounterInternal* FCounterProvider::CreateCounter(const TCHAR* Name, const TCHAR* Description, ECounterDisplayHint DisplayHint)
 {
 	FCounterInternal* Counter = new FCounterInternal(Allocator, Counters.Num(), Name, Description, DisplayHint);
 	Counters.Add(Counter);
@@ -164,22 +164,22 @@ static void InsertOp(FCounterInternal& Counter, double Time, ECounterOpType Type
 	}
 }
 
-void FCountersProvider::Add(FCounterInternal& Counter, double Time, int64 Value)
+void FCounterProvider::Add(FCounterInternal& Counter, double Time, int64 Value)
 {
 	InsertOp(Counter, Time, CounterOpType_Add, Value);
 }
 
-void FCountersProvider::Add(FCounterInternal& Counter, double Time, double Value)
+void FCounterProvider::Add(FCounterInternal& Counter, double Time, double Value)
 {
 	InsertOp(Counter, Time, CounterOpType_Add, Value);
 }
 
-void FCountersProvider::Set(FCounterInternal& Counter, double Time, int64 Value)
+void FCounterProvider::Set(FCounterInternal& Counter, double Time, int64 Value)
 {
 	InsertOp(Counter, Time, CounterOpType_Set, Value);
 }
 
-void FCountersProvider::Set(FCounterInternal& Counter, double Time, double Value)
+void FCounterProvider::Set(FCounterInternal& Counter, double Time, double Value)
 {
 	InsertOp(Counter, Time, CounterOpType_Set, Value);
 }
diff --git a/Engine/Source/Developer/TraceServices/Private/Model/Counters.h b/Engine/Source/Developer/TraceServices/Private/Model/Counters.h
--- a/Engine/Source/Developer/TraceServices/Private/Model/Counters.h
+++ b/Engine/Source/Developer/TraceServices/Private/Model/Counters.h
@@ -48,12 +48,12 @@ private:
 	TPagedArray<double> DoubleOpArguments;
 };
 
-class FCountersProvider
-	: public ICountersProvider
+class FCounterProvider
+	: public ICounterProvider
 {
 public:
-	FCountersProvider(FSlabAllocator& Allocator, FAnalysisSessionLock& SessionLock);
-	virtual ~FCountersProvider();
+	FCounterProvider(FSlabAllocator& Allocator, FAnalysisSessionLock& SessionLock);
+	virtual ~FCounterProvider();
 	virtual uint64 GetCounterCount() const override { return Counters.Num(); }
 	virtual void EnumerateCounters(TFunctionRef<void(const ICounter&)> Callback) const override;
 
diff --git a/Engine/Source/Developer/TraceServices/Private/Model/TimingProfiler.cpp b/Engine/Source/Developer/TraceServices/Private/Model/TimingProfiler.cpp
--- a/Engine/Source/Developer/TraceServices/Private/Model/TimingProfiler.cpp
+++ b/Engine/Source/Developer/TraceServices/Private/Model/TimingProfiler.cpp
@@ -51,7 +51,7 @@ FTimingProfilerTimer& FTimingProfilerProvider::AddTimerInternal(const TCHAR* Nam
 	return Timer;
 }
 
-TSharedRef<FTimingProfilerProvider::TimelineInternal> FTimingProfilerProvider::EditCpuThreadTimeline(uint32 ThreadId)
+FTimingProfilerProvider::TimelineInternal& FTimingProfilerProvider::EditCpuThreadTimeline(uint32 ThreadId)
 {
 	SessionLock.WriteAccessCheck();
 
@@ -61,20 +61,20 @@ TSharedRef<FTimingProfilerProvider::TimelineInternal> FTimingProfilerProvider::E
 		uint32 TimelineIndex = Timelines.Num();
 		CpuThreadTimelineIndexMap.Add(ThreadId, TimelineIndex);
 		Timelines.Add(Timeline);
-		return Timeline;
+		return Timeline.Get();
 	}
 	else
 	{
 		uint32 TimelineIndex = CpuThreadTimelineIndexMap[ThreadId];
-		return Timelines[TimelineIndex];
+		return Timelines[TimelineIndex].Get();
 	}
 }
 
-TSharedRef<FTimingProfilerProvider::TimelineInternal> FTimingProfilerProvider::EditGpuTimeline()
+FTimingProfilerProvider::TimelineInternal& FTimingProfilerProvider::EditGpuTimeline()
 {
 	SessionLock.WriteAccessCheck();
 
-	return Timelines[GpuTimelineIndex];
+	return Timelines[GpuTimelineIndex].Get();
 }
 
 bool FTimingProfilerProvider::GetCpuThreadTimelineIndex(uint32 ThreadId, uint32& OutTimelineIndex) const
diff --git a/Engine/Source/Developer/TraceServices/Private/Model/TimingProfiler.h b/Engine/Source/Developer/TraceServices/Private/Model/TimingProfiler.h
--- a/Engine/Source/Developer/TraceServices/Private/Model/TimingProfiler.h
+++ b/Engine/Source/Developer/TraceServices/Private/Model/TimingProfiler.h
@@ -23,8 +23,8 @@ public:
 	virtual ~FTimingProfilerProvider();
 	uint32 AddCpuTimer(const TCHAR* Name);
 	uint32 AddGpuTimer(const TCHAR* Name);
-	TSharedRef<TimelineInternal> EditCpuThreadTimeline(uint32 ThreadId);
-	TSharedRef<TimelineInternal> EditGpuTimeline();
+	TimelineInternal& EditCpuThreadTimeline(uint32 ThreadId);
+	TimelineInternal& EditGpuTimeline();
 	virtual bool GetCpuThreadTimelineIndex(uint32 ThreadId, uint32& OutTimelineIndex) const override;
 	virtual bool GetGpuTimelineIndex(uint32& OutTimelineIndex) const override;
 	virtual bool ReadTimeline(uint32 Index, TFunctionRef<void(const Timeline&)> Callback) const override;
diff --git a/Engine/Source/Developer/TraceServices/Public/TraceServices/AnalysisService.h b/Engine/Source/Developer/TraceServices/Public/TraceServices/AnalysisService.h
--- a/Engine/Source/Developer/TraceServices/Public/TraceServices/AnalysisService.h
+++ b/Engine/Source/Developer/TraceServices/Public/TraceServices/AnalysisService.h
@@ -327,10 +327,10 @@ public:
 	virtual void EnumerateFloatValues(double IntervalStart, double IntervalEnd, TFunctionRef<void(double, double)> Callback) const = 0;
 };
 
-class ICountersProvider
+class ICounterProvider
 {
 public:
-	virtual ~ICountersProvider() = default;
+	virtual ~ICounterProvider() = default;
 	virtual uint64 GetCounterCount() const = 0;
 	virtual void EnumerateCounters(TFunctionRef<void(const ICounter&)> Callback) const = 0;
 };
@@ -346,14 +346,14 @@ public:
 
 	virtual void BeginRead() const = 0;
 	virtual void EndRead() const = 0;
-	virtual void ReadBookmarkProvider(TFunctionRef<void(const IBookmarkProvider&)> Callback) const = 0;
-	virtual void ReadLogProvider(TFunctionRef<void(const ILogProvider&)> Callback) const = 0;
-	virtual void ReadThreadProvider(TFunctionRef<void(const IThreadProvider&)> Callback) const = 0;
-	virtual void ReadFramesProvider(TFunctionRef<void(const IFrameProvider&)> Callback) const = 0;
-	virtual void ReadTimingProfilerProvider(TFunctionRef<void(const ITimingProfilerProvider&)> Callback) const = 0;
-	virtual void ReadFileActivityProvider(TFunctionRef<void(const IFileActivityProvider&)> Callback) const = 0;
-	virtual void ReadLoadTimeProfilerProvider(TFunctionRef<void(const ILoadTimeProfilerProvider&)> Callback) const = 0;
-	virtual void ReadCountersProvider(TFunctionRef<void(const ICountersProvider&)> Callback) const = 0;
+	virtual const IBookmarkProvider& ReadBookmarkProvider() const = 0;
+	virtual const ILogProvider& ReadLogProvider() const = 0;
+	virtual const IThreadProvider& ReadThreadProvider() const = 0;
+	virtual const IFrameProvider& ReadFrameProvider() const = 0;
+	virtual const ITimingProfilerProvider& ReadTimingProfilerProvider() const = 0;
+	virtual const IFileActivityProvider& ReadFileActivityProvider() const = 0;
+	virtual const ILoadTimeProfilerProvider& ReadLoadTimeProfilerProvider() const = 0;
+	virtual const ICounterProvider& ReadCounterProvider() const = 0;
 };
 
 struct FAnalysisSessionReadScope
@@ -378,7 +378,6 @@ class IAnalysisService
 public:
 	virtual TSharedPtr<const IAnalysisSession> Analyze(const TCHAR* SessionName, TUniquePtr<Trace::IInDataStream>&& DataStream) = 0;
 	virtual TSharedPtr<const IAnalysisSession> StartAnalysis(const TCHAR* SessionName, TUniquePtr<Trace::IInDataStream>&& DataStream) = 0;
-	virtual TSharedPtr<const IAnalysisSession> MockAnalysis() = 0;
 
 	DECLARE_EVENT_OneParam(IAnalysisService, FAnalysisStartedEvent, TSharedRef<const IAnalysisSession>)
 	virtual FAnalysisStartedEvent& OnAnalysisStarted() = 0;
