Compile fix

diff --git a/Engine/Source/Developer/TraceInsights/Private/Insights/Widgets/SStatsView.cpp b/Engine/Source/Developer/TraceInsights/Private/Insights/Widgets/SStatsView.cpp
--- a/Engine/Source/Developer/TraceInsights/Private/Insights/Widgets/SStatsView.cpp
+++ b/Engine/Source/Developer/TraceInsights/Private/Insights/Widgets/SStatsView.cpp
@@ -1748,7 +1748,6 @@ void SStatsView::UpdateStats(double StartTime, double EndTime)
 
 				// Compute histogram.
 				// Iterate again through all counters.
-				const Trace::ICounterProvider& CountersProvider = Session->ReadCounterProvider();
 				CountersProvider.EnumerateCounters([&CalculationHelperDbl, &CalculationHelperInt](const Trace::ICounter& Counter)
 				{
 					if (Counter.GetDisplayHint() == Trace::CounterDisplayHint_FloatingPoint)
diff --git a/Engine/Source/Developer/TraceInsights/Private/Insights/Widgets/STimingView.cpp b/Engine/Source/Developer/TraceInsights/Private/Insights/Widgets/STimingView.cpp
--- a/Engine/Source/Developer/TraceInsights/Private/Insights/Widgets/STimingView.cpp
+++ b/Engine/Source/Developer/TraceInsights/Private/Insights/Widgets/STimingView.cpp
@@ -2536,7 +2536,6 @@ bool STimingView::SearchTimingEvent(const double InStartTime,
 	
 		if (Ctx.bFound)
 		{
-			const Trace::ITimingProfilerProvider& TimingProfilerProvider = Session->ReadTimingProfilerProvider();
 			TimingProfilerProvider.ReadTimeline(Ctx.TimingEvent.Track->GetId(), [&Ctx](const Trace::ITimingProfilerProvider::Timeline& Timeline)
 			{
 				struct FEnumerationState
