[Insights]
  - Start Page: Changed "Live" button to open the latest available live session and to be disabled if there are no live sessions.
  - Start Page: Added "Latest" button to open the latest available session, if any (otherwise disabled).
  - Start Page: Fixed Load to not create the profiler tabs if dialog is canceled.
  - Changed Load to open, by default, the LocalSessionDirectory instead of ProfilingDir.
  - Timing Insights: Removed Live and Load buttons from toolbar.
  - Minor code style fixes.

diff --git a/Engine/Source/Developer/TraceInsights/Private/Insights/InsightsCommands.cpp b/Engine/Source/Developer/TraceInsights/Private/Insights/InsightsCommands.cpp
--- a/Engine/Source/Developer/TraceInsights/Private/Insights/InsightsCommands.cpp
+++ b/Engine/Source/Developer/TraceInsights/Private/Insights/InsightsCommands.cpp
@@ -36,7 +36,7 @@ FInsightsCommands::FInsightsCommands()
 PRAGMA_DISABLE_OPTIMIZATION
 void FInsightsCommands::RegisterCommands()
 {
-	UI_COMMAND(InsightsManager_Live, "Live", "Loads profiler data from a live trace session", EUserInterfaceActionType::Button, FInputChord());
+	UI_COMMAND(InsightsManager_Live, "Live", "Loads profiler data from a live or last trace session", EUserInterfaceActionType::Button, FInputChord());
 	UI_COMMAND(InsightsManager_Load, "Load...", "Loads profiler data from a trace file", EUserInterfaceActionType::Button, FInputChord(EModifierKey::Control, EKeys::L));
 	UI_COMMAND(ToggleDebugInfo, "Debug", "Toggles the display of debug info", EUserInterfaceActionType::ToggleButton, FInputChord(EModifierKey::Control, EKeys::D));
 	UI_COMMAND(OpenSettings, "Settings", "Opens the Unreal Insights settings", EUserInterfaceActionType::Button, FInputChord(EModifierKey::Control, EKeys::O));
@@ -95,7 +95,11 @@ void FInsightsActionManager::Map_InsightsManager_Live()
 
 void FInsightsActionManager::InsightsManager_Live_Execute()
 {
-	This->CreateLiveSession();
+	This->LoadLastLiveSession();
+	if (!This->GetSession().IsValid())
+	{
+		This->LoadLastSession();
+	}
 }
 
 ////////////////////////////////////////////////////////////////////////////////////////////////////
@@ -122,8 +126,11 @@ void FInsightsActionManager::Map_InsightsManager_Load()
 
 void FInsightsActionManager::InsightsManager_Load_Execute()
 {
+	//const FString ProfilingDirectory(FPaths::ConvertRelativePathToFull(*FPaths::ProfilingDir()));
+	TSharedRef<Trace::ISessionService> SessionService = This->GetSessionService();
+	const FString ProfilingDirectory(FPaths::ConvertRelativePathToFull(SessionService->GetLocalSessionDirectory()));
+
 	TArray<FString> OutFiles;
-	const FString ProfilingDirectory = *FPaths::ConvertRelativePathToFull(*FPaths::ProfilingDir());
 
 	bool bOpened = false;
 
@@ -135,10 +142,10 @@ void FInsightsActionManager::InsightsManager_Load_Execute()
 		bOpened = DesktopPlatform->OpenFileDialog
 		(
 			FSlateApplication::Get().FindBestParentWindowHandleForDialogs(nullptr),
-			LOCTEXT("UnrealInsights_LoadTrace_FileDesc", "Open trace capture file...").ToString(),
+			LOCTEXT("LoadTrace_FileDesc", "Open trace file...").ToString(),
 			ProfilingDirectory,
 			TEXT(""),
-			LOCTEXT("UnrealInsights_LoadTrace_FileFilter", "Trace files (*.utrace)|*.utrace|All files (*.*)|*.*").ToString(),
+			LOCTEXT("LoadTrace_FileFilter", "Trace files (*.utrace)|*.utrace|All files (*.*)|*.*").ToString(),
 			EFileDialogFlags::None,
 			OutFiles
 		);
diff --git a/Engine/Source/Developer/TraceInsights/Private/Insights/InsightsManager.cpp b/Engine/Source/Developer/TraceInsights/Private/Insights/InsightsManager.cpp
--- a/Engine/Source/Developer/TraceInsights/Private/Insights/InsightsManager.cpp
+++ b/Engine/Source/Developer/TraceInsights/Private/Insights/InsightsManager.cpp
@@ -2,15 +2,15 @@
 
 #include "InsightsManager.h"
 
+#include "Framework/Application/SlateApplication.h"
 #include "Modules/ModuleManager.h"
 #include "Templates/ScopedPointer.h"
 #include "Templates/UniquePtr.h"
-#include "TraceServices/AnalysisService.h"
-#include "TraceServices/SessionService.h"
 
 // Insights
 #include "Insights/TimingProfilerManager.h"
 #include "Insights/Widgets/SStartPageWindow.h"
+#include "Insights/Widgets/STimingProfilerWindow.h"
 
 ////////////////////////////////////////////////////////////////////////////////////////////////////
 
@@ -121,8 +121,11 @@ bool FInsightsManager::Tick(float DeltaTime)
 
 void FInsightsManager::ResetSession()
 {
-	Session = nullptr;
-	OnSessionChanged();
+	if (Session.IsValid())
+	{
+		Session.Reset();
+		OnSessionChanged();
+	}
 }
 
 ////////////////////////////////////////////////////////////////////////////////////////////////////
@@ -140,34 +143,131 @@ void FInsightsManager::OnSessionChanged()
 
 ////////////////////////////////////////////////////////////////////////////////////////////////////
 
-void FInsightsManager::CreateLiveSession()
+void FInsightsManager::SpawnAndActivateTabs()
+{
+	// Open tabs for profilers.
+	if (FGlobalTabmanager::Get()->CanSpawnTab(FInsightsManagerTabs::TimingProfilerTabId))
+	{
+		FGlobalTabmanager::Get()->InvokeTab(FInsightsManagerTabs::TimingProfilerTabId);
+	}
+	if (FGlobalTabmanager::Get()->CanSpawnTab(FInsightsManagerTabs::IoProfilerTabId))
+	{
+		FGlobalTabmanager::Get()->InvokeTab(FInsightsManagerTabs::IoProfilerTabId);
+	}
+
+	// Ensure Timing Insights / Timing View is the active tab / view.
+	if (TSharedPtr<SDockTab> TimingInsightsTab = FGlobalTabmanager::Get()->FindExistingLiveTab(FInsightsManagerTabs::TimingProfilerTabId))
+	{
+		TimingInsightsTab->ActivateInParent(ETabActivationCause::SetDirectly);
+
+		//TOOD: FTimingProfilerManager::Get()->SpawnAndActivateTabs();
+		TSharedPtr<class STimingProfilerWindow> Wnd = FTimingProfilerManager::Get()->GetProfilerWindow();
+		if (Wnd)
+		{
+			TSharedPtr<FTabManager> TabManager = Wnd->GetTabManager();
+
+			if (TSharedPtr<SDockTab> TimingViewTab = TabManager->FindExistingLiveTab(FTimingProfilerTabs::TimingViewID))
+			{
+				TimingViewTab->ActivateInParent(ETabActivationCause::SetDirectly);
+				FSlateApplication::Get().SetKeyboardFocus(TimingViewTab->GetContent());
+			}
+		}
+	}
+}
+
+////////////////////////////////////////////////////////////////////////////////////////////////////
+
+bool FInsightsManager::IsAnyLiveSessionAvailable(Trace::FSessionHandle& OutLastLiveSessionHandle) const
+{
+	TArray<Trace::FSessionHandle> AvailableSessions;
+	SessionService->GetAvailableSessions(AvailableSessions);
+
+	for (int32 SessionIndex = AvailableSessions.Num() - 1; SessionIndex >= 0; --SessionIndex)
+	{
+		Trace::FSessionInfo SessionInfo;
+		SessionService->GetSessionInfo(AvailableSessions[SessionIndex], SessionInfo);
+
+		if (SessionInfo.bIsLive)
+		{
+			OutLastLiveSessionHandle = AvailableSessions[SessionIndex];
+			return true;
+		}
+	}
+
+	return false;
+}
+
+////////////////////////////////////////////////////////////////////////////////////////////////////
+
+bool FInsightsManager::IsAnySessionAvailable(Trace::FSessionHandle& OutLastSessionHandle) const
+{
+	TArray<Trace::FSessionHandle> AvailableSessions;
+	SessionService->GetAvailableSessions(AvailableSessions);
+
+	if (AvailableSessions.Num() > 0)
+	{
+		OutLastSessionHandle = AvailableSessions.Last();
+		return true;
+	}
+
+	return false;
+}
+
+////////////////////////////////////////////////////////////////////////////////////////////////////
+
+void FInsightsManager::LoadLastLiveSession()
 {
 	ResetSession();
 
 	TArray<Trace::FSessionHandle> AvailableSessions;
 	SessionService->GetAvailableSessions(AvailableSessions);
 
-	for (Trace::FSessionHandle SessionHandle : AvailableSessions)
+	// Iterate in reverse order as we want the most recent live session first.
+	for (int32 SessionIndex = AvailableSessions.Num() - 1; SessionIndex >= 0; --SessionIndex)
 	{
+		Trace::FSessionHandle SessionHandle = AvailableSessions[SessionIndex];
+
 		Trace::FSessionInfo SessionInfo;
 		SessionService->GetSessionInfo(SessionHandle, SessionInfo);
+
 		if (SessionInfo.bIsLive)
 		{
-			TUniquePtr<Trace::IInDataStream> DataStream(SessionService->OpenSessionStream(SessionHandle));
-			check(DataStream);
-			Session = AnalysisService->StartAnalysis(TEXT("Live session"), MoveTemp(DataStream));
+			LoadSession(SessionHandle);
 			break;
 		}
 	}
+}
+
+////////////////////////////////////////////////////////////////////////////////////////////////////
 
-	if (!Session && AvailableSessions.Num())
+void FInsightsManager::LoadLastSession()
+{
+	ResetSession();
+
+	TArray<Trace::FSessionHandle> AvailableSessions;
+	SessionService->GetAvailableSessions(AvailableSessions);
+
+	if (AvailableSessions.Num() > 0)
 	{
-		Trace::FSessionHandle SessionHandle = AvailableSessions.Last();
-		TUniquePtr<Trace::IInDataStream> DataStream(SessionService->OpenSessionStream(SessionHandle));
-		check(DataStream);
-		Session = AnalysisService->StartAnalysis(TEXT("Latest session"), MoveTemp(DataStream));
+		LoadSession(AvailableSessions.Last());
 	}
+}
+
+////////////////////////////////////////////////////////////////////////////////////////////////////
+
+void FInsightsManager::LoadSession(Trace::FSessionHandle SessionHandle)
+{
+	ResetSession();
+
+	Trace::FSessionInfo SessionInfo;
+	SessionService->GetSessionInfo(SessionHandle, SessionInfo);
+
+	TUniquePtr<Trace::IInDataStream> DataStream(SessionService->OpenSessionStream(SessionHandle));
+	check(DataStream);
+
+	Session = AnalysisService->StartAnalysis(SessionInfo.Name, MoveTemp(DataStream));
 
+	SpawnAndActivateTabs();
 	OnSessionChanged();
 }
 
@@ -182,6 +282,7 @@ void FInsightsManager::LoadTraceFile(const FString& TraceFilepath)
 
 	Session = AnalysisService->StartAnalysis(*TraceFilepath, MoveTemp(DataStream));
 
+	SpawnAndActivateTabs();
 	OnSessionChanged();
 }
 
diff --git a/Engine/Source/Developer/TraceInsights/Private/Insights/InsightsManager.h b/Engine/Source/Developer/TraceInsights/Private/Insights/InsightsManager.h
--- a/Engine/Source/Developer/TraceInsights/Private/Insights/InsightsManager.h
+++ b/Engine/Source/Developer/TraceInsights/Private/Insights/InsightsManager.h
@@ -5,6 +5,8 @@
 #include "Containers/Ticker.h"
 #include "CoreMinimal.h"
 #include "Framework/Commands/UICommandList.h"
+#include "TraceServices/AnalysisService.h"
+#include "TraceServices/SessionService.h"
 
 // Insights
 #include "Insights/InsightsCommands.h"
@@ -14,13 +16,6 @@
 
 class SStartPageWindow;
 
-namespace Trace
-{
-	class IAnalysisService;
-	class IAnalysisSession;
-	class ISessionService;
-}
-
 ////////////////////////////////////////////////////////////////////////////////////////////////////
 
 struct FInsightsManagerTabs
@@ -142,11 +137,23 @@ public:
 
 	////////////////////////////////////////////////////////////////////////////////////////////////////
 
-	/** Creates a new profiler session instance and loads a live trace. */
-	void CreateLiveSession();
+	bool IsAnyLiveSessionAvailable(Trace::FSessionHandle& OutLastLiveSessionHandle) const;
+	bool IsAnySessionAvailable(Trace::FSessionHandle& OutLastSessionHandle) const;
+
+	/** Creates a new analysis session instance and loads the latest available trace session that is live. */
+	void LoadLastLiveSession();
+
+	/** Creates a new analysis session instance and loads the latest available trace session. */
+	void LoadLastSession();
 
 	/**
-	 * Creates a new profiler session instance and loads a trace file from the specified location.
+	 * Creates a new analysis session instance using specified session handle.
+	 * @param SessionHandle - The handle for session to analyze
+	 */
+	void LoadSession(Trace::FSessionHandle SessionHandle);
+
+	/**
+	 * Creates a new analysis session instance and loads a trace file from the specified location.
 	 * @param TraceFilepath - The path to the trace file
 	 */
 	void LoadTraceFile(const FString& TraceFilepath);
@@ -182,6 +189,8 @@ private:
 
 	void OnSessionChanged();
 
+	void SpawnAndActivateTabs();
+
 private:
 	/** The delegate to be invoked when this manager ticks. */
 	FTickerDelegate OnTick;
diff --git a/Engine/Source/Developer/TraceInsights/Private/Insights/Widgets/SFrameTrack.cpp b/Engine/Source/Developer/TraceInsights/Private/Insights/Widgets/SFrameTrack.cpp
--- a/Engine/Source/Developer/TraceInsights/Private/Insights/Widgets/SFrameTrack.cpp
+++ b/Engine/Source/Developer/TraceInsights/Private/Insights/Widgets/SFrameTrack.cpp
@@ -187,12 +187,13 @@ void SFrameTrack::Tick(const FGeometry& AllottedGeometry, const double InCurrent
 			Trace::FAnalysisSessionReadScope SessionReadScope(*Session.Get());
 
 			const Trace::IFrameProvider& FramesProvider = Session->ReadFrameProvider();
+
 			for (int32 FrameType = 0; FrameType < TraceFrameType_Count; ++FrameType)
 			{
 				FFrameTrackTimeline& CachedTimeline = CachedTimelines.FindOrAdd(FrameType);
 				CachedTimeline.Id = FrameType;
 
-				int32 NumFrames = FramesProvider.GetFrameCount(ETraceFrameType(FrameType));
+				int32 NumFrames = FramesProvider.GetFrameCount(static_cast<ETraceFrameType>(FrameType));
 				if (NumFrames > Viewport.MaxIndex)
 				{
 					Viewport.SetMinMaxIndexInterval(0, NumFrames);
@@ -240,6 +241,7 @@ void SFrameTrack::UpdateState()
 		Trace::FAnalysisSessionReadScope SessionReadScope(*Session.Get());
 
 		const Trace::IFrameProvider& FramesProvider = Session->ReadFrameProvider();
+
 		for (int32 FrameType = 0; FrameType < TraceFrameType_Count; ++FrameType)
 		{
 			FFrameTrackTimeline& Timeline = CachedTimelines.FindOrAdd(FrameType);
@@ -247,9 +249,10 @@ void SFrameTrack::UpdateState()
 
 			FFrameTrackTimelineBuilder Builder(Timeline, Viewport);
 
-			int32 StartIndex = FMath::Max(0, Viewport.GetIndexAtViewportX(0.0f));
-			int32 EndIndex = Viewport.GetIndexAtViewportX(Viewport.Width);
-			FramesProvider.EnumerateFrames(ETraceFrameType(FrameType), static_cast<uint64>(StartIndex), static_cast<uint64>(EndIndex), [&Builder](const Trace::FFrame& Frame)
+			uint64 StartIndex = static_cast<uint64>(FMath::Max(0, Viewport.GetIndexAtViewportX(0.0f)));
+			uint64 EndIndex = static_cast<uint64>(Viewport.GetIndexAtViewportX(Viewport.Width));
+
+			FramesProvider.EnumerateFrames(static_cast<ETraceFrameType>(FrameType), StartIndex, EndIndex, [&Builder](const Trace::FFrame& Frame)
 			{
 				Builder.AddFrame(Frame);
 			});
diff --git a/Engine/Source/Developer/TraceInsights/Private/Insights/Widgets/SLogView.cpp b/Engine/Source/Developer/TraceInsights/Private/Insights/Widgets/SLogView.cpp
--- a/Engine/Source/Developer/TraceInsights/Private/Insights/Widgets/SLogView.cpp
+++ b/Engine/Source/Developer/TraceInsights/Private/Insights/Widgets/SLogView.cpp
@@ -705,6 +705,7 @@ void SLogView::Tick(const FGeometry& AllottedGeometry, const double InCurrentTim
 	{
 		Trace::FAnalysisSessionReadScope SessionReadScope(*Session.Get());
 		const Trace::ILogProvider& LogProvider = Session->ReadLogProvider();
+
 		NewMessageCount = static_cast<int32>(LogProvider.GetMessageCount());
 
 		//TODO: show only categories that are used in current trace
@@ -728,7 +729,9 @@ void SLogView::Tick(const FGeometry& AllottedGeometry, const double InCurrentTim
 			{
 				FString CategoryStr(Category.Name);
 				if (CategoryStr.StartsWith(TEXT("Log")))
+				{
 					CategoryStr = CategoryStr.RightChop(3);
+				}
 				if (Categories.Contains(FName(*CategoryStr)))
 				{
 					UE_LOG(TimingProfiler, Log, TEXT("[LogView] Duplicated Log Category: \"%s\""), Category.Name);
diff --git a/Engine/Source/Developer/TraceInsights/Private/Insights/Widgets/SStartPageWindow.cpp b/Engine/Source/Developer/TraceInsights/Private/Insights/Widgets/SStartPageWindow.cpp
--- a/Engine/Source/Developer/TraceInsights/Private/Insights/Widgets/SStartPageWindow.cpp
+++ b/Engine/Source/Developer/TraceInsights/Private/Insights/Widgets/SStartPageWindow.cpp
@@ -2,9 +2,12 @@
 
 #include "SStartPageWindow.h"
 
+#include "DesktopPlatformModule.h"
 #include "EditorStyleSet.h"
 #include "Framework/Application/SlateApplication.h"
 #include "Framework/Docking/WorkspaceItem.h"
+#include "Framework/MultiBox/MultiBoxBuilder.h"
+#include "HAL/FileManagerGeneric.h"
 #include "SlateOptMacros.h"
 #include "TraceServices/SessionService.h"
 #include "Widgets/Docking/SDockTab.h"
@@ -16,6 +19,7 @@
 #include "Widgets/SBoxPanel.h"
 #include "Widgets/Text/STextBlock.h"
 #include "Widgets/Input/SButton.h"
+#include "Widgets/Input/SComboButton.h"
 //#include "WorkspaceMenuStructure.h"
 //#include "WorkspaceMenuStructureModule.h"
 
@@ -80,6 +84,10 @@ static FText GetTextForNotification(const EInsightsNotificationType NotificatonT
 
 SStartPageWindow::SStartPageWindow()
 	: DurationActive(0.0f)
+	, bIsAnyLiveSessionAvailable(false)
+	, LastLiveSessionHandle(static_cast<Trace::FSessionHandle>(0))
+	, bIsAnySessionAvailable(false)
+	, LastSessionHandle(static_cast<Trace::FSessionHandle>(0))
 {
 }
 
@@ -99,8 +107,6 @@ SStartPageWindow::~SStartPageWindow()
 
 BEGIN_SLATE_FUNCTION_BUILD_OPTIMIZATION
 
-////////////////////////////////////////////////////////////////////////////////////////////////////
-
 void SStartPageWindow::Construct(const FArguments& InArgs)
 {
 	ChildSlot
@@ -152,9 +158,10 @@ void SStartPageWindow::Construct(const FArguments& InArgs)
 							[
 								SNew(STextBlock)
 								.Text(LOCTEXT("SelectTraceHint", "\
-- Use \"Live\" button to start a live session.\n\
-   If no live session, it will load the last trace file from the Local Session Directory.\n\
+- Use \"Live\" button to start analysis for the latest available session that is live, if any.\n\
+- Use \"Latest\" button to start analysis for the latest available session, if any.\n\
 - Use \"Load\" button to load a trace file (*.utrace).\n\
+- Use the combo box to choose from available sessions.\n\
 - Drag and drop a *.utrace file over this window."))
 								//.Justification(ETextJustify::Center)
 							]
@@ -173,7 +180,8 @@ void SStartPageWindow::Construct(const FArguments& InArgs)
 								[
 									SNew(SButton)
 									.OnClicked(this, &SStartPageWindow::Live_OnClicked)
-									.ToolTipText(LOCTEXT("LiveButtonTooltip", "Start a live session or load the last trace file from the Local Session Directory."))
+									.IsEnabled(this, &SStartPageWindow::Live_IsEnabled)
+									.ToolTipText(LOCTEXT("LiveButtonTooltip", "Start analysis for the latest available session that is live, if any."))
 									.ContentPadding(8.0f)
 									.Content()
 									[
@@ -204,6 +212,41 @@ void SStartPageWindow::Construct(const FArguments& InArgs)
 							+ SHorizontalBox::Slot()
 								.AutoWidth()
 								.Padding(3.0f, 0.0f)
+								[
+									SNew(SButton)
+									.OnClicked(this, &SStartPageWindow::Last_OnClicked)
+									.IsEnabled(this, &SStartPageWindow::Last_IsEnabled)
+									.ToolTipText(LOCTEXT("LastButtonTooltip", "Start analysis for the latest available session, if any."))
+									.ContentPadding(8.0f)
+									.Content()
+									[
+										SNew(SHorizontalBox)
+
+										+SHorizontalBox::Slot()
+											.AutoWidth()
+											[
+												SNew(SImage)
+												.Image(FInsightsStyle::GetBrush("Live.Icon.Large"))
+											]
+
+										+SHorizontalBox::Slot()
+											.AutoWidth()
+											.VAlign(VAlign_Center)
+											[
+												SNew(SBox)
+												.WidthOverride(100.0f)
+												[
+													SNew(STextBlock)
+													.Font(FCoreStyle::GetDefaultFontStyle("Regular", 18))
+													.Text(LOCTEXT("LastButtonText", "Latest"))
+												]
+											]
+									]
+								]
+
+							+ SHorizontalBox::Slot()
+								.AutoWidth()
+								.Padding(3.0f, 0.0f, 0.0f, 0.0f)
 								[
 									SNew(SButton)
 									.OnClicked(this, &SStartPageWindow::Load_OnClicked)
@@ -234,6 +277,17 @@ void SStartPageWindow::Construct(const FArguments& InArgs)
 											]
 									]
 								]
+
+							+ SHorizontalBox::Slot()
+								.AutoWidth()
+								.Padding(0.0f, 0.0f, 0.0f, 0.0f)
+								[
+									SNew(SComboButton)
+									.ToolTipText(LOCTEXT("SessionList_Tooltip", "Choose from Most Recently Used Sessions or from Available Sessions"))
+									.OnGetMenuContent(this, &SStartPageWindow::MakeSessionListMenu)
+									.HasDownArrow(true)
+									.ContentPadding(FMargin(1.0f, 1.0f, 1.0f, 1.0f))
+								]
 						]
 
 					+ SVerticalBox::Slot()
@@ -373,6 +427,23 @@ END_SLATE_FUNCTION_BUILD_OPTIMIZATION
 
 ////////////////////////////////////////////////////////////////////////////////////////////////////
 
+void SStartPageWindow::Tick(const FGeometry& AllottedGeometry, const double InCurrentTime, const float InDeltaTime)
+{
+	// We need to check if there is any available session, but not too often.
+	static uint64 NextTimestamp = 0;
+	uint64 Time = FPlatformTime::Cycles64();
+	if (Time > NextTimestamp)
+	{
+		const uint64 WaitTime = static_cast<uint64>(0.5 / FPlatformTime::GetSecondsPerCycle64()); // 500ms
+		NextTimestamp = Time + WaitTime;
+
+		bIsAnySessionAvailable = FInsightsManager::Get()->IsAnySessionAvailable(LastSessionHandle);
+		bIsAnyLiveSessionAvailable = FInsightsManager::Get()->IsAnyLiveSessionAvailable(LastLiveSessionHandle);
+	}
+}
+
+////////////////////////////////////////////////////////////////////////////////////////////////////
+
 EVisibility SStartPageWindow::IsSessionOverlayVisible() const
 {
 	if (FInsightsManager::Get()->GetSession().IsValid())
@@ -589,10 +660,7 @@ FReply SStartPageWindow::OnDrop(const FGeometry& MyGeometry, const FDragDropEven
 				const FString DraggedFileExtension = FPaths::GetExtension(Files[0], true);
 				if (DraggedFileExtension == TEXT(".utrace"))
 				{
-					SpawnAndActivateTabs();
-
 					FInsightsManager::Get()->LoadTraceFile(Files[0]);
-
 					return FReply::Handled();
 				}
 			}
@@ -604,53 +672,132 @@ FReply SStartPageWindow::OnDrop(const FGeometry& MyGeometry, const FDragDropEven
 
 ////////////////////////////////////////////////////////////////////////////////////////////////////
 
-void SStartPageWindow::SpawnAndActivateTabs()
+bool SStartPageWindow::Live_IsEnabled() const
 {
-	// Open tabs for profilers.
-	if (FGlobalTabmanager::Get()->CanSpawnTab(FInsightsManagerTabs::TimingProfilerTabId))
-	{
-		FGlobalTabmanager::Get()->InvokeTab(FInsightsManagerTabs::TimingProfilerTabId);
-	}
-	if (FGlobalTabmanager::Get()->CanSpawnTab(FInsightsManagerTabs::IoProfilerTabId))
+	return bIsAnyLiveSessionAvailable;
+}
+
+////////////////////////////////////////////////////////////////////////////////////////////////////
+
+bool SStartPageWindow::Last_IsEnabled() const
+{
+	return bIsAnySessionAvailable;
+}
+
+////////////////////////////////////////////////////////////////////////////////////////////////////
+
+FReply SStartPageWindow::Live_OnClicked()
+{
+	FInsightsManager::Get()->LoadLastLiveSession();
+	return FReply::Handled();
+}
+
+////////////////////////////////////////////////////////////////////////////////////////////////////
+
+FReply SStartPageWindow::Last_OnClicked()
+{
+	FInsightsManager::Get()->LoadLastSession();
+	return FReply::Handled();
+}
+
+////////////////////////////////////////////////////////////////////////////////////////////////////
+
+FReply SStartPageWindow::Load_OnClicked()
+{
+	//const FString ProfilingDirectory(FPaths::ConvertRelativePathToFull(*FPaths::ProfilingDir()));
+	TSharedRef<Trace::ISessionService> SessionService = FInsightsManager::Get()->GetSessionService();
+	const FString ProfilingDirectory(FPaths::ConvertRelativePathToFull(SessionService->GetLocalSessionDirectory()));
+
+	TArray<FString> OutFiles;
+	bool bOpened = false;
+
+	IDesktopPlatform* DesktopPlatform = FDesktopPlatformModule::Get();
+	if (DesktopPlatform != nullptr)
 	{
-		FGlobalTabmanager::Get()->InvokeTab(FInsightsManagerTabs::IoProfilerTabId);
+		FSlateApplication::Get().CloseToolTip();
+
+		bOpened = DesktopPlatform->OpenFileDialog
+		(
+			FSlateApplication::Get().FindBestParentWindowHandleForDialogs(nullptr),
+			LOCTEXT("LoadTrace_FileDesc", "Open trace file...").ToString(),
+			ProfilingDirectory,
+			TEXT(""),
+			LOCTEXT("LoadTrace_FileFilter", "Trace files (*.utrace)|*.utrace|All files (*.*)|*.*").ToString(),
+			EFileDialogFlags::None,
+			OutFiles
+		);
 	}
 
-	// Ensure Timing Insights / Timing View is the active tab / view.
-	if (TSharedPtr<SDockTab> TimingInsightsTab = FGlobalTabmanager::Get()->FindExistingLiveTab(FInsightsManagerTabs::TimingProfilerTabId))
+	if (bOpened == true)
 	{
-		TimingInsightsTab->ActivateInParent(ETabActivationCause::SetDirectly);
-
-		TSharedPtr<class STimingProfilerWindow> Wnd = FTimingProfilerManager::Get()->GetProfilerWindow();
-		if (Wnd)
+		if (OutFiles.Num() == 1)
 		{
-			TSharedPtr<FTabManager> TabManager = Wnd->GetTabManager();
-
-			if (TSharedPtr<SDockTab> TimingViewTab = TabManager->FindExistingLiveTab(FTimingProfilerTabs::TimingViewID))
-			{
-				TimingViewTab->ActivateInParent(ETabActivationCause::SetDirectly);
-				FSlateApplication::Get().SetKeyboardFocus(TimingViewTab->GetContent());
-			}
+			FInsightsManager::Get()->LoadTraceFile(OutFiles[0]);
 		}
 	}
+
+	return FReply::Handled();
 }
 
 ////////////////////////////////////////////////////////////////////////////////////////////////////
 
-FReply SStartPageWindow::Live_OnClicked()
+void SStartPageWindow::LoadTraceFile(const TCHAR* TraceFile)
 {
-	SpawnAndActivateTabs();
-	FInsightsManager::Get()->GetCommandList()->ExecuteAction(FInsightsManager::GetCommands().InsightsManager_Live.ToSharedRef());
-	return FReply::Handled();
+	FInsightsManager::Get()->LoadTraceFile(FString(TraceFile));
 }
 
 ////////////////////////////////////////////////////////////////////////////////////////////////////
 
-FReply SStartPageWindow::Load_OnClicked()
+void SStartPageWindow::LoadSession(Trace::FSessionHandle SessionHandle)
 {
-	SpawnAndActivateTabs();
-	FInsightsManager::Get()->GetCommandList()->ExecuteAction(FInsightsManager::GetCommands().InsightsManager_Load.ToSharedRef());
-	return FReply::Handled();
+	FInsightsManager::Get()->LoadSession(SessionHandle);
+}
+
+////////////////////////////////////////////////////////////////////////////////////////////////////
+
+TSharedRef<SWidget> SStartPageWindow::MakeSessionListMenu()
+{
+	FMenuBuilder MenuBuilder(/*bInShouldCloseWindowAfterMenuSelection=*/true, nullptr);
+
+	MenuBuilder.BeginSection("MostRecentlyUsedSessions", LOCTEXT("MostRecentlyUsedSessionsHeading", "Most Recently Used Sessions"));
+	{
+		//TODO: MRU
+	}
+	MenuBuilder.EndSection();
+
+	MenuBuilder.BeginSection("AvailableSessions", LOCTEXT("AvailableSessionsHeading", "Available Sessions"));
+	{
+		TSharedRef<Trace::ISessionService> SessionService = FInsightsManager::Get()->GetSessionService();
+		TArray<Trace::FSessionHandle> AvailableSessions;
+		SessionService->GetAvailableSessions(AvailableSessions);
+
+		// Iterate in reverse order as we want most recent sessions first.
+		for (int32 SessionIndex = AvailableSessions.Num() - 1; SessionIndex >= 0; --SessionIndex)
+		{
+			Trace::FSessionHandle SessionHandle = AvailableSessions[SessionIndex];
+
+			Trace::FSessionInfo SessionInfo;
+			SessionService->GetSessionInfo(SessionHandle, SessionInfo);
+
+			FText Label = FText::FromString(SessionInfo.Name);
+			if (SessionInfo.bIsLive)
+			{
+				Label = FText::Format(LOCTEXT("LiveSessionTextFmt", "{0} (LIVE!)"), Label);
+			}
+
+			MenuBuilder.AddMenuEntry(
+				Label,
+				TAttribute<FText>(), // no tooltip
+				FSlateIcon(),
+				FUIAction(FExecuteAction::CreateSP(this, &SStartPageWindow::LoadSession, SessionHandle)),
+				NAME_None,
+				EUserInterfaceActionType::Button
+			);
+		}
+	}
+	MenuBuilder.EndSection();
+
+	return MenuBuilder.MakeWidget();
 }
 
 ////////////////////////////////////////////////////////////////////////////////////////////////////
diff --git a/Engine/Source/Developer/TraceInsights/Private/Insights/Widgets/SStartPageWindow.h b/Engine/Source/Developer/TraceInsights/Private/Insights/Widgets/SStartPageWindow.h
--- a/Engine/Source/Developer/TraceInsights/Private/Insights/Widgets/SStartPageWindow.h
+++ b/Engine/Source/Developer/TraceInsights/Private/Insights/Widgets/SStartPageWindow.h
@@ -67,11 +67,18 @@ private:
 	void SendingServiceSideCapture_Cancel(const FString Filename);
 	void SendingServiceSideCapture_Load(const FString Filename);
 
-	void SpawnAndActivateTabs();
+	bool Live_IsEnabled() const;
+	bool Last_IsEnabled() const;
 
 	FReply Live_OnClicked();
+	FReply Last_OnClicked();
 	FReply Load_OnClicked();
 
+	void LoadTraceFile(const TCHAR* TraceFile);
+	void LoadSession(Trace::FSessionHandle SessionHandle);
+
+	TSharedRef<SWidget> MakeSessionListMenu();
+
 	FText GetLocalSessionDirectory() const;
 	FReply ExploreLocalSessionDirectory_OnClicked();
 
@@ -91,7 +98,7 @@ private:
 	 * @param  InCurrentTime  Current absolute real time
 	 * @param  InDeltaTime  Real time passed since last tick
 	 */
-	//virtual void Tick(const FGeometry& AllottedGeometry, const double InCurrentTime, const float InDeltaTime) override;
+	virtual void Tick(const FGeometry& AllottedGeometry, const double InCurrentTime, const float InDeltaTime) override;
 
 	/**
 	 * The system will use this event to notify a widget that the cursor has entered it. This event is NOT bubbled.
@@ -157,4 +164,10 @@ private:
 
 	/** Holds all widgets for the profiler window like menu bar, toolbar and tabs. */
 	TSharedPtr<SVerticalBox> MainContentPanel;
+
+	bool bIsAnyLiveSessionAvailable;
+	Trace::FSessionHandle LastLiveSessionHandle;
+
+	bool bIsAnySessionAvailable;
+	Trace::FSessionHandle LastSessionHandle;
 };
diff --git a/Engine/Source/Developer/TraceInsights/Private/Insights/Widgets/STimingProfilerToolbar.cpp b/Engine/Source/Developer/TraceInsights/Private/Insights/Widgets/STimingProfilerToolbar.cpp
--- a/Engine/Source/Developer/TraceInsights/Private/Insights/Widgets/STimingProfilerToolbar.cpp
+++ b/Engine/Source/Developer/TraceInsights/Private/Insights/Widgets/STimingProfilerToolbar.cpp
@@ -43,12 +43,12 @@ void STimingProfilerToolbar::Construct(const FArguments& InArgs)
 	{
 		static void FillToolbar1(FToolBarBuilder& ToolbarBuilder)
 		{
-			ToolbarBuilder.BeginSection("Session");
-			{
-				ToolbarBuilder.AddToolBarButton(FInsightsCommands::Get().InsightsManager_Live);
-				ToolbarBuilder.AddToolBarButton(FInsightsCommands::Get().InsightsManager_Load);
-			}
-			ToolbarBuilder.EndSection();
+			//ToolbarBuilder.BeginSection("Session");
+			//{
+			//	ToolbarBuilder.AddToolBarButton(FInsightsCommands::Get().InsightsManager_Live);
+			//	ToolbarBuilder.AddToolBarButton(FInsightsCommands::Get().InsightsManager_Load);
+			//}
+			//ToolbarBuilder.EndSection();
 			ToolbarBuilder.BeginSection("View");
 			{
 				ToolbarBuilder.AddToolBarButton(FTimingProfilerCommands::Get().ToggleFramesTrackVisibility);
diff --git a/Engine/Source/Developer/TraceInsights/Private/Insights/Widgets/STimingView.cpp b/Engine/Source/Developer/TraceInsights/Private/Insights/Widgets/STimingView.cpp
--- a/Engine/Source/Developer/TraceInsights/Private/Insights/Widgets/STimingView.cpp
+++ b/Engine/Source/Developer/TraceInsights/Private/Insights/Widgets/STimingView.cpp
@@ -22,6 +22,7 @@
 #include "TraceServices/AnalysisService.h"
 #include "TraceServices/SessionService.h"
 #include "Widgets/Input/SButton.h"
+#include "Widgets/Input/SComboButton.h"
 #include "Widgets/Layout/SScrollBar.h"
 #include "Widgets/Layout/SSpacer.h"
 #include "Widgets/SOverlay.h"
@@ -404,7 +405,9 @@ void STimingView::Tick(const FGeometry& AllottedGeometry, const double InCurrent
 
 		// Sync CachedTimelines and TimingEventsTracks with available timelines on analysis side.
 		// TODO: can we make this more efficient?
+
 		const Trace::ITimingProfilerProvider& TimingProfilerProvider = Session->ReadTimingProfilerProvider();
+
 		// Check if we have a GPU track.
 		uint32 GpuTimelineIndex;
 		if (TimingProfilerProvider.GetGpuTimelineIndex(GpuTimelineIndex))
@@ -417,8 +420,8 @@ void STimingView::Tick(const FGeometry& AllottedGeometry, const double InCurrent
 		}
 
 		// Check available CPU tracks.
-		const Trace::IThreadProvider& ThreadProvider = Session->ReadThreadProvider();
 		int32 Order = 1;
+		const Trace::IThreadProvider& ThreadProvider = Session->ReadThreadProvider();
 		ThreadProvider.EnumerateThreads([this, &Order, &bIsTimingEventsTrackDirty, &TimingProfilerProvider](const Trace::FThreadInfo& ThreadInfo)
 		{
 			uint32 CpuTimelineIndex;
@@ -442,7 +445,7 @@ void STimingView::Tick(const FGeometry& AllottedGeometry, const double InCurrent
 				Order++;
 			}
 		});
-		
+
 		if (IoActivityTrack == nullptr)
 		{
 			// Note: This track is hardcoded for now, as the I/O timelines are just prototypes for now (will be removed once the functionality is moved in analyzer).
@@ -1022,6 +1025,7 @@ void STimingView::DrawCpuGpuTimelineTrack(FTimingViewDrawHelper& Helper, FTiming
 			Trace::FAnalysisSessionReadScope SessionReadScope(*Session.Get());
 
 			const Trace::ITimingProfilerProvider& TimingProfilerProvider = Session->ReadTimingProfilerProvider();
+
 			TimingProfilerProvider.ReadTimers([this, &Helper, &Track, &TimingProfilerProvider](const Trace::FTimingProfilerTimer* Timers, uint64 TimersCount)
 			{
 				TimingProfilerProvider.ReadTimeline(Track.GetId(), [this, &Helper, &Track, Timers](const Trace::ITimingProfilerProvider::Timeline& Timeline)
@@ -1043,7 +1047,7 @@ void STimingView::DrawCpuGpuTimelineTrack(FTimingViewDrawHelper& Helper, FTiming
 					}
 				});
 			});
-		
+
 			Helper.EndTimeline(Track);
 		}
 	}
@@ -1167,6 +1171,7 @@ void STimingView::UpdateIo()
 			Trace::FAnalysisSessionReadScope SessionReadScope(*Session.Get());
 
 			const Trace::IFileActivityProvider& FileActivityProvider = Session->ReadFileActivityProvider();
+
 			// Enumerate all IO events and cache them.
 			FileActivityProvider.EnumerateFileActivity([this](const Trace::FFileInfo& FileInfo, const Trace::IFileActivityProvider::Timeline& Timeline)
 			{
@@ -2510,8 +2515,8 @@ bool STimingView::SearchTimingEvent(const double InStartTime,
 	{
 		Trace::FAnalysisSessionReadScope SessionReadScope(*Session.Get());
 
-		//if (Ctx.TimingEvent.Track->GetType() == TrackType::Timing)
 		const Trace::ITimingProfilerProvider& TimingProfilerProvider = Session->ReadTimingProfilerProvider();
+
 		TimingProfilerProvider.ReadTimeline(Ctx.TimingEvent.Track->GetId(), [&Ctx](const Trace::ITimingProfilerProvider::Timeline& Timeline)
 		{
 			Timeline.EnumerateEvents(Ctx.StartTime, Ctx.EndTime, [&Ctx](double EventStartTime, double EventEndTime, uint32 EventDepth, const Trace::FTimingProfilerEvent& Event)
@@ -2533,7 +2538,7 @@ bool STimingView::SearchTimingEvent(const double InStartTime,
 				}
 			});
 		});
-	
+
 		if (Ctx.bFound)
 		{
 			TimingProfilerProvider.ReadTimeline(Ctx.TimingEvent.Track->GetId(), [&Ctx](const Trace::ITimingProfilerProvider::Timeline& Timeline)
@@ -2826,7 +2831,7 @@ TSharedRef<SWidget> STimingView::MakeTracksFilterMenu()
 {
 	FMenuBuilder MenuBuilder(/*bInShouldCloseWindowAfterMenuSelection=*/true, nullptr);
 
-	MenuBuilder.BeginSection("TracksFilter", LOCTEXT("TracksFilterHeading", "Quick Filter"));
+	MenuBuilder.BeginSection("QuickFilter", LOCTEXT("TracksFilterHeading", "Quick Filter"));
 	{
 		MenuBuilder.AddMenuEntry(
 			LOCTEXT("ShowAllGpuTracks", "Show/Hide All GPU Tracks (U)"),
diff --git a/Engine/Source/Developer/TraceInsights/Private/Insights/Widgets/STimingView.h b/Engine/Source/Developer/TraceInsights/Private/Insights/Widgets/STimingView.h
--- a/Engine/Source/Developer/TraceInsights/Private/Insights/Widgets/STimingView.h
+++ b/Engine/Source/Developer/TraceInsights/Private/Insights/Widgets/STimingView.h
@@ -338,7 +338,7 @@ protected:
 
 	////////////////////////////////////////////////////////////
 
-	//enum FTimingTrackType
+	//TODO: enum class FTimingTrackType
 	//{
 	//	TimeRuler,
 	//	Markers,
