[Insights] LogView: Fixed selection to not be lost after applying (or resetting) a filter; also now will automatically scroll to bring selected log into view.

diff --git a/Engine/Source/Developer/TraceInsights/Private/Insights/Widgets/SLogView.cpp b/Engine/Source/Developer/TraceInsights/Private/Insights/Widgets/SLogView.cpp
--- a/Engine/Source/Developer/TraceInsights/Private/Insights/Widgets/SLogView.cpp
+++ b/Engine/Source/Developer/TraceInsights/Private/Insights/Widgets/SLogView.cpp
@@ -2,6 +2,7 @@
 
 #include "SLogView.h"
 
+#include "Algo/BinarySearch.h"
 #include "Async/AsyncWork.h"
 #include "EditorStyleSet.h"
 #include "Framework/Application/SlateApplication.h"
@@ -805,7 +806,12 @@ void SLogView::Tick(const FGeometry& AllottedGeometry, const double InCurrentTim
 
 				const int32 NumAddedMessages = NewMessageCount - TotalNumMessages;
 				TotalNumMessages = NewMessageCount;
+				TSharedPtr<FLogMessage> SelectedLogMessage = GetSelectedLogMessage();
 				ListView->RebuildList();
+				if (SelectedLogMessage.IsValid())
+				{
+					SelectedLogMessageByLogIndex(SelectedLogMessage->GetIndex());
+				}
 				bIsDirty = false;
 				DirtyStopwatch.Reset();
 				UpdateStatsText();
@@ -861,7 +867,12 @@ void SLogView::Tick(const FGeometry& AllottedGeometry, const double InCurrentTim
 			}
 
 			TotalNumMessages = Task.GetEndIndex();
+			TSharedPtr<FLogMessage> SelectedLogMessage = GetSelectedLogMessage();
 			ListView->RebuildList();
+			if (SelectedLogMessage.IsValid())
+			{
+				SelectedLogMessageByLogIndex(SelectedLogMessage->GetIndex());
+			}
 			bIsDirty = false;
 			DirtyStopwatch.Reset();
 			UpdateStatsText();
@@ -893,6 +904,19 @@ TSharedPtr<FLogMessage> SLogView::GetSelectedLogMessage() const
 
 ////////////////////////////////////////////////////////////////////////////////////////////////////
 
+void SLogView::SelectedLogMessageByLogIndex(int32 LogIndex)
+{
+	// We are assuming the Messages list is sorted by log index...
+	int32 MessageIndex = Algo::BinarySearchBy(Messages, LogIndex, &FLogMessage::GetIndex);
+	if (MessageIndex != INDEX_NONE)
+	{
+		ListView->SetItemSelection(Messages[MessageIndex], true);
+		ListView->RequestScrollIntoView(Messages[MessageIndex]);
+	}
+}
+
+////////////////////////////////////////////////////////////////////////////////////////////////////
+
 void SLogView::OnSelectionChanged(TSharedPtr<FLogMessage> LogMessage, ESelectInfo::Type SelectInfo)
 {
 	TSharedPtr<STimingProfilerWindow> Window = FTimingProfilerManager::Get()->GetProfilerWindow();
diff --git a/Engine/Source/Developer/TraceInsights/Private/Insights/Widgets/SLogView.h b/Engine/Source/Developer/TraceInsights/Private/Insights/Widgets/SLogView.h
--- a/Engine/Source/Developer/TraceInsights/Private/Insights/Widgets/SLogView.h
+++ b/Engine/Source/Developer/TraceInsights/Private/Insights/Widgets/SLogView.h
@@ -66,6 +66,7 @@ public:
 	FLogMessageCache& GetCache() { return Cache; }
 
 	TSharedPtr<FLogMessage> GetSelectedLogMessage() const;
+	void SelectedLogMessageByLogIndex(int32 LogIndex);
 
 	FText GetFilterText() const { return FilterTextBox->GetText(); }
 
