Fixed deprecation warnings

diff --git a/Engine/Source/Developer/TraceInsights/Private/Insights/InsightsCommands.cpp b/Engine/Source/Developer/TraceInsights/Private/Insights/InsightsCommands.cpp
--- a/Engine/Source/Developer/TraceInsights/Private/Insights/InsightsCommands.cpp
+++ b/Engine/Source/Developer/TraceInsights/Private/Insights/InsightsCommands.cpp
@@ -52,7 +52,7 @@ void FInsightsActionManager::Map_ToggleDebugInfo_Global()
 	This->CommandList->MapAction(This->GetCommands().ToggleDebugInfo, ToggleDebugInfo_Custom());
 }
 
-const FUIAction FInsightsActionManager::ToggleDebugInfo_Custom() const
+const FUIAction FInsightsActionManager::ToggleDebugInfo_Custom()
 {
 	FUIAction UIAction;
 	UIAction.ExecuteAction = FExecuteAction::CreateRaw(this, &FInsightsActionManager::ToggleDebugInfo_Execute);
@@ -178,7 +178,7 @@ void FInsightsActionManager::Map_OpenSettings_Global()
 
 ////////////////////////////////////////////////////////////////////////////////////////////////////
 
-const FUIAction FInsightsActionManager::OpenSettings_Custom() const
+const FUIAction FInsightsActionManager::OpenSettings_Custom()
 {
 	FUIAction UIAction;
 	UIAction.ExecuteAction = FExecuteAction::CreateRaw(this, &FInsightsActionManager::OpenSettings_Execute);
diff --git a/Engine/Source/Developer/TraceInsights/Private/Insights/InsightsCommands.h b/Engine/Source/Developer/TraceInsights/Private/Insights/InsightsCommands.h
--- a/Engine/Source/Developer/TraceInsights/Private/Insights/InsightsCommands.h
+++ b/Engine/Source/Developer/TraceInsights/Private/Insights/InsightsCommands.h
@@ -78,7 +78,7 @@ protected:
 
 public:
 	void Map_ToggleDebugInfo_Global(); /**< Maps UI command info ToggleDebugInfo with the specified UI command list. */
-	const FUIAction ToggleDebugInfo_Custom() const; /**< UI action for ToggleDebugInfo command. */
+	const FUIAction ToggleDebugInfo_Custom(); /**< UI action for ToggleDebugInfo command. */
 protected:
 	void ToggleDebugInfo_Execute(); /**< Handles FExecuteAction for ToggleDebugInfo. */
 	bool ToggleDebugInfo_CanExecute() const; /**< Handles FCanExecuteAction for ToggleDebugInfo. */
@@ -89,7 +89,7 @@ protected:
 
 public:
 	void Map_OpenSettings_Global(); /**< Maps UI command info OpenSettings with the specified UI command list. */
-	const FUIAction OpenSettings_Custom() const; /**< UI action for OpenSettings command. */
+	const FUIAction OpenSettings_Custom(); /**< UI action for OpenSettings command. */
 protected:
 	void OpenSettings_Execute(); /**< Handles FExecuteAction for OpenSettings. */
 	bool OpenSettings_CanExecute() const; /**< Handles FCanExecuteAction for OpenSettings. */
diff --git a/Engine/Source/Developer/TraceInsights/Private/Insights/IoProfilerCommands.cpp b/Engine/Source/Developer/TraceInsights/Private/Insights/IoProfilerCommands.cpp
--- a/Engine/Source/Developer/TraceInsights/Private/Insights/IoProfilerCommands.cpp
+++ b/Engine/Source/Developer/TraceInsights/Private/Insights/IoProfilerCommands.cpp
@@ -67,7 +67,7 @@ PRAGMA_ENABLE_OPTIMIZATION
 		This->CommandList->MapAction(This->GetCommands().CmdName, CmdName##_Custom());\
 	}\
 	\
-	const FUIAction FIoProfilerActionManager::CmdName##_Custom() const\
+	const FUIAction FIoProfilerActionManager::CmdName##_Custom() \
 	{\
 		FUIAction UIAction;\
 		UIAction.ExecuteAction = FExecuteAction::CreateRaw(this, &FIoProfilerActionManager::CmdName##_Execute);\
diff --git a/Engine/Source/Developer/TraceInsights/Private/Insights/IoProfilerCommands.h b/Engine/Source/Developer/TraceInsights/Private/Insights/IoProfilerCommands.h
--- a/Engine/Source/Developer/TraceInsights/Private/Insights/IoProfilerCommands.h
+++ b/Engine/Source/Developer/TraceInsights/Private/Insights/IoProfilerCommands.h
@@ -76,7 +76,7 @@ private:
 #define DECLARE_TOGGLE_COMMAND(CmdName)\
 public:\
 	void Map_##CmdName##_Global(); /**< Maps UI command info CmdName with the specified UI command list. */\
-	const FUIAction CmdName##_Custom() const; /**< UI action for CmdName command. */\
+	const FUIAction CmdName##_Custom(); /**< UI action for CmdName command. */\
 protected:\
 	void CmdName##_Execute(); /**< Handles FExecuteAction for CmdName. */\
 	bool CmdName##_CanExecute() const; /**< Handles FCanExecuteAction for CmdName. */\
diff --git a/Engine/Source/Developer/TraceInsights/Private/Insights/TimingProfilerCommands.cpp b/Engine/Source/Developer/TraceInsights/Private/Insights/TimingProfilerCommands.cpp
--- a/Engine/Source/Developer/TraceInsights/Private/Insights/TimingProfilerCommands.cpp
+++ b/Engine/Source/Developer/TraceInsights/Private/Insights/TimingProfilerCommands.cpp
@@ -72,7 +72,7 @@ PRAGMA_ENABLE_OPTIMIZATION
 		This->CommandList->MapAction(This->GetCommands().CmdName, CmdName##_Custom());\
 	}\
 	\
-	const FUIAction FTimingProfilerActionManager::CmdName##_Custom() const\
+	const FUIAction FTimingProfilerActionManager::CmdName##_Custom() \
 	{\
 		FUIAction UIAction;\
 		UIAction.ExecuteAction = FExecuteAction::CreateRaw(this, &FTimingProfilerActionManager::CmdName##_Execute);\
diff --git a/Engine/Source/Developer/TraceInsights/Private/Insights/TimingProfilerCommands.h b/Engine/Source/Developer/TraceInsights/Private/Insights/TimingProfilerCommands.h
--- a/Engine/Source/Developer/TraceInsights/Private/Insights/TimingProfilerCommands.h
+++ b/Engine/Source/Developer/TraceInsights/Private/Insights/TimingProfilerCommands.h
@@ -91,7 +91,7 @@ private:
 #define DECLARE_TOGGLE_COMMAND(CmdName)\
 public:\
 	void Map_##CmdName##_Global(); /**< Maps UI command info CmdName with the specified UI command list. */\
-	const FUIAction CmdName##_Custom() const; /**< UI action for CmdName command. */\
+	const FUIAction CmdName##_Custom(); /**< UI action for CmdName command. */\
 protected:\
 	void CmdName##_Execute(); /**< Handles FExecuteAction for CmdName. */\
 	bool CmdName##_CanExecute() const; /**< Handles FCanExecuteAction for CmdName. */\
diff --git a/Engine/Source/Developer/TraceInsights/Private/Insights/Widgets/SStatsView.cpp b/Engine/Source/Developer/TraceInsights/Private/Insights/Widgets/SStatsView.cpp
--- a/Engine/Source/Developer/TraceInsights/Private/Insights/Widgets/SStatsView.cpp
+++ b/Engine/Source/Developer/TraceInsights/Private/Insights/Widgets/SStatsView.cpp
@@ -213,7 +213,7 @@ END_SLATE_FUNCTION_BUILD_OPTIMIZATION
 
 ////////////////////////////////////////////////////////////////////////////////////////////////////
 
-TSharedPtr<SWidget> SStatsView::TreeView_GetMenuContent() const
+TSharedPtr<SWidget> SStatsView::TreeView_GetMenuContent()
 {
 	const TArray<FStatsNodePtr> SelectedStatsNodes = TreeView->GetSelectedItems();
 	const int32 NumSelectedStatsNodes = SelectedStatsNodes.Num();
diff --git a/Engine/Source/Developer/TraceInsights/Private/Insights/Widgets/SStatsView.h b/Engine/Source/Developer/TraceInsights/Private/Insights/Widgets/SStatsView.h
--- a/Engine/Source/Developer/TraceInsights/Private/Insights/Widgets/SStatsView.h
+++ b/Engine/Source/Developer/TraceInsights/Private/Insights/Widgets/SStatsView.h
@@ -112,7 +112,7 @@ protected:
 	////////////////////////////////////////////////////////////////////////////////////////////////////
 	// Tree View - Context Menu
 
-	TSharedPtr<SWidget> TreeView_GetMenuContent() const;
+	TSharedPtr<SWidget> TreeView_GetMenuContent();
 	void TreeView_BuildSortByMenu(FMenuBuilder& MenuBuilder);
 	void TreeView_BuildViewColumnMenu(FMenuBuilder& MenuBuilder);
 
diff --git a/Engine/Source/Developer/TraceInsights/Private/Insights/Widgets/STimersView.cpp b/Engine/Source/Developer/TraceInsights/Private/Insights/Widgets/STimersView.cpp
--- a/Engine/Source/Developer/TraceInsights/Private/Insights/Widgets/STimersView.cpp
+++ b/Engine/Source/Developer/TraceInsights/Private/Insights/Widgets/STimersView.cpp
@@ -220,7 +220,7 @@ END_SLATE_FUNCTION_BUILD_OPTIMIZATION
 
 ////////////////////////////////////////////////////////////////////////////////////////////////////
 
-TSharedPtr<SWidget> STimersView::TreeView_GetMenuContent() const
+TSharedPtr<SWidget> STimersView::TreeView_GetMenuContent()
 {
 	const TArray<FTimerNodePtr> SelectedTimerNodes = TreeView->GetSelectedItems();
 	const int32 NumSelectedTimerNodes = SelectedTimerNodes.Num();
diff --git a/Engine/Source/Developer/TraceInsights/Private/Insights/Widgets/STimersView.h b/Engine/Source/Developer/TraceInsights/Private/Insights/Widgets/STimersView.h
--- a/Engine/Source/Developer/TraceInsights/Private/Insights/Widgets/STimersView.h
+++ b/Engine/Source/Developer/TraceInsights/Private/Insights/Widgets/STimersView.h
@@ -92,7 +92,7 @@ protected:
 	////////////////////////////////////////////////////////////////////////////////////////////////////
 	// Tree View - Context Menu
 
-	TSharedPtr<SWidget> TreeView_GetMenuContent() const;
+	TSharedPtr<SWidget> TreeView_GetMenuContent();
 	void TreeView_BuildSortByMenu(FMenuBuilder& MenuBuilder);
 	void TreeView_BuildViewColumnMenu(FMenuBuilder& MenuBuilder);
 
