[Insights] Start Page: Added control for trace modules.

diff --git a/Engine/Source/Developer/TraceInsights/Private/Insights/InsightsManager.cpp b/Engine/Source/Developer/TraceInsights/Private/Insights/InsightsManager.cpp
--- a/Engine/Source/Developer/TraceInsights/Private/Insights/InsightsManager.cpp
+++ b/Engine/Source/Developer/TraceInsights/Private/Insights/InsightsManager.cpp
@@ -24,9 +24,12 @@ TSharedPtr<FInsightsManager> FInsightsManager::Instance = nullptr;
 
 ////////////////////////////////////////////////////////////////////////////////////////////////////
 
-FInsightsManager::FInsightsManager(TSharedRef<Trace::IAnalysisService> InTraceAnalysisService, TSharedRef<Trace::ISessionService> InSessionService)
+FInsightsManager::FInsightsManager(TSharedRef<Trace::IAnalysisService> InTraceAnalysisService,
+								   TSharedRef<Trace::ISessionService> InTraceSessionService,
+								   TSharedRef<Trace::IModuleService> InTraceModuleService)
 	: AnalysisService(InTraceAnalysisService)
-	, SessionService(InSessionService)
+	, SessionService(InTraceSessionService)
+	, ModuleService(InTraceModuleService)
 	, CommandList(new FUICommandList())
 	, ActionManager(this)
 	, Settings()
@@ -267,8 +270,8 @@ void FInsightsManager::LoadSession(Trace::FSessionHandle SessionHandle)
 	{
 		Session = AnalysisService->StartAnalysis(SessionInfo.Name, MoveTemp(DataStream));
 		SpawnAndActivateTabs();
+		OnSessionChanged();
 	}
-	OnSessionChanged();
 }
 
 ////////////////////////////////////////////////////////////////////////////////////////////////////
@@ -282,8 +285,8 @@ void FInsightsManager::LoadTraceFile(const FString& TraceFilepath)
 	{
 		Session = AnalysisService->StartAnalysis(*TraceFilepath, MoveTemp(DataStream));
 		SpawnAndActivateTabs();
+		OnSessionChanged();
 	}
-	OnSessionChanged();
 }
 
 ////////////////////////////////////////////////////////////////////////////////////////////////////
diff --git a/Engine/Source/Developer/TraceInsights/Private/Insights/InsightsManager.h b/Engine/Source/Developer/TraceInsights/Private/Insights/InsightsManager.h
--- a/Engine/Source/Developer/TraceInsights/Private/Insights/InsightsManager.h
+++ b/Engine/Source/Developer/TraceInsights/Private/Insights/InsightsManager.h
@@ -6,6 +6,7 @@
 #include "CoreMinimal.h"
 #include "Framework/Commands/UICommandList.h"
 #include "TraceServices/AnalysisService.h"
+#include "TraceServices/ModuleService.h"
 #include "TraceServices/SessionService.h"
 
 // Insights
@@ -64,7 +65,9 @@ class FInsightsManager
 
 public:
 	/** Creates the main manager, only one instance can exist. */
-	FInsightsManager(TSharedRef<Trace::IAnalysisService> TraceAnalysisService, TSharedRef<Trace::ISessionService> SessionService);
+	FInsightsManager(TSharedRef<Trace::IAnalysisService> TraceAnalysisService,
+					 TSharedRef<Trace::ISessionService> SessionService,
+					 TSharedRef<Trace::IModuleService> TraceModuleService);
 
 	/** Virtual destructor. */
 	virtual ~FInsightsManager();
@@ -73,15 +76,18 @@ public:
 	 * Creates an instance of the main manager and initializes global instance with the previously created instance of the manager.
 	 * @param TraceAnalysisService The trace analysis service
 	 * @param TraceSessionService  The trace session service
+	 * @param TraceModuleService   The trace module service
 	 */
-	static TSharedPtr<FInsightsManager> Initialize(TSharedRef<Trace::IAnalysisService> TraceAnalysisService, TSharedRef<Trace::ISessionService> TraceSessionService)
+	static TSharedPtr<FInsightsManager> Initialize(TSharedRef<Trace::IAnalysisService> TraceAnalysisService,
+												   TSharedRef<Trace::ISessionService> TraceSessionService,
+												   TSharedRef<Trace::IModuleService> TraceModuleService)
 	{
 		if (FInsightsManager::Instance.IsValid())
 		{
 			FInsightsManager::Instance.Reset();
 		}
 
-		FInsightsManager::Instance = MakeShareable(new FInsightsManager(TraceAnalysisService, TraceSessionService));
+		FInsightsManager::Instance = MakeShareable(new FInsightsManager(TraceAnalysisService, TraceSessionService, TraceModuleService));
 		FInsightsManager::Instance->PostConstructor();
 
 		return FInsightsManager::Instance;
@@ -98,6 +104,7 @@ public:
 
 	TSharedRef<Trace::IAnalysisService> GetAnalysisService() const { return AnalysisService; }
 	TSharedRef<Trace::ISessionService> GetSessionService() const { return SessionService; }
+	TSharedRef<Trace::IModuleService> GetModuleService() const { return ModuleService; }
 
 	/** @return an instance of the trace analysis session. */
 	TSharedPtr<const Trace::IAnalysisSession> GetSession() const;
@@ -200,6 +207,7 @@ private:
 
 	TSharedRef<Trace::IAnalysisService> AnalysisService;
 	TSharedRef<Trace::ISessionService> SessionService;
+	TSharedRef<Trace::IModuleService> ModuleService;
 
 	/** The trace analysis session. */
 	TSharedPtr<const Trace::IAnalysisSession> Session;
diff --git a/Engine/Source/Developer/TraceInsights/Private/Insights/TraceInsightsModule.cpp b/Engine/Source/Developer/TraceInsights/Private/Insights/TraceInsightsModule.cpp
--- a/Engine/Source/Developer/TraceInsights/Private/Insights/TraceInsightsModule.cpp
+++ b/Engine/Source/Developer/TraceInsights/Private/Insights/TraceInsightsModule.cpp
@@ -63,6 +63,7 @@ protected:
 protected:
 	TSharedPtr<Trace::IAnalysisService> TraceAnalysisService;
 	TSharedPtr<Trace::ISessionService> TraceSessionService;
+	TSharedPtr<Trace::IModuleService> TraceModuleService;
 };
 
 ////////////////////////////////////////////////////////////////////////////////////////////////////
@@ -78,11 +79,14 @@ void FTraceInsightsModule::StartupModule()
 	ITraceServicesModule& TraceServicesModule = FModuleManager::LoadModuleChecked<ITraceServicesModule>("TraceServices");
 	TraceAnalysisService = TraceServicesModule.GetAnalysisService();
 	TraceSessionService = TraceServicesModule.GetSessionService();
+	TraceModuleService = TraceServicesModule.GetModuleService();
+
+	// Starts the Trace Recorder service.
 	TraceSessionService->StartRecorderServer();
 
 	FInsightsStyle::Initialize();
 
-	FInsightsManager::Initialize(TraceAnalysisService.ToSharedRef(), TraceSessionService.ToSharedRef());
+	FInsightsManager::Initialize(TraceAnalysisService.ToSharedRef(), TraceSessionService.ToSharedRef(), TraceModuleService.ToSharedRef());
 	FTimingProfilerManager::Initialize();
 	FIoProfilerManager::Initialize();
 
diff --git a/Engine/Source/Developer/TraceInsights/Private/Insights/Widgets/SStartPageWindow.cpp b/Engine/Source/Developer/TraceInsights/Private/Insights/Widgets/SStartPageWindow.cpp
--- a/Engine/Source/Developer/TraceInsights/Private/Insights/Widgets/SStartPageWindow.cpp
+++ b/Engine/Source/Developer/TraceInsights/Private/Insights/Widgets/SStartPageWindow.cpp
@@ -346,63 +346,23 @@ void SStartPageWindow::Construct(const FArguments& InArgs)
 						.HAlign(HAlign_Center)
 						.Padding(3.0f, 3.0f)
 						[
-							SNew(SBorder)
-							.BorderImage(FEditorStyle::GetBrush("NotificationList.ItemBackground"))
-							.Padding(8.0f)
-							[
-								SNew(SVerticalBox)
-
-								+ SVerticalBox::Slot()
-									.AutoHeight()
-									.HAlign(HAlign_Center)
-									.Padding(0.0, 2.0f)
-									[
-										SNew(STextBlock)
-										.Text(LOCTEXT("RecorderStatusTitle", "Trace Recorder Status:"))
-										.ColorAndOpacity(FLinearColor::Gray)
-									]
-
-								+ SVerticalBox::Slot()
-									.AutoHeight()
-									.HAlign(HAlign_Center)
-									.Padding(0.0, 2.0f)
-									[
-										SNew(SHorizontalBox)
-
-										+ SHorizontalBox::Slot()
-											.AutoWidth()
-											.Padding(2.0f, 0.0f)
-											.VAlign(VAlign_Center)
-											[
-												SNew(STextBlock)
-												.Text(this, &SStartPageWindow::GetRecorderStatusText)
-											]
+							SNew(SHorizontalBox)
 
-										+ SHorizontalBox::Slot()
-											.AutoWidth()
-											.Padding(2.0f, 0.0f)
-											.VAlign(VAlign_Center)
-											[
-												SNew(SButton)
-												.Text(LOCTEXT("StartRecorder", "Start"))
-												.ToolTipText(LOCTEXT("StartRecorderToolTip", "Start the Trace Recorder"))
-												.OnClicked(this, &SStartPageWindow::StartTraceRecorder_OnClicked)
-												.Visibility(this, &SStartPageWindow::StartTraceRecorder_Visibility)
-											]
+							+ SHorizontalBox::Slot()
+								.AutoWidth()
+								.VAlign(VAlign_Top)
+								.Padding(3.0f, 0.0f)
+								[
+									ConstructRecorderStatus()
+								]
 
-										+ SHorizontalBox::Slot()
-											.AutoWidth()
-											.Padding(2.0f, 0.0f)
-											.VAlign(VAlign_Center)
-											[
-												SNew(SButton)
-												.Text(LOCTEXT("StopRecorder", "Stop"))
-												.ToolTipText(LOCTEXT("StopRecorderToolTip", "Stop the Trace Recorder"))
-												.OnClicked(this, &SStartPageWindow::StopTraceRecorder_OnClicked)
-												.Visibility(this, &SStartPageWindow::StopTraceRecorder_Visibility)
-											]
-									]
-							]
+							+ SHorizontalBox::Slot()
+								.AutoWidth()
+								.VAlign(VAlign_Top)
+								.Padding(3.0f, 0.0f)
+								[
+									ConstructModuleList()
+								]
 						]
 				]
 
@@ -423,10 +383,175 @@ void SStartPageWindow::Construct(const FArguments& InArgs)
 		];
 }
 
+////////////////////////////////////////////////////////////////////////////////////////////////////
+
+TSharedRef<SWidget> SStartPageWindow::ConstructRecorderStatus()
+{
+	TSharedRef<SWidget> BorderWidget =
+		SNew(SBorder)
+		.BorderImage(FEditorStyle::GetBrush("NotificationList.ItemBackground"))
+		.Padding(8.0f)
+		[
+			SNew(SVerticalBox)
+
+			+ SVerticalBox::Slot()
+				.AutoHeight()
+				.HAlign(HAlign_Center)
+				.Padding(0.0, 2.0f)
+				[
+					SNew(STextBlock)
+					.Text(LOCTEXT("RecorderStatusTitle", "Trace Recorder Status:"))
+					.ColorAndOpacity(FLinearColor::Gray)
+				]
+
+			+ SVerticalBox::Slot()
+				.AutoHeight()
+				.HAlign(HAlign_Center)
+				.Padding(0.0, 2.0f)
+				[
+					SNew(SHorizontalBox)
+
+					+ SHorizontalBox::Slot()
+						.AutoWidth()
+						.Padding(2.0f, 0.0f)
+						.VAlign(VAlign_Center)
+						[
+							SNew(STextBlock)
+							.Text(this, &SStartPageWindow::GetRecorderStatusText)
+						]
+
+					+ SHorizontalBox::Slot()
+						.AutoWidth()
+						.Padding(2.0f, 0.0f)
+						.VAlign(VAlign_Center)
+						[
+							SNew(SButton)
+							.Text(LOCTEXT("StartRecorder", "Start"))
+							.ToolTipText(LOCTEXT("StartRecorderToolTip", "Start the Trace Recorder"))
+							.OnClicked(this, &SStartPageWindow::StartTraceRecorder_OnClicked)
+							.Visibility(this, &SStartPageWindow::StartTraceRecorder_Visibility)
+						]
+
+					+ SHorizontalBox::Slot()
+						.AutoWidth()
+						.Padding(2.0f, 0.0f)
+						.VAlign(VAlign_Center)
+						[
+							SNew(SButton)
+							.Text(LOCTEXT("StopRecorder", "Stop"))
+							.ToolTipText(LOCTEXT("StopRecorderToolTip", "Stop the Trace Recorder"))
+							.OnClicked(this, &SStartPageWindow::StopTraceRecorder_OnClicked)
+							.Visibility(this, &SStartPageWindow::StopTraceRecorder_Visibility)
+						]
+				]
+		];
+
+	return BorderWidget;
+}
+
+////////////////////////////////////////////////////////////////////////////////////////////////////
+
+TSharedRef<SWidget> SStartPageWindow::ConstructModuleList()
+{
+	TSharedPtr<SVerticalBox> VerticalBox;
+
+	TSharedRef<SWidget> BorderWidget =
+		SNew(SBorder)
+		.BorderImage(FEditorStyle::GetBrush("NotificationList.ItemBackground"))
+		.Padding(8.0f)
+		[
+			SNew(SVerticalBox)
+
+			+ SVerticalBox::Slot()
+				.AutoHeight()
+				.HAlign(HAlign_Center)
+				.Padding(0.0, 2.0f)
+				[
+					SNew(STextBlock)
+					.Text(LOCTEXT("ModulesTitle", "Modules:"))
+					.ColorAndOpacity(FLinearColor::Gray)
+				]
+
+			+ SVerticalBox::Slot()
+				.AutoHeight()
+				.HAlign(HAlign_Center)
+				.Padding(0.0, 2.0f)
+				[
+					SAssignNew(VerticalBox, SVerticalBox)
+				]
+		];
+
+	TSharedRef<Trace::IModuleService> ModuleService = FInsightsManager::Get()->GetModuleService();
+	ModuleService->GetAvailableModules(AvailableModules);
+
+	const bool bDefaultIsModuleEnabledState = true;
+	AvailableModulesEnabledState.Reset();
+
+	for (int32 ModuleIndex = 0; ModuleIndex < AvailableModules.Num(); ++ModuleIndex)
+	{
+		const Trace::FModuleInfo& Module = AvailableModules[ModuleIndex];
+
+		AvailableModulesEnabledState.Add(bDefaultIsModuleEnabledState);
+		ModuleService->SetModuleEnabled(Module.Name, bDefaultIsModuleEnabledState);
+
+		VerticalBox->AddSlot()
+			.AutoHeight()
+			.Padding(0.0f, 1.0f)
+			[
+				SNew(SCheckBox)
+				.IsChecked(this, &SStartPageWindow::Module_IsChecked, ModuleIndex)
+				.OnCheckStateChanged(this, &SStartPageWindow::Module_OnCheckStateChanged, ModuleIndex)
+				[
+					SNew(STextBlock)
+					.Text(FText::FromString(Module.DisplayName))
+				]
+			];
+	}
+
+	return BorderWidget;
+}
+
 END_SLATE_FUNCTION_BUILD_OPTIMIZATION
 
 ////////////////////////////////////////////////////////////////////////////////////////////////////
 
+ECheckBoxState SStartPageWindow::Module_IsChecked(int32 ModuleIndex) const
+{
+	if (ModuleIndex >= 0 && ModuleIndex < AvailableModules.Num())
+	{
+		//const Trace::FModuleInfo& Module = AvailableModules[ModuleIndex];
+		//const bool bIsEnabled = Module.bIsEnabled;
+
+		ensure(AvailableModulesEnabledState.Num() == AvailableModules.Num());
+		const bool bIsEnabled = AvailableModulesEnabledState[ModuleIndex];
+
+		return bIsEnabled ? ECheckBoxState::Checked : ECheckBoxState::Unchecked;
+	}
+	else
+	{
+		return ECheckBoxState::Undetermined;
+	}
+}
+
+////////////////////////////////////////////////////////////////////////////////////////////////////
+
+void SStartPageWindow::Module_OnCheckStateChanged(ECheckBoxState NewRadioState, int32 ModuleIndex)
+{
+	if (ModuleIndex >= 0 && ModuleIndex < AvailableModules.Num())
+	{
+		bool bIsEnabled = NewRadioState == ECheckBoxState::Checked;
+
+		ensure(AvailableModulesEnabledState.Num() == AvailableModules.Num());
+		AvailableModulesEnabledState[ModuleIndex] = bIsEnabled;
+
+		TSharedRef<Trace::IModuleService> ModuleService = FInsightsManager::Get()->GetModuleService();
+		Trace::FModuleInfo& Module = AvailableModules[ModuleIndex];
+		ModuleService->SetModuleEnabled(Module.Name, bIsEnabled);
+	}
+}
+
+////////////////////////////////////////////////////////////////////////////////////////////////////
+
 void SStartPageWindow::Tick(const FGeometry& AllottedGeometry, const double InCurrentTime, const float InDeltaTime)
 {
 	// We need to check if there is any available session, but not too often.
diff --git a/Engine/Source/Developer/TraceInsights/Private/Insights/Widgets/SStartPageWindow.h b/Engine/Source/Developer/TraceInsights/Private/Insights/Widgets/SStartPageWindow.h
--- a/Engine/Source/Developer/TraceInsights/Private/Insights/Widgets/SStartPageWindow.h
+++ b/Engine/Source/Developer/TraceInsights/Private/Insights/Widgets/SStartPageWindow.h
@@ -56,10 +56,11 @@ public:
 	void OpenSettings();
 	void CloseSettings();
 
-	void ShowHideTab(const FName& TabID, bool bIsVisible);
-
 private:
-	/** Callback for determining the visibility of the 'Select a session' overlay. */
+	TSharedRef<SWidget> ConstructRecorderStatus();
+	TSharedRef<SWidget> ConstructModuleList();
+
+	/** Callback for determining the visibility of the 'Please select a trace' overlay. */
 	EVisibility IsSessionOverlayVisible() const;
 
 	bool IsSessionValid() const;
@@ -88,6 +89,9 @@ private:
 	FReply StartTraceRecorder_OnClicked();
 	FReply StopTraceRecorder_OnClicked();
 
+	ECheckBoxState Module_IsChecked(int32 ModuleIndex) const;
+	void Module_OnCheckStateChanged(ECheckBoxState NewRadioState, int32 ModuleIndex);
+
 	/** Updates the amount of time the profiler has been active. */
 	EActiveTimerReturnType UpdateActiveDuration(double InCurrentTime, float InDeltaTime);
 
@@ -170,4 +174,7 @@ private:
 
 	bool bIsAnySessionAvailable;
 	Trace::FSessionHandle LastSessionHandle;
+
+	TArray<Trace::FModuleInfo> AvailableModules;
+	TArray<bool> AvailableModulesEnabledState;
 };
