[Insights]
  - Timing view: Added "Loading - Main Thread" and "Loading - Async Thread" tracks (Load Time profiler).
  - Timing view: Added 5 coloring/naming modes for Loading* tracks: 1) by package type, 2) by export type, 3) by package name, 4)  by export class name, 5) mixed packag ename / export class name
  - Timing view: Added keyboard shortcuts (keys 1 to 5) to easily change between the 5 coloring/naming modes in Loading* tracks.
  - Timing view: Added specific tooltip for timing events hovered with mouse in Loading* tracks (shows Package and Export info).
  - Timing view: Changed "Tracks" filtering to allow visibility switches only for thread groups (not for individual tracks).
  - Timing view: Changed the time ruler's background to be darker for invalid time regions (outside session time range).
  - Markers track: Made vertical lines more transparent (was 0.5, now 0.25).
  - Timing view: Made all scrollbars the same size (~7 px). Also decreased transparency for horizontal scrollbar from Frames track.
  - Timing view: Updated key shortcuts: Y = on/off GPU tracks; U = on/off CPU tracks; L = on/off Asset Loading tracks.
  - Start Page: Disabled all modules by default.

diff --git a/Engine/Source/Developer/TraceInsights/Private/Insights/Common/TimeUtils.cpp b/Engine/Source/Developer/TraceInsights/Private/Insights/Common/TimeUtils.cpp
--- a/Engine/Source/Developer/TraceInsights/Private/Insights/Common/TimeUtils.cpp
+++ b/Engine/Source/Developer/TraceInsights/Private/Insights/Common/TimeUtils.cpp
@@ -764,7 +764,7 @@ PRAGMA_ENABLE_OPTIMIZATION
 PRAGMA_DISABLE_OPTIMIZATION
 void TestOptimizationIssue()
 {
-	static const double Ns = 0.000000001;
+	constexpr double Ns = 0.000000001;
 	int32 D1 = GetNumDigits(Ns);
 	int32 D2 = GetNumDigitsOptDisabled(Ns);
 	//UE_LOG(TimingProfiler, Log, TEXT("D1 = %d, D2 = %d"), D1, D2);
diff --git a/Engine/Source/Developer/TraceInsights/Private/Insights/Common/TimeUtils.h b/Engine/Source/Developer/TraceInsights/Private/Insights/Common/TimeUtils.h
--- a/Engine/Source/Developer/TraceInsights/Private/Insights/Common/TimeUtils.h
+++ b/Engine/Source/Developer/TraceInsights/Private/Insights/Common/TimeUtils.h
@@ -22,15 +22,15 @@ double is a 64 bit IEEE 754 double precision Floating Point Number
 
 namespace TimeUtils
 {
-	static const double Picosecond = 0.000000000001;
-	static const double Nanosecond = 0.000000001;
-	static const double Microsecond = 0.000001;
-	static const double Milisecond = 0.001;
-	static const double Second = 1.0;
-	static const double Minute = 60.0;
-	static const double Hour = 3600.0;
-	static const double Day = 86400.0;
-	static const double Week = 604800.0;
+	static constexpr double Picosecond = 0.000000000001;
+	static constexpr double Nanosecond = 0.000000001;
+	static constexpr double Microsecond = 0.000001;
+	static constexpr double Milisecond = 0.001;
+	static constexpr double Second = 1.0;
+	static constexpr double Minute = 60.0;
+	static constexpr double Hour = 3600.0;
+	static constexpr double Day = 86400.0;
+	static constexpr double Week = 604800.0;
 
 	struct FTimeSplit
 	{
diff --git a/Engine/Source/Developer/TraceInsights/Private/Insights/ViewModels/BaseTimingTrack.h b/Engine/Source/Developer/TraceInsights/Private/Insights/ViewModels/BaseTimingTrack.h
--- a/Engine/Source/Developer/TraceInsights/Private/Insights/ViewModels/BaseTimingTrack.h
+++ b/Engine/Source/Developer/TraceInsights/Private/Insights/ViewModels/BaseTimingTrack.h
@@ -78,6 +78,8 @@ public:
 	virtual void Tick(const double CurrentTime, const float DeltaTime) {}
 	virtual void Update(const FTimingTrackViewport& Viewport) {}
 
+	static uint64 GenerateId() { return IdGenerator++; }
+
 private:
 	uint64 Id;
 
@@ -85,6 +87,8 @@ private:
 	float Height; // height, in Slate units
 
 	ETimingTrackFlags Flags;
+
+	static uint64 IdGenerator;
 };
 
 ////////////////////////////////////////////////////////////////////////////////////////////////////
diff --git a/Engine/Source/Developer/TraceInsights/Private/Insights/ViewModels/FrameTrackHelper.cpp b/Engine/Source/Developer/TraceInsights/Private/Insights/ViewModels/FrameTrackHelper.cpp
--- a/Engine/Source/Developer/TraceInsights/Private/Insights/ViewModels/FrameTrackHelper.cpp
+++ b/Engine/Source/Developer/TraceInsights/Private/Insights/ViewModels/FrameTrackHelper.cpp
@@ -281,10 +281,10 @@ void FFrameTrackDrawHelper::DrawHighlightedInterval(const FFrameTrackTimeline& T
 		float X1 = Index1 * SampleW;
 		float X2 = Index2 * SampleW;
 
-		const float Y1 = 10.0f; // allows 10px for the horizontal scrollbar (one displayed on top of the track)
+		constexpr float Y1 = 12.0f; // allows 12px for the horizontal scrollbar (one displayed on top of the track)
 		const float Y2 = Viewport.Height;
-		const float D = 2.0f; // line thickness (for both horizontal and vertical lines)
-		const float H = 10.0f; // height of corner lines
+		constexpr float D = 2.0f; // line thickness (for both horizontal and vertical lines)
+		constexpr float H = 10.0f; // height of corner lines
 
 		const FLinearColor Color(1.0f, 1.0f, 1.0f, 1.0f);
 
diff --git a/Engine/Source/Developer/TraceInsights/Private/Insights/ViewModels/GraphTrack.cpp b/Engine/Source/Developer/TraceInsights/Private/Insights/ViewModels/GraphTrack.cpp
--- a/Engine/Source/Developer/TraceInsights/Private/Insights/ViewModels/GraphTrack.cpp
+++ b/Engine/Source/Developer/TraceInsights/Private/Insights/ViewModels/GraphTrack.cpp
@@ -74,8 +74,8 @@ FGraphTrack::~FGraphTrack()
 
 void FGraphTrack::UpdateHoveredState(float MouseX, float MouseY, const FTimingTrackViewport& Viewport)
 {
-	static const float HeaderWidth = 100.0f;
-	static const float HeaderHeight = 14.0f;
+	constexpr float HeaderWidth = 100.0f;
+	constexpr float HeaderHeight = 14.0f;
 
 	if (MouseY >= GetPosY() && MouseY < GetPosY() + GetHeight())
 	{
@@ -159,9 +159,9 @@ void FGraphTrack::DrawSeries(const FGraphTrackSeries& Series, FDrawContext& Draw
 			for (int32 Index = 0; Index < NumPoints; ++Index)
 			{
 				const FVector2D& Pt = Series.Points[Index];
-				const float PtX = Pt.X - GraphTrackPointVisualSize / 2.0f - 1.0f;
-				const float PtY = GetPosY() + Pt.Y - GraphTrackPointVisualSize / 2.0f - 1.0f;
-				DrawContext.DrawBox(PtX, PtY, GraphTrackPointVisualSize + 2.0f, GraphTrackPointVisualSize + 2.0f, PointBrush, Series.BorderColor);
+				const float PtX = Pt.X - PointVisualSize / 2.0f - 1.0f;
+				const float PtY = GetPosY() + Pt.Y - PointVisualSize / 2.0f - 1.0f;
+				DrawContext.DrawBox(PtX, PtY, PointVisualSize + 2.0f, PointVisualSize + 2.0f, PointBrush, Series.BorderColor);
 			}
 			DrawContext.LayerId++;
 		}
@@ -170,9 +170,9 @@ void FGraphTrack::DrawSeries(const FGraphTrackSeries& Series, FDrawContext& Draw
 		for (int32 Index = 0; Index < NumPoints; ++Index)
 		{
 			const FVector2D& Pt = Series.Points[Index];
-			const float PtX = Pt.X - GraphTrackPointVisualSize / 2.0f;
-			const float PtY = GetPosY() + Pt.Y - GraphTrackPointVisualSize / 2.0f;
-			DrawContext.DrawBox(PtX, PtY, GraphTrackPointVisualSize, GraphTrackPointVisualSize, PointBrush, Series.Color);
+			const float PtX = Pt.X - PointVisualSize / 2.0f;
+			const float PtY = GetPosY() + Pt.Y - PointVisualSize / 2.0f;
+			DrawContext.DrawBox(PtX, PtY, PointVisualSize, PointVisualSize, PointBrush, Series.Color);
 		}
 		DrawContext.LayerId++;
 
@@ -183,7 +183,7 @@ void FGraphTrack::DrawSeries(const FGraphTrackSeries& Series, FDrawContext& Draw
 		if (bDrawPointsWithBorder)
 		{
 			// Draw borders.
-			const float BorderPtSize = GraphTrackPointVisualSize;
+			const float BorderPtSize = PointVisualSize;
 			//FVector2D BorderRotationPoint(BorderPtSize / 2.0f, BorderPtSize / 2.0f);
 			for (int32 Index = 0; Index < NumPoints; ++Index)
 			{
@@ -197,7 +197,7 @@ void FGraphTrack::DrawSeries(const FGraphTrackSeries& Series, FDrawContext& Draw
 		}
 
 		// Draw points as rectangles.
-		const float PtSize = GraphTrackPointVisualSize - 2.0f;
+		const float PtSize = PointVisualSize - 2.0f;
 		//FVector2D RotationPoint(PtSize / 2.0f, PtSize / 2.0f);
 		for (int32 Index = 0; Index < NumPoints; ++Index)
 		{
@@ -260,8 +260,8 @@ void FRandomGraphTrack::GenerateSeries(FGraphTrackSeries& Series, const FTimingT
 	//////////////////////////////////////////////////
 	// Generate random events.
 
-	static const double MinDeltaTime = 0.0000001; // 100ns
-	static const double MaxDeltaTime = 0.01; // 100ms
+	constexpr double MinDeltaTime = 0.0000001; // 100ns
+	constexpr double MaxDeltaTime = 0.01; // 100ms
 	const float MinValue = 0;
 	const float MaxValue = GetHeight();
 
@@ -444,7 +444,7 @@ void FGraphTrackBuilder::BeginPoints()
 	PointsCurrentX = -DBL_MAX;
 
 	PointsAtCurrentX.Reset();
-	int32 MaxPointsPerLineScan = FMath::CeilToInt(Track.GetHeight() / GraphTrackPointDY);
+	int32 MaxPointsPerLineScan = FMath::CeilToInt(Track.GetHeight() / FGraphTrack::PointSizeY);
 	if (MaxPointsPerLineScan > 0)
 	{
 		PointsAtCurrentX.AddDefaulted(MaxPointsPerLineScan);
@@ -456,15 +456,15 @@ void FGraphTrackBuilder::BeginPoints()
 void FGraphTrackBuilder::AddPoint(double Time, double Value)
 {
 	const float X = Viewport.TimeToSlateUnitsRounded(Time);
-	if (X < -GraphTrackPointVisualSize / 2.0f || X >= Viewport.Width + GraphTrackPointVisualSize / 2.0f)
+	if (X < -FGraphTrack::PointVisualSize / 2.0f || X >= Viewport.Width + FGraphTrack::PointVisualSize / 2.0f)
 	{
 		return;
 	}
 
 	// Align the X with a grid of GraphTrackPointDX pixels in size, in the global space (i.e. scroll independent).
-	const double AlignedX = FMath::RoundToDouble(Time * Viewport.ScaleX / GraphTrackPointDX) * GraphTrackPointDX;
+	const double AlignedX = FMath::RoundToDouble(Time * Viewport.ScaleX / FGraphTrack::PointSizeX) * FGraphTrack::PointSizeX;
 
-	if (AlignedX > PointsCurrentX + GraphTrackPointDX - 0.5)
+	if (AlignedX > PointsCurrentX + FGraphTrack::PointSizeX - 0.5)
 	{
 		FlushPoints();
 		PointsCurrentX = AlignedX;
@@ -472,7 +472,7 @@ void FGraphTrackBuilder::AddPoint(double Time, double Value)
 
 	const float Y = Track.GetYForValue(Value);
 
-	int32 Index = FMath::RoundToInt(Y / GraphTrackPointDY);
+	int32 Index = FMath::RoundToInt(Y / FGraphTrack::PointSizeY);
 	if (Index >= 0 && Index < PointsAtCurrentX.Num())
 	{
 		FPointInfo& Pt = PointsAtCurrentX[Index];
diff --git a/Engine/Source/Developer/TraceInsights/Private/Insights/ViewModels/GraphTrack.h b/Engine/Source/Developer/TraceInsights/Private/Insights/ViewModels/GraphTrack.h
--- a/Engine/Source/Developer/TraceInsights/Private/Insights/ViewModels/GraphTrack.h
+++ b/Engine/Source/Developer/TraceInsights/Private/Insights/ViewModels/GraphTrack.h
@@ -14,17 +14,6 @@ class FTimingTrackViewport;
 
 ////////////////////////////////////////////////////////////////////////////////////////////////////
 
-static const float DefaultGraphTrackH = 100.0f; // in pixels (Slate units)
-
-// Visual size of points (in pixels).
-static const float GraphTrackPointVisualSize = 7.0f;
-
-// Size of points (in pixels) used in reduction algorithm.
-static const double GraphTrackPointDX = 3.0f;
-static const float GraphTrackPointDY = 3.0f;
-
-////////////////////////////////////////////////////////////////////////////////////////////////////
-
 struct FGraphBox
 {
 	float X;
@@ -69,6 +58,14 @@ class FGraphTrack : public FBaseTimingTrack
 {
 	friend class FGraphTrackBuilder;
 
+private:
+	// Visual size of points (in pixels).
+	static constexpr float PointVisualSize = 7.0f;
+
+	// Size of points (in pixels) used in reduction algorithm.
+	static constexpr double PointSizeX = 3.0f;
+	static constexpr float PointSizeY = 3.0f;
+
 public:
 	FGraphTrack(uint64 InTrackId);
 	virtual ~FGraphTrack();
diff --git a/Engine/Source/Developer/TraceInsights/Private/Insights/ViewModels/MarkersTimingTrack.cpp b/Engine/Source/Developer/TraceInsights/Private/Insights/ViewModels/MarkersTimingTrack.cpp
--- a/Engine/Source/Developer/TraceInsights/Private/Insights/ViewModels/MarkersTimingTrack.cpp
+++ b/Engine/Source/Developer/TraceInsights/Private/Insights/ViewModels/MarkersTimingTrack.cpp
@@ -72,8 +72,8 @@ void FMarkersTimingTrack::Reset()
 
 void FMarkersTimingTrack::UpdateHeight()
 {
-	static const float BookmarksTrackHeight = 14.0f;
-	static const float TimeMarkersTrackHeight = 28.0f;
+	constexpr float BookmarksTrackHeight = 14.0f;
+	constexpr float TimeMarkersTrackHeight = 28.0f;
 
 	if (bUseOnlyBookmarks)
 	{
@@ -89,8 +89,8 @@ void FMarkersTimingTrack::UpdateHeight()
 
 void FMarkersTimingTrack::UpdateHoveredState(float MouseX, float MouseY, const FTimingTrackViewport& Viewport)
 {
-	static const float HeaderWidth = 80.0f;
-	static const float HeaderHeight = 14.0f;
+	constexpr float HeaderWidth = 80.0f;
+	constexpr float HeaderHeight = 14.0f;
 
 	if (MouseY >= GetPosY() && MouseY < GetPosY() + GetHeight())
 	{
@@ -415,7 +415,7 @@ void FTimeMarkerTrackBuilder::Flush(float AvailableTextW)
 			Box.X = LastX1;
 			Box.W = LastX2 - LastX1;
 			Box.Color = Color;
-			Box.Color.A = 0.5f;
+			Box.Color.A = 0.25f;
 		}
 
 		if (AvailableTextW > 6.0f)
diff --git a/Engine/Source/Developer/TraceInsights/Private/Insights/ViewModels/TimeRulerTrack.cpp b/Engine/Source/Developer/TraceInsights/Private/Insights/ViewModels/TimeRulerTrack.cpp
--- a/Engine/Source/Developer/TraceInsights/Private/Insights/ViewModels/TimeRulerTrack.cpp
+++ b/Engine/Source/Developer/TraceInsights/Private/Insights/ViewModels/TimeRulerTrack.cpp
@@ -36,7 +36,7 @@ void FTimeRulerTrack::Reset()
 {
 	FBaseTimingTrack::Reset();
 
-	static const float TimeRulerHeight = 22.0f;
+	constexpr float TimeRulerHeight = 22.0f;
 	SetHeight(TimeRulerHeight);
 }
 
@@ -44,8 +44,8 @@ void FTimeRulerTrack::Reset()
 
 void FTimeRulerTrack::UpdateHoveredState(float MouseX, float MouseY, const FTimingTrackViewport& Viewport)
 {
-	//static const float HeaderWidth = 100.0f;
-	//static const float HeaderHeight = 14.0f;
+	//constexpr float HeaderWidth = 100.0f;
+	//constexpr float HeaderHeight = 14.0f;
 
 	if (MouseY >= GetPosY() && MouseY < GetPosY() + GetHeight())
 	{
@@ -60,6 +60,50 @@ void FTimeRulerTrack::UpdateHoveredState(float MouseX, float MouseY, const FTimi
 
 ////////////////////////////////////////////////////////////////////////////////////////////////////
 
+void FTimeRulerTrack::DrawBackground(FDrawContext& DrawContext, const FTimingTrackViewport& Viewport) const
+{
+	const float X0 = Viewport.TimeToSlateUnitsRounded(0.0);
+	const float X1 = Viewport.TimeToSlateUnitsRounded(Viewport.MaxValidTime);
+	const float W = FMath::CeilToFloat(Viewport.Width);
+	const float H = GetHeight();
+
+	const FLinearColor InvalidAreaColor(0.08f, 0.07f, 0.07f, 1.0f);
+	const FLinearColor ValidAreaColor(0.09f, 0.09f, 0.09f, 1.0f);
+
+	if (X0 >= W || X1 <= 0.0f)
+	{
+		// Draw invalid area (entire view).
+		DrawContext.DrawBox(0.0f, 0.0f, W, H, WhiteBrush, InvalidAreaColor);
+	}
+	else // X0 < W && X1 > 0
+	{
+		if (X0 > 0.0f)
+		{
+			// Draw invalid area (left).
+			DrawContext.DrawBox(0.0f, 0.0f, X0, H, WhiteBrush, InvalidAreaColor);
+		}
+
+		if (X1 < W)
+		{
+			// Draw invalid area (right).
+			DrawContext.DrawBox(X1, 0.0f, W - X1, H, WhiteBrush, InvalidAreaColor);
+		}
+
+		float ValidX0 = FMath::Max(X0, 0.0f);
+		float ValidX1 = FMath::Min(X1, W);
+
+		if (ValidX1 > ValidX0)
+		{
+			// Draw valid area.
+			DrawContext.DrawBox(ValidX0, 0.0f, ValidX1 - ValidX0, H, WhiteBrush, ValidAreaColor);
+		}
+	}
+
+	DrawContext.LayerId++;
+}
+
+////////////////////////////////////////////////////////////////////////////////////////////////////
+
 void FTimeRulerTrack::Draw(FDrawContext& DrawContext, const FTimingTrackViewport& Viewport, const FVector2D& MousePosition, const bool bIsSelecting, const double SelectionStartTime, const double SelectionEndTime) const
 {
 	const float MinorTickMark = 5.0f;
@@ -80,8 +124,7 @@ void FTimeRulerTrack::Draw(FDrawContext& DrawContext, const FTimingTrackViewport
 	float MajorOX = static_cast<float>(FMath::RoundToDouble(MajorN * static_cast<double>(MajorTickMark) - VX));
 
 	// Draw the time ruler's background.
-	DrawContext.DrawBox(0.0f, GetPosY(), Viewport.Width, GetHeight(), WhiteBrush, FLinearColor(0.09f, 0.09f, 0.09f, 1.0f));
-	DrawContext.LayerId++;
+	DrawBackground(DrawContext, Viewport);
 
 	// Draw the minor tick marks.
 	for (float X = MinorOX; X < Viewport.Width; X += MinorTickMark)
diff --git a/Engine/Source/Developer/TraceInsights/Private/Insights/ViewModels/TimeRulerTrack.h b/Engine/Source/Developer/TraceInsights/Private/Insights/ViewModels/TimeRulerTrack.h
--- a/Engine/Source/Developer/TraceInsights/Private/Insights/ViewModels/TimeRulerTrack.h
+++ b/Engine/Source/Developer/TraceInsights/Private/Insights/ViewModels/TimeRulerTrack.h
@@ -29,6 +29,9 @@ public:
 				const double SelectionStartTime = 0.0,
 				const double SelectionEndTime = 0.0) const;
 
+private:
+	void DrawBackground(FDrawContext& DrawContext, const FTimingTrackViewport& Viewport) const;
+
 public:
 	// Slate resources
 	const FSlateBrush* WhiteBrush;
diff --git a/Engine/Source/Developer/TraceInsights/Private/Insights/ViewModels/TimingEvent.h b/Engine/Source/Developer/TraceInsights/Private/Insights/ViewModels/TimingEvent.h
--- a/Engine/Source/Developer/TraceInsights/Private/Insights/ViewModels/TimingEvent.h
+++ b/Engine/Source/Developer/TraceInsights/Private/Insights/ViewModels/TimingEvent.h
@@ -3,6 +3,8 @@
 #pragma once
 
 #include "CoreMinimal.h"
+#include "TraceServices/Model/LoadTimeProfiler.h"
+
 
 // Insights
 #include "Insights/ViewModels/TimerNode.h"
@@ -17,6 +19,7 @@ struct FTimingEvent
 	double StartTime;
 	double EndTime;
 	double ExclusiveTime;
+	Trace::FLoadTimeProfilerCpuEvent LoadingInfo;
 
 	void Reset()
 	{
@@ -26,6 +29,10 @@ struct FTimingEvent
 		StartTime = 0.0;
 		EndTime = 0.0;
 		ExclusiveTime = 0.0;
+		LoadingInfo.Package = nullptr;
+		LoadingInfo.Export = nullptr;
+		LoadingInfo.PackageEventType = LoadTimeProfilerPackageEventType_None;
+		LoadingInfo.ExportEventType = LoadTimeProfilerObjectEventType_None;
 	}
 
 	bool IsValidTrack() const { return Track != nullptr; }
diff --git a/Engine/Source/Developer/TraceInsights/Private/Insights/ViewModels/TimingEventsTrack.cpp b/Engine/Source/Developer/TraceInsights/Private/Insights/ViewModels/TimingEventsTrack.cpp
--- a/Engine/Source/Developer/TraceInsights/Private/Insights/ViewModels/TimingEventsTrack.cpp
+++ b/Engine/Source/Developer/TraceInsights/Private/Insights/ViewModels/TimingEventsTrack.cpp
@@ -34,41 +34,58 @@ void FTimingEventsTrackLayout::ForceCompactMode()
 
 void FTimingEventsTrackLayout::Update()
 {
-	const float LayoutTransitionSpeed = 0.25f;
+	constexpr float LayoutTransitionSpeed = 0.25f;
 
 	if (bIsCompactMode)
 	{
 		if (EventH > CompactLayoutEventH)
+		{
 			EventH -= LayoutTransitionSpeed;
+		}
 		if (EventDY > CompactLayoutEventDY)
+		{
 			EventDY -= LayoutTransitionSpeed;
+		}
 		if (TimelineDY > CompactLayoutTimelineDY)
+		{
 			TimelineDY -= LayoutTransitionSpeed;
+		}
 	}
 	else
 	{
 		if (EventH < NormalLayoutEventH)
+		{
 			EventH += LayoutTransitionSpeed;
+		}
 		if (EventDY < NormalLayoutEventDY)
+		{
 			EventDY += LayoutTransitionSpeed;
+		}
 		if (TimelineDY < NormalLayoutTimelineDY)
+		{
 			TimelineDY += LayoutTransitionSpeed;
+		}
 	}
 
 	if (MinTimelineH > TargetMinTimelineH)
+	{
 		MinTimelineH -= LayoutTransitionSpeed;
+	}
 	else if (MinTimelineH < TargetMinTimelineH)
+	{
 		MinTimelineH += LayoutTransitionSpeed;
+	}
 }
 
 ////////////////////////////////////////////////////////////////////////////////////////////////////
 // FTimingEventsTrack
 ////////////////////////////////////////////////////////////////////////////////////////////////////
 
-FTimingEventsTrack::FTimingEventsTrack(uint64 InTrackId, ETimingEventsTrackType InType, const FString& InName)
+FTimingEventsTrack::FTimingEventsTrack(uint64 InTrackId, ETimingEventsTrackType InType, const FString& InName, const TCHAR* InGroupName)
 	: FBaseTimingTrack(InTrackId)
 	, Type(InType)
 	, Name(InName)
+	, GroupName(InGroupName)
 	, Depth(-1)
 	, bIsCollapsed(false)
 {
@@ -94,8 +111,8 @@ void FTimingEventsTrack::Reset()
 
 void FTimingEventsTrack::UpdateHoveredState(float MouseX, float MouseY, const FTimingTrackViewport& Viewport)
 {
-	static const float HeaderWidth = 100.0f;
-	static const float HeaderHeight = 14.0f;
+	constexpr float HeaderWidth = 100.0f;
+	constexpr float HeaderHeight = 14.0f;
 
 	if (MouseY >= GetPosY() && MouseY < GetPosY() + GetHeight())
 	{
diff --git a/Engine/Source/Developer/TraceInsights/Private/Insights/ViewModels/TimingEventsTrack.h b/Engine/Source/Developer/TraceInsights/Private/Insights/ViewModels/TimingEventsTrack.h
--- a/Engine/Source/Developer/TraceInsights/Private/Insights/ViewModels/TimingEventsTrack.h
+++ b/Engine/Source/Developer/TraceInsights/Private/Insights/ViewModels/TimingEventsTrack.h
@@ -45,7 +45,8 @@ enum class ETimingEventsTrackType
 {
 	Cpu,
 	Gpu,
-	Io
+	Loading,
+	Io,
 };
 
 ////////////////////////////////////////////////////////////////////////////////////////////////////
@@ -53,12 +54,14 @@ enum class ETimingEventsTrackType
 class FTimingEventsTrack : public FBaseTimingTrack
 {
 public:
-	FTimingEventsTrack(uint64 InTrackId, ETimingEventsTrackType InType, const FString& InName);
+	FTimingEventsTrack(uint64 InTrackId, ETimingEventsTrackType InType, const FString& InName, const TCHAR* InGroupName);
 	virtual ~FTimingEventsTrack();
 
 	ETimingEventsTrackType GetType() const { return Type; }
 	const FString& GetName() const { return Name; }
 
+	const TCHAR* GetGroupName() const { return GroupName; };
+
 	void SetOrder(int32 InOrder) { Order = InOrder; }
 	int32 GetOrder() const { return Order; }
 
@@ -71,6 +74,7 @@ public:
 public:
 	ETimingEventsTrackType Type;
 	FString Name;
+	const TCHAR* GroupName;
 	int32 Order;
 
 	int32 Depth; // number of lanes == Depth + 1
diff --git a/Engine/Source/Developer/TraceInsights/Private/Insights/ViewModels/TimingViewDrawHelper.cpp b/Engine/Source/Developer/TraceInsights/Private/Insights/ViewModels/TimingViewDrawHelper.cpp
--- a/Engine/Source/Developer/TraceInsights/Private/Insights/ViewModels/TimingViewDrawHelper.cpp
+++ b/Engine/Source/Developer/TraceInsights/Private/Insights/ViewModels/TimingViewDrawHelper.cpp
@@ -21,17 +21,6 @@
 
 ////////////////////////////////////////////////////////////////////////////////////////////////////
 
-static const int32 EventBorderLayer = 0;
-static const int32 EventFillLayer = 1;
-static const int32 EventTextLayer = 2;
-
-static const int32 TimelineHeaderLayer = 3;
-static const int32 TimelineTextLayer = 4;
-
-static const int32 NumLayers = 5;
-
-////////////////////////////////////////////////////////////////////////////////////////////////////
-
 FTimingViewDrawHelper::FTimingViewDrawHelper(const FDrawContext& InDC, const FTimingTrackViewport& InViewport, const FTimingEventsTrackLayout& InLayout)
 	: DrawContext(InDC)
 	, Viewport(InViewport)
@@ -58,8 +47,8 @@ FTimingViewDrawHelper::~FTimingViewDrawHelper()
 
 void FTimingViewDrawHelper::DrawBackground() const
 {
-	float X0 = Viewport.TimeToSlateUnitsRounded(0.0);
-	float X1 = Viewport.TimeToSlateUnitsRounded(Viewport.MaxValidTime);
+	const float X0 = Viewport.TimeToSlateUnitsRounded(0.0);
+	const float X1 = Viewport.TimeToSlateUnitsRounded(Viewport.MaxValidTime);
 	const float W = FMath::CeilToFloat(Viewport.Width);
 	const float H = FMath::CeilToFloat(Viewport.Height);
 
@@ -112,7 +101,7 @@ void FTimingViewDrawHelper::BeginTimelines()
 		Y = FMath::Min(Y, Viewport.Height);
 
 		// Draw invalid area (top).
-		DrawContext.DrawBox(DrawContext.LayerId + TimelineHeaderLayer, ValidX0, 0.0f, ValidX1 - ValidX0, Y, BackgroundAreaBrush, InvalidAreaColor);
+		DrawContext.DrawBox(DrawContext.LayerId + ToInt32(EDrawLayer::TimelineHeader), ValidX0, 0.0f, ValidX1 - ValidX0, Y, BackgroundAreaBrush, InvalidAreaColor);
 	}
 }
 
@@ -236,14 +225,14 @@ void FTimingViewDrawHelper::EndTimeline(FTimingEventsTrack& Track)
 		}
 
 		// Draw a horizontal line between timelines.
-		DrawContext.DrawBox(DrawContext.LayerId + TimelineHeaderLayer, 0, Y, Viewport.Width, 1.0f, WhiteBrush, Color);
+		DrawContext.DrawBox(DrawContext.LayerId + ToInt32(EDrawLayer::TimelineHeader), 0, Y, Viewport.Width, 1.0f, WhiteBrush, Color);
 
 		// Draw name of timeline.
 		//const FString Name = FString::Printf(TEXT("%d Y: %g, H: %g, VO: %g, VY: %g"), TimelineIndex, Track.Y, Track.H, Viewport.TopOffset, Viewport.ScrollPosY); //debug
 		const TSharedRef<FSlateFontMeasure> FontMeasureService = FSlateApplication::Get().GetRenderer()->GetFontMeasureService();
 		float NameWidth = FontMeasureService->Measure(Track.GetName(), EventFont).X;
-		DrawContext.DrawBox(DrawContext.LayerId + TimelineHeaderLayer, 0.0f, Y + 1.0f, NameWidth + 4.0f, 12.0f, WhiteBrush, Color);
-		DrawContext.DrawText(DrawContext.LayerId + TimelineTextLayer, 2.0f, Y, Track.GetName(), EventFont, TextColor);
+		DrawContext.DrawBox(DrawContext.LayerId + ToInt32(EDrawLayer::TimelineHeader), 0.0f, Y + 1.0f, NameWidth + 4.0f, 12.0f, WhiteBrush, Color);
+		DrawContext.DrawText(DrawContext.LayerId + ToInt32(EDrawLayer::TimelineText), 2.0f, Y, Track.GetName(), EventFont, TextColor);
 	}
 
 	TimelineTopY += Track.GetHeight();
@@ -400,15 +389,15 @@ void FTimingViewDrawHelper::AddEvent(double EventStartTime, double EventEndTime,
 	//EventColorFill.B = (Color & 0xFF) * MulL;
 	//EventColorFill.A = 1.0f;
 
-	//static const float BorderColorFactor = 0.75f; // darker border
-	static const float BorderColorFactor = 1.25f; // brighter border
+	//constexpr float BorderColorFactor = 0.75f; // darker border
+	constexpr float BorderColorFactor = 1.25f; // brighter border
 
 	//////////////////////////////////////////////////
 
 	if (EventW > 2.0f && Layout.EventH > 2.0f)
 	{
 		// Fill inside of the timing event box.
-		DrawContext.DrawBox(DrawContext.LayerId + EventFillLayer, EventX1 + 1.0f, EventY + 1.0f, EventW - 2.0f, Layout.EventH - 2.0f, WhiteBrush, EventColorFill);
+		DrawContext.DrawBox(DrawContext.LayerId + ToInt32(EDrawLayer::EventFill), EventX1 + 1.0f, EventY + 1.0f, EventW - 2.0f, Layout.EventH - 2.0f, WhiteBrush, EventColorFill);
 		Stats.NumDrawBoxes++;
 	}
 
@@ -426,7 +415,7 @@ void FTimingViewDrawHelper::AddEvent(double EventStartTime, double EventEndTime,
 		}
 
 		const FLinearColor EventColorBorder(EventColorFill.R * BorderColorFactor, EventColorFill.G * BorderColorFactor, EventColorFill.B * BorderColorFactor, 1.0f);
-		DrawContext.DrawBox(DrawContext.LayerId + EventBorderLayer, EventX1, EventY, EventW, Layout.EventH, EventsBorderBrush, EventColorBorder);
+		DrawContext.DrawBox(DrawContext.LayerId + ToInt32(EDrawLayer::EventBorder), EventX1, EventY, EventW, Layout.EventH, EventsBorderBrush, EventColorBorder);
 		Stats.NumDrawBorders++;
 	}
 	else // 1px or 2px boxes
@@ -480,7 +469,7 @@ void FTimingViewDrawHelper::AddEvent(double EventStartTime, double EventEndTime,
 			// Ref: https://books.google.ro/books?id=0pVr7dhmdWYC
 			const bool bIsDarkColor = (EventColorFill.ComputeLuminance() < 0.4f);
 
-			DrawContext.DrawText(DrawContext.LayerId + EventTextLayer, EventX1 + 2.0f, EventY + 1.0f, Name, 0, LastWholeCharacterIndex + 1, EventFont, bIsDarkColor ? FLinearColor::White : FLinearColor::Black);
+			DrawContext.DrawText(DrawContext.LayerId + ToInt32(EDrawLayer::EventText), EventX1 + 2.0f, EventY + 1.0f, Name, 0, LastWholeCharacterIndex + 1, EventFont, bIsDarkColor ? FLinearColor::White : FLinearColor::Black);
 			Stats.NumDrawTexts++;
 		}
 	}
@@ -490,7 +479,7 @@ void FTimingViewDrawHelper::AddEvent(double EventStartTime, double EventEndTime,
 
 void FTimingViewDrawHelper::DrawBox(const FBoxData& Box, const float EventY, const float EventH)
 {
-	DrawContext.DrawBox(DrawContext.LayerId + EventBorderLayer, Box.X1, EventY, Box.X2 - Box.X1, EventH, WhiteBrush, Box.LinearColor);
+	DrawContext.DrawBox(DrawContext.LayerId + ToInt32(EDrawLayer::EventBorder), Box.X1, EventY, Box.X2 - Box.X1, EventH, WhiteBrush, Box.LinearColor);
 	Stats.NumDrawBoxes++;
 }
 
@@ -501,7 +490,7 @@ void FTimingViewDrawHelper::EndTimelines()
 	float Y = Viewport.TopOffset + TimelineTopY;
 
 	// Draw a last horizontal line.
-	DrawContext.DrawBox(DrawContext.LayerId + TimelineHeaderLayer, 0, Y, Viewport.Width, 1.0f, WhiteBrush, FLinearColor(0.05f, 0.05f, 0.05f, 1.0f));
+	DrawContext.DrawBox(DrawContext.LayerId + ToInt32(EDrawLayer::TimelineHeader), 0, Y, Viewport.Width, 1.0f, WhiteBrush, FLinearColor(0.05f, 0.05f, 0.05f, 1.0f));
 	Y += 1.0f;
 
 	if (Y < Viewport.Height && ValidX1 > ValidX0)
@@ -509,10 +498,10 @@ void FTimingViewDrawHelper::EndTimelines()
 		Y = FMath::Max(Y, 0.0f);
 
 		// Draw invalid area (bottom).
-		DrawContext.DrawBox(DrawContext.LayerId + TimelineHeaderLayer, ValidX0, Y, ValidX1 - ValidX0, Viewport.Height - Y, BackgroundAreaBrush, InvalidAreaColor);
+		DrawContext.DrawBox(DrawContext.LayerId + ToInt32(EDrawLayer::TimelineHeader), ValidX0, Y, ValidX1 - ValidX0, Viewport.Height - Y, BackgroundAreaBrush, InvalidAreaColor);
 	}
 
-	DrawContext.LayerId += NumLayers;
+	DrawContext.LayerId += ToInt32(EDrawLayer::Count);
 }
 
 ////////////////////////////////////////////////////////////////////////////////////////////////////
diff --git a/Engine/Source/Developer/TraceInsights/Private/Insights/ViewModels/TimingViewDrawHelper.h b/Engine/Source/Developer/TraceInsights/Private/Insights/ViewModels/TimingViewDrawHelper.h
--- a/Engine/Source/Developer/TraceInsights/Private/Insights/ViewModels/TimingViewDrawHelper.h
+++ b/Engine/Source/Developer/TraceInsights/Private/Insights/ViewModels/TimingViewDrawHelper.h
@@ -28,6 +28,18 @@ public:
 	};
 
 private:
+	enum class EDrawLayer : int32
+	{
+		EventBorder,
+		EventFill,
+		EventText,
+		TimelineHeader,
+		TimelineText,
+
+		Count,
+	};
+	static int32 ToInt32(EDrawLayer Layer) { return static_cast<int32>(Layer); }
+
 	struct FBoxData
 	{
 		float X1;
diff --git a/Engine/Source/Developer/TraceInsights/Private/Insights/Widgets/SFrameTrack.cpp b/Engine/Source/Developer/TraceInsights/Private/Insights/Widgets/SFrameTrack.cpp
--- a/Engine/Source/Developer/TraceInsights/Private/Insights/Widgets/SFrameTrack.cpp
+++ b/Engine/Source/Developer/TraceInsights/Private/Insights/Widgets/SFrameTrack.cpp
@@ -108,8 +108,8 @@ void SFrameTrack::Construct(const FArguments& InArgs)
 			.Orientation(Orient_Horizontal)
 			.AlwaysShowScrollbar(false)
 			.Visibility(EVisibility::Visible)
-			.Thickness(FVector2D(4.0f, 4.0f))
-			.RenderOpacity(0.5)
+			.Thickness(FVector2D(5.0f, 5.0f))
+			.RenderOpacity(0.75)
 			.OnUserScrolled(this, &SFrameTrack::HorizontalScrollBar_OnUserScrolled)
 		]
 	];
@@ -700,7 +700,7 @@ FReply SFrameTrack::OnMouseWheel(const FGeometry& MyGeometry, const FPointerEven
 	{
 		// Zoom in/out vertically.
 		const float Delta = MouseEvent.GetWheelDelta();
-		static const float ZoomStep = 0.25f; // as percent
+		constexpr float ZoomStep = 0.25f; // as percent
 		float ScaleY;
 
 		if (Delta > 0)
@@ -729,7 +729,7 @@ FReply SFrameTrack::OnMouseWheel(const FGeometry& MyGeometry, const FPointerEven
 
 void SFrameTrack::ZoomHorizontally(const float Delta, const float X)
 {
-	static const float ZoomStep = 0.25f; // as percent
+	constexpr float ZoomStep = 0.25f; // as percent
 	float ScaleX;
 
 	if (Delta > 0)
diff --git a/Engine/Source/Developer/TraceInsights/Private/Insights/Widgets/SStartPageWindow.cpp b/Engine/Source/Developer/TraceInsights/Private/Insights/Widgets/SStartPageWindow.cpp
--- a/Engine/Source/Developer/TraceInsights/Private/Insights/Widgets/SStartPageWindow.cpp
+++ b/Engine/Source/Developer/TraceInsights/Private/Insights/Widgets/SStartPageWindow.cpp
@@ -484,15 +484,15 @@ TSharedRef<SWidget> SStartPageWindow::ConstructModuleList()
 	TSharedRef<Trace::IModuleService> ModuleService = FInsightsManager::Get()->GetModuleService();
 	ModuleService->GetAvailableModules(AvailableModules);
 
-	const bool bDefaultIsModuleEnabledState = true;
+	constexpr bool bDefaultModuleEnableState = false;
 	AvailableModulesEnabledState.Reset();
 
 	for (int32 ModuleIndex = 0; ModuleIndex < AvailableModules.Num(); ++ModuleIndex)
 	{
 		const Trace::FModuleInfo& Module = AvailableModules[ModuleIndex];
 
-		AvailableModulesEnabledState.Add(bDefaultIsModuleEnabledState);
-		ModuleService->SetModuleEnabled(Module.Name, bDefaultIsModuleEnabledState);
+		AvailableModulesEnabledState.Add(bDefaultModuleEnableState);
+		ModuleService->SetModuleEnabled(Module.Name, bDefaultModuleEnableState);
 
 		VerticalBox->AddSlot()
 			.AutoHeight()
diff --git a/Engine/Source/Developer/TraceInsights/Private/Insights/Widgets/STimingView.cpp b/Engine/Source/Developer/TraceInsights/Private/Insights/Widgets/STimingView.cpp
--- a/Engine/Source/Developer/TraceInsights/Private/Insights/Widgets/STimingView.cpp
+++ b/Engine/Source/Developer/TraceInsights/Private/Insights/Widgets/STimingView.cpp
@@ -35,6 +35,7 @@
 #include "Insights/InsightsManager.h"
 #include "Insights/TimingProfilerCommon.h"
 #include "Insights/TimingProfilerManager.h"
+#include "Insights/ViewModels/BaseTimingTrack.h"
 #include "Insights/ViewModels/TimingViewDrawHelper.h"
 #include "Insights/Widgets/SStatsView.h"
 #include "Insights/Widgets/STimersView.h"
@@ -47,7 +48,11 @@
 #define DEBUG_TIMING_TRACK 0
 #define LOCTEXT_NAMESPACE "STimingView"
 
-static const float MinTooltipWidth = 110.0f;
+// start auto generated ids from a big number (MSB set to 1) to avoid collisions with ids for gpu/cpu tracks based on 32bit timeline index
+uint64 FBaseTimingTrack::IdGenerator = (1ULL << 63);
+
+const TCHAR* GetName(ELoadTimeProfilerPackageEventType Type);
+const TCHAR* GetName(ELoadTimeProfilerObjectEventType Type);
 
 ////////////////////////////////////////////////////////////////////////////////////////////////////
 
@@ -95,7 +100,7 @@ void STimingView::Construct(const FArguments& InArgs)
 			.Orientation(Orient_Horizontal)
 			.AlwaysShowScrollbar(false)
 			.Visibility(EVisibility::Visible)
-			.Thickness(FVector2D(8.0f, 8.0f))
+			.Thickness(FVector2D(5.0f, 5.0f))
 			//.RenderOpacity(0.75)
 			.OnUserScrolled(this, &STimingView::HorizontalScrollBar_OnUserScrolled)
 		]
@@ -108,7 +113,7 @@ void STimingView::Construct(const FArguments& InArgs)
 			.Orientation(Orient_Vertical)
 			.AlwaysShowScrollbar(false)
 			.Visibility(EVisibility::Visible)
-			.Thickness(FVector2D(8.0f, 8.0f))
+			.Thickness(FVector2D(5.0f, 5.0f))
 			//.RenderOpacity(0.75)
 			.OnUserScrolled(this, &STimingView::VerticalScrollBar_OnUserScrolled)
 		]
@@ -164,6 +169,7 @@ void STimingView::Reset()
 {
 	bShowHideAllGpuTracks = true;
 	bShowHideAllCpuTracks = true;
+	bShowHideAllLoadingTracks = true;
 	bShowHideAllIoTracks = false;
 
 	//////////////////////////////////////////////////
@@ -175,6 +181,8 @@ void STimingView::Reset()
 
 	//////////////////////////////////////////////////
 
+	ThreadGroups.Reset();
+
 	for (const auto& KV : CachedTimelines)
 	{
 		delete KV.Value;
@@ -195,6 +203,11 @@ void STimingView::Reset()
 
 	//////////////////////////////////////////////////
 
+	LoadingMainThreadTrack = nullptr;
+	LoadingAsyncThreadTrack = nullptr;
+
+	LoadingGetEventNameFn = FLoadingTrackGetEventNameDelegate::CreateRaw(this, &STimingView::GetLoadTimeProfilerEventName);
+
 	IoOverviewTrack = nullptr;
 	IoActivityTrack = nullptr;
 
@@ -394,11 +407,27 @@ void STimingView::Tick(const FGeometry& AllottedGeometry, const double InCurrent
 
 		bool bIsTimingEventsTrackDirty = false;
 
+		if (LoadingMainThreadTrack == nullptr)
+		{
+			uint64 TrackId = FBaseTimingTrack::GenerateId();
+			LoadingMainThreadTrack = AddTimingEventsTrack(TrackId, ETimingEventsTrackType::Loading, TEXT("Loading - Main Thread"), nullptr, -3);
+			LoadingMainThreadTrack->SetVisibilityFlag(bShowHideAllLoadingTracks);
+			bIsTimingEventsTrackDirty = true;
+		}
+
+		if (LoadingAsyncThreadTrack == nullptr)
+		{
+			uint64 TrackId = FBaseTimingTrack::GenerateId();
+			LoadingAsyncThreadTrack = AddTimingEventsTrack(TrackId, ETimingEventsTrackType::Loading, TEXT("Loading - Async Thread"), nullptr, -2);
+			LoadingAsyncThreadTrack->SetVisibilityFlag(bShowHideAllLoadingTracks);
+			bIsTimingEventsTrackDirty = true;
+		}
+
 		if (IoOverviewTrack == nullptr)
 		{
-			// Note: This track is hardcoded for now, as the I/O timelines are just prototypes for now (will be removed once the functionality is moved in analyzer).
-			uint64 OverviewTrackId = (0xF << 28) | 1;
-			IoOverviewTrack = AddTimingEventsTrack(OverviewTrackId, ETimingEventsTrackType::Io, TEXT("I/O Overview"), -1);
+			// Note: The I/O timelines are just prototypes for now (will be removed once the functionality is moved in analyzer).
+			const uint64 TrackId = FBaseTimingTrack::GenerateId();
+			IoOverviewTrack = AddTimingEventsTrack(TrackId, ETimingEventsTrackType::Io, TEXT("I/O Overview"), nullptr, -1);
 			IoOverviewTrack->SetVisibilityFlag(bShowHideAllIoTracks);
 			bIsTimingEventsTrackDirty = true;
 		}
@@ -412,9 +441,10 @@ void STimingView::Tick(const FGeometry& AllottedGeometry, const double InCurrent
 		uint32 GpuTimelineIndex;
 		if (TimingProfilerProvider.GetGpuTimelineIndex(GpuTimelineIndex))
 		{
-			if (!CachedTimelines.Contains(GpuTimelineIndex))
+			const uint64 TrackId = static_cast<uint64>(GpuTimelineIndex);
+			if (!CachedTimelines.Contains(TrackId))
 			{
-				AddTimingEventsTrack(GpuTimelineIndex, ETimingEventsTrackType::Gpu, TEXT("GPU"), 0);
+				AddTimingEventsTrack(TrackId, ETimingEventsTrackType::Gpu, TEXT("GPU"), nullptr, 0);
 				bIsTimingEventsTrackDirty = true;
 			}
 		}
@@ -424,33 +454,61 @@ void STimingView::Tick(const FGeometry& AllottedGeometry, const double InCurrent
 		const Trace::IThreadProvider& ThreadProvider = Trace::ReadThreadProvider(*Session.Get());
 		ThreadProvider.EnumerateThreads([this, &Order, &bIsTimingEventsTrackDirty, &TimingProfilerProvider](const Trace::FThreadInfo& ThreadInfo)
 		{
+			const TCHAR* GroupName = ThreadInfo.GroupName;
+			if (GroupName == nullptr)
+			{
+				GroupName = ThreadInfo.Name;
+			}
+
+			bool bIsGroupVisible = true;
+			if (GroupName != nullptr)
+			{
+				if (!ThreadGroups.Contains(GroupName))
+				{
+					ThreadGroups.Add(GroupName, { GroupName, bIsGroupVisible, 0 });
+				}
+				else
+				{
+					bIsGroupVisible = ThreadGroups[GroupName].bIsVisible;
+				}
+			}
+
 			uint32 CpuTimelineIndex;
 			if (TimingProfilerProvider.GetCpuThreadTimelineIndex(ThreadInfo.Id, CpuTimelineIndex))
 			{
-				if (!CachedTimelines.Contains(CpuTimelineIndex))
+				FTimingEventsTrack* Track = nullptr;
+				const uint64 TrackId = static_cast<uint64>(CpuTimelineIndex);
+
+				if (!CachedTimelines.Contains(TrackId))
 				{
 					FString TrackName(ThreadInfo.Name && *ThreadInfo.Name ? ThreadInfo.Name : FString::Printf(TEXT("Thread %u")));
-					AddTimingEventsTrack(CpuTimelineIndex, ETimingEventsTrackType::Cpu, TrackName, Order);
+					Track = AddTimingEventsTrack(TrackId, ETimingEventsTrackType::Cpu, TrackName, GroupName, Order);
 					bIsTimingEventsTrackDirty = true;
 				}
 				else
 				{
-					FTimingEventsTrack* Track = CachedTimelines[CpuTimelineIndex];
+					Track = CachedTimelines[TrackId];
 					if (Track->GetOrder() != Order)
 					{
 						Track->SetOrder(Order);
 						bIsTimingEventsTrackDirty = true;
 					}
 				}
+
+				Track->SetVisibilityFlag(bIsGroupVisible);
+
+				FThreadGroup& ThreadGroup = ThreadGroups[GroupName];
+				ThreadGroup.NumTimelines++;
+				
 				Order++;
 			}
 		});
 
 		if (IoActivityTrack == nullptr)
 		{
-			// Note: This track is hardcoded for now, as the I/O timelines are just prototypes for now (will be removed once the functionality is moved in analyzer).
-			uint64 ActivityTrackId = (0xF << 28) | 2;
-			IoActivityTrack = AddTimingEventsTrack(ActivityTrackId, ETimingEventsTrackType::Io, TEXT("I/O Activity"), 999999);
+			// Note: The I/O timelines are just prototypes for now (will be removed once the functionality is moved in analyzer).
+			uint64 TrackId = FBaseTimingTrack::GenerateId();
+			IoActivityTrack = AddTimingEventsTrack(TrackId, ETimingEventsTrackType::Io, TEXT("I/O Activity"), nullptr, 999999);
 			IoActivityTrack->SetVisibilityFlag(bShowHideAllIoTracks);
 			bIsTimingEventsTrackDirty = true;
 		}
@@ -630,11 +688,15 @@ int32 STimingView::OnPaint(const FPaintArgs& Args, const FGeometry& AllottedGeom
 				//TODO: Track.Draw(Helper);
 				if (Track.GetType() == ETimingEventsTrackType::Cpu)
 				{
-					DrawCpuGpuTimelineTrack(Helper, Track);
+					DrawTimingProfilerTrack(Helper, Track);
 				}
 				else if (Track.GetType() == ETimingEventsTrackType::Gpu)
 				{
-					DrawCpuGpuTimelineTrack(Helper, Track);
+					DrawTimingProfilerTrack(Helper, Track);
+				}
+				else if (Track.GetType() == ETimingEventsTrackType::Loading)
+				{
+					DrawLoadTimeProfilerTrack(Helper, Track);
 				}
 				else if (Track.GetType() == ETimingEventsTrackType::Io)
 				{
@@ -897,90 +959,218 @@ int32 STimingView::OnPaint(const FPaintArgs& Args, const FGeometry& AllottedGeom
 		// Draw info about selected event (bottom-right corner).
 		if (SelectedTimingEvent.IsValid())
 		{
-			const FTimerNodePtr TimerNodePtr = GetTimerNode(SelectedTimingEvent.TypeId);
+			if (SelectedTimingEvent.Track->GetType() == ETimingEventsTrackType::Cpu ||
+				SelectedTimingEvent.Track->GetType() == ETimingEventsTrackType::Gpu)
+			{
+				const FTimerNodePtr TimerNodePtr = GetTimerNode(SelectedTimingEvent.TypeId);
 
-			FString Str = FString::Printf(TEXT("%s (Incl.: %s, Excl.: %s)"),
-				TimerNodePtr ? *(TimerNodePtr->GetName().ToString()) : TEXT("N/A"),
-				*TimeUtils::FormatTimeAuto(SelectedTimingEvent.Duration()),
-				*TimeUtils::FormatTimeAuto(SelectedTimingEvent.ExclusiveTime));
+				FString Str = FString::Printf(TEXT("%s (Incl.: %s, Excl.: %s)"),
+					TimerNodePtr ? *(TimerNodePtr->GetName().ToString()) : TEXT("N/A"),
+					*TimeUtils::FormatTimeAuto(SelectedTimingEvent.Duration()),
+					*TimeUtils::FormatTimeAuto(SelectedTimingEvent.ExclusiveTime));
 
-			const FVector2D Size = FontMeasureService->Measure(Str, MainFont);
-			const float X = Viewport.Width - Size.X - 23.0f;
-			const float Y = Viewport.Height - Size.Y - 18.0f;
+				const FVector2D Size = FontMeasureService->Measure(Str, MainFont);
+				const float X = Viewport.Width - Size.X - 23.0f;
+				const float Y = Viewport.Height - Size.Y - 18.0f;
 
-			const FLinearColor BackgroundColor(0.05f, 0.05f, 0.05f, 1.0f);
-			const FLinearColor TextColor(0.7f, 0.7f, 0.7f, 1.0f);
+				const FLinearColor BackgroundColor(0.05f, 0.05f, 0.05f, 1.0f);
+				const FLinearColor TextColor(0.7f, 0.7f, 0.7f, 1.0f);
 
-			DrawContext.DrawBox(X - 8.0f, Y - 2.0f, Size.X + 16.0f, Size.Y + 4.0f, WhiteBrush, BackgroundColor);
-			DrawContext.LayerId++;
+				DrawContext.DrawBox(X - 8.0f, Y - 2.0f, Size.X + 16.0f, Size.Y + 4.0f, WhiteBrush, BackgroundColor);
+				DrawContext.LayerId++;
 
-			DrawContext.DrawText(X, Y, Str, MainFont, TextColor);
-			DrawContext.LayerId++;
+				DrawContext.DrawText(X, Y, Str, MainFont, TextColor);
+				DrawContext.LayerId++;
+			}
+			else if (SelectedTimingEvent.Track->GetType() == ETimingEventsTrackType::Loading)
+			{
+				//TODO: ...
+			}
 		}
 
 		// Draw info about hovered event (like a tooltip at mouse position).
 		if (HoveredTimingEvent.IsValid() && !MousePosition.IsZero())
 		{
-			const FTimerNodePtr TimerNodePtr = GetTimerNode(HoveredTimingEvent.TypeId);
-			FString Name = TimerNodePtr ? TimerNodePtr->GetName().ToString() : TEXT("N/A");
-
-			float W = FontMeasureService->Measure(Name, MainFont).X;
-			if (W < MinTooltipWidth)
-			{
-				W = MinTooltipWidth;
-			}
-			if (TooltipWidth != W)
+			if (HoveredTimingEvent.Track->GetType() == ETimingEventsTrackType::Cpu ||
+				HoveredTimingEvent.Track->GetType() == ETimingEventsTrackType::Gpu)
 			{
-				TooltipWidth = TooltipWidth * 0.75f + W * 0.25f;
-			}
-			const float MaxX = FMath::Max(0.0f, Viewport.Width - TooltipWidth - 21.0f);
-			const float X = FMath::Clamp<float>(MousePosition.X + 18.0f, 0.0f, MaxX);
+				const FTimerNodePtr TimerNodePtr = GetTimerNode(HoveredTimingEvent.TypeId);
+				FString Name = TimerNodePtr ? TimerNodePtr->GetName().ToString() : TEXT("N/A");
 
-			const float LineH = 14.0f;
-			const float H = 4 * LineH;
-			const float MaxY = FMath::Max(0.0f, Viewport.Height - H - 18.0f);
-			float Y = FMath::Clamp<float>(MousePosition.Y + 18.0f, 0.0f, MaxY);
+				float W = FontMeasureService->Measure(Name, MainFont).X;
 
-			const float Alpha = 1.0f - FMath::Abs(TooltipWidth - W) / W;
-			if (TooltipAlpha < Alpha)
-			{
-				TooltipAlpha = TooltipAlpha * 0.9f + Alpha * 0.1f;
+				if (W < MinTooltipWidth)
+				{
+					W = MinTooltipWidth;
+				}
+				if (TooltipWidth != W)
+				{
+					TooltipWidth = TooltipWidth * 0.75f + W * 0.25f;
+				}
+
+				const float MaxX = FMath::Max(0.0f, Viewport.Width - TooltipWidth - 21.0f);
+				const float X = FMath::Clamp<float>(MousePosition.X + 18.0f, 0.0f, MaxX);
+
+				constexpr float LineH = 14.0f;
+				constexpr float H = 4 * LineH;
+				const float MaxY = FMath::Max(0.0f, Viewport.Height - H - 18.0f);
+				float Y = FMath::Clamp<float>(MousePosition.Y + 18.0f, 0.0f, MaxY);
+
+				const float Alpha = 1.0f - FMath::Abs(TooltipWidth - W) / W;
+				if (TooltipAlpha < Alpha)
+				{
+					TooltipAlpha = TooltipAlpha * 0.9f + Alpha * 0.1f;
+				}
+				else
+				{
+					TooltipAlpha = Alpha;
+				}
+
+				const FLinearColor BackgroundColor(0.05f, 0.05f, 0.05f, TooltipAlpha);
+				const FLinearColor NameColor(0.9f, 0.9f, 0.5f, TooltipAlpha);
+				const FLinearColor TextColor(0.6f, 0.6f, 0.6f, TooltipAlpha);
+				const FLinearColor ValueColor(1.0f, 1.0f, 1.0f, TooltipAlpha);
+
+				DrawContext.DrawBox(X - 6.0f, Y - 3.0f, TooltipWidth + 12.0f, H + 6.0f, WhiteBrush, BackgroundColor);
+				DrawContext.LayerId++;
+
+				DrawContext.DrawText(X, Y, Name, MainFont, NameColor);
+				Y += LineH;
+
+				const float ValueX = X + 58.0f;
+
+				DrawContext.DrawText(X + 3.0f, Y, TEXT("Incl. Time:"), MainFont, TextColor);
+				FString InclStr = TimeUtils::FormatTimeAuto(HoveredTimingEvent.Duration());
+				DrawContext.DrawText(ValueX, Y, InclStr, MainFont, ValueColor);
+				Y += LineH;
+
+				DrawContext.DrawText(X, Y, TEXT("Excl. Time:"), MainFont, TextColor);
+				FString ExclStr = TimeUtils::FormatTimeAuto(HoveredTimingEvent.ExclusiveTime);
+				DrawContext.DrawText(ValueX, Y, ExclStr, MainFont, ValueColor);
+				Y += LineH;
+
+				DrawContext.DrawText(X + 24.0f, Y, TEXT("Depth:"), MainFont, TextColor);
+				FString DepthStr = FString::Printf(TEXT("%d"), HoveredTimingEvent.Depth);
+				DrawContext.DrawText(ValueX, Y, DepthStr, MainFont, ValueColor);
+				Y += LineH;
+
+				DrawContext.LayerId++;
 			}
-			else
+			else if (HoveredTimingEvent.Track->GetType() == ETimingEventsTrackType::Loading)
 			{
-				TooltipAlpha = Alpha;
-			}
+				FString Name(LoadingGetEventNameFn.Execute(HoveredTimingEvent.Depth, HoveredTimingEvent.LoadingInfo));
 
-			const FLinearColor BackgroundColor(0.05f, 0.05f, 0.05f, TooltipAlpha);
-			const FLinearColor NameColor(0.9f, 0.9f, 0.5f, TooltipAlpha);
-			//const FLinearColor NameColor(0.85f, 0.5f, 0.03f, TooltipAlpha);
-			const FLinearColor TextColor(0.6f, 0.6f, 0.6f, TooltipAlpha);
-			const FLinearColor ValueColor(1.0f, 1.0f, 1.0f, TooltipAlpha);
+				const Trace::FPackageInfo* Package = HoveredTimingEvent.LoadingInfo.Package;
+				const Trace::FPackageExportInfo* Export = HoveredTimingEvent.LoadingInfo.Export;
 
-			DrawContext.DrawBox(X - 6.0f, Y - 3.0f, TooltipWidth + 12.0f, H + 6.0f, WhiteBrush, BackgroundColor);
-			DrawContext.LayerId++;
+				FString PackageName = Package ? Package->Name : TEXT("N/A");
+
+				constexpr float ValueOffsetX = 100.0f;
+				constexpr float MinValueTextWidth = 220.0f;
 
-			DrawContext.DrawText(X, Y, Name, MainFont, NameColor);
-			Y += LineH;
+				float W = FontMeasureService->Measure(Name, MainFont).X;
+
+				float PackageNameW = ValueOffsetX + FontMeasureService->Measure(PackageName, MainFont).X;
+				if (W < PackageNameW)
+				{
+					W = PackageNameW;
+				}
+				if (W < ValueOffsetX + MinValueTextWidth)
+				{
+					W = ValueOffsetX + MinValueTextWidth;
+				}
+				if (W < MinTooltipWidth)
+				{
+					W = MinTooltipWidth;
+				}
+				if (TooltipWidth != W)
+				{
+					TooltipWidth = TooltipWidth * 0.75f + W * 0.25f;
+				}
 
-			const float ValueX = X + 58.0f;
+				const float MaxX = FMath::Max(0.0f, Viewport.Width - TooltipWidth - 21.0f);
+				const float X = FMath::Clamp<float>(MousePosition.X + 18.0f, 0.0f, MaxX);
+				const float ValueX = X + ValueOffsetX;
 
-			DrawContext.DrawText(X + 3.0f, Y, TEXT("Incl. Time:"), MainFont, TextColor);
-			FString InclStr = TimeUtils::FormatTimeAuto(HoveredTimingEvent.Duration());
-			DrawContext.DrawText(ValueX, Y, InclStr, MainFont, ValueColor);
-			Y += LineH;
+				constexpr float LineH = 14.0f;
+				constexpr float H = 10 * LineH;
+				const float MaxY = FMath::Max(0.0f, Viewport.Height - H - 18.0f);
+				float Y = FMath::Clamp<float>(MousePosition.Y + 18.0f, 0.0f, MaxY);
 
-			DrawContext.DrawText(X, Y, TEXT("Excl. Time:"), MainFont, TextColor);
-			FString ExclStr = TimeUtils::FormatTimeAuto(HoveredTimingEvent.ExclusiveTime);
-			DrawContext.DrawText(ValueX, Y, ExclStr, MainFont, ValueColor);
-			Y += LineH;
+				const float Alpha = 1.0f - FMath::Abs(TooltipWidth - W) / W;
+				if (TooltipAlpha < Alpha)
+				{
+					TooltipAlpha = TooltipAlpha * 0.9f + Alpha * 0.1f;
+				}
+				else
+				{
+					TooltipAlpha = Alpha;
+				}
 
-			DrawContext.DrawText(X + 24.0f, Y, TEXT("Depth:"), MainFont, TextColor);
-			FString DepthStr = FString::Printf(TEXT("%d"), HoveredTimingEvent.Depth);
-			DrawContext.DrawText(ValueX, Y, DepthStr, MainFont, ValueColor);
-			Y += LineH;
+				const FLinearColor BackgroundColor(0.05f, 0.05f, 0.05f, TooltipAlpha);
+				const FLinearColor NameColor(0.9f, 0.9f, 0.5f, TooltipAlpha);
+				const FLinearColor TextColor(0.6f, 0.6f, 0.6f, TooltipAlpha);
+				const FLinearColor ValueColor(1.0f, 1.0f, 1.0f, TooltipAlpha);
 
-			DrawContext.LayerId++;
+				DrawContext.DrawBox(X - 6.0f, Y - 3.0f, TooltipWidth + 12.0f, H + 6.0f, WhiteBrush, BackgroundColor);
+				DrawContext.LayerId++;
+
+				DrawContext.DrawText(X, Y, Name, MainFont, NameColor);
+				Y += LineH;
+
+				DrawContext.DrawText(X, Y, TEXT("Duration:"), MainFont, TextColor);
+				FString InclStr = TimeUtils::FormatTimeAuto(HoveredTimingEvent.Duration());
+				DrawContext.DrawText(ValueX, Y, InclStr, MainFont, ValueColor);
+				Y += LineH;
+
+				DrawContext.DrawText(X, Y, TEXT("Depth:"), MainFont, TextColor);
+				FString DepthStr = FString::Printf(TEXT("%d"), HoveredTimingEvent.Depth);
+				DrawContext.DrawText(ValueX, Y, DepthStr, MainFont, ValueColor);
+				Y += LineH;
+
+				DrawContext.DrawText(X, Y, TEXT("Package Type:"), MainFont, TextColor);
+				DrawContext.DrawText(ValueX, Y, GetName(HoveredTimingEvent.LoadingInfo.PackageEventType), MainFont, ValueColor);
+				Y += LineH;
+
+				if (Package)
+				{
+					DrawContext.DrawText(X, Y, TEXT("Package Name:"), MainFont, TextColor);
+					DrawContext.DrawText(ValueX, Y, PackageName, MainFont, ValueColor);
+					Y += LineH;
+
+					DrawContext.DrawText(X, Y, TEXT("Header Size:"), MainFont, TextColor);
+					FString HeaderSizeStr = FString::Printf(TEXT("%d bytes"), Package->Summary.TotalHeaderSize);
+					DrawContext.DrawText(ValueX, Y, HeaderSizeStr, MainFont, ValueColor);
+					Y += LineH;
+
+					DrawContext.DrawText(X, Y, TEXT("Package Summary:"), MainFont, TextColor);
+					FString SummaryStr = FString::Printf(TEXT("%d names, %d imports, %d exports"), Package->Summary.NameCount, Package->Summary.ImportCount, Package->Summary.ExportCount);
+					DrawContext.DrawText(ValueX, Y, SummaryStr, MainFont, ValueColor);
+					Y += LineH;
+				}
+
+				{
+					DrawContext.DrawText(X, Y, TEXT("Export Type:"), MainFont, TextColor);
+					FString ExportTypeStr = FString::Printf(TEXT("%s%s"), GetName(HoveredTimingEvent.LoadingInfo.ExportEventType), Export && Export->IsAsset ? TEXT(" [asset]") : TEXT(""));
+					DrawContext.DrawText(ValueX, Y, ExportTypeStr, MainFont, ValueColor);
+					Y += LineH;
+				}
+
+				if (Export)
+				{
+					DrawContext.DrawText(X, Y, TEXT("Export Class:"), MainFont, TextColor);
+					FString ClassStr = FString::Printf(TEXT("%s"), Export->Class ? Export->Class->Name : TEXT("N/A"));
+					DrawContext.DrawText(ValueX, Y, ClassStr, MainFont, ValueColor);
+					Y += LineH;
+
+					DrawContext.DrawText(X, Y, TEXT("Serial:"), MainFont, TextColor);
+					FString SerialStr = FString::Printf(TEXT("Offset: %llu, Size: %llu%s"), Export->SerialOffset, Export->SerialSize);
+					DrawContext.DrawText(ValueX, Y, SerialStr, MainFont, ValueColor);
+					Y += LineH;
+				}
+
+				DrawContext.LayerId++;
+			}
 		}
 		else
 		{
@@ -994,9 +1184,9 @@ int32 STimingView::OnPaint(const FPaintArgs& Args, const FGeometry& AllottedGeom
 
 ////////////////////////////////////////////////////////////////////////////////////////////////////
 
-FTimingEventsTrack* STimingView::AddTimingEventsTrack(uint32 TrackId, ETimingEventsTrackType TrackType, const FString& TrackName, int32 Order)
+FTimingEventsTrack* STimingView::AddTimingEventsTrack(uint64 TrackId, ETimingEventsTrackType TrackType, const FString& TrackName, const TCHAR* GroupName, int32 Order)
 {
-	FTimingEventsTrack* Track = new FTimingEventsTrack(TrackId, TrackType, TrackName);
+	FTimingEventsTrack* Track = new FTimingEventsTrack(TrackId, TrackType, TrackName, GroupName);
 
 	Track->SetOrder(Order);
 
@@ -1015,7 +1205,7 @@ FTimingEventsTrack* STimingView::AddTimingEventsTrack(uint32 TrackId, ETimingEve
 
 ////////////////////////////////////////////////////////////////////////////////////////////////////
 
-void STimingView::DrawCpuGpuTimelineTrack(FTimingViewDrawHelper& Helper, FTimingEventsTrack& Track) const
+void STimingView::DrawTimingProfilerTrack(FTimingViewDrawHelper& Helper, FTimingEventsTrack& Track) const
 {
 	if (Helper.BeginTimeline(Track))
 	{
@@ -1053,6 +1243,143 @@ void STimingView::DrawCpuGpuTimelineTrack(FTimingViewDrawHelper& Helper, FTiming
 	}
 }
 
+////////////////////////////////////////////////////////////////////////////////////////////////////
+
+const TCHAR* GetName(ELoadTimeProfilerPackageEventType Type)
+{
+	switch (Type)
+	{
+		case LoadTimeProfilerPackageEventType_CreateLinker:				return TEXT("CreateLinker");
+		case LoadTimeProfilerPackageEventType_FinishLinker:				return TEXT("FinishLinker");
+		case LoadTimeProfilerPackageEventType_StartImportPackages:		return TEXT("StartImportPackages");
+		case LoadTimeProfilerPackageEventType_SetupImports:				return TEXT("SetupImports");
+		case LoadTimeProfilerPackageEventType_SetupExports:				return TEXT("SetupExports");
+		case LoadTimeProfilerPackageEventType_ProcessImportsAndExports:	return TEXT("ProcessImportsAndExports");
+		case LoadTimeProfilerPackageEventType_ExportsDone:				return TEXT("ExportsDone");
+		case LoadTimeProfilerPackageEventType_PostLoadWait:				return TEXT("PostLoadWait");
+		case LoadTimeProfilerPackageEventType_StartPostLoad:			return TEXT("StartPostLoad");
+		case LoadTimeProfilerPackageEventType_Tick:						return TEXT("Tick");
+		case LoadTimeProfilerPackageEventType_Finish:					return TEXT("Finish");
+		case LoadTimeProfilerPackageEventType_None:						return TEXT("None");
+		default:														return TEXT("");
+	}
+}
+
+const TCHAR* GetName(ELoadTimeProfilerObjectEventType Type)
+{
+	switch (Type)
+	{
+		case LoadTimeProfilerObjectEventType_Create:	return TEXT("Create");
+		case LoadTimeProfilerObjectEventType_Serialize:	return TEXT("Serialize");
+		case LoadTimeProfilerObjectEventType_PostLoad:	return TEXT("PostLoad");
+		case LoadTimeProfilerObjectEventType_None:		return TEXT("None");
+		default:										return TEXT("");
+	};
+}
+
+const TCHAR* STimingView::GetLoadTimeProfilerEventNameByPackageEventType(uint32 Depth, const Trace::FLoadTimeProfilerCpuEvent& Event) const
+{
+	return GetName(Event.PackageEventType);
+}
+
+const TCHAR* STimingView::GetLoadTimeProfilerEventNameByExportEventType(uint32 Depth, const Trace::FLoadTimeProfilerCpuEvent& Event) const
+{
+	return GetName(Event.ExportEventType);
+}
+
+const TCHAR* STimingView::GetLoadTimeProfilerEventNameByPackageName(uint32 Depth, const Trace::FLoadTimeProfilerCpuEvent& Event) const
+{
+	return Event.Package ?  Event.Package->Name : TEXT("");
+}
+
+const TCHAR* STimingView::GetLoadTimeProfilerEventNameByExportClassName(uint32 Depth, const Trace::FLoadTimeProfilerCpuEvent& Event) const
+{
+	return Event.Export && Event.Export->Class ? Event.Export->Class->Name : TEXT("");
+}
+
+const TCHAR* STimingView::GetLoadTimeProfilerEventName(uint32 Depth, const Trace::FLoadTimeProfilerCpuEvent& Event) const
+{
+	if (Depth == 0)
+	{
+		if (Event.Package)
+		{
+			return Event.Package->Name;
+		}
+	}
+	
+	if (Event.Export && Event.Export->Class)
+	{
+		return Event.Export->Class->Name;
+	}
+
+	return TEXT("");
+}
+
+////////////////////////////////////////////////////////////////////////////////////////////////////
+
+void STimingView::DrawLoadTimeProfilerTrack(FTimingViewDrawHelper& Helper, FTimingEventsTrack& Track) const
+{
+	if (Helper.BeginTimeline(Track))
+	{
+		TSharedPtr<const Trace::IAnalysisSession> Session = FInsightsManager::Get()->GetSession();
+		if (Session.IsValid() && Trace::ReadLoadTimeProfilerProvider(*Session.Get()))
+		{
+			Trace::FAnalysisSessionReadScope SessionReadScope(*Session.Get());
+
+			const Trace::ILoadTimeProfilerProvider& LoadTimeProfilerProvider = *Trace::ReadLoadTimeProfilerProvider(*Session.Get());
+
+			if (&Track == LoadingMainThreadTrack)
+			{
+				LoadTimeProfilerProvider.ReadMainThreadCpuTimeline([this, &Helper, &Track](const Trace::ILoadTimeProfilerProvider::CpuTimeline& Timeline)
+				{
+					if (bUseDownSampling)
+					{
+						const double SecondsPerPixel = 1.0 / Helper.GetViewport().ScaleX;
+						Timeline.EnumerateEventsDownSampled(Helper.GetViewport().StartTime, Helper.GetViewport().EndTime, SecondsPerPixel, [this, &Helper](double StartTime, double EndTime, uint32 Depth, const Trace::FLoadTimeProfilerCpuEvent& Event)
+						{
+							const TCHAR* Name = LoadingGetEventNameFn.Execute(Depth, Event);
+							Helper.AddEvent(StartTime, EndTime, Depth, Name);
+						});
+					}
+					else
+					{
+						Timeline.EnumerateEvents(Helper.GetViewport().StartTime, Helper.GetViewport().EndTime, [this, &Helper](double StartTime, double EndTime, uint32 Depth, const Trace::FLoadTimeProfilerCpuEvent& Event)
+						{
+							const TCHAR* Name = LoadingGetEventNameFn.Execute(Depth, Event);
+							Helper.AddEvent(StartTime, EndTime, Depth, Name);
+						});
+					}
+				});
+			}
+			else if (&Track == LoadingAsyncThreadTrack)
+			{
+				LoadTimeProfilerProvider.ReadAsyncLoadingThreadCpuTimeline([this, &Helper, &Track](const Trace::ILoadTimeProfilerProvider::CpuTimeline& Timeline)
+				{
+					if (bUseDownSampling)
+					{
+						const double SecondsPerPixel = 1.0 / Helper.GetViewport().ScaleX;
+						Timeline.EnumerateEventsDownSampled(Helper.GetViewport().StartTime, Helper.GetViewport().EndTime, SecondsPerPixel, [this, &Helper](double StartTime, double EndTime, uint32 Depth, const Trace::FLoadTimeProfilerCpuEvent& Event)
+						{
+							const TCHAR* Name = LoadingGetEventNameFn.Execute(Depth, Event);
+							Helper.AddEvent(StartTime, EndTime, Depth, Name);
+						});
+					}
+					else
+					{
+						Timeline.EnumerateEvents(Helper.GetViewport().StartTime, Helper.GetViewport().EndTime, [this, &Helper](double StartTime, double EndTime, uint32 Depth, const Trace::FLoadTimeProfilerCpuEvent& Event)
+						{
+							const TCHAR* Name = LoadingGetEventNameFn.Execute(Depth, Event);
+							Helper.AddEvent(StartTime, EndTime, Depth, Name);
+						});
+					}
+				});
+			}
+		}
+
+		Helper.EndTimeline(Track);
+	}
+}
+
 ////////////////////////////////////////////////////////////////////////////////////////////////////
 // "I/O - File Activity" prototype
 ////////////////////////////////////////////////////////////////////////////////////////////////////
@@ -1155,7 +1482,6 @@ void STimingView::DrawIoOverviewTrack(FTimingViewDrawHelper& Helper, FTimingEven
 }
 
 ////////////////////////////////////////////////////////////////////////////////////////////////////
-//PRAGMA_DISABLE_OPTIMIZATION
 
 void STimingView::UpdateIo()
 {
@@ -1327,7 +1653,6 @@ void STimingView::UpdateIo()
 	}
 }
 
-//PRAGMA_ENABLE_OPTIMIZATION
 ////////////////////////////////////////////////////////////////////////////////////////////////////
 
 void STimingView::DrawIoActivityTrack(FTimingViewDrawHelper& Helper, FTimingEventsTrack& Track) const
@@ -1841,7 +2166,7 @@ FReply STimingView::OnMouseWheel(const FGeometry& MyGeometry, const FPointerEven
 	if (MouseEvent.GetModifierKeys().IsShiftDown())
 	{
 		// Scroll vertically.
-		static const float ScrollSpeedY = 16.0f * 3;
+		constexpr float ScrollSpeedY = 16.0f * 3;
 		Viewport.ScrollPosY -= ScrollSpeedY * MouseEvent.GetWheelDelta();
 		UpdateVerticalScrollBar();
 		bIsVerticalViewportDirty = true;
@@ -1856,7 +2181,7 @@ FReply STimingView::OnMouseWheel(const FGeometry& MyGeometry, const FPointerEven
 	{
 		// Zoom in/out horizontally.
 		const double Delta = MouseEvent.GetWheelDelta();
-		static const double ZoomStep = 0.25; // as percent
+		constexpr double ZoomStep = 0.25; // as percent
 		double ScaleX;
 
 		if (Delta > 0)
@@ -2143,12 +2468,21 @@ FReply STimingView::OnKeyDown(const FGeometry& MyGeometry, const FKeyEvent& InKe
 		bAreTimingEventsTracksDirty = true;
 		return FReply::Handled();
 	}
-	else if (InKeyEvent.GetKey() == EKeys::U) // debug: toggles Timing (GPU, CPU) tracks on/off
+	else if (InKeyEvent.GetKey() == EKeys::Y) // debug: toggles GPU track on/off
 	{
 		ShowHideAllGpuTracks_Execute();
+		return FReply::Handled();
+	}
+	else if (InKeyEvent.GetKey() == EKeys::U) // debug: toggles CPU tracks on/off
+	{
 		ShowHideAllCpuTracks_Execute();
 		return FReply::Handled();
 	}
+	else if (InKeyEvent.GetKey() == EKeys::L)  // debug: toggles Loading tracks on/off
+	{
+		ShowHideAllLoadingTracks_Execute();
+		return FReply::Handled();
+	}
 	else if (InKeyEvent.GetKey() == EKeys::I)  // debug: toggles IO tracks on/off
 	{
 		ShowHideAllIoTracks_Execute();
@@ -2160,6 +2494,31 @@ FReply STimingView::OnKeyDown(const FGeometry& MyGeometry, const FKeyEvent& InKe
 		bForceIoEventsUpdate = true;
 		return FReply::Handled();
 	}
+	else if (InKeyEvent.GetKey() == EKeys::One)
+	{
+		LoadingGetEventNameFn = FLoadingTrackGetEventNameDelegate::CreateRaw(this, &STimingView::GetLoadTimeProfilerEventNameByPackageEventType);
+		return FReply::Handled();
+	}
+	else if (InKeyEvent.GetKey() == EKeys::Two)
+	{
+		LoadingGetEventNameFn = FLoadingTrackGetEventNameDelegate::CreateRaw(this, &STimingView::GetLoadTimeProfilerEventNameByExportEventType);
+		return FReply::Handled();
+	}
+	else if (InKeyEvent.GetKey() == EKeys::Three)
+	{
+		LoadingGetEventNameFn = FLoadingTrackGetEventNameDelegate::CreateRaw(this, &STimingView::GetLoadTimeProfilerEventNameByPackageName);
+		return FReply::Handled();
+	}
+	else if (InKeyEvent.GetKey() == EKeys::Four)
+	{
+		LoadingGetEventNameFn = FLoadingTrackGetEventNameDelegate::CreateRaw(this, &STimingView::GetLoadTimeProfilerEventNameByExportClassName);
+		return FReply::Handled();
+	}
+	else if (InKeyEvent.GetKey() == EKeys::Five)
+	{
+		LoadingGetEventNameFn = FLoadingTrackGetEventNameDelegate::CreateRaw(this, &STimingView::GetLoadTimeProfilerEventName);
+		return FReply::Handled();
+	}
 
 	return SCompoundWidget::OnKeyDown(MyGeometry, InKeyEvent);
 }
@@ -2458,8 +2817,8 @@ void STimingView::UpdateHoveredTimingEvent(float MX, float MY)
 
 				const double StartTime = Viewport.SlateUnitsToTime(MX);
 				const double EndTime = StartTime + 2.0 / Viewport.ScaleX; // +2px
-				const bool bStopAtFirstMatch = true; // get first one matching
-				const bool bSearchForLargestEvent = false;
+				constexpr bool bStopAtFirstMatch = true; // get first one matching
+				constexpr bool bSearchForLargestEvent = false;
 				SearchTimingEvent(StartTime, EndTime,
 					[Depth](double, double, uint32 EventDepth, uint32)
 					{
@@ -2480,7 +2839,7 @@ bool STimingView::SearchTimingEvent(const double InStartTime,
 									TFunctionRef<bool(double, double, uint32, uint32)> InPredicate,
 									FTimingEvent& InOutTimingEvent,
 									bool bInStopAtFirstMatch,
-									bool bInSearchForLargestEvent)
+									bool bInSearchForLargestEvent) const
 {
 	struct FSearchTimingEventContext
 	{
@@ -2506,6 +2865,46 @@ bool STimingView::SearchTimingEvent(const double InStartTime,
 			, LargestDuration(-1.0)
 		{
 		}
+
+		void CheckEvent(double EventStartTime, double EventEndTime, uint32 EventDepth, const Trace::FTimingProfilerEvent& Event)
+		{
+			if (bContinueSearching && Predicate(EventStartTime, EventEndTime, EventDepth, Event.TimerIndex))
+			{
+				if (!bSearchForLargestEvent || EventEndTime - EventStartTime > LargestDuration)
+				{
+					LargestDuration = EventEndTime - EventStartTime;
+
+					TimingEvent.TypeId = Event.TimerIndex;
+					TimingEvent.Depth = EventDepth;
+					TimingEvent.StartTime = EventStartTime;
+					TimingEvent.EndTime = EventEndTime;
+
+					bFound = true;
+					bContinueSearching = !bStopAtFirstMatch || bSearchForLargestEvent;
+				}
+			}
+		}
+
+		void CheckEvent(double EventStartTime, double EventEndTime, uint32 EventDepth, const Trace::FLoadTimeProfilerCpuEvent& Event)
+		{
+			if (bContinueSearching && Predicate(EventStartTime, EventEndTime, EventDepth, 0))
+			{
+				if (!bSearchForLargestEvent || EventEndTime - EventStartTime > LargestDuration)
+				{
+					LargestDuration = EventEndTime - EventStartTime;
+
+					TimingEvent.TypeId = 0;
+					TimingEvent.Depth = EventDepth;
+					TimingEvent.StartTime = EventStartTime;
+					TimingEvent.EndTime = EventEndTime;
+
+					TimingEvent.LoadingInfo = Event;
+
+					bFound = true;
+					bContinueSearching = !bStopAtFirstMatch || bSearchForLargestEvent;
+				}
+			}
+		}
 	};
 
 	FSearchTimingEventContext Ctx(InStartTime, InEndTime, InPredicate, InOutTimingEvent, bInStopAtFirstMatch, bInSearchForLargestEvent);
@@ -2515,83 +2914,102 @@ bool STimingView::SearchTimingEvent(const double InStartTime,
 	{
 		Trace::FAnalysisSessionReadScope SessionReadScope(*Session.Get());
 
-		const Trace::ITimingProfilerProvider& TimingProfilerProvider = *Trace::ReadTimingProfilerProvider(*Session.Get());
+		const FTimingEventsTrack* Track = Ctx.TimingEvent.Track;
 
-		TimingProfilerProvider.ReadTimeline(Ctx.TimingEvent.Track->GetId(), [&Ctx](const Trace::ITimingProfilerProvider::Timeline& Timeline)
+		if (Track->GetType() == ETimingEventsTrackType::Cpu ||
+			Track->GetType() == ETimingEventsTrackType::Gpu)
 		{
-			Timeline.EnumerateEvents(Ctx.StartTime, Ctx.EndTime, [&Ctx](double EventStartTime, double EventEndTime, uint32 EventDepth, const Trace::FTimingProfilerEvent& Event)
+			const Trace::ITimingProfilerProvider& TimingProfilerProvider = *Trace::ReadTimingProfilerProvider(*Session.Get());
+
+			TimingProfilerProvider.ReadTimeline(Track->GetId(), [&Ctx](const Trace::ITimingProfilerProvider::Timeline& Timeline)
 			{
-				if (Ctx.bContinueSearching && Ctx.Predicate(EventStartTime, EventEndTime, EventDepth, Event.TimerIndex))
+				Timeline.EnumerateEvents(Ctx.StartTime, Ctx.EndTime, [&Ctx](double EventStartTime, double EventEndTime, uint32 EventDepth, const Trace::FTimingProfilerEvent& Event)
 				{
-					if (!Ctx.bSearchForLargestEvent || EventEndTime - EventStartTime > Ctx.LargestDuration)
-					{
-						Ctx.LargestDuration = EventEndTime - EventStartTime;
-
-						Ctx.TimingEvent.TypeId = Event.TimerIndex;
-						Ctx.TimingEvent.Depth = EventDepth;
-						Ctx.TimingEvent.StartTime = EventStartTime;
-						Ctx.TimingEvent.EndTime = EventEndTime;
-
-						Ctx.bFound = true;
-						Ctx.bContinueSearching = !Ctx.bStopAtFirstMatch || Ctx.bSearchForLargestEvent;
-					}
-				}
+					Ctx.CheckEvent(EventStartTime, EventEndTime, EventDepth, Event);
+				});
 			});
-		});
 
-		if (Ctx.bFound)
-		{
-			TimingProfilerProvider.ReadTimeline(Ctx.TimingEvent.Track->GetId(), [&Ctx](const Trace::ITimingProfilerProvider::Timeline& Timeline)
+			if (Ctx.bFound)
 			{
-				struct FEnumerationState
-				{
-					double EventStartTime;
-					double EventEndTime;
-					uint64 EventDepth;
-
-					uint64 CurrentDepth;
-					double LastTime;
-					double ExclusiveTime;
-					bool IsInEventScope;
-				};
-				FEnumerationState State;
-
-				State.EventStartTime = Ctx.TimingEvent.StartTime;
-				State.EventEndTime = Ctx.TimingEvent.EndTime;
-				State.EventDepth = Ctx.TimingEvent.Depth;
-
-				State.CurrentDepth = 0;
-				State.LastTime = 0.0;
-				State.ExclusiveTime = 0.0;
-				State.IsInEventScope = false;
-
-				Timeline.EnumerateEvents(Ctx.TimingEvent.StartTime, Ctx.TimingEvent.EndTime, [&State](bool IsEnter, double Time, const Trace::FTimingProfilerEvent& Event)
+				// Compute Exclusive Time.
+				TimingProfilerProvider.ReadTimeline(Track->GetId(), [&Ctx](const Trace::ITimingProfilerProvider::Timeline& Timeline)
 				{
-					if (IsEnter)
+					struct FEnumerationState
 					{
-						if (State.IsInEventScope && State.CurrentDepth == State.EventDepth + 1)
+						double EventStartTime;
+						double EventEndTime;
+						uint64 EventDepth;
+
+						uint64 CurrentDepth;
+						double LastTime;
+						double ExclusiveTime;
+						bool IsInEventScope;
+					};
+					FEnumerationState State;
+
+					State.EventStartTime = Ctx.TimingEvent.StartTime;
+					State.EventEndTime = Ctx.TimingEvent.EndTime;
+					State.EventDepth = Ctx.TimingEvent.Depth;
+
+					State.CurrentDepth = 0;
+					State.LastTime = 0.0;
+					State.ExclusiveTime = 0.0;
+					State.IsInEventScope = false;
+
+					Timeline.EnumerateEvents(Ctx.TimingEvent.StartTime, Ctx.TimingEvent.EndTime, [&State](bool IsEnter, double Time, const Trace::FTimingProfilerEvent& Event)
+					{
+						if (IsEnter)
 						{
-							State.ExclusiveTime += Time - State.LastTime;
+							if (State.IsInEventScope && State.CurrentDepth == State.EventDepth + 1)
+							{
+								State.ExclusiveTime += Time - State.LastTime;
+							}
+							if (State.CurrentDepth == State.EventDepth && Time == State.EventStartTime)
+							{
+								State.IsInEventScope = true;
+							}
+							++State.CurrentDepth;
 						}
-						if (State.CurrentDepth == State.EventDepth && Time == State.EventStartTime)
+						else
 						{
-							State.IsInEventScope = true;
+							--State.CurrentDepth;
+							if (State.CurrentDepth == State.EventDepth && Time == State.EventEndTime)
+							{
+								State.IsInEventScope = false;
+								State.ExclusiveTime += Time - State.LastTime;
+							}
 						}
-						++State.CurrentDepth;
-					}
-					else
+						State.LastTime = Time;
+					});
+
+					Ctx.TimingEvent.ExclusiveTime = State.ExclusiveTime;
+				});
+			}
+		}
+		else if (Track->GetType() == ETimingEventsTrackType::Loading)
+		{
+			const Trace::ILoadTimeProfilerProvider& LoadTimeProfilerProvider = *Trace::ReadLoadTimeProfilerProvider(*Session.Get());
+
+			if (Track == LoadingMainThreadTrack)
+			{
+				LoadTimeProfilerProvider.ReadMainThreadCpuTimeline([&Ctx](const Trace::ILoadTimeProfilerProvider::CpuTimeline& Timeline)
+				{
+					Timeline.EnumerateEvents(Ctx.StartTime, Ctx.EndTime, [&Ctx](double EventStartTime, double EventEndTime, uint32 EventDepth, const Trace::FLoadTimeProfilerCpuEvent& Event)
 					{
-						--State.CurrentDepth;
-						if (State.CurrentDepth == State.EventDepth && Time == State.EventEndTime)
-						{
-							State.IsInEventScope = false;
-							State.ExclusiveTime += Time - State.LastTime;
-						}
-					}
-					State.LastTime = Time;
+						Ctx.CheckEvent(EventStartTime, EventEndTime, EventDepth, Event);
+					});
 				});
-				Ctx.TimingEvent.ExclusiveTime = State.ExclusiveTime;
-			});
+			}
+			else if (Track == LoadingAsyncThreadTrack)
+			{
+				LoadTimeProfilerProvider.ReadAsyncLoadingThreadCpuTimeline([&Ctx](const Trace::ILoadTimeProfilerProvider::CpuTimeline& Timeline)
+				{
+					Timeline.EnumerateEvents(Ctx.StartTime, Ctx.EndTime, [&Ctx](double EventStartTime, double EventEndTime, uint32 EventDepth, const Trace::FLoadTimeProfilerCpuEvent& Event)
+					{
+						Ctx.CheckEvent(EventStartTime, EventEndTime, EventDepth, Event);
+					});
+				});
+			}
 		}
 	}
 
@@ -2834,8 +3252,8 @@ TSharedRef<SWidget> STimingView::MakeTracksFilterMenu()
 	MenuBuilder.BeginSection("QuickFilter", LOCTEXT("TracksFilterHeading", "Quick Filter"));
 	{
 		MenuBuilder.AddMenuEntry(
-			LOCTEXT("ShowAllGpuTracks", "Show/Hide All GPU Tracks (U)"),
-			LOCTEXT("ShowAllGpuTracks_Tooltip", "Show/hide all GPU tracks"),
+			LOCTEXT("ShowAllGpuTracks", "GPU Track (Y)"),
+			LOCTEXT("ShowAllGpuTracks_Tooltip", "Show/hide the GPU track"),
 			FSlateIcon(),
 			FUIAction(FExecuteAction::CreateSP(this, &STimingView::ShowHideAllGpuTracks_Execute),
 				FCanExecuteAction::CreateLambda([] { return true; }),
@@ -2845,8 +3263,8 @@ TSharedRef<SWidget> STimingView::MakeTracksFilterMenu()
 		);
 
 		MenuBuilder.AddMenuEntry(
-			LOCTEXT("ShowAllCpuTracks", "Show/Hide All CPU Tracks (U)"),
-			LOCTEXT("ShowAllCpuTracks_Tooltip", "Show/hide all CPU tracks"),
+			LOCTEXT("ShowAllCpuTracks", "CPU Thread Tracks (U)"),
+			LOCTEXT("ShowAllCpuTracks_Tooltip", "Show/hide all CPU tracks (and all CPU thread groups)"),
 			FSlateIcon(),
 			FUIAction(FExecuteAction::CreateSP(this, &STimingView::ShowHideAllCpuTracks_Execute),
 				FCanExecuteAction::CreateLambda([] { return true; }),
@@ -2856,8 +3274,19 @@ TSharedRef<SWidget> STimingView::MakeTracksFilterMenu()
 		);
 
 		MenuBuilder.AddMenuEntry(
-			LOCTEXT("ShowAllIoTracks", "Show/Hide All I/O Tracks (I)"),
-			LOCTEXT("ShowAllIoTracks_Tooltip", "Show/hide all I/O tracks"),
+			LOCTEXT("ShowAllLoadingTracks", "Asset Loading Tracks (L)"),
+			LOCTEXT("ShowAllLoadingTracks_Tooltip", "Show/hide the Asset Loading tracks"),
+			FSlateIcon(),
+			FUIAction(FExecuteAction::CreateSP(this, &STimingView::ShowHideAllLoadingTracks_Execute),
+				FCanExecuteAction::CreateLambda([] { return true; }),
+				FIsActionChecked::CreateSP(this, &STimingView::ShowHideAllLoadingTracks_IsChecked)),
+			NAME_None,
+			EUserInterfaceActionType::ToggleButton
+		);
+
+		MenuBuilder.AddMenuEntry(
+			LOCTEXT("ShowAllIoTracks", "I/O Tracks (I)"),
+			LOCTEXT("ShowAllIoTracks_Tooltip", "Show/hide the I/O (File Activity) tracks"),
 			FSlateIcon(),
 			FUIAction(FExecuteAction::CreateSP(this, &STimingView::ShowHideAllIoTracks_Execute),
 				FCanExecuteAction::CreateLambda([] { return true; }),
@@ -2879,15 +3308,42 @@ TSharedRef<SWidget> STimingView::MakeTracksFilterMenu()
 	}
 	MenuBuilder.EndSection();
 
-	MenuBuilder.BeginSection("Tracks", LOCTEXT("TracksHeading", "Tracks"));
-	CreateTracksMenu(MenuBuilder);
+	MenuBuilder.BeginSection("ThreadGroups", LOCTEXT("ThreadGroupsHeading", "CPU Thread Groups"));
+	CreateThreadGroupsMenu(MenuBuilder);
 	MenuBuilder.EndSection();
 
+	//MenuBuilder.BeginSection("Tracks", LOCTEXT("TracksHeading", "Tracks"));
+	//CreateTracksMenu(MenuBuilder);
+	//MenuBuilder.EndSection();
+
 	return MenuBuilder.MakeWidget();
 }
 
 ////////////////////////////////////////////////////////////////////////////////////////////////////
 
+void STimingView::CreateThreadGroupsMenu(FMenuBuilder& MenuBuilder)
+{
+	for (const auto& KV : ThreadGroups)
+	{
+		const FThreadGroup& ThreadGroup = KV.Value;
+		if (ThreadGroup.NumTimelines > 0)
+		{
+			MenuBuilder.AddMenuEntry(
+				FText::FromString(ThreadGroup.Name),
+				TAttribute<FText>(), // no tooltip
+				FSlateIcon(),
+				FUIAction(FExecuteAction::CreateSP(this, &STimingView::ToggleTrackVisibilityByGroup_Execute, ThreadGroup.Name),
+					FCanExecuteAction::CreateLambda([] { return true; }),
+					FIsActionChecked::CreateSP(this, &STimingView::ToggleTrackVisibilityByGroup_IsChecked, ThreadGroup.Name)),
+				NAME_None,
+				EUserInterfaceActionType::ToggleButton
+			);
+		}
+	}
+}
+
+////////////////////////////////////////////////////////////////////////////////////////////////////
+
 void STimingView::CreateTracksMenu(FMenuBuilder& MenuBuilder)
 {
 	for (int32 TrackIndex = 0; TrackIndex < TimingEventsTracks.Num(); ++TrackIndex)
@@ -2931,6 +3387,11 @@ void STimingView::ShowHideAllCpuTracks_Execute()
 		}
 	}
 
+	for (auto& KV: ThreadGroups)
+	{
+		KV.Value.bIsVisible = bShowHideAllCpuTracks;
+	}
+
 	HoveredTimingEvent.Reset();
 	SelectedTimingEvent.Reset();
 
@@ -2967,6 +3428,31 @@ void STimingView::ShowHideAllGpuTracks_Execute()
 
 ////////////////////////////////////////////////////////////////////////////////////////////////////
 
+bool STimingView::ShowHideAllLoadingTracks_IsChecked() const
+{
+	return bShowHideAllLoadingTracks;
+}
+
+////////////////////////////////////////////////////////////////////////////////////////////////////
+
+void STimingView::ShowHideAllLoadingTracks_Execute()
+{
+	bShowHideAllLoadingTracks = !bShowHideAllLoadingTracks;
+
+	for (int32 TrackIndex = 0; TrackIndex < TimingEventsTracks.Num(); ++TrackIndex)
+	{
+		FTimingEventsTrack& Track = *TimingEventsTracks[TrackIndex];
+		if (Track.GetType() == ETimingEventsTrackType::Loading)
+		{
+			Track.SetVisibilityFlag(bShowHideAllLoadingTracks);
+		}
+	}
+
+	bAreTimingEventsTracksDirty = true;
+}
+
+////////////////////////////////////////////////////////////////////////////////////////////////////
+
 bool STimingView::ShowHideAllIoTracks_IsChecked() const
 {
 	return bShowHideAllIoTracks;
@@ -3055,4 +3541,40 @@ void STimingView::ToggleTrackVisibility_Execute(uint64 InTrackId)
 
 ////////////////////////////////////////////////////////////////////////////////////////////////////
 
+bool STimingView::ToggleTrackVisibilityByGroup_IsChecked(const TCHAR* InGroupName) const
+{
+	if (ThreadGroups.Contains(InGroupName))
+	{
+		const FThreadGroup& ThreadGroup = ThreadGroups[InGroupName];
+		return ThreadGroup.bIsVisible;
+	}
+	return false;
+}
+
+////////////////////////////////////////////////////////////////////////////////////////////////////
+
+void STimingView::ToggleTrackVisibilityByGroup_Execute(const TCHAR* InGroupName)
+{
+	if (ThreadGroups.Contains(InGroupName))
+	{
+		FThreadGroup& ThreadGroup = ThreadGroups[InGroupName];
+		ThreadGroup.bIsVisible = !ThreadGroup.bIsVisible;
+
+		for (auto& KV : CachedTimelines)
+		{
+			FTimingEventsTrack& Track = *KV.Value;
+			if (Track.GetType() == ETimingEventsTrackType::Cpu &&
+				Track.GetGroupName() == InGroupName)
+			{
+				Track.SetVisibilityFlag(ThreadGroup.bIsVisible);
+			}
+		}
+
+		HoveredTimingEvent.Reset();
+		SelectedTimingEvent.Reset();
+	}
+}
+
+////////////////////////////////////////////////////////////////////////////////////////////////////
+
 #undef LOCTEXT_NAMESPACE
diff --git a/Engine/Source/Developer/TraceInsights/Private/Insights/Widgets/STimingView.h b/Engine/Source/Developer/TraceInsights/Private/Insights/Widgets/STimingView.h
--- a/Engine/Source/Developer/TraceInsights/Private/Insights/Widgets/STimingView.h
+++ b/Engine/Source/Developer/TraceInsights/Private/Insights/Widgets/STimingView.h
@@ -35,9 +35,15 @@ DECLARE_DELEGATE_RetVal(int32, FHoveredEventChangedDelegate);
 /** The delegate to be invoked when the selected timing event has changed. Returns id of timing event or -1 (if no event is hovered). */
 DECLARE_DELEGATE_RetVal(int32, FSelectedEventChangedDelegate);
 
+/** Defines FLoadingTrackGetEventNameDelegate delegate interface. Returns the name for a timing event in a Loading track. */
+DECLARE_DELEGATE_RetVal_TwoParams(const TCHAR*, FLoadingTrackGetEventNameDelegate, uint32 /*Depth*/, const Trace::FLoadTimeProfilerCpuEvent& /*Event*/);
+
 /** A custom widget used to display timing events. */
 class STimingView : public SCompoundWidget
 {
+private:
+	static constexpr float MinTooltipWidth = 110.0f;
+
 public:
 	/** Default constructor. */
 	STimingView();
@@ -66,6 +72,7 @@ public:
 	void Construct(const FArguments& InArgs);
 
 	TSharedRef<SWidget> MakeTracksFilterMenu();
+	void CreateThreadGroupsMenu(FMenuBuilder& MenuBuilder);
 	void CreateTracksMenu(FMenuBuilder& MenuBuilder);
 
 	bool ShowHideAllCpuTracks_IsChecked() const;
@@ -74,6 +81,9 @@ public:
 	bool ShowHideAllGpuTracks_IsChecked() const;
 	void ShowHideAllGpuTracks_Execute();
 
+	bool ShowHideAllLoadingTracks_IsChecked() const;
+	void ShowHideAllLoadingTracks_Execute();
+
 	bool ShowHideAllIoTracks_IsChecked() const;
 	void ShowHideAllIoTracks_Execute();
 
@@ -83,6 +93,9 @@ public:
 	bool ToggleTrackVisibility_IsChecked(uint64 InTrackId) const;
 	void ToggleTrackVisibility_Execute(uint64 InTrackId);
 
+	bool ToggleTrackVisibilityByGroup_IsChecked(const TCHAR* InGroupName) const;
+	void ToggleTrackVisibilityByGroup_Execute(const TCHAR* InGroupName);
+
 	/** Resets internal widget's data to the default one. */
 	void Reset();
 
@@ -245,9 +258,16 @@ protected:
 		return FVector2D(16.0f, 16.0f);
 	}
 
-	FTimingEventsTrack* AddTimingEventsTrack(uint32 TrackId, ETimingEventsTrackType TrackType, const FString& TrackName, int32 Order);
+	FTimingEventsTrack* AddTimingEventsTrack(uint64 TrackId, ETimingEventsTrackType TrackType, const FString& TrackName, const TCHAR* GroupName, int32 Order);
+
+	void DrawTimingProfilerTrack(FTimingViewDrawHelper& Helper, FTimingEventsTrack& Track) const;
 
-	void DrawCpuGpuTimelineTrack(FTimingViewDrawHelper& Helper, FTimingEventsTrack& Track) const;
+	const TCHAR* GetLoadTimeProfilerEventNameByPackageEventType(uint32 Depth, const Trace::FLoadTimeProfilerCpuEvent& Event) const;
+	const TCHAR* GetLoadTimeProfilerEventNameByExportEventType(uint32 Depth, const Trace::FLoadTimeProfilerCpuEvent& Event) const;
+	const TCHAR* GetLoadTimeProfilerEventNameByPackageName(uint32 Depth, const Trace::FLoadTimeProfilerCpuEvent& Event) const;
+	const TCHAR* GetLoadTimeProfilerEventNameByExportClassName(uint32 Depth, const Trace::FLoadTimeProfilerCpuEvent& Event) const;
+	const TCHAR* GetLoadTimeProfilerEventName(uint32 Depth, const Trace::FLoadTimeProfilerCpuEvent& Event) const;
+	void DrawLoadTimeProfilerTrack(FTimingViewDrawHelper& Helper, FTimingEventsTrack& Track) const;
 
 	void UpdateIo();
 	void DrawIoOverviewTrack(FTimingViewDrawHelper& Helper, FTimingEventsTrack& Track) const;
@@ -282,12 +302,12 @@ protected:
 
 	void UpdateHoveredTimingEvent(float MX, float MY);
 
-	static bool SearchTimingEvent(const double InStartTime,
-								  const double InEndTime,
-								  TFunctionRef<bool(double, double, uint32, uint32)> InPredicate,
-								  FTimingEvent& InOutTimingEvent,
-								  bool bInStopAtFirstMatch,
-								  bool bInSearchForLargestEvent);
+	bool SearchTimingEvent(const double InStartTime,
+						   const double InEndTime,
+						   TFunctionRef<bool(double, double, uint32, uint32)> InPredicate,
+						   FTimingEvent& InOutTimingEvent,
+						   bool bInStopAtFirstMatch,
+						   bool bInSearchForLargestEvent) const;
 
 	void OnSelectedTimingEventChanged();
 	void SelectHoveredTimingEvent();
@@ -326,6 +346,7 @@ protected:
 protected:
 	bool bShowHideAllGpuTracks;
 	bool bShowHideAllCpuTracks;
+	bool bShowHideAllLoadingTracks;
 	bool bShowHideAllIoTracks;
 
 	////////////////////////////////////////////////////////////
@@ -338,18 +359,16 @@ protected:
 
 	////////////////////////////////////////////////////////////
 
-	//TODO: enum class FTimingTrackType
-	//{
-	//	TimeRuler,
-	//	Markers,
-	//	CpuEvents,
-	//	GpuEvents,
-	//	IoEvents,
-	//	Graph,
-	//	Other
-	//};
+	struct FThreadGroup
+	{
+		const TCHAR* Name;
+		bool bIsVisible;
+		uint32 NumTimelines;
+	};
+
+	mutable TMap<const TCHAR*, FThreadGroup> ThreadGroups; /**< thread groups */
 
-	mutable TMap<uint32, FTimingEventsTrack*> CachedTimelines; /**< all tracks */
+	mutable TMap<uint64, FTimingEventsTrack*> CachedTimelines; /**< all tracks */
 
 	//TODO: TArray<FBaseTimingTrack*> TopTracks; /**< tracks docked on top, in order to be displayed (top to bottom) */
 	//TODO: TArray<FBaseTimingTrack*> ScrollableTracks; /**< tracks in scrollable area, in order to be displayed (top to bottom) */
@@ -366,7 +385,15 @@ protected:
 	bool bUseDownSampling;
 
 	////////////////////////////////////////////////////////////
-	// IO
+	// Load time
+
+	FTimingEventsTrack* LoadingMainThreadTrack;
+	FTimingEventsTrack* LoadingAsyncThreadTrack;
+
+	FLoadingTrackGetEventNameDelegate LoadingGetEventNameFn;
+
+	////////////////////////////////////////////////////////////
+	// File activity (I/O)
 
 	FTimingEventsTrack* IoOverviewTrack;
 	FTimingEventsTrack* IoActivityTrack;
