[Insights] Fixed UE-75236 Unreal Insights - Timing Insights Menu does not communicate window status with toolbar and vice versa.

diff --git a/Engine/Source/Developer/TraceInsights/Private/Insights/TimingProfilerCommands.cpp b/Engine/Source/Developer/TraceInsights/Private/Insights/TimingProfilerCommands.cpp
--- a/Engine/Source/Developer/TraceInsights/Private/Insights/TimingProfilerCommands.cpp
+++ b/Engine/Source/Developer/TraceInsights/Private/Insights/TimingProfilerCommands.cpp
@@ -98,12 +98,12 @@ PRAGMA_ENABLE_OPTIMIZATION
 		return b##IsEnabled ? ECheckBoxState::Checked : ECheckBoxState::Unchecked;\
 	}
 
-IMPLEMENT_TOGGLE_COMMAND(ToggleFramesTrackVisibility, IsFramesTrackVisible, SetFramesTrackVisible)
-IMPLEMENT_TOGGLE_COMMAND(ToggleGraphTrackVisibility, IsGraphTrackVisible, SetGraphTrackVisible)
-IMPLEMENT_TOGGLE_COMMAND(ToggleTimingViewVisibility, IsTimingViewVisible, SetTimingViewVisible)
-IMPLEMENT_TOGGLE_COMMAND(ToggleTimersViewVisibility, IsTimersViewVisible, SetTimersViewVisible)
-IMPLEMENT_TOGGLE_COMMAND(ToggleStatsCountersViewVisibility, IsStatsCountersViewVisible, SetStatsCountersViewVisible)
-IMPLEMENT_TOGGLE_COMMAND(ToggleLogViewVisibility, IsLogViewVisible, SetLogViewVisible)
+IMPLEMENT_TOGGLE_COMMAND(ToggleFramesTrackVisibility, IsFramesTrackVisible, ShowHideFramesTrack)
+IMPLEMENT_TOGGLE_COMMAND(ToggleGraphTrackVisibility, IsGraphTrackVisible, ShowHideGraphTrack)
+IMPLEMENT_TOGGLE_COMMAND(ToggleTimingViewVisibility, IsTimingViewVisible, ShowHideTimingView)
+IMPLEMENT_TOGGLE_COMMAND(ToggleTimersViewVisibility, IsTimersViewVisible, ShowHideTimersView)
+IMPLEMENT_TOGGLE_COMMAND(ToggleStatsCountersViewVisibility, IsStatsCountersViewVisible, ShowHideStatsCountersView)
+IMPLEMENT_TOGGLE_COMMAND(ToggleLogViewVisibility, IsLogViewVisible, ShowHideLogView)
 
 #undef IMPLEMENT_TOGGLE_COMMAND
 
diff --git a/Engine/Source/Developer/TraceInsights/Private/Insights/TimingProfilerManager.cpp b/Engine/Source/Developer/TraceInsights/Private/Insights/TimingProfilerManager.cpp
--- a/Engine/Source/Developer/TraceInsights/Private/Insights/TimingProfilerManager.cpp
+++ b/Engine/Source/Developer/TraceInsights/Private/Insights/TimingProfilerManager.cpp
@@ -154,7 +154,7 @@ void FTimingProfilerManager::OnSessionChanged()
 
 ////////////////////////////////////////////////////////////////////////////////////////////////////
 
-void FTimingProfilerManager::SetFramesTrackVisible(const bool bFramesTrackVisibleState)
+void FTimingProfilerManager::ShowHideFramesTrack(const bool bFramesTrackVisibleState)
 {
 	bool bWasFramesTrackVisible = bIsFramesTrackVisible;
 
@@ -169,7 +169,7 @@ void FTimingProfilerManager::SetFramesTrackVisible(const bool bFramesTrackVisibl
 
 ////////////////////////////////////////////////////////////////////////////////////////////////////
 
-void FTimingProfilerManager::SetGraphTrackVisible(const bool bGraphTrackVisibleState)
+void FTimingProfilerManager::ShowHideGraphTrack(const bool bGraphTrackVisibleState)
 {
 	bool bWasGraphTrackVisible = bIsGraphTrackVisible;
 
@@ -184,7 +184,7 @@ void FTimingProfilerManager::SetGraphTrackVisible(const bool bGraphTrackVisibleS
 
 ////////////////////////////////////////////////////////////////////////////////////////////////////
 
-void FTimingProfilerManager::SetTimingViewVisible(const bool bTimingViewVisibleState)
+void FTimingProfilerManager::ShowHideTimingView(const bool bTimingViewVisibleState)
 {
 	bool bWasTimingViewVisible = bIsTimingViewVisible;
 
@@ -199,7 +199,7 @@ void FTimingProfilerManager::SetTimingViewVisible(const bool bTimingViewVisibleS
 
 ////////////////////////////////////////////////////////////////////////////////////////////////////
 
-void FTimingProfilerManager::SetTimersViewVisible(const bool bTimersViewVisibleState)
+void FTimingProfilerManager::ShowHideTimersView(const bool bTimersViewVisibleState)
 {
 	bool bWasTimersViewVisible = bIsTimersViewVisible;
 
@@ -214,7 +214,7 @@ void FTimingProfilerManager::SetTimersViewVisible(const bool bTimersViewVisibleS
 
 ////////////////////////////////////////////////////////////////////////////////////////////////////
 
-void FTimingProfilerManager::SetStatsCountersViewVisible(const bool bStatsCountersViewVisibleState)
+void FTimingProfilerManager::ShowHideStatsCountersView(const bool bStatsCountersViewVisibleState)
 {
 	bool bWasStatsCountersViewVisible = bIsStatsCountersViewVisible;
 
@@ -229,7 +229,7 @@ void FTimingProfilerManager::SetStatsCountersViewVisible(const bool bStatsCounte
 
 ////////////////////////////////////////////////////////////////////////////////////////////////////
 
-void FTimingProfilerManager::SetLogViewVisible(const bool bLogViewVisibleState)
+void FTimingProfilerManager::ShowHideLogView(const bool bLogViewVisibleState)
 {
 	bool bWasLogViewVisible = bIsLogViewVisible;
 
diff --git a/Engine/Source/Developer/TraceInsights/Private/Insights/TimingProfilerManager.h b/Engine/Source/Developer/TraceInsights/Private/Insights/TimingProfilerManager.h
--- a/Engine/Source/Developer/TraceInsights/Private/Insights/TimingProfilerManager.h
+++ b/Engine/Source/Developer/TraceInsights/Private/Insights/TimingProfilerManager.h
@@ -99,27 +99,33 @@ public:
 
 	/** @return true, if Frames Track is visible */
 	const bool IsFramesTrackVisible() const { return bIsFramesTrackVisible; }
-	void SetFramesTrackVisible(const bool bFramesTrackVisibleState);
+	void SetFramesTrackVisible(const bool bFramesTrackVisibleState) { bIsFramesTrackVisible = bFramesTrackVisibleState; }
+	void ShowHideFramesTrack(const bool bFramesTrackVisibleState);
 
 	/** @return true, if Graph Track is visible */
 	const bool IsGraphTrackVisible() const { return bIsGraphTrackVisible; }
-	void SetGraphTrackVisible(const bool bGraphTrackVisibleState);
+	void SetGraphTrackVisible(const bool bGraphTrackVisibleState) { bIsGraphTrackVisible = bGraphTrackVisibleState; }
+	void ShowHideGraphTrack(const bool bGraphTrackVisibleState);
 
 	/** @return true, if Timing View is visible */
 	const bool IsTimingViewVisible() const { return bIsTimingViewVisible; }
-	void SetTimingViewVisible(const bool bTimingViewVisibleState);
+	void SetTimingViewVisible(const bool bTimingViewVisibleState) { bIsTimingViewVisible = bTimingViewVisibleState; }
+	void ShowHideTimingView(const bool bTimingViewVisibleState);
 
 	/** @return true, if Timers View is visible */
 	const bool IsTimersViewVisible() const { return bIsTimersViewVisible; }
-	void SetTimersViewVisible(const bool bTimersViewVisibleState);
+	void SetTimersViewVisible(const bool bTimersViewVisibleState) { bIsTimersViewVisible = bTimersViewVisibleState; }
+	void ShowHideTimersView(const bool bTimersViewVisibleState);
 
 	/** @return true, if Stats Counters View is visible */
 	const bool IsStatsCountersViewVisible() const { return bIsStatsCountersViewVisible; }
-	void SetStatsCountersViewVisible(const bool bStatsCountersViewVisibleState);
+	void SetStatsCountersViewVisible(const bool bStatsCountersViewVisibleState) { bIsStatsCountersViewVisible = bStatsCountersViewVisibleState; }
+	void ShowHideStatsCountersView(const bool bStatsCountersViewVisibleState);
 
 	/** @return true, if Log View is visible */
 	const bool IsLogViewVisible() const { return bIsLogViewVisible; }
-	void SetLogViewVisible(const bool bLogViewVisibleState);
+	void SetLogViewVisible(const bool bLogViewVisibleState) { bIsLogViewVisible = bLogViewVisibleState; }
+	void ShowHideLogView(const bool bLogViewVisibleState);
 
 	////////////////////////////////////////////////////////////////////////////////////////////////////
 
diff --git a/Engine/Source/Developer/TraceInsights/Private/Insights/Widgets/STimingProfilerWindow.cpp b/Engine/Source/Developer/TraceInsights/Private/Insights/Widgets/STimingProfilerWindow.cpp
--- a/Engine/Source/Developer/TraceInsights/Private/Insights/Widgets/STimingProfilerWindow.cpp
+++ b/Engine/Source/Developer/TraceInsights/Private/Insights/Widgets/STimingProfilerWindow.cpp
@@ -88,6 +88,8 @@ TSharedRef<SDockTab> STimingProfilerWindow::SpawnTab_Toolbar(const FSpawnTabArgs
 
 TSharedRef<SDockTab> STimingProfilerWindow::SpawnTab_FramesTrack(const FSpawnTabArgs& Args)
 {
+	FTimingProfilerManager::Get()->SetFramesTrackVisible(true);
+
 	return SNew(SDockTab)
 		.ShouldAutosize(false)
 		.TabRole(ETabRole::PanelTab)
@@ -101,6 +103,8 @@ TSharedRef<SDockTab> STimingProfilerWindow::SpawnTab_FramesTrack(const FSpawnTab
 
 TSharedRef<SDockTab> STimingProfilerWindow::SpawnTab_GraphTrack(const FSpawnTabArgs& Args)
 {
+	FTimingProfilerManager::Get()->SetGraphTrackVisible(true);
+
 	return SNew(SDockTab)
 		.ShouldAutosize(false)
 		.TabRole(ETabRole::PanelTab)
@@ -114,6 +118,8 @@ TSharedRef<SDockTab> STimingProfilerWindow::SpawnTab_GraphTrack(const FSpawnTabA
 
 TSharedRef<SDockTab> STimingProfilerWindow::SpawnTab_TimingView(const FSpawnTabArgs& Args)
 {
+	FTimingProfilerManager::Get()->SetTimingViewVisible(true);
+
 	return SNew(SDockTab)
 		.ShouldAutosize(false)
 		.TabRole(ETabRole::PanelTab)
@@ -127,6 +133,8 @@ TSharedRef<SDockTab> STimingProfilerWindow::SpawnTab_TimingView(const FSpawnTabA
 
 TSharedRef<SDockTab> STimingProfilerWindow::SpawnTab_Timers(const FSpawnTabArgs& Args)
 {
+	FTimingProfilerManager::Get()->SetTimersViewVisible(true);
+
 	return SNew(SDockTab)
 		.ShouldAutosize(false)
 		.TabRole(ETabRole::PanelTab)
@@ -140,6 +148,8 @@ TSharedRef<SDockTab> STimingProfilerWindow::SpawnTab_Timers(const FSpawnTabArgs&
 
 TSharedRef<SDockTab> STimingProfilerWindow::SpawnTab_StatsCounters(const FSpawnTabArgs& Args)
 {
+	FTimingProfilerManager::Get()->SetStatsCountersViewVisible(true);
+
 	return SNew(SDockTab)
 		.ShouldAutosize(false)
 		.TabRole(ETabRole::PanelTab)
@@ -153,6 +163,8 @@ TSharedRef<SDockTab> STimingProfilerWindow::SpawnTab_StatsCounters(const FSpawnT
 
 TSharedRef<SDockTab> STimingProfilerWindow::SpawnTab_LogView(const FSpawnTabArgs& Args)
 {
+	FTimingProfilerManager::Get()->SetLogViewVisible(true);
+
 	return SNew(SDockTab)
 		.ShouldAutosize(false)
 		.TabRole(ETabRole::PanelTab)
