[Insights] Timers view:  Enabled Thread filtering for aggregated stats. I.e. the aggregated statistics computation for selected time range in Timing view will not include gpu/cpu timeline tracks that are hidden.

diff --git a/Engine/Source/Developer/TraceInsights/Private/Insights/ViewModels/TimingEventsTrack.cpp b/Engine/Source/Developer/TraceInsights/Private/Insights/ViewModels/TimingEventsTrack.cpp
--- a/Engine/Source/Developer/TraceInsights/Private/Insights/ViewModels/TimingEventsTrack.cpp
+++ b/Engine/Source/Developer/TraceInsights/Private/Insights/ViewModels/TimingEventsTrack.cpp
@@ -86,6 +86,8 @@ FTimingEventsTrack::FTimingEventsTrack(uint64 InTrackId, ETimingEventsTrackType
 	, Type(InType)
 	, Name(InName)
 	, GroupName(InGroupName)
+	, ThreadId(0)
+	, Order(0)
 	, Depth(-1)
 	, bIsCollapsed(false)
 {
diff --git a/Engine/Source/Developer/TraceInsights/Private/Insights/ViewModels/TimingEventsTrack.h b/Engine/Source/Developer/TraceInsights/Private/Insights/ViewModels/TimingEventsTrack.h
--- a/Engine/Source/Developer/TraceInsights/Private/Insights/ViewModels/TimingEventsTrack.h
+++ b/Engine/Source/Developer/TraceInsights/Private/Insights/ViewModels/TimingEventsTrack.h
@@ -62,6 +62,9 @@ public:
 
 	const TCHAR* GetGroupName() const { return GroupName; };
 
+	void SetThreadId(uint32 InThreadId) { ThreadId = InThreadId; }
+	uint32 GetThreadId() const { return ThreadId; }
+
 	void SetOrder(int32 InOrder) { Order = InOrder; }
 	int32 GetOrder() const { return Order; }
 
@@ -75,8 +78,8 @@ public:
 	ETimingEventsTrackType Type;
 	FString Name;
 	const TCHAR* GroupName;
+	uint32 ThreadId;
 	int32 Order;
-
 	int32 Depth; // number of lanes == Depth + 1
 	bool bIsCollapsed;
 
diff --git a/Engine/Source/Developer/TraceInsights/Private/Insights/ViewModels/TimingTrackViewport.h b/Engine/Source/Developer/TraceInsights/Private/Insights/ViewModels/TimingTrackViewport.h
--- a/Engine/Source/Developer/TraceInsights/Private/Insights/ViewModels/TimingTrackViewport.h
+++ b/Engine/Source/Developer/TraceInsights/Private/Insights/ViewModels/TimingTrackViewport.h
@@ -34,7 +34,7 @@ public:
 		Height = 0.0f;
 
 		MinValidTime = 0.0;
-		MaxValidTime = 0.0; // 1m
+		MaxValidTime = 0.0;
 		StartTime = 0.0;
 		EndTime = 0.0;
 		MinScaleX = (5 * 20) / 3600.0; // 1h between major tick marks
diff --git a/Engine/Source/Developer/TraceInsights/Private/Insights/Widgets/STimersView.cpp b/Engine/Source/Developer/TraceInsights/Private/Insights/Widgets/STimersView.cpp
--- a/Engine/Source/Developer/TraceInsights/Private/Insights/Widgets/STimersView.cpp
+++ b/Engine/Source/Developer/TraceInsights/Private/Insights/Widgets/STimersView.cpp
@@ -22,6 +22,8 @@
 #include "Insights/ViewModels/TimersViewColumnFactory.h"
 #include "Insights/Widgets/STimersViewTooltip.h"
 #include "Insights/Widgets/STimerTableRow.h"
+#include "Insights/Widgets/STimingProfilerWindow.h"
+#include "Insights/Widgets/STimingView.h"
 
 #define LOCTEXT_NAMESPACE "STimersView"
 
@@ -1485,16 +1487,21 @@ void STimersView::UpdateStats(double StartTime, double EndTime)
 
 	if (Session.IsValid() && StartTime < EndTime && Trace::ReadTimingProfilerProvider(*Session.Get()))
 	{
-		auto ThreadFilter = [](uint32 ThreadId)
+		TSharedPtr<STimingProfilerWindow> Wnd = FTimingProfilerManager::Get()->GetProfilerWindow();
+		TSharedPtr<STimingView> TimingView = Wnd.IsValid() ? Wnd->TimingView : nullptr;
+
+		auto ThreadFilter = [&TimingView](uint32 ThreadId)
 		{
-			return true;
+			return !TimingView.IsValid() || TimingView->IsCpuTrackVisible(ThreadId);
 		};
 
+		const bool bIsGpuTrackVisible = TimingView.IsValid() && TimingView->IsGpuTrackVisible();
+
 		TUniquePtr<Trace::ITable<Trace::FTimingProfilerAggregatedStats>> AggregationResultTable;
 		{
 			Trace::FAnalysisSessionReadScope SessionReadScope(*Session.Get());
 			const Trace::ITimingProfilerProvider& TimingProfilerProvider = *Trace::ReadTimingProfilerProvider(*Session.Get());
-			AggregationResultTable.Reset(TimingProfilerProvider.CreateAggregation(StartTime, EndTime, ThreadFilter, true));
+			AggregationResultTable.Reset(TimingProfilerProvider.CreateAggregation(StartTime, EndTime, ThreadFilter, bIsGpuTrackVisible));
 		}
 
 		TUniquePtr<Trace::ITableReader<Trace::FTimingProfilerAggregatedStats>> AggregationResultTableReader(AggregationResultTable->CreateReader());
diff --git a/Engine/Source/Developer/TraceInsights/Private/Insights/Widgets/STimingView.cpp b/Engine/Source/Developer/TraceInsights/Private/Insights/Widgets/STimingView.cpp
--- a/Engine/Source/Developer/TraceInsights/Private/Insights/Widgets/STimingView.cpp
+++ b/Engine/Source/Developer/TraceInsights/Private/Insights/Widgets/STimingView.cpp
@@ -57,9 +57,9 @@ const TCHAR* GetName(ELoadTimeProfilerObjectEventType Type);
 ////////////////////////////////////////////////////////////////////////////////////////////////////
 
 STimingView::STimingView()
-	: TimeRulerTrack(0xFFFF0001)
-	, MarkersTrack(0xFFFF0002)
-	, GraphTrack(0xFFFF0003)
+	: TimeRulerTrack(FBaseTimingTrack::GenerateId())
+	, MarkersTrack(FBaseTimingTrack::GenerateId())
+	, GraphTrack(FBaseTimingTrack::GenerateId())
 	, TooltipWidth(MinTooltipWidth)
 	, TooltipAlpha(0.0f)
 	, WhiteBrush(FCoreStyle::Get().GetBrush("WhiteBrush"))
@@ -181,17 +181,16 @@ void STimingView::Reset()
 
 	//////////////////////////////////////////////////
 
-	ThreadGroups.Reset();
-
 	for (const auto& KV : CachedTimelines)
 	{
 		delete KV.Value;
 	}
 	CachedTimelines.Reset();
 
-	//TopTracks.Reset();
-	//ScrollableTracks.Reset();
-	//BottomTracks.Reset();
+	//TODO: TopTracks.Reset();
+	//TODO: ScrollableTracks.Reset();
+	//TODO: BottomTracks.Reset();
+	//TODO: ForegroundTracks.Reset();
 
 	//////////////////////////////////////////////////
 
@@ -203,6 +202,13 @@ void STimingView::Reset()
 
 	//////////////////////////////////////////////////
 
+	GpuTrack = nullptr;
+	CpuTracks.Reset();
+
+	ThreadGroups.Reset();
+
+	//////////////////////////////////////////////////
+
 	LoadingMainThreadTrack = nullptr;
 	LoadingAsyncThreadTrack = nullptr;
 
@@ -444,7 +450,7 @@ void STimingView::Tick(const FGeometry& AllottedGeometry, const double InCurrent
 			const uint64 TrackId = static_cast<uint64>(GpuTimelineIndex);
 			if (!CachedTimelines.Contains(TrackId))
 			{
-				AddTimingEventsTrack(TrackId, ETimingEventsTrackType::Gpu, TEXT("GPU"), nullptr, 0);
+				GpuTrack = AddTimingEventsTrack(TrackId, ETimingEventsTrackType::Gpu, TEXT("GPU"), nullptr, 0);
 				bIsTimingEventsTrackDirty = true;
 			}
 		}
@@ -483,6 +489,8 @@ void STimingView::Tick(const FGeometry& AllottedGeometry, const double InCurrent
 				{
 					FString TrackName(ThreadInfo.Name && *ThreadInfo.Name ? ThreadInfo.Name : FString::Printf(TEXT("Thread %u")));
 					Track = AddTimingEventsTrack(TrackId, ETimingEventsTrackType::Cpu, TrackName, GroupName, Order);
+					Track->SetThreadId(ThreadInfo.Id);
+					CpuTracks.Add(ThreadInfo.Id, Track);
 					bIsTimingEventsTrackDirty = true;
 				}
 				else
@@ -1205,6 +1213,20 @@ FTimingEventsTrack* STimingView::AddTimingEventsTrack(uint64 TrackId, ETimingEve
 
 ////////////////////////////////////////////////////////////////////////////////////////////////////
 
+bool STimingView::IsGpuTrackVisible() const
+{
+	return GpuTrack != nullptr && GpuTrack->IsVisible();
+}
+
+////////////////////////////////////////////////////////////////////////////////////////////////////
+
+bool STimingView::IsCpuTrackVisible(uint32 ThreadId) const
+{
+	return CpuTracks.Contains(ThreadId) && CpuTracks[ThreadId]->IsVisible();
+}
+
+////////////////////////////////////////////////////////////////////////////////////////////////////
+
 void STimingView::DrawTimingProfilerTrack(FTimingViewDrawHelper& Helper, FTimingEventsTrack& Track) const
 {
 	if (Helper.BeginTimeline(Track))
diff --git a/Engine/Source/Developer/TraceInsights/Private/Insights/Widgets/STimingView.h b/Engine/Source/Developer/TraceInsights/Private/Insights/Widgets/STimingView.h
--- a/Engine/Source/Developer/TraceInsights/Private/Insights/Widgets/STimingView.h
+++ b/Engine/Source/Developer/TraceInsights/Private/Insights/Widgets/STimingView.h
@@ -252,6 +252,9 @@ public:
 
 	void SetGraphTrackVisible(bool bOnOff) { GraphTrack.SetVisibilityFlag(bOnOff); }
 
+	bool IsGpuTrackVisible() const;
+	bool IsCpuTrackVisible(uint32 ThreadId) const;
+
 protected:
 	virtual FVector2D ComputeDesiredSize(float) const override
 	{
@@ -359,16 +362,10 @@ protected:
 
 	////////////////////////////////////////////////////////////
 
-	struct FThreadGroup
-	{
-		const TCHAR* Name;
-		bool bIsVisible;
-		uint32 NumTimelines;
-	};
-
-	mutable TMap<const TCHAR*, FThreadGroup> ThreadGroups; /**< thread groups */
-
-	mutable TMap<uint64, FTimingEventsTrack*> CachedTimelines; /**< all tracks */
+	/** All created tracks.
+	  * Maps track id to track pointer.
+	  */
+	TMap<uint64, FTimingEventsTrack*> CachedTimelines;
 
 	//TODO: TArray<FBaseTimingTrack*> TopTracks; /**< tracks docked on top, in order to be displayed (top to bottom) */
 	//TODO: TArray<FBaseTimingTrack*> ScrollableTracks; /**< tracks in scrollable area, in order to be displayed (top to bottom) */
@@ -376,16 +373,30 @@ protected:
 	//TODO: TArray<FBaseTimingTrack*> ForegorundTracks; /**< tracks to draw over top/scrollable/bottom tracks (can use entire area), in order to be displayed (back to front) */
 
 	////////////////////////////////////////////////////////////
-	// Timing Events Tracks (GPU, CPU, IO)
 
-	TArray<FTimingEventsTrack*> TimingEventsTracks; /**< timing events tracks in order to be displayed */
+	TArray<FTimingEventsTrack*> TimingEventsTracks; /**< all timing events tracks in order to be displayed */
 
 	bool bAreTimingEventsTracksDirty;
 
 	bool bUseDownSampling;
 
 	////////////////////////////////////////////////////////////
-	// Load time
+	// Cpu/Gpu
+
+	FTimingEventsTrack* GpuTrack;
+	TMap<uint32, FTimingEventsTrack*> CpuTracks; /**< maps thread id to track pointer */
+
+	struct FThreadGroup
+	{
+		const TCHAR* Name; /**< the thread group name; pointer to string owned by ThreadProvider */
+		bool bIsVisible;
+		uint32 NumTimelines;
+	};
+
+	TMap<const TCHAR*, FThreadGroup> ThreadGroups; /**< maps thread group name to thread group info */
+
+	////////////////////////////////////////////////////////////
+	// Asset Loading
 
 	FTimingEventsTrack* LoadingMainThreadTrack;
 	FTimingEventsTrack* LoadingAsyncThreadTrack;
