[Insights]
  - Updated default visible tracks in Asset Loading tab to show only loading tracks and the CPu threads associated with it (main gamne thread and the async loading thread).
  - Timing view: Fixed clicking on an event in a Loading* track to not update selection in Timers view.

diff --git a/Engine/Source/Developer/TraceInsights/Private/Insights/Widgets/SIoProfilerWindow.cpp b/Engine/Source/Developer/TraceInsights/Private/Insights/Widgets/SIoProfilerWindow.cpp
--- a/Engine/Source/Developer/TraceInsights/Private/Insights/Widgets/SIoProfilerWindow.cpp
+++ b/Engine/Source/Developer/TraceInsights/Private/Insights/Widgets/SIoProfilerWindow.cpp
@@ -102,7 +102,7 @@ void SIoProfilerWindow::Construct(const FArguments& InArgs)
 				]
 		];
 
-	TimingView->SetGraphTrackVisible(true);
+	TimingView->EnableAssetLoadingMode();
 }
 
 END_SLATE_FUNCTION_BUILD_OPTIMIZATION
diff --git a/Engine/Source/Developer/TraceInsights/Private/Insights/Widgets/STimingView.cpp b/Engine/Source/Developer/TraceInsights/Private/Insights/Widgets/STimingView.cpp
--- a/Engine/Source/Developer/TraceInsights/Private/Insights/Widgets/STimingView.cpp
+++ b/Engine/Source/Developer/TraceInsights/Private/Insights/Widgets/STimingView.cpp
@@ -57,7 +57,8 @@ const TCHAR* GetName(ELoadTimeProfilerObjectEventType Type);
 ////////////////////////////////////////////////////////////////////////////////////////////////////
 
 STimingView::STimingView()
-	: TimeRulerTrack(FBaseTimingTrack::GenerateId())
+	: bAssetLoadingMode(false)
+	, TimeRulerTrack(FBaseTimingTrack::GenerateId())
 	, MarkersTrack(FBaseTimingTrack::GenerateId())
 	, GraphTrack(FBaseTimingTrack::GenerateId())
 	, TooltipWidth(MinTooltipWidth)
@@ -85,7 +86,7 @@ void STimingView::Construct(const FArguments& InArgs)
 {
 	OnSelectionChanged = InArgs._OnSelectionChanged;
 	OnHoveredEventChanged = InArgs._OnHoveredEventChanged;
-	SelectedEventChanged = InArgs._SelectedEventChanged;
+	OnSelectedEventChanged = InArgs._OnSelectedEventChanged;
 
 	ChildSlot
 	[
@@ -212,6 +213,9 @@ void STimingView::Reset()
 	LoadingMainThreadTrack = nullptr;
 	LoadingAsyncThreadTrack = nullptr;
 
+	LoadingMainThreadId = 0;
+	LoadingAsyncThreadId = 0;
+
 	LoadingGetEventNameFn = FLoadingTrackGetEventNameDelegate::CreateRaw(this, &STimingView::GetLoadTimeProfilerEventName);
 
 	IoOverviewTrack = nullptr;
@@ -275,6 +279,22 @@ void STimingView::Reset()
 	DrawDurationHistory.Reset();
 	OnPaintDurationHistory.Reset();
 	LastOnPaintTime = FPlatformTime::Cycles64();
+
+	if (bAssetLoadingMode)
+	{
+		EnableAssetLoadingMode();
+	}
+}
+
+////////////////////////////////////////////////////////////////////////////////////////////////////
+
+void STimingView::EnableAssetLoadingMode()
+{
+	bAssetLoadingMode = true;
+
+	bShowHideAllGpuTracks = false;
+	bShowHideAllCpuTracks = false;
+	bShowHideAllLoadingTracks = true;
 }
 
 ////////////////////////////////////////////////////////////////////////////////////////////////////
@@ -389,7 +409,7 @@ void STimingView::Tick(const FGeometry& AllottedGeometry, const double InCurrent
 	UpdateIo();
 
 	TSharedPtr<const Trace::IAnalysisSession> Session = FInsightsManager::Get()->GetSession();
-	if (Session && Trace::ReadTimingProfilerProvider(*Session.Get()))
+	if (Session)
 	{
 		Trace::FAnalysisSessionReadScope SessionReadScope(*Session.Get());
 
@@ -413,112 +433,140 @@ void STimingView::Tick(const FGeometry& AllottedGeometry, const double InCurrent
 
 		bool bIsTimingEventsTrackDirty = false;
 
-		if (LoadingMainThreadTrack == nullptr)
+		if (Trace::ReadLoadTimeProfilerProvider(*Session.Get()))
 		{
-			uint64 TrackId = FBaseTimingTrack::GenerateId();
-			LoadingMainThreadTrack = AddTimingEventsTrack(TrackId, ETimingEventsTrackType::Loading, TEXT("Loading - Main Thread"), nullptr, -3);
-			LoadingMainThreadTrack->SetVisibilityFlag(bShowHideAllLoadingTracks);
-			bIsTimingEventsTrackDirty = true;
-		}
-
-		if (LoadingAsyncThreadTrack == nullptr)
-		{
-			uint64 TrackId = FBaseTimingTrack::GenerateId();
-			LoadingAsyncThreadTrack = AddTimingEventsTrack(TrackId, ETimingEventsTrackType::Loading, TEXT("Loading - Async Thread"), nullptr, -2);
-			LoadingAsyncThreadTrack->SetVisibilityFlag(bShowHideAllLoadingTracks);
-			bIsTimingEventsTrackDirty = true;
-		}
-
-		if (IoOverviewTrack == nullptr)
-		{
-			// Note: The I/O timelines are just prototypes for now (will be removed once the functionality is moved in analyzer).
-			const uint64 TrackId = FBaseTimingTrack::GenerateId();
-			IoOverviewTrack = AddTimingEventsTrack(TrackId, ETimingEventsTrackType::Io, TEXT("I/O Overview"), nullptr, -1);
-			IoOverviewTrack->SetVisibilityFlag(bShowHideAllIoTracks);
-			bIsTimingEventsTrackDirty = true;
-		}
+			const Trace::ILoadTimeProfilerProvider& LoadTimeProfilerProvider = *Trace::ReadLoadTimeProfilerProvider(*Session.Get());
 
-		// Sync CachedTimelines and TimingEventsTracks with available timelines on analysis side.
-		// TODO: can we make this more efficient?
+			LoadingMainThreadId = LoadTimeProfilerProvider.GetMainThreadId();
+			LoadingAsyncThreadId = LoadTimeProfilerProvider.GetAsyncLoadingThreadId();
 
-		const Trace::ITimingProfilerProvider& TimingProfilerProvider = *Trace::ReadTimingProfilerProvider(*Session.Get());
+			if (LoadingMainThreadTrack == nullptr)
+			{
+				uint64 TrackId = FBaseTimingTrack::GenerateId();
+				LoadingMainThreadTrack = AddTimingEventsTrack(TrackId, ETimingEventsTrackType::Loading, TEXT("Loading - Main Thread"), nullptr, -3);
+				LoadingMainThreadTrack->SetVisibilityFlag(bShowHideAllLoadingTracks);
+				bIsTimingEventsTrackDirty = true;
+			}
 
-		// Check if we have a GPU track.
-		uint32 GpuTimelineIndex;
-		if (TimingProfilerProvider.GetGpuTimelineIndex(GpuTimelineIndex))
-		{
-			const uint64 TrackId = static_cast<uint64>(GpuTimelineIndex);
-			if (!CachedTimelines.Contains(TrackId))
+			if (LoadingAsyncThreadTrack == nullptr)
 			{
-				GpuTrack = AddTimingEventsTrack(TrackId, ETimingEventsTrackType::Gpu, TEXT("GPU"), nullptr, 0);
+				uint64 TrackId = FBaseTimingTrack::GenerateId();
+				LoadingAsyncThreadTrack = AddTimingEventsTrack(TrackId, ETimingEventsTrackType::Loading, TEXT("Loading - Async Thread"), nullptr, -2);
+				LoadingAsyncThreadTrack->SetVisibilityFlag(bShowHideAllLoadingTracks);
 				bIsTimingEventsTrackDirty = true;
 			}
 		}
 
-		// Check available CPU tracks.
-		int32 Order = 1;
-		const Trace::IThreadProvider& ThreadProvider = Trace::ReadThreadProvider(*Session.Get());
-		ThreadProvider.EnumerateThreads([this, &Order, &bIsTimingEventsTrackDirty, &TimingProfilerProvider](const Trace::FThreadInfo& ThreadInfo)
+		if (Trace::ReadFileActivityProvider(*Session.Get()))
 		{
-			const TCHAR* GroupName = ThreadInfo.GroupName;
-			if (GroupName == nullptr)
+			if (IoOverviewTrack == nullptr)
 			{
-				GroupName = ThreadInfo.Name;
+				// Note: The I/O timelines are just prototypes for now (will be removed once the functionality is moved in analyzer).
+				const uint64 TrackId = FBaseTimingTrack::GenerateId();
+				IoOverviewTrack = AddTimingEventsTrack(TrackId, ETimingEventsTrackType::Io, TEXT("I/O Overview"), nullptr, -1);
+				IoOverviewTrack->SetVisibilityFlag(bShowHideAllIoTracks);
+				bIsTimingEventsTrackDirty = true;
 			}
+		}
+
+		if (Trace::ReadTimingProfilerProvider(*Session.Get()))
+		{
+			const Trace::ITimingProfilerProvider& TimingProfilerProvider = *Trace::ReadTimingProfilerProvider(*Session.Get());
 
-			bool bIsGroupVisible = true;
-			if (GroupName != nullptr)
+			// Check if we have a GPU track.
+			uint32 GpuTimelineIndex;
+			if (TimingProfilerProvider.GetGpuTimelineIndex(GpuTimelineIndex))
 			{
-				if (!ThreadGroups.Contains(GroupName))
-				{
-					ThreadGroups.Add(GroupName, { GroupName, bIsGroupVisible, 0 });
-				}
-				else
+				const uint64 TrackId = static_cast<uint64>(GpuTimelineIndex);
+				if (!CachedTimelines.Contains(TrackId))
 				{
-					bIsGroupVisible = ThreadGroups[GroupName].bIsVisible;
+					GpuTrack = AddTimingEventsTrack(TrackId, ETimingEventsTrackType::Gpu, TEXT("GPU"), nullptr, 0);
+					GpuTrack->SetVisibilityFlag(bShowHideAllGpuTracks);
+					bIsTimingEventsTrackDirty = true;
 				}
 			}
 
-			uint32 CpuTimelineIndex;
-			if (TimingProfilerProvider.GetCpuThreadTimelineIndex(ThreadInfo.Id, CpuTimelineIndex))
+			// Check available CPU tracks.
+			int32 Order = 1;
+			const Trace::IThreadProvider& ThreadProvider = Trace::ReadThreadProvider(*Session.Get());
+			ThreadProvider.EnumerateThreads([this, &Order, &bIsTimingEventsTrackDirty, &TimingProfilerProvider](const Trace::FThreadInfo& ThreadInfo)
 			{
-				FTimingEventsTrack* Track = nullptr;
-				const uint64 TrackId = static_cast<uint64>(CpuTimelineIndex);
+				const TCHAR* GroupName = ThreadInfo.GroupName;
+				if (GroupName == nullptr)
+				{
+					GroupName = ThreadInfo.Name;
+				}
 
-				if (!CachedTimelines.Contains(TrackId))
+				bool bIsGroupVisible = bShowHideAllCpuTracks;
+				if (GroupName != nullptr)
 				{
-					FString TrackName(ThreadInfo.Name && *ThreadInfo.Name ? ThreadInfo.Name : FString::Printf(TEXT("Thread %u")));
-					Track = AddTimingEventsTrack(TrackId, ETimingEventsTrackType::Cpu, TrackName, GroupName, Order);
-					Track->SetThreadId(ThreadInfo.Id);
-					CpuTracks.Add(ThreadInfo.Id, Track);
-					bIsTimingEventsTrackDirty = true;
+					if (!ThreadGroups.Contains(GroupName))
+					{
+						ThreadGroups.Add(GroupName, { GroupName, bIsGroupVisible, 0 });
+					}
+					else
+					{
+						bIsGroupVisible = ThreadGroups[GroupName].bIsVisible;
+					}
 				}
-				else
+
+				uint32 CpuTimelineIndex;
+				if (TimingProfilerProvider.GetCpuThreadTimelineIndex(ThreadInfo.Id, CpuTimelineIndex))
 				{
-					Track = CachedTimelines[TrackId];
-					if (Track->GetOrder() != Order)
+					FTimingEventsTrack* Track = nullptr;
+					const uint64 TrackId = static_cast<uint64>(CpuTimelineIndex);
+
+					if (!CachedTimelines.Contains(TrackId))
 					{
-						Track->SetOrder(Order);
+						FString TrackName(ThreadInfo.Name && *ThreadInfo.Name ? ThreadInfo.Name : FString::Printf(TEXT("Thread %u")));
+
+						// Create new Timing Events track for the CPU thread.
+						Track = AddTimingEventsTrack(TrackId, ETimingEventsTrackType::Cpu, TrackName, GroupName, Order);
+
+						Track->SetThreadId(ThreadInfo.Id);
+						CpuTracks.Add(ThreadInfo.Id, Track);
+
+						FThreadGroup& ThreadGroup = ThreadGroups[GroupName];
+						ThreadGroup.NumTimelines++;
+
+						if (bAssetLoadingMode && (ThreadInfo.Id == LoadingMainThreadId || ThreadInfo.Id == LoadingAsyncThreadId))
+						{
+							Track->SetVisibilityFlag(true);
+							ThreadGroup.bIsVisible = true;
+						}
+						else
+						{
+							Track->SetVisibilityFlag(bIsGroupVisible);
+						}
+
 						bIsTimingEventsTrackDirty = true;
 					}
-				}
+					else
+					{
+						Track = CachedTimelines[TrackId];
 
-				Track->SetVisibilityFlag(bIsGroupVisible);
+						if (Track->GetOrder() != Order)
+						{
+							Track->SetOrder(Order);
+							bIsTimingEventsTrackDirty = true;
+						}
+					}
 
-				FThreadGroup& ThreadGroup = ThreadGroups[GroupName];
-				ThreadGroup.NumTimelines++;
-				
-				Order++;
-			}
-		});
+					Order++;
+				}
+			});
+		}
 
-		if (IoActivityTrack == nullptr)
+		if (Trace::ReadFileActivityProvider(*Session.Get()))
 		{
-			// Note: The I/O timelines are just prototypes for now (will be removed once the functionality is moved in analyzer).
-			uint64 TrackId = FBaseTimingTrack::GenerateId();
-			IoActivityTrack = AddTimingEventsTrack(TrackId, ETimingEventsTrackType::Io, TEXT("I/O Activity"), nullptr, 999999);
-			IoActivityTrack->SetVisibilityFlag(bShowHideAllIoTracks);
-			bIsTimingEventsTrackDirty = true;
+			if (IoActivityTrack == nullptr)
+			{
+				// Note: The I/O timelines are just prototypes for now (will be removed once the functionality is moved in analyzer).
+				uint64 TrackId = FBaseTimingTrack::GenerateId();
+				IoActivityTrack = AddTimingEventsTrack(TrackId, ETimingEventsTrackType::Io, TEXT("I/O Activity"), nullptr, 999999);
+				IoActivityTrack->SetVisibilityFlag(bShowHideAllIoTracks);
+				bIsTimingEventsTrackDirty = true;
+			}
 		}
 
 		if (bIsTimingEventsTrackDirty)
@@ -1192,18 +1240,27 @@ int32 STimingView::OnPaint(const FPaintArgs& Args, const FGeometry& AllottedGeom
 
 ////////////////////////////////////////////////////////////////////////////////////////////////////
 
+const TCHAR* GetName(ETimingEventsTrackType Type)
+{
+	switch (Type)
+	{
+		case ETimingEventsTrackType::Gpu:		return TEXT("GPU");
+		case ETimingEventsTrackType::Cpu:		return TEXT("CPU");
+		case ETimingEventsTrackType::Loading:	return TEXT("Loading");
+		case ETimingEventsTrackType::Io:		return TEXT("I/O");
+		default:								return TEXT("unknown");
+	}
+}
+
+////////////////////////////////////////////////////////////////////////////////////////////////////
+
 FTimingEventsTrack* STimingView::AddTimingEventsTrack(uint64 TrackId, ETimingEventsTrackType TrackType, const FString& TrackName, const TCHAR* GroupName, int32 Order)
 {
 	FTimingEventsTrack* Track = new FTimingEventsTrack(TrackId, TrackType, TrackName, GroupName);
 
 	Track->SetOrder(Order);
 
-	UE_LOG(TimingProfiler, Log, TEXT("New Timing Events Track (%d) : %s (\"%s\")"),
-		TimingEventsTracks.Num() + 1,
-		TrackType == ETimingEventsTrackType::Gpu ? TEXT("GPU") :
-		TrackType == ETimingEventsTrackType::Cpu ? TEXT("CPU") :
-		TrackType == ETimingEventsTrackType::Io ? TEXT("I/O") : TEXT("unknown"),
-		*TrackName);
+	UE_LOG(TimingProfiler, Log, TEXT("New Timing Events Track (%d) : %s (\"%s\")"), TimingEventsTracks.Num() + 1, GetName(TrackType), *TrackName);
 
 	CachedTimelines.Add(TrackId, Track);
 	TimingEventsTracks.Add(Track);
@@ -1328,7 +1385,7 @@ const TCHAR* STimingView::GetLoadTimeProfilerEventName(uint32 Depth, const Trace
 			return Event.Package->Name;
 		}
 	}
-	
+
 	if (Event.Export && Event.Export->Class)
 	{
 		return Event.Export->Class->Name;
@@ -2000,19 +2057,7 @@ FReply STimingView::OnMouseButtonUp(const FGeometry& MyGeometry, const FPointerE
 			else if (bIsSelecting)
 			{
 				//TODO: SelectionChangedEvent.Broadcast(SelectionStartTime, SelectionEndTime);
-
-				TSharedPtr<STimingProfilerWindow> Wnd = FTimingProfilerManager::Get()->GetProfilerWindow();
-				if (Wnd)
-				{
-					if (Wnd->TimersView)
-					{
-						Wnd->TimersView->UpdateStats(SelectionStartTime, SelectionEndTime);
-					}
-					if (Wnd->StatsView)
-					{
-						Wnd->StatsView->UpdateStats(SelectionStartTime, SelectionEndTime);
-					}
-				}
+				UpdateAggregatedStats();
 
 				bIsSelecting = false;
 			}
@@ -2712,6 +2757,13 @@ void STimingView::SelectTimeInterval(double IntervalStartTime, double IntervalDu
 	LastSelectionType = ESelectionType::TimeRange;
 
 	//TODO: SelectionChangedEvent.Broadcast(SelectionStartTime, SelectionEndTime);
+	UpdateAggregatedStats();
+}
+
+////////////////////////////////////////////////////////////////////////////////////////////////////
+
+void STimingView::UpdateAggregatedStats()
+{
 	TSharedPtr<STimingProfilerWindow> Wnd = FTimingProfilerManager::Get()->GetProfilerWindow();
 	if (Wnd)
 	{
@@ -3043,7 +3095,9 @@ bool STimingView::SearchTimingEvent(const double InStartTime,
 void STimingView::OnSelectedTimingEventChanged()
 {
 	// Select the timer node coresponding to timing event type of selected timing event.
-	if (SelectedTimingEvent.IsValid())
+	if (SelectedTimingEvent.IsValid() &&
+		(SelectedTimingEvent.Track->GetType() == ETimingEventsTrackType::Cpu ||
+		 SelectedTimingEvent.Track->GetType() == ETimingEventsTrackType::Gpu))
 	{
 		TSharedPtr<STimingProfilerWindow> Wnd = FTimingProfilerManager::Get()->GetProfilerWindow();
 		if (Wnd)
diff --git a/Engine/Source/Developer/TraceInsights/Private/Insights/Widgets/STimingView.h b/Engine/Source/Developer/TraceInsights/Private/Insights/Widgets/STimingView.h
--- a/Engine/Source/Developer/TraceInsights/Private/Insights/Widgets/STimingView.h
+++ b/Engine/Source/Developer/TraceInsights/Private/Insights/Widgets/STimingView.h
@@ -54,14 +54,14 @@ public:
 	SLATE_BEGIN_ARGS(STimingView)
 		: _OnSelectionChanged()
 		, _OnHoveredEventChanged()
-		, _SelectedEventChanged()
+		, _OnSelectedEventChanged()
 		{
 			_Clipping = EWidgetClipping::ClipToBounds;
 		}
 
 		SLATE_EVENT(FSelectionChangedDelegate, OnSelectionChanged)
 		SLATE_EVENT(FHoveredEventChangedDelegate, OnHoveredEventChanged)
-		SLATE_EVENT(FSelectedEventChangedDelegate, SelectedEventChanged)
+		SLATE_EVENT(FSelectedEventChangedDelegate, OnSelectedEventChanged)
 	SLATE_END_ARGS()
 
 	/**
@@ -96,6 +96,8 @@ public:
 	bool ToggleTrackVisibilityByGroup_IsChecked(const TCHAR* InGroupName) const;
 	void ToggleTrackVisibilityByGroup_Execute(const TCHAR* InGroupName);
 
+	void EnableAssetLoadingMode();
+
 	/** Resets internal widget's data to the default one. */
 	void Reset();
 
@@ -250,8 +252,6 @@ public:
 	bool IsDrawOnlyBookmarksEnabled() { return MarkersTrack.IsBookmarksTrack(); }
 	void SetDrawOnlyBookmarks(bool bOnOff);
 
-	void SetGraphTrackVisible(bool bOnOff) { GraphTrack.SetVisibilityFlag(bOnOff); }
-
 	bool IsGpuTrackVisible() const;
 	bool IsCpuTrackVisible(uint32 ThreadId) const;
 
@@ -303,6 +303,8 @@ protected:
 
 	////////////////////////////////////////////////////////////////////////////////////////////////////
 
+	void UpdateAggregatedStats();
+
 	void UpdateHoveredTimingEvent(float MX, float MY);
 
 	bool SearchTimingEvent(const double InStartTime,
@@ -401,15 +403,20 @@ protected:
 	FTimingEventsTrack* LoadingMainThreadTrack;
 	FTimingEventsTrack* LoadingAsyncThreadTrack;
 
+	uint32 LoadingMainThreadId;
+	uint32 LoadingAsyncThreadId;
+
 	FLoadingTrackGetEventNameDelegate LoadingGetEventNameFn;
 
+	bool bAssetLoadingMode;
+
 	////////////////////////////////////////////////////////////
 	// File activity (I/O)
 
 	FTimingEventsTrack* IoOverviewTrack;
 	FTimingEventsTrack* IoActivityTrack;
 
-	mutable bool bForceIoEventsUpdate;
+	bool bForceIoEventsUpdate;
 
 	bool bMergeIoLanes;
 
@@ -423,7 +430,7 @@ protected:
 	};
 
 	/** All IO events, cached. */
-	mutable TArray<FIoTimingEvent> AllIoEvents;
+	TArray<FIoTimingEvent> AllIoEvents;
 
 	////////////////////////////////////////////////////////////
 
@@ -531,5 +538,5 @@ protected:
 
 	FSelectionChangedDelegate OnSelectionChanged;
 	FHoveredEventChangedDelegate OnHoveredEventChanged;
-	FSelectedEventChangedDelegate SelectedEventChanged;
+	FSelectedEventChangedDelegate OnSelectedEventChanged;
 };
