[Insights]
  - Asset Loading: Added Object Type Aggregation and Events Aggregation statistics (top 10 stats, sorted by Total time) as a fixed tooltip in Timing view.

diff --git a/Engine/Source/Developer/TraceInsights/Private/Insights/InsightsManager.cpp b/Engine/Source/Developer/TraceInsights/Private/Insights/InsightsManager.cpp
--- a/Engine/Source/Developer/TraceInsights/Private/Insights/InsightsManager.cpp
+++ b/Engine/Source/Developer/TraceInsights/Private/Insights/InsightsManager.cpp
@@ -182,19 +182,13 @@ void FInsightsManager::SpawnAndActivateTabs()
 
 bool FInsightsManager::IsAnyLiveSessionAvailable(Trace::FSessionHandle& OutLastLiveSessionHandle) const
 {
-	TArray<Trace::FSessionHandle> AvailableSessions;
-	SessionService->GetAvailableSessions(AvailableSessions);
+	TArray<Trace::FSessionHandle> LiveSessions;
+	SessionService->GetLiveSessions(LiveSessions);
 
-	for (int32 SessionIndex = AvailableSessions.Num() - 1; SessionIndex >= 0; --SessionIndex)
+	if (LiveSessions.Num() > 0)
 	{
-		Trace::FSessionInfo SessionInfo;
-		SessionService->GetSessionInfo(AvailableSessions[SessionIndex], SessionInfo);
-
-		if (SessionInfo.bIsLive)
-		{
-			OutLastLiveSessionHandle = AvailableSessions[SessionIndex];
-			return true;
-		}
+		OutLastLiveSessionHandle = LiveSessions.Last();
+		return true;
 	}
 
 	return false;
diff --git a/Engine/Source/Developer/TraceInsights/Private/Insights/Widgets/STimingView.cpp b/Engine/Source/Developer/TraceInsights/Private/Insights/Widgets/STimingView.cpp
--- a/Engine/Source/Developer/TraceInsights/Private/Insights/Widgets/STimingView.cpp
+++ b/Engine/Source/Developer/TraceInsights/Private/Insights/Widgets/STimingView.cpp
@@ -218,6 +218,9 @@ void STimingView::Reset()
 
 	LoadingGetEventNameFn = FLoadingTrackGetEventNameDelegate::CreateRaw(this, &STimingView::GetLoadTimeProfilerEventName);
 
+	EventAggregationStr.Reset();
+	ObjectTypeAggregationStr.Reset();
+
 	IoOverviewTrack = nullptr;
 	IoActivityTrack = nullptr;
 
@@ -1235,6 +1238,41 @@ int32 STimingView::OnPaint(const FPaintArgs& Args, const FGeometry& AllottedGeom
 		}
 	}
 
+	if (bAssetLoadingMode)
+	{
+		const FLinearColor BackgroundColor(0.01f, 0.01f, 0.01f, 0.9f);
+		const FLinearColor TextColor(1.0f, 1.0f, 1.0f, 0.9f);
+
+		constexpr float MarginX = 8.0f;
+		constexpr float MarginY = 8.0f;
+
+		constexpr float BorderX = 4.0f;
+		constexpr float BorderY = 4.0f;
+		float Y = ViewHeight;
+
+		if (ObjectTypeAggregationStr.Len() > 0)
+		{
+			FVector2D TextSize = FontMeasureService->Measure(ObjectTypeAggregationStr, MainFont);
+			const float X = ViewWidth - MarginX - TextSize.X - 2 * BorderX;
+			Y -= MarginY + TextSize.Y + 2 * BorderY;
+			DrawContext.DrawBox(X, Y, TextSize.X + 2 * BorderX, TextSize.Y + 2 * BorderY, WhiteBrush, BackgroundColor);
+			DrawContext.LayerId++;
+			DrawContext.DrawText(X + BorderX, Y + BorderY, ObjectTypeAggregationStr, MainFont, TextColor);
+			DrawContext.LayerId++;
+		}
+
+		if (EventAggregationStr.Len() > 0)
+		{
+			FVector2D TextSize = FontMeasureService->Measure(EventAggregationStr, MainFont);
+			const float X = ViewWidth - MarginX - TextSize.X - 2 * BorderX;
+			Y -= MarginY + TextSize.Y + 2 * BorderY;
+			DrawContext.DrawBox(X, Y, TextSize.X + 2 * BorderX, TextSize.Y + 2 * BorderY, WhiteBrush, BackgroundColor);
+			DrawContext.LayerId++;
+			DrawContext.DrawText(X + BorderX, Y + BorderY, EventAggregationStr, MainFont, TextColor);
+			DrawContext.LayerId++;
+		}
+	}
+
 	return SCompoundWidget::OnPaint(Args, AllottedGeometry, MyCullingRect, OutDrawElements, LayerId, InWidgetStyle, bParentEnabled && IsEnabled());
 }
 
@@ -1988,6 +2026,8 @@ FReply STimingView::OnMouseButtonDown(const FGeometry& MyGeometry, const FPointe
 				SelectionEndTime = SelectionStartTime;
 				LastSelectionType = ESelectionType::None;
 				//TODO: SelectionChangingEvent.Broadcast(SelectionStartTime, SelectionEndTime);
+				EventAggregationStr.Reset();
+				ObjectTypeAggregationStr.Reset();
 			}
 
 			// Capture mouse, so we can drag outside this widget.
@@ -2075,6 +2115,8 @@ FReply STimingView::OnMouseButtonUp(const FGeometry& MyGeometry, const FPointerE
 					SelectionEndTime = SelectionStartTime = 0.0;
 					LastSelectionType = ESelectionType::None;
 					//TODO: SelectionChangedEvent.Broadcast(SelectionStartTime, SelectionEndTime);
+					EventAggregationStr.Reset();
+					ObjectTypeAggregationStr.Reset();
 				}
 			}
 
@@ -2184,6 +2226,8 @@ FReply STimingView::OnMouseMove(const FGeometry& MyGeometry, const FPointerEvent
 			}
 			LastSelectionType = ESelectionType::TimeRange;
 			//TODO: SelectionChangingEvent.Broadcast(SelectionStartTime, SelectionEndTime);
+			EventAggregationStr.Reset();
+			ObjectTypeAggregationStr.Reset();
 		}
 
 		Reply = FReply::Handled();
@@ -2776,6 +2820,105 @@ void STimingView::UpdateAggregatedStats()
 			Wnd->StatsView->UpdateStats(SelectionStartTime, SelectionEndTime);
 		}
 	}
+
+	if (bAssetLoadingMode)
+	{
+		TSharedPtr<const Trace::IAnalysisSession> Session = FInsightsManager::Get()->GetSession();
+		if (Session.IsValid() && Trace::ReadLoadTimeProfilerProvider(*Session.Get()))
+		{
+			Trace::FAnalysisSessionReadScope SessionReadScope(*Session.Get());
+
+			const Trace::ILoadTimeProfilerProvider& LoadTimeProfilerProvider = *Trace::ReadLoadTimeProfilerProvider(*Session.Get());
+
+			Trace::ITable<Trace::FLoadTimeProfilerAggregatedStats>* EventAggregationTable = LoadTimeProfilerProvider.CreateEventAggregation(SelectionStartTime, SelectionEndTime);
+			Trace::ITable<Trace::FLoadTimeProfilerAggregatedStats>* ObjectTypeAggregationTable = LoadTimeProfilerProvider.CreateObjectTypeAggregation(SelectionStartTime, SelectionEndTime);
+
+			auto TableToString = [](const TCHAR* TableName, Trace::ITable<Trace::FLoadTimeProfilerAggregatedStats>* Table) -> FString
+			{
+				FString Str;
+
+				const uint32 TotalRowCount = static_cast<uint32>(Table->GetRowCount());
+
+				Str.Append(TableName).Append(": ").Append(FText::AsNumber(TotalRowCount).ToString()).Append(" records, sorted by Total time");
+
+				constexpr uint32 MaxRowCount = 10; // only get first 10 records
+				const uint32 RowCount = FMath::Min<uint32>(TotalRowCount, MaxRowCount);
+
+				const Trace::ITableLayout& TableLayout = Table->GetLayout();
+				const int32 ColumnCount = TableLayout.GetColumnCount();
+
+				auto Reader = Table->CreateReader();
+
+				//////////////////////////////////////////////////
+
+				struct FSortedIndexEntry
+				{
+					uint32 RowIndex;
+					double Value;
+				};
+
+				TArray<FSortedIndexEntry> SortedIndex;
+				SortedIndex.Reserve(TotalRowCount);
+
+				constexpr uint32 TotalTimeColumnIndex = 2;
+				ensure(TableLayout.GetColumnType(TotalTimeColumnIndex) == Trace::ETableColumnType::TableColumnType_Double);
+
+				for (uint32 RowIndex = 0; RowIndex < TotalRowCount; ++RowIndex)
+				{
+					Reader->SetRowIndex(RowIndex);
+					double Value = Reader->GetValueDouble(TotalTimeColumnIndex);
+					SortedIndex.Add({ RowIndex, Value });
+				}
+
+				SortedIndex.Sort([](const FSortedIndexEntry& A, const FSortedIndexEntry& B) { return A.Value > B.Value; });
+
+				//////////////////////////////////////////////////
+
+				for (uint32 Index = 0; Index < RowCount; ++Index)
+				{
+					uint32 RowIndex = SortedIndex[Index].RowIndex;
+					Reader->SetRowIndex(RowIndex);
+
+					Str.Append("\n").Append(FText::AsNumber(Index + 1).ToString()).Append(". ");
+
+					for (int32 ColumnIndex = 0; ColumnIndex < ColumnCount; ++ColumnIndex)
+					{
+						Str.Append(TableLayout.GetColumnName(ColumnIndex)).Append("=");
+
+						Trace::ETableColumnType ColumnType = TableLayout.GetColumnType(ColumnIndex);
+						switch (ColumnType)
+						{
+						case Trace::ETableColumnType::TableColumnType_Bool:
+							Str.Append(Reader->GetValueBool(ColumnIndex) ? TEXT("True") : TEXT("False")).Append(", ");
+							break;
+						case Trace::ETableColumnType::TableColumnType_Int:
+							Str.Append(FText::AsNumber(Reader->GetValueInt(ColumnIndex)).ToString()).Append(", ");
+							break;
+						case Trace::ETableColumnType::TableColumnType_Float:
+							Str.Append(FText::AsNumber(Reader->GetValueFloat(ColumnIndex)).ToString()).Append(", ");
+							break;
+						case Trace::ETableColumnType::TableColumnType_Double:
+							Str.Append(FText::AsNumber(Reader->GetValueDouble(ColumnIndex)).ToString()).Append(", ");
+							break;
+						case Trace::ETableColumnType::TableColumnType_CString:
+							Str.Append(Reader->GetValueCString(ColumnIndex)).Append(", ");
+							break;
+						}
+					}
+				}
+
+				if (RowCount != TotalRowCount)
+				{
+					Str.Append("\n[...]");
+				}
+
+				return Str;
+			};
+
+			EventAggregationStr = TableToString(TEXT("Event Aggregation"), EventAggregationTable);
+			ObjectTypeAggregationStr = TableToString(TEXT("Object Type Aggregation"), ObjectTypeAggregationTable);
+		}
+	}
 }
 
 ////////////////////////////////////////////////////////////////////////////////////////////////////
@@ -2984,7 +3127,7 @@ bool STimingView::SearchTimingEvent(const double InStartTime,
 	FSearchTimingEventContext Ctx(InStartTime, InEndTime, InPredicate, InOutTimingEvent, bInStopAtFirstMatch, bInSearchForLargestEvent);
 
 	TSharedPtr<const Trace::IAnalysisSession> Session = FInsightsManager::Get()->GetSession();
-	if (Session.IsValid() && Trace::ReadTimingProfilerProvider(*Session.Get()))
+	if (Session.IsValid())
 	{
 		Trace::FAnalysisSessionReadScope SessionReadScope(*Session.Get());
 
@@ -2993,96 +3136,102 @@ bool STimingView::SearchTimingEvent(const double InStartTime,
 		if (Track->GetType() == ETimingEventsTrackType::Cpu ||
 			Track->GetType() == ETimingEventsTrackType::Gpu)
 		{
-			const Trace::ITimingProfilerProvider& TimingProfilerProvider = *Trace::ReadTimingProfilerProvider(*Session.Get());
-
-			TimingProfilerProvider.ReadTimeline(Track->GetId(), [&Ctx](const Trace::ITimingProfilerProvider::Timeline& Timeline)
+			if (Trace::ReadTimingProfilerProvider(*Session.Get()))
 			{
-				Timeline.EnumerateEvents(Ctx.StartTime, Ctx.EndTime, [&Ctx](double EventStartTime, double EventEndTime, uint32 EventDepth, const Trace::FTimingProfilerEvent& Event)
-				{
-					Ctx.CheckEvent(EventStartTime, EventEndTime, EventDepth, Event);
-				});
-			});
+				const Trace::ITimingProfilerProvider& TimingProfilerProvider = *Trace::ReadTimingProfilerProvider(*Session.Get());
 
-			if (Ctx.bFound)
-			{
-				// Compute Exclusive Time.
 				TimingProfilerProvider.ReadTimeline(Track->GetId(), [&Ctx](const Trace::ITimingProfilerProvider::Timeline& Timeline)
 				{
-					struct FEnumerationState
+					Timeline.EnumerateEvents(Ctx.StartTime, Ctx.EndTime, [&Ctx](double EventStartTime, double EventEndTime, uint32 EventDepth, const Trace::FTimingProfilerEvent& Event)
 					{
-						double EventStartTime;
-						double EventEndTime;
-						uint64 EventDepth;
-
-						uint64 CurrentDepth;
-						double LastTime;
-						double ExclusiveTime;
-						bool IsInEventScope;
-					};
-					FEnumerationState State;
-
-					State.EventStartTime = Ctx.TimingEvent.StartTime;
-					State.EventEndTime = Ctx.TimingEvent.EndTime;
-					State.EventDepth = Ctx.TimingEvent.Depth;
-
-					State.CurrentDepth = 0;
-					State.LastTime = 0.0;
-					State.ExclusiveTime = 0.0;
-					State.IsInEventScope = false;
-
-					Timeline.EnumerateEvents(Ctx.TimingEvent.StartTime, Ctx.TimingEvent.EndTime, [&State](bool IsEnter, double Time, const Trace::FTimingProfilerEvent& Event)
+						Ctx.CheckEvent(EventStartTime, EventEndTime, EventDepth, Event);
+					});
+				});
+
+				if (Ctx.bFound)
+				{
+					// Compute Exclusive Time.
+					TimingProfilerProvider.ReadTimeline(Track->GetId(), [&Ctx](const Trace::ITimingProfilerProvider::Timeline& Timeline)
 					{
-						if (IsEnter)
+						struct FEnumerationState
 						{
-							if (State.IsInEventScope && State.CurrentDepth == State.EventDepth + 1)
-							{
-								State.ExclusiveTime += Time - State.LastTime;
-							}
-							if (State.CurrentDepth == State.EventDepth && Time == State.EventStartTime)
+							double EventStartTime;
+							double EventEndTime;
+							uint64 EventDepth;
+
+							uint64 CurrentDepth;
+							double LastTime;
+							double ExclusiveTime;
+							bool IsInEventScope;
+						};
+						FEnumerationState State;
+
+						State.EventStartTime = Ctx.TimingEvent.StartTime;
+						State.EventEndTime = Ctx.TimingEvent.EndTime;
+						State.EventDepth = Ctx.TimingEvent.Depth;
+
+						State.CurrentDepth = 0;
+						State.LastTime = 0.0;
+						State.ExclusiveTime = 0.0;
+						State.IsInEventScope = false;
+
+						Timeline.EnumerateEvents(Ctx.TimingEvent.StartTime, Ctx.TimingEvent.EndTime, [&State](bool IsEnter, double Time, const Trace::FTimingProfilerEvent& Event)
+						{
+							if (IsEnter)
 							{
-								State.IsInEventScope = true;
+								if (State.IsInEventScope && State.CurrentDepth == State.EventDepth + 1)
+								{
+									State.ExclusiveTime += Time - State.LastTime;
+								}
+								if (State.CurrentDepth == State.EventDepth && Time == State.EventStartTime)
+								{
+									State.IsInEventScope = true;
+								}
+								++State.CurrentDepth;
 							}
-							++State.CurrentDepth;
-						}
-						else
-						{
-							--State.CurrentDepth;
-							if (State.CurrentDepth == State.EventDepth && Time == State.EventEndTime)
+							else
 							{
-								State.IsInEventScope = false;
-								State.ExclusiveTime += Time - State.LastTime;
+								--State.CurrentDepth;
+								if (State.CurrentDepth == State.EventDepth && Time == State.EventEndTime)
+								{
+									State.IsInEventScope = false;
+									State.ExclusiveTime += Time - State.LastTime;
+								}
 							}
-						}
-						State.LastTime = Time;
-					});
+							State.LastTime = Time;
+						});
 
-					Ctx.TimingEvent.ExclusiveTime = State.ExclusiveTime;
-				});
+						Ctx.TimingEvent.ExclusiveTime = State.ExclusiveTime;
+					});
+				}
 			}
 		}
 		else if (Track->GetType() == ETimingEventsTrackType::Loading)
 		{
-			const Trace::ILoadTimeProfilerProvider& LoadTimeProfilerProvider = *Trace::ReadLoadTimeProfilerProvider(*Session.Get());
-
-			if (Track == LoadingMainThreadTrack)
+			if (Trace::ReadLoadTimeProfilerProvider(*Session.Get()))
 			{
-				LoadTimeProfilerProvider.ReadMainThreadCpuTimeline([&Ctx](const Trace::ILoadTimeProfilerProvider::CpuTimeline& Timeline)
+				const Trace::ILoadTimeProfilerProvider& LoadTimeProfilerProvider = *Trace::ReadLoadTimeProfilerProvider(*Session.Get());
+
+				if (Track == LoadingMainThreadTrack)
 				{
-					Timeline.EnumerateEvents(Ctx.StartTime, Ctx.EndTime, [&Ctx](double EventStartTime, double EventEndTime, uint32 EventDepth, const Trace::FLoadTimeProfilerCpuEvent& Event)
+					LoadTimeProfilerProvider.ReadMainThreadCpuTimeline([&Ctx](const Trace::ILoadTimeProfilerProvider::CpuTimeline& Timeline)
 					{
-						Ctx.CheckEvent(EventStartTime, EventEndTime, EventDepth, Event);
+						Timeline.EnumerateEvents(Ctx.StartTime, Ctx.EndTime, [&Ctx](double EventStartTime, double EventEndTime, uint32 EventDepth, const Trace::FLoadTimeProfilerCpuEvent& Event)
+						{
+							Ctx.CheckEvent(EventStartTime, EventEndTime, EventDepth, Event);
+						});
 					});
-				});
-			}
-			else if (Track == LoadingAsyncThreadTrack)
-			{
-				LoadTimeProfilerProvider.ReadAsyncLoadingThreadCpuTimeline([&Ctx](const Trace::ILoadTimeProfilerProvider::CpuTimeline& Timeline)
+				}
+				else if (Track == LoadingAsyncThreadTrack)
 				{
-					Timeline.EnumerateEvents(Ctx.StartTime, Ctx.EndTime, [&Ctx](double EventStartTime, double EventEndTime, uint32 EventDepth, const Trace::FLoadTimeProfilerCpuEvent& Event)
+					LoadTimeProfilerProvider.ReadAsyncLoadingThreadCpuTimeline([&Ctx](const Trace::ILoadTimeProfilerProvider::CpuTimeline& Timeline)
 					{
-						Ctx.CheckEvent(EventStartTime, EventEndTime, EventDepth, Event);
+						Timeline.EnumerateEvents(Ctx.StartTime, Ctx.EndTime, [&Ctx](double EventStartTime, double EventEndTime, uint32 EventDepth, const Trace::FLoadTimeProfilerCpuEvent& Event)
+						{
+							Ctx.CheckEvent(EventStartTime, EventEndTime, EventDepth, Event);
+						});
 					});
-				});
+				}
 			}
 		}
 	}
diff --git a/Engine/Source/Developer/TraceInsights/Private/Insights/Widgets/STimingView.h b/Engine/Source/Developer/TraceInsights/Private/Insights/Widgets/STimingView.h
--- a/Engine/Source/Developer/TraceInsights/Private/Insights/Widgets/STimingView.h
+++ b/Engine/Source/Developer/TraceInsights/Private/Insights/Widgets/STimingView.h
@@ -410,6 +410,9 @@ protected:
 
 	bool bAssetLoadingMode;
 
+	FString EventAggregationStr;
+	FString ObjectTypeAggregationStr;
+
 	////////////////////////////////////////////////////////////
 	// File activity (I/O)
 
