[Insights]
  - Start Page: Updated UI layout for Trace Recorder.
  - Start Page: Added "Connect" button for creating a new connection.
  - Start Page: Added list of connections (live sessions) + "Refresh" button.
  - Start Page: For selected of connection, listed the available modules. Now the state of modules can be set per connection.

diff --git a/Engine/Source/Developer/TraceInsights/Private/Insights/Widgets/SStartPageWindow.cpp b/Engine/Source/Developer/TraceInsights/Private/Insights/Widgets/SStartPageWindow.cpp
--- a/Engine/Source/Developer/TraceInsights/Private/Insights/Widgets/SStartPageWindow.cpp
+++ b/Engine/Source/Developer/TraceInsights/Private/Insights/Widgets/SStartPageWindow.cpp
@@ -9,7 +9,6 @@
 #include "Framework/MultiBox/MultiBoxBuilder.h"
 #include "HAL/FileManagerGeneric.h"
 #include "SlateOptMacros.h"
-#include "TraceServices/SessionService.h"
 #include "Widgets/Docking/SDockTab.h"
 #include "Widgets/Images/SImage.h"
 #include "Widgets/Layout/SBorder.h"
@@ -19,7 +18,9 @@
 #include "Widgets/SBoxPanel.h"
 #include "Widgets/Text/STextBlock.h"
 #include "Widgets/Input/SButton.h"
+#include "Widgets/Input/SEditableTextBox.h"
 #include "Widgets/Input/SComboButton.h"
+#include "Widgets/Views/SListView.h"
 //#include "WorkspaceMenuStructure.h"
 //#include "WorkspaceMenuStructureModule.h"
 
@@ -43,7 +44,6 @@
 
 ////////////////////////////////////////////////////////////////////////////////////////////////////
 
-
 static FText GetTextForNotification(const EInsightsNotificationType NotificatonType, const ELoadingProgressState ProgressState, const FString& Filename, const float ProgressPercent = 0.0f)
 {
 	FText Result;
@@ -80,6 +80,106 @@ static FText GetTextForNotification(const EInsightsNotificationType NotificatonT
 	return Result;
 }
 
+////////////////////////////////////////////////////////////////////////////////////////////////////
+// SConnectionRow
+////////////////////////////////////////////////////////////////////////////////////////////////////
+
+class SConnectionRow : public SMultiColumnTableRow<TSharedPtr<FRecorderConnection>>
+{
+	SLATE_BEGIN_ARGS(SConnectionRow) {}
+	SLATE_END_ARGS()
+
+public:
+	/**
+	 * Constructs the widget.
+	 *
+	 * @param InArgs The construction arguments.
+	 * @param InConnection The connection displayed by this row.
+	 * @param InOwnerTableView The table to which the row must be added.
+	 */
+	void Construct(const FArguments& InArgs, TSharedPtr<FRecorderConnection> InConnection, TSharedRef<SStartPageWindow> InParentWidget, const TSharedRef<STableViewBase>& InOwnerTableView)
+	{
+		WeakConnection = MoveTemp(InConnection);
+		WeakParentWidget = InParentWidget;
+
+		SMultiColumnTableRow<TSharedPtr<FRecorderConnection>>::Construct(FSuperRowType::FArguments(), InOwnerTableView);
+	}
+
+public:
+	virtual TSharedRef<SWidget> GenerateWidgetForColumn(const FName& ColumnName) override
+	{
+		if (ColumnName == FName(TEXT("Name")))
+		{
+			return SNew(SBox)
+				.Padding(FMargin(4.0, 0.0))
+				[
+					SNew(STextBlock)
+					.Text(this, &SConnectionRow::GetConnectionName)
+					.ToolTipText(this, &SConnectionRow::GetConnectionName)
+				];
+		}
+		else if (ColumnName == FName(TEXT("Uri")))
+		{
+			return SNew(SBox)
+				.Padding(FMargin(4.0, 0.0))
+				[
+					SNew(STextBlock)
+					.Text(this, &SConnectionRow::GetConnectionUri)
+					.ToolTipText(this, &SConnectionRow::GetConnectionUri)
+				];
+		}
+		else
+		{
+			return SNew(STextBlock).Text(LOCTEXT("UnknownColumn", "Unknown Column"));
+		}
+	}
+
+	FText GetConnectionName() const
+	{
+		TSharedPtr<FRecorderConnection> ConnectionPin = WeakConnection.Pin();
+		if (ConnectionPin.IsValid())
+		{
+			return ConnectionPin->Name;
+		}
+		else
+		{
+			return FText();
+		}
+	}
+
+	FText GetConnectionUri() const
+	{
+		TSharedPtr<FRecorderConnection> ConnectionPin = WeakConnection.Pin();
+		if (ConnectionPin.IsValid())
+		{
+			return ConnectionPin->Uri;
+		}
+		else
+		{
+			return FText();
+		}
+	}
+
+	FText GetRowToolTip() const
+	{
+		TSharedPtr<FRecorderConnection> ConnectionPin = WeakConnection.Pin();
+		if (ConnectionPin.IsValid())
+		{
+			return ConnectionPin->Uri;
+		}
+		else
+		{
+			return FText();
+		}
+	}
+
+private:
+	TWeakPtr<FRecorderConnection> WeakConnection;
+	TWeakPtr<SStartPageWindow> WeakParentWidget;
+};
+
+////////////////////////////////////////////////////////////////////////////////////////////////////
+// SStartPageWindow
 ////////////////////////////////////////////////////////////////////////////////////////////////////
 
 SStartPageWindow::SStartPageWindow()
@@ -346,23 +446,12 @@ void SStartPageWindow::Construct(const FArguments& InArgs)
 						.HAlign(HAlign_Center)
 						.Padding(3.0f, 3.0f)
 						[
-							SNew(SHorizontalBox)
-
-							+ SHorizontalBox::Slot()
-								.AutoWidth()
-								.VAlign(VAlign_Top)
-								.Padding(3.0f, 0.0f)
-								[
-									ConstructRecorderStatus()
-								]
-
-							+ SHorizontalBox::Slot()
-								.AutoWidth()
-								.VAlign(VAlign_Top)
-								.Padding(3.0f, 0.0f)
-								[
-									ConstructModuleList()
-								]
+							SNew(SBorder)
+							.BorderImage(FEditorStyle::GetBrush("NotificationList.ItemBackground"))
+							.Padding(8.0f)
+							[
+								ConstructRecoderPanel()
+							]
 						]
 				]
 
@@ -385,101 +474,206 @@ void SStartPageWindow::Construct(const FArguments& InArgs)
 
 ////////////////////////////////////////////////////////////////////////////////////////////////////
 
-TSharedRef<SWidget> SStartPageWindow::ConstructRecorderStatus()
+TSharedRef<SWidget> SStartPageWindow::ConstructRecoderPanel()
 {
-	TSharedRef<SWidget> BorderWidget =
-		SNew(SBorder)
-		.BorderImage(FEditorStyle::GetBrush("NotificationList.ItemBackground"))
-		.Padding(8.0f)
+	TSharedRef<SWidget> Widget = SNew(SVerticalBox)
+
+	+ SVerticalBox::Slot()
+		.AutoHeight()
+		.HAlign(HAlign_Center)
+		.Padding(0.0, 2.0f)
 		[
-			SNew(SVerticalBox)
+			SNew(STextBlock)
+			.Text(LOCTEXT("RecorderPanelTitle", "Trace Recorder"))
+			.Font(FCoreStyle::GetDefaultFontStyle("Bold", 11))
+			.ColorAndOpacity(FLinearColor::Gray)
+		]
 
-			+ SVerticalBox::Slot()
-				.AutoHeight()
-				.HAlign(HAlign_Center)
-				.Padding(0.0, 2.0f)
+	+ SVerticalBox::Slot()
+		.AutoHeight()
+		.HAlign(HAlign_Left)
+		.Padding(0.0, 2.0f)
+		[
+			SNew(SHorizontalBox)
+
+			+ SHorizontalBox::Slot()
+				.AutoWidth()
+				.Padding(0.0f, 0.0f, 2.0f, 0.0f)
+				.VAlign(VAlign_Center)
 				[
 					SNew(STextBlock)
-					.Text(LOCTEXT("RecorderStatusTitle", "Trace Recorder Status:"))
+					.Text(LOCTEXT("RecorderStatusTitle", "Status:"))
 					.ColorAndOpacity(FLinearColor::Gray)
 				]
 
-			+ SVerticalBox::Slot()
-				.AutoHeight()
-				.HAlign(HAlign_Center)
-				.Padding(0.0, 2.0f)
+			+ SHorizontalBox::Slot()
+				.AutoWidth()
+				.Padding(2.0f, 0.0f)
+				.VAlign(VAlign_Center)
 				[
-					SNew(SHorizontalBox)
-
-					+ SHorizontalBox::Slot()
-						.AutoWidth()
-						.Padding(2.0f, 0.0f)
-						.VAlign(VAlign_Center)
-						[
-							SNew(STextBlock)
-							.Text(this, &SStartPageWindow::GetRecorderStatusText)
-						]
+					SNew(STextBlock)
+					.Text(this, &SStartPageWindow::GetRecorderStatusText)
+				]
 
-					+ SHorizontalBox::Slot()
-						.AutoWidth()
-						.Padding(2.0f, 0.0f)
-						.VAlign(VAlign_Center)
-						[
-							SNew(SButton)
-							.Text(LOCTEXT("StartRecorder", "Start"))
-							.ToolTipText(LOCTEXT("StartRecorderToolTip", "Start the Trace Recorder"))
-							.OnClicked(this, &SStartPageWindow::StartTraceRecorder_OnClicked)
-							.Visibility(this, &SStartPageWindow::StartTraceRecorder_Visibility)
-						]
+			+ SHorizontalBox::Slot()
+				.AutoWidth()
+				.Padding(2.0f, 0.0f)
+				.VAlign(VAlign_Center)
+				[
+					SNew(SButton)
+					.Text(LOCTEXT("StartRecorder", "Start"))
+					.ToolTipText(LOCTEXT("StartRecorderToolTip", "Start the Trace Recorder"))
+					.OnClicked(this, &SStartPageWindow::StartTraceRecorder_OnClicked)
+					.Visibility(this, &SStartPageWindow::StartTraceRecorder_Visibility)
+				]
 
-					+ SHorizontalBox::Slot()
-						.AutoWidth()
-						.Padding(2.0f, 0.0f)
-						.VAlign(VAlign_Center)
-						[
-							SNew(SButton)
-							.Text(LOCTEXT("StopRecorder", "Stop"))
-							.ToolTipText(LOCTEXT("StopRecorderToolTip", "Stop the Trace Recorder"))
-							.OnClicked(this, &SStartPageWindow::StopTraceRecorder_OnClicked)
-							.Visibility(this, &SStartPageWindow::StopTraceRecorder_Visibility)
-						]
+			+ SHorizontalBox::Slot()
+				.AutoWidth()
+				.Padding(2.0f, 0.0f)
+				.VAlign(VAlign_Center)
+				[
+					SNew(SButton)
+					.Text(LOCTEXT("StopRecorder", "Stop"))
+					.ToolTipText(LOCTEXT("StopRecorderToolTip", "Stop the Trace Recorder"))
+					.OnClicked(this, &SStartPageWindow::StopTraceRecorder_OnClicked)
+					.Visibility(this, &SStartPageWindow::StopTraceRecorder_Visibility)
 				]
-		];
+		]
 
-	return BorderWidget;
-}
+	+ SVerticalBox::Slot()
+		.AutoHeight()
+		.HAlign(HAlign_Fill)
+		.Padding(0.0, 2.0f)
+		[
+			SNew(SHorizontalBox)
+			.Visibility(this, &SStartPageWindow::StopTraceRecorder_Visibility)
 
-////////////////////////////////////////////////////////////////////////////////////////////////////
+			+ SHorizontalBox::Slot()
+				.AutoWidth()
+				.Padding(0.0f, 0.0f, 2.0f, 0.0f)
+				.VAlign(VAlign_Center)
+				[
+					SNew(STextBlock)
+					.Text(LOCTEXT("HostTitle", "Host:"))
+					.ColorAndOpacity(FLinearColor::Gray)
+				]
 
-TSharedRef<SWidget> SStartPageWindow::ConstructModuleList()
-{
-	TSharedPtr<SVerticalBox> VerticalBox;
+			+ SHorizontalBox::Slot()
+				.FillWidth(1.0)
+				.VAlign(VAlign_Center)
+				[
+					SAssignNew(HostTextBox, SEditableTextBox)
+				]
 
-	TSharedRef<SWidget> BorderWidget =
-		SNew(SBorder)
-		.BorderImage(FEditorStyle::GetBrush("NotificationList.ItemBackground"))
-		.Padding(8.0f)
+			+ SHorizontalBox::Slot()
+				.AutoWidth()
+				.Padding(2.0f, 0.0f)
+				.VAlign(VAlign_Center)
+				[
+					SNew(SButton)
+					.Text(LOCTEXT("Connect", "Connect"))
+					.ToolTipText(LOCTEXT("ConnectToolTip", "Try connecting to host."))
+					.OnClicked(this, &SStartPageWindow::Connect_OnClicked)
+				]
+		]
+
+	+ SVerticalBox::Slot()
+		.AutoHeight()
+		.HAlign(HAlign_Left)
+		.Padding(0.0, 2.0f, 0.0f, 1.0f)
 		[
-			SNew(SVerticalBox)
+			SNew(SHorizontalBox)
+			.Visibility(this, &SStartPageWindow::StopTraceRecorder_Visibility)
 
-			+ SVerticalBox::Slot()
-				.AutoHeight()
-				.HAlign(HAlign_Center)
-				.Padding(0.0, 2.0f)
+			+ SHorizontalBox::Slot()
+				.AutoWidth()
+				.Padding(0.0f, 0.0f, 2.0f, 0.0f)
+				.VAlign(VAlign_Center)
 				[
 					SNew(STextBlock)
-					.Text(LOCTEXT("ModulesTitle", "Modules:"))
+					.Text_Lambda([this]() -> FText
+					{
+						return FText::Format(LOCTEXT("ConnectionListTitle", "Connections (live sessions): {0}"), FText::AsNumber(Connections.Num()));
+					})
 					.ColorAndOpacity(FLinearColor::Gray)
 				]
 
-			+ SVerticalBox::Slot()
-				.AutoHeight()
-				.HAlign(HAlign_Center)
-				.Padding(0.0, 2.0f)
+			+ SHorizontalBox::Slot()
+				.AutoWidth()
+				.Padding(2.0f, 0.0f)
+				.VAlign(VAlign_Center)
 				[
-					SAssignNew(VerticalBox, SVerticalBox)
+					SNew(SButton)
+					.Text(LOCTEXT("Refresh", "Refresh"))
+					.ToolTipText(LOCTEXT("RefreshToolTip", "Refresh the connection list (live sessions)."))
+					.OnClicked(this, &SStartPageWindow::RefreshConnections_OnClicked)
 				]
-		];
+		]
+
+	+ SVerticalBox::Slot()
+		.AutoHeight()
+		.HAlign(HAlign_Left)
+		.Padding(0.0, 1.0f, 0.0f, 2.0f)
+		[
+			SNew(SBox)
+			.WidthOverride(640.0f)
+			.MaxDesiredHeight(78.0f)
+			.Visibility(this, &SStartPageWindow::Connections_Visibility)
+			[
+				SAssignNew(ConnectionsListView, SListView<TSharedPtr<FRecorderConnection>>)
+				.ItemHeight(20.0f)
+				.SelectionMode(ESelectionMode::Single)
+				.OnSelectionChanged(this, &SStartPageWindow::Connections_OnSelectionChanged)
+				.ListItemsSource(&Connections)
+				.OnGenerateRow(this, &SStartPageWindow::Connections_OnGenerateRow)
+				.ConsumeMouseWheel(EConsumeMouseWheel::Always)
+				//.OnContextMenuOpening(FOnContextMenuOpening::CreateSP(this, &SStartPageWindow::Connections_GetContextMenu))
+				.HeaderRow
+				(
+					SNew(SHeaderRow)
+
+					+ SHeaderRow::Column(FName(TEXT("Name")))
+					.FillWidth(0.2f)
+					.DefaultLabel(LOCTEXT("NameColumn", "Name"))
+
+					+ SHeaderRow::Column(FName(TEXT("Uri")))
+					.FillWidth(0.8f)
+					.DefaultLabel(LOCTEXT("UriColumn", "URI"))
+				)
+			]
+		]
+
+	+ SVerticalBox::Slot()
+		.AutoHeight()
+		.HAlign(HAlign_Left)
+		.Padding(0.0, 2.0f, 0.0f, 1.0f)
+		[
+			SNew(STextBlock)
+			.Visibility(this, &SStartPageWindow::Modules_Visibility)
+			.Text(LOCTEXT("ModulesListTitle", "Modules for selected connection:"))
+			.ColorAndOpacity(FLinearColor::Gray)
+		]
+
+	+ SVerticalBox::Slot()
+		.AutoHeight()
+		.HAlign(HAlign_Left)
+		.Padding(0.0, 1.0f, 0.0f, 2.0f)
+		[
+			ConstructModuleList()
+		]
+	;
+
+	RefreshConnectionList();
+
+	return Widget;
+}
+
+////////////////////////////////////////////////////////////////////////////////////////////////////
+
+TSharedRef<SWidget> SStartPageWindow::ConstructModuleList()
+{
+	TSharedRef<SVerticalBox> VerticalBox = SNew(SVerticalBox)
+		.Visibility(this, &SStartPageWindow::Modules_Visibility);
 
 	TSharedRef<Trace::IModuleService> ModuleService = FInsightsManager::Get()->GetModuleService();
 	ModuleService->GetAvailableModules(AvailableModules);
@@ -508,23 +702,33 @@ TSharedRef<SWidget> SStartPageWindow::ConstructModuleList()
 			];
 	}
 
-	return BorderWidget;
+	return VerticalBox;
+}
+
+////////////////////////////////////////////////////////////////////////////////////////////////////
+
+TSharedRef<ITableRow> SStartPageWindow::Connections_OnGenerateRow(TSharedPtr<FRecorderConnection> InConnection, const TSharedRef<STableViewBase>& OwnerTable)
+{
+	return SNew(SConnectionRow, InConnection, SharedThis(this), OwnerTable);
 }
 
 END_SLATE_FUNCTION_BUILD_OPTIMIZATION
 
 ////////////////////////////////////////////////////////////////////////////////////////////////////
 
+EVisibility SStartPageWindow::Modules_Visibility() const
+{
+	TSharedRef<Trace::ISessionService> SessionService = FInsightsManager::Get()->GetSessionService();
+	return (SessionService->IsRecorderServerRunning() && Connections.Num() > 0 && SelectedConnection.IsValid()) ? EVisibility::Visible : EVisibility::Collapsed;
+}
+
+////////////////////////////////////////////////////////////////////////////////////////////////////
+
 ECheckBoxState SStartPageWindow::Module_IsChecked(int32 ModuleIndex) const
 {
-	if (ModuleIndex >= 0 && ModuleIndex < AvailableModules.Num())
+	if (ModuleIndex >= 0 && ModuleIndex < AvailableModulesEnabledState.Num())
 	{
-		//const Trace::FModuleInfo& Module = AvailableModules[ModuleIndex];
-		//const bool bIsEnabled = Module.bIsEnabled;
-
-		ensure(AvailableModulesEnabledState.Num() == AvailableModules.Num());
 		const bool bIsEnabled = AvailableModulesEnabledState[ModuleIndex];
-
 		return bIsEnabled ? ECheckBoxState::Checked : ECheckBoxState::Unchecked;
 	}
 	else
@@ -537,21 +741,137 @@ ECheckBoxState SStartPageWindow::Module_IsChecked(int32 ModuleIndex) const
 
 void SStartPageWindow::Module_OnCheckStateChanged(ECheckBoxState NewRadioState, int32 ModuleIndex)
 {
-	if (ModuleIndex >= 0 && ModuleIndex < AvailableModules.Num())
+	ensure(AvailableModulesEnabledState.Num() == AvailableModules.Num());
+
+	if (SelectedConnection.IsValid() &&
+		ModuleIndex >= 0 && ModuleIndex < AvailableModules.Num())
 	{
 		bool bIsEnabled = NewRadioState == ECheckBoxState::Checked;
 
-		ensure(AvailableModulesEnabledState.Num() == AvailableModules.Num());
-		AvailableModulesEnabledState[ModuleIndex] = bIsEnabled;
+		TSharedRef<Trace::ISessionService> SessionService = FInsightsManager::Get()->GetSessionService();
+		
+		SessionService->SetModuleEnabled(SelectedConnection->SessionHandle, AvailableModules[ModuleIndex].Name, bIsEnabled);
+		//AvailableModulesEnabledState[ModuleIndex] = bIsEnabled;
+		AvailableModulesEnabledState[ModuleIndex] = SessionService->IsModuleEnabled(SelectedConnection->SessionHandle, AvailableModules[ModuleIndex].Name);
+	}
+}
+
+////////////////////////////////////////////////////////////////////////////////////////////////////
+
+FReply SStartPageWindow::RefreshConnections_OnClicked()
+{
+	RefreshConnectionList();
+	return FReply::Handled();
+}
+
+////////////////////////////////////////////////////////////////////////////////////////////////////
+
+FReply SStartPageWindow::Connect_OnClicked()
+{
+	TSharedRef<Trace::ISessionService> SessionService = FInsightsManager::Get()->GetSessionService();
+
+	FText HostText = HostTextBox->GetText();
+	if (HostText.IsEmptyOrWhitespace())
+	{
+		//...
+	}
+	else if (SessionService->ConnectSession(*HostText.ToString()))
+	{
+		FNotificationInfo NotificationInfo(FText::Format(LOCTEXT("ConnectSuccess", "Successfully connected to \"{0}\"!"), HostText));
+		NotificationInfo.bFireAndForget = false;
+		NotificationInfo.bUseLargeFont = false;
+		NotificationInfo.bUseSuccessFailIcons = true;
+		NotificationInfo.ExpireDuration = 10.0f;
+		SNotificationItemWeak NotificationItem = NotificationList->AddNotification(NotificationInfo);
+		NotificationItem.Pin()->SetCompletionState(SNotificationItem::CS_Success);
+		NotificationItem.Pin()->ExpireAndFadeout();
+		ActiveNotifications.Add(TEXT("ConnectSuccess"), NotificationItem);
+	}
+	else
+	{
+		FNotificationInfo NotificationInfo(FText::Format(LOCTEXT("ConnectFailed", "Failed to connect to \"{0}\"!"), HostText));
+		NotificationInfo.bFireAndForget = false;
+		NotificationInfo.bUseLargeFont = false;
+		NotificationInfo.bUseSuccessFailIcons = true;
+		NotificationInfo.ExpireDuration = 10.0f;
+		SNotificationItemWeak NotificationItem = NotificationList->AddNotification(NotificationInfo);
+		NotificationItem.Pin()->SetCompletionState(SNotificationItem::CS_Fail);
+		NotificationItem.Pin()->ExpireAndFadeout();
+		ActiveNotifications.Add(TEXT("ConnectFailed"), NotificationItem);
+	}
+
+	RefreshConnectionList();
+	return FReply::Handled();
+}
+
+////////////////////////////////////////////////////////////////////////////////////////////////////
+
+void SStartPageWindow::RefreshConnectionList()
+{
+	TSharedRef<Trace::ISessionService> SessionService = FInsightsManager::Get()->GetSessionService();
+
+	TArray<Trace::FSessionHandle> LiveSessions;
+	SessionService->GetLiveSessions(LiveSessions);
+
+	TSharedPtr<FRecorderConnection> NewSelectedConnection;
+
+	Connections.Reset();
+
+	for (Trace::FSessionHandle SessionHandle : LiveSessions)
+	{
+		Trace::FSessionInfo SessionInfo;
+		SessionService->GetSessionInfo(SessionHandle, SessionInfo);
+
+		Connections.Add(MakeShareable(new FRecorderConnection(SessionHandle, SessionInfo)));
+
+		if (SelectedConnection && SelectedConnection->SessionHandle == SessionHandle)
+		{
+			NewSelectedConnection = Connections.Last();
+		}
+	}
+
+	ConnectionsListView->RebuildList();
+
+	if (NewSelectedConnection.IsValid())
+	{
+		ConnectionsListView->SetItemSelection(NewSelectedConnection, true);
+		ConnectionsListView->RequestScrollIntoView(NewSelectedConnection);
+	}
+}
+
+////////////////////////////////////////////////////////////////////////////////////////////////////
+
+void SStartPageWindow::Connections_OnSelectionChanged(TSharedPtr<FRecorderConnection> Connection, ESelectInfo::Type SelectInfo)
+{
+	SelectedConnection = Connection;
 
-		TSharedRef<Trace::IModuleService> ModuleService = FInsightsManager::Get()->GetModuleService();
-		Trace::FModuleInfo& Module = AvailableModules[ModuleIndex];
-		ModuleService->SetModuleEnabled(Module.Name, bIsEnabled);
+	if (Connection.IsValid())
+	{
+		TSharedRef<Trace::ISessionService> SessionService = FInsightsManager::Get()->GetSessionService();
+		for (int32 Index = 0; Index < AvailableModulesEnabledState.Num(); ++Index)
+		{
+			AvailableModulesEnabledState[Index] = SessionService->IsModuleEnabled(SelectedConnection->SessionHandle, AvailableModules[Index].Name);
+		}
+	}
+	else
+	{
+		for (int32 Index = 0; Index < AvailableModulesEnabledState.Num(); ++Index)
+		{
+			AvailableModulesEnabledState[Index] = false;
+		}
 	}
 }
 
 ////////////////////////////////////////////////////////////////////////////////////////////////////
 
+EVisibility SStartPageWindow::Connections_Visibility() const
+{
+	TSharedRef<Trace::ISessionService> SessionService = FInsightsManager::Get()->GetSessionService();
+	return (SessionService->IsRecorderServerRunning() && Connections.Num() > 0) ? EVisibility::Visible : EVisibility::Collapsed;
+}
+
+////////////////////////////////////////////////////////////////////////////////////////////////////
+
 void SStartPageWindow::Tick(const FGeometry& AllottedGeometry, const double InCurrentTime, const float InDeltaTime)
 {
 	// We need to check if there is any available session, but not too often.
@@ -563,7 +883,26 @@ void SStartPageWindow::Tick(const FGeometry& AllottedGeometry, const double InCu
 		NextTimestamp = Time + WaitTime;
 
 		bIsAnySessionAvailable = FInsightsManager::Get()->IsAnySessionAvailable(LastSessionHandle);
-		bIsAnyLiveSessionAvailable = FInsightsManager::Get()->IsAnyLiveSessionAvailable(LastLiveSessionHandle);
+		//bIsAnyLiveSessionAvailable = FInsightsManager::Get()->IsAnyLiveSessionAvailable(LastLiveSessionHandle);
+
+		TSharedRef<Trace::ISessionService> SessionService = FInsightsManager::Get()->GetSessionService();
+		TArray<Trace::FSessionHandle> LiveSessions;
+		SessionService->GetLiveSessions(LiveSessions);
+
+		if (LiveSessions.Num() > 0)
+		{
+			LastLiveSessionHandle = LiveSessions.Last();
+			bIsAnyLiveSessionAvailable = true;
+		}
+		else
+		{
+			bIsAnyLiveSessionAvailable = false;
+		}
+
+		if (LiveSessions.Num() != Connections.Num())
+		{
+			RefreshConnectionList();
+		}
 	}
 }
 
@@ -882,6 +1221,8 @@ void SStartPageWindow::LoadSession(Trace::FSessionHandle SessionHandle)
 
 TSharedRef<SWidget> SStartPageWindow::MakeSessionListMenu()
 {
+	RefreshConnectionList();
+
 	FMenuBuilder MenuBuilder(/*bInShouldCloseWindowAfterMenuSelection=*/true, nullptr);
 
 	MenuBuilder.BeginSection("MostRecentlyUsedSessions", LOCTEXT("MostRecentlyUsedSessionsHeading", "Most Recently Used Sessions"));
@@ -893,6 +1234,7 @@ TSharedRef<SWidget> SStartPageWindow::MakeSessionListMenu()
 	MenuBuilder.BeginSection("AvailableSessions", LOCTEXT("AvailableSessionsHeading", "Available Sessions"));
 	{
 		TSharedRef<Trace::ISessionService> SessionService = FInsightsManager::Get()->GetSessionService();
+
 		TArray<Trace::FSessionHandle> AvailableSessions;
 		SessionService->GetAvailableSessions(AvailableSessions);
 
@@ -980,6 +1322,7 @@ FReply SStartPageWindow::StartTraceRecorder_OnClicked()
 {
 	TSharedRef<Trace::ISessionService> SessionService = FInsightsManager::Get()->GetSessionService();
 	SessionService->StartRecorderServer();
+	RefreshConnectionList();
 	return FReply::Handled();
 }
 
@@ -989,6 +1332,7 @@ FReply SStartPageWindow::StopTraceRecorder_OnClicked()
 {
 	TSharedRef<Trace::ISessionService> SessionService = FInsightsManager::Get()->GetSessionService();
 	SessionService->StopRecorderServer();
+	RefreshConnectionList();
 	return FReply::Handled();
 }
 
diff --git a/Engine/Source/Developer/TraceInsights/Private/Insights/Widgets/SStartPageWindow.h b/Engine/Source/Developer/TraceInsights/Private/Insights/Widgets/SStartPageWindow.h
--- a/Engine/Source/Developer/TraceInsights/Private/Insights/Widgets/SStartPageWindow.h
+++ b/Engine/Source/Developer/TraceInsights/Private/Insights/Widgets/SStartPageWindow.h
@@ -8,6 +8,7 @@
 #include "Layout/Visibility.h"
 #include "Misc/Guid.h"
 #include "SlateFwd.h"
+#include "TraceServices/SessionService.h"
 #include "Widgets/DeclarativeSyntaxSupport.h"
 #include "Widgets/Docking/SDockTab.h"
 #include "Widgets/Layout/SSplitter.h"
@@ -21,6 +22,7 @@
 
 class FActiveTimerHandle;
 class SVerticalBox;
+class SEditableTextBox;
 
 ////////////////////////////////////////////////////////////////////////////////////////////////////
 
@@ -35,6 +37,19 @@ typedef TWeakPtr<class SNotificationItem> SNotificationItemWeak;
 
 ////////////////////////////////////////////////////////////////////////////////////////////////////
 
+struct FRecorderConnection
+{
+	Trace::FSessionHandle SessionHandle;
+	FText Name;
+	FText Uri;
+
+	FRecorderConnection(const Trace::FSessionHandle InSessionHandle, const Trace::FSessionInfo& InSessionInfo)
+		: SessionHandle(InSessionHandle)
+		, Name(FText::FromString(InSessionInfo.Name))
+		, Uri(FText::FromString(InSessionInfo.Uri))
+	{}
+};
+
 /** Implements the Start Page window. */
 class SStartPageWindow : public SCompoundWidget
 {
@@ -57,9 +72,12 @@ public:
 	void CloseSettings();
 
 private:
-	TSharedRef<SWidget> ConstructRecorderStatus();
+	TSharedRef<SWidget> ConstructRecoderPanel();
 	TSharedRef<SWidget> ConstructModuleList();
 
+	/** Generate a new row for the Connections list view. */
+	TSharedRef<ITableRow> Connections_OnGenerateRow(TSharedPtr<FRecorderConnection> InConnection, const TSharedRef<STableViewBase>& OwnerTable);
+
 	/** Callback for determining the visibility of the 'Please select a trace' overlay. */
 	EVisibility IsSessionOverlayVisible() const;
 
@@ -89,9 +107,16 @@ private:
 	FReply StartTraceRecorder_OnClicked();
 	FReply StopTraceRecorder_OnClicked();
 
+	EVisibility Modules_Visibility() const;
 	ECheckBoxState Module_IsChecked(int32 ModuleIndex) const;
 	void Module_OnCheckStateChanged(ECheckBoxState NewRadioState, int32 ModuleIndex);
 
+	FReply RefreshConnections_OnClicked();
+	FReply Connect_OnClicked();
+	void RefreshConnectionList();
+	void Connections_OnSelectionChanged(TSharedPtr<FRecorderConnection> Connection, ESelectInfo::Type SelectInfo);
+	EVisibility Connections_Visibility() const;
+
 	/** Updates the amount of time the profiler has been active. */
 	EActiveTimerReturnType UpdateActiveDuration(double InCurrentTime, float InDeltaTime);
 
@@ -175,6 +200,10 @@ private:
 	bool bIsAnySessionAvailable;
 	Trace::FSessionHandle LastSessionHandle;
 
+	TSharedPtr<SListView<TSharedPtr<FRecorderConnection>>> ConnectionsListView;
+	TArray<TSharedPtr<FRecorderConnection>> Connections;
+	TSharedPtr<SEditableTextBox> HostTextBox;
+	TSharedPtr<FRecorderConnection> SelectedConnection;
 	TArray<Trace::FModuleInfo> AvailableModules;
 	TArray<bool> AvailableModulesEnabledState;
 };
