[Insights ] Fixed UE-75588 (Timing Insights Timing View spawns an empty box when right clicking).

diff --git a/Engine/Source/Developer/TraceInsights/Private/Insights/Widgets/STimingView.cpp b/Engine/Source/Developer/TraceInsights/Private/Insights/Widgets/STimingView.cpp
--- a/Engine/Source/Developer/TraceInsights/Private/Insights/Widgets/STimingView.cpp
+++ b/Engine/Source/Developer/TraceInsights/Private/Insights/Widgets/STimingView.cpp
@@ -2652,17 +2652,20 @@ FReply STimingView::OnKeyUp(const FGeometry& MyGeometry, const FKeyEvent& InKeyE
 
 void STimingView::ShowContextMenu(const FVector2D& ScreenSpacePosition, const FPointerEvent& MouseEvent)
 {
+	/* TODO
 	TSharedPtr<FUICommandList> ProfilerCommandList = FTimingProfilerManager::Get()->GetCommandList();
 	const FTimingProfilerCommands& ProfilerCommands = FTimingProfilerManager::GetCommands();
-	const FTimingProfilerActionManager& ProfilerActionManager = FTimingProfilerManager::GetActionManager();
+	const FTimingViewCommands& Commands = FTimingViewCommands::Get();
 
 	const bool bShouldCloseWindowAfterMenuSelection = true;
 	FMenuBuilder MenuBuilder(bShouldCloseWindowAfterMenuSelection, ProfilerCommandList);
 
 	MenuBuilder.BeginSection(TEXT("Misc"), LOCTEXT("Miscellaneous", "Miscellaneous"));
 	{
-		//MenuBuilder.AddMenuEntry(FTimingProfilerManager::GetCommands().EventGraph_SelectAllFrames);
-		//MenuBuilder.AddMenuEntry(FTimingProfilerManager::GetCommands().ProfilerManager_ToggleLivePreview);
+		MenuBuilder.AddMenuEntry(Commands.ShowAllGpuTracks);
+		MenuBuilder.AddMenuEntry(Commands.ShowAllCpuTracks);
+		MenuBuilder.AddMenuEntry(ProfilerCommands.ToggleTimersViewVisibility);
+		MenuBuilder.AddMenuEntry(ProfilerCommands.ToggleStatsCountersViewVisibility);
 	}
 	MenuBuilder.EndSection();
 
@@ -2670,6 +2673,7 @@ void STimingView::ShowContextMenu(const FVector2D& ScreenSpacePosition, const FP
 
 	FWidgetPath EventPath = MouseEvent.GetEventPath() != nullptr ? *MouseEvent.GetEventPath() : FWidgetPath();
 	FSlateApplication::Get().PushMenu(SharedThis(this), EventPath, MenuWidget, ScreenSpacePosition, FPopupTransitionEffect::ContextMenu);
+	*/
 }
 
 ////////////////////////////////////////////////////////////////////////////////////////////////////
