[Insights]
  - Frames track: Fixed panning with mouse (to be able to horizontally pan also using left mouse button, not only with right mouse button).
  - Asset Loading: Improved table visualisation for aggregated stats (Event Aggregation and Object Type Aggregation top10 tables).
  - Asset Loading: Minor change to tooltip for hovered asset loading event (reanmed "Package Type" and "Export Type" to "Package Event" and "Export Event" respectivelly).
  - Timers / Stats Counters: Fixed default sorting column/mode not working as intended.
  - Timers / Stats Counters: Fixed UE-75578 (Timing Insights Timer's & Stats Counters Windows do not use standard UE4 radio button style). Made selected option for radio buttons be enabled all the time.

diff --git a/Engine/Source/Developer/TraceInsights/Private/Insights/Common/PaintUtils.h b/Engine/Source/Developer/TraceInsights/Private/Insights/Common/PaintUtils.h
--- a/Engine/Source/Developer/TraceInsights/Private/Insights/Common/PaintUtils.h
+++ b/Engine/Source/Developer/TraceInsights/Private/Insights/Common/PaintUtils.h
@@ -74,6 +74,22 @@ struct FDrawContext
 		FSlateDrawElement::MakeText(ElementList, InLayer, MAKE_PAINT_GEOMETRY_PT(Geometry, X, Y), Text, StartIndex, EndIndex, Font, DrawEffects, Color);
 	}
 
+	inline void DrawTextAligned(EHorizontalAlignment HAlign, const float X, const float Y, const FString& Text, const FSlateFontInfo& Font, const FLinearColor& Color) const
+	{
+		const TSharedRef<FSlateFontMeasure> FontMeasureService = FSlateApplication::Get().GetRenderer()->GetFontMeasureService();
+		float TextWidth = FontMeasureService->Measure(Text, Font).X;
+		float TextX = X;
+		if (HAlign == HAlign_Right)
+		{
+			TextX -= TextWidth;
+		}
+		else if (HAlign == HAlign_Center)
+		{
+			TextX -= TextWidth / 2;
+		}
+		FSlateDrawElement::MakeText(ElementList, LayerId, MAKE_PAINT_GEOMETRY_PT(Geometry, TextX, Y), Text, Font, DrawEffects, Color);
+	}
+
 	//inline void IncrementLayer() const { ++LayerId; }
 
 	// Accessors
diff --git a/Engine/Source/Developer/TraceInsights/Private/Insights/Widgets/SFrameTrack.cpp b/Engine/Source/Developer/TraceInsights/Private/Insights/Widgets/SFrameTrack.cpp
--- a/Engine/Source/Developer/TraceInsights/Private/Insights/Widgets/SFrameTrack.cpp
+++ b/Engine/Source/Developer/TraceInsights/Private/Insights/Widgets/SFrameTrack.cpp
@@ -72,7 +72,7 @@ void SFrameTrack::Reset()
 	bIsLMB_Pressed = false;
 	bIsRMB_Pressed = false;
 
-	bIsRMB_Scrolling = false;
+	bIsScrolling = false;
 
 	SelectionStartFrameIndex = 0;
 	SelectionEndFrameIndex = 0;
@@ -134,7 +134,7 @@ void SFrameTrack::Tick(const FGeometry& AllottedGeometry, const double InCurrent
 
 	ThisGeometry = AllottedGeometry;
 
-	if (!bIsRMB_Scrolling)
+	if (!bIsScrolling)
 	{
 		// Elastic snap to horizontal limits (+15%).
 
@@ -583,7 +583,12 @@ FReply SFrameTrack::OnMouseButtonUp(const FGeometry& MyGeometry, const FPointerE
 			// Release the mouse.
 			Reply = FReply::Handled().ReleaseMouseCapture();
 
-			if (bIsValidForMouseClick)
+			if (bIsScrolling)
+			{
+				bIsScrolling = false;
+				CursorType = EFrameTrackCursor::Default;
+			}
+			else if (bIsValidForMouseClick)
 			{
 				SelectFrameAtMousePosition(MousePositionOnButtonUp.X, MousePositionOnButtonUp.Y);
 			}
@@ -598,10 +603,10 @@ FReply SFrameTrack::OnMouseButtonUp(const FGeometry& MyGeometry, const FPointerE
 			// Release mouse as we no longer scroll.
 			Reply = FReply::Handled().ReleaseMouseCapture();
 
-			if (bIsRMB_Scrolling)
+			if (bIsScrolling)
 			{
+				bIsScrolling = false;
 				CursorType = EFrameTrackCursor::Default;
-				bIsRMB_Scrolling = false;
 			}
 			else if (bIsValidForMouseClick)
 			{
@@ -633,8 +638,18 @@ FReply SFrameTrack::OnMouseMove(const FGeometry& MyGeometry, const FPointerEvent
 	{
 		if (HasMouseCapture() && !MouseEvent.GetCursorDelta().IsZero())
 		{
-			// Inform other widgets that we have moved the selection box.
-			//SelectionBoxChangedEvent.Broadcast(SelectionBoxFrameStart, SelectionBoxFrameEnd);
+			if (!bIsScrolling)
+			{
+				bIsScrolling = true;
+				CursorType = EFrameTrackCursor::Hand;
+
+				HoveredSample.Reset();
+			}
+
+			float PosX = ViewportPosXOnButtonDown + (MousePositionOnButtonDown.X - MousePosition.X);
+			Viewport.ScrollAtIndex(Viewport.GetIndexAtPosX(PosX));
+			UpdateHorizontalScrollBar();
+			bIsStateDirty = true;
 
 			Reply = FReply::Handled();
 		}
@@ -643,10 +658,13 @@ FReply SFrameTrack::OnMouseMove(const FGeometry& MyGeometry, const FPointerEvent
 	{
 		if (HasMouseCapture() && !MouseEvent.GetCursorDelta().IsZero())
 		{
-			bIsRMB_Scrolling = true;
-			CursorType = EFrameTrackCursor::Hand;
+			if (!bIsScrolling)
+			{
+				bIsScrolling = true;
+				CursorType = EFrameTrackCursor::Hand;
 
-			HoveredSample.Reset();
+				HoveredSample.Reset();
+			}
 
 			float PosX = ViewportPosXOnButtonDown + (MousePositionOnButtonDown.X - MousePosition.X);
 			Viewport.ScrollAtIndex(Viewport.GetIndexAtPosX(PosX));
diff --git a/Engine/Source/Developer/TraceInsights/Private/Insights/Widgets/SFrameTrack.h b/Engine/Source/Developer/TraceInsights/Private/Insights/Widgets/SFrameTrack.h
--- a/Engine/Source/Developer/TraceInsights/Private/Insights/Widgets/SFrameTrack.h
+++ b/Engine/Source/Developer/TraceInsights/Private/Insights/Widgets/SFrameTrack.h
@@ -164,7 +164,7 @@ protected:
 	bool bIsRMB_Pressed;
 
 	/** True, if the user is currently interactively scrolling the view by holding the right mouse button and dragging. */
-	bool bIsRMB_Scrolling;
+	bool bIsScrolling;
 
 	//////////////////////////////////////////////////
 
diff --git a/Engine/Source/Developer/TraceInsights/Private/Insights/Widgets/SStatsView.cpp b/Engine/Source/Developer/TraceInsights/Private/Insights/Widgets/SStatsView.cpp
--- a/Engine/Source/Developer/TraceInsights/Private/Insights/Widgets/SStatsView.cpp
+++ b/Engine/Source/Developer/TraceInsights/Private/Insights/Widgets/SStatsView.cpp
@@ -27,16 +27,25 @@
 
 ////////////////////////////////////////////////////////////////////////////////////////////////////
 
-const EColumnSortMode::Type SStatsView::DefaultColumnSortMode(EColumnSortMode::Ascending);
-const FName SStatsView::DefaultColumnBeingSorted(FStatsViewColumns::NameColumnID);
+const EColumnSortMode::Type SStatsView::GetDefaultColumnSortMode()
+{
+	return EColumnSortMode::Ascending;
+}
+
+////////////////////////////////////////////////////////////////////////////////////////////////////
+
+const FName SStatsView::GetDefaultColumnBeingSorted()
+{
+	return FStatsViewColumns::NameColumnID;
+}
 
 ////////////////////////////////////////////////////////////////////////////////////////////////////
 
 SStatsView::SStatsView()
 	: bExpansionSaved(false)
 	, GroupingMode(EStatsGroupingMode::Flat)
-	, ColumnSortMode(DefaultColumnSortMode)
-	, ColumnBeingSorted(DefaultColumnBeingSorted)
+	, ColumnSortMode(GetDefaultColumnSortMode())
+	, ColumnBeingSorted(GetDefaultColumnBeingSorted())
 {
 	FMemory::Memset(bStatsNodeIsVisible, 1);
 }
@@ -1160,12 +1169,7 @@ bool SStatsView::HeaderMenu_SortMode_IsChecked(const FName ColumnId, const EColu
 bool SStatsView::HeaderMenu_SortMode_CanExecute(const FName ColumnId, const EColumnSortMode::Type InSortMode) const
 {
 	const FStatsViewColumn& Column = TreeViewHeaderColumns.FindChecked(ColumnId);
-	const bool bIsValid = Column.bCanBeSorted();
-
-	bool bCanExecute = ColumnBeingSorted != ColumnId ? true : ColumnSortMode != InSortMode;
-	bCanExecute = bCanExecute && bIsValid;
-
-	return bCanExecute;
+	return Column.bCanBeSorted();
 }
 
 ////////////////////////////////////////////////////////////////////////////////////////////////////
@@ -1189,7 +1193,7 @@ bool SStatsView::ContextMenu_SortMode_IsChecked(const EColumnSortMode::Type InSo
 
 bool SStatsView::ContextMenu_SortMode_CanExecute(const EColumnSortMode::Type InSortMode) const
 {
-	return ColumnSortMode != InSortMode;
+	return true; //ColumnSortMode != InSortMode;
 }
 
 ////////////////////////////////////////////////////////////////////////////////////////////////////
@@ -1213,7 +1217,7 @@ bool SStatsView::ContextMenu_SortByColumn_IsChecked(const FName ColumnId)
 
 bool SStatsView::ContextMenu_SortByColumn_CanExecute(const FName ColumnId) const
 {
-	return ColumnId != ColumnBeingSorted;
+	return true; //ColumnId != ColumnBeingSorted;
 }
 
 ////////////////////////////////////////////////////////////////////////////////////////////////////
@@ -1289,8 +1293,8 @@ bool SStatsView::ContextMenu_ShowAllColumns_CanExecute() const
 
 void SStatsView::ContextMenu_ShowAllColumns_Execute()
 {
-	ColumnSortMode = DefaultColumnSortMode;
-	ColumnBeingSorted = DefaultColumnBeingSorted;
+	ColumnSortMode = GetDefaultColumnSortMode();
+	ColumnBeingSorted = GetDefaultColumnBeingSorted();
 
 	const int32 NumColumns = FStatsViewColumnFactory::Get().Collection.Num();
 	for (int32 ColumnIndex = 0; ColumnIndex < NumColumns; ColumnIndex++)
@@ -1367,8 +1371,8 @@ bool SStatsView::ContextMenu_ResetColumns_CanExecute() const
 
 void SStatsView::ContextMenu_ResetColumns_Execute()
 {
-	ColumnSortMode = DefaultColumnSortMode;
-	ColumnBeingSorted = DefaultColumnBeingSorted;
+	ColumnSortMode = GetDefaultColumnSortMode();
+	ColumnBeingSorted = GetDefaultColumnBeingSorted();
 
 	const int32 NumColumns = FStatsViewColumnFactory::Get().Collection.Num();
 	for (int32 ColumnIndex = 0; ColumnIndex < NumColumns; ColumnIndex++)
diff --git a/Engine/Source/Developer/TraceInsights/Private/Insights/Widgets/SStatsView.h b/Engine/Source/Developer/TraceInsights/Private/Insights/Widgets/SStatsView.h
--- a/Engine/Source/Developer/TraceInsights/Private/Insights/Widgets/SStatsView.h
+++ b/Engine/Source/Developer/TraceInsights/Private/Insights/Widgets/SStatsView.h
@@ -187,6 +187,9 @@ protected:
 	////////////////////////////////////////////////////////////////////////////////////////////////////
 	// Sorting
 
+	static const EColumnSortMode::Type GetDefaultColumnSortMode();
+	static const FName GetDefaultColumnBeingSorted();
+
 	void SortStats();
 
 	EColumnSortMode::Type GetSortModeForColumn(const FName ColumnId) const;
@@ -329,9 +332,6 @@ protected:
 	/** Name of the column currently being sorted, NAME_None if sorting is disabled. */
 	FName ColumnBeingSorted;
 
-	static const EColumnSortMode::Type DefaultColumnSortMode;
-	static const FName DefaultColumnBeingSorted;
-
 	//////////////////////////////////////////////////
 
 	double StatsStartTime;
diff --git a/Engine/Source/Developer/TraceInsights/Private/Insights/Widgets/STimersView.cpp b/Engine/Source/Developer/TraceInsights/Private/Insights/Widgets/STimersView.cpp
--- a/Engine/Source/Developer/TraceInsights/Private/Insights/Widgets/STimersView.cpp
+++ b/Engine/Source/Developer/TraceInsights/Private/Insights/Widgets/STimersView.cpp
@@ -29,16 +29,25 @@
 
 ////////////////////////////////////////////////////////////////////////////////////////////////////
 
-const EColumnSortMode::Type STimersView::DefaultColumnSortMode(EColumnSortMode::Descending);
-const FName STimersView::DefaultColumnBeingSorted(FTimersViewColumns::TotalInclusiveTimeColumnID);
+const EColumnSortMode::Type STimersView::GetDefaultColumnSortMode()
+{
+	return EColumnSortMode::Descending;
+}
+
+////////////////////////////////////////////////////////////////////////////////////////////////////
+
+const FName STimersView::GetDefaultColumnBeingSorted()
+{
+	return FTimersViewColumns::TotalInclusiveTimeColumnID;
+}
 
 ////////////////////////////////////////////////////////////////////////////////////////////////////
 
 STimersView::STimersView()
 	: bExpansionSaved(false)
 	, GroupingMode(ETimerGroupingMode::ByType)
-	, ColumnSortMode(DefaultColumnSortMode)
-	, ColumnBeingSorted(DefaultColumnBeingSorted)
+	, ColumnSortMode(GetDefaultColumnSortMode())
+	, ColumnBeingSorted(GetDefaultColumnBeingSorted())
 {
 	FMemory::Memset(bTimerNodeIsVisible, 1);
 }
@@ -1189,12 +1198,7 @@ bool STimersView::HeaderMenu_SortMode_IsChecked(const FName ColumnId, const ECol
 bool STimersView::HeaderMenu_SortMode_CanExecute(const FName ColumnId, const EColumnSortMode::Type InSortMode) const
 {
 	const FTimersViewColumn& Column = TreeViewHeaderColumns.FindChecked(ColumnId);
-	const bool bIsValid = Column.bCanBeSorted();
-
-	bool bCanExecute = ColumnBeingSorted != ColumnId ? true : ColumnSortMode != InSortMode;
-	bCanExecute = bCanExecute && bIsValid;
-
-	return bCanExecute;
+	return Column.bCanBeSorted();
 }
 
 ////////////////////////////////////////////////////////////////////////////////////////////////////
@@ -1218,7 +1222,7 @@ bool STimersView::ContextMenu_SortMode_IsChecked(const EColumnSortMode::Type InS
 
 bool STimersView::ContextMenu_SortMode_CanExecute(const EColumnSortMode::Type InSortMode) const
 {
-	return ColumnSortMode != InSortMode;
+	return true; //ColumnSortMode != InSortMode;
 }
 
 ////////////////////////////////////////////////////////////////////////////////////////////////////
@@ -1242,7 +1246,7 @@ bool STimersView::ContextMenu_SortByColumn_IsChecked(const FName ColumnId)
 
 bool STimersView::ContextMenu_SortByColumn_CanExecute(const FName ColumnId) const
 {
-	return ColumnId != ColumnBeingSorted;
+	return true; //ColumnId != ColumnBeingSorted;
 }
 
 ////////////////////////////////////////////////////////////////////////////////////////////////////
@@ -1318,8 +1322,8 @@ bool STimersView::ContextMenu_ShowAllColumns_CanExecute() const
 
 void STimersView::ContextMenu_ShowAllColumns_Execute()
 {
-	ColumnSortMode = DefaultColumnSortMode;
-	ColumnBeingSorted = DefaultColumnBeingSorted;
+	ColumnSortMode = GetDefaultColumnSortMode();
+	ColumnBeingSorted = GetDefaultColumnBeingSorted();
 
 	const int32 NumColumns = FTimersViewColumnFactory::Get().Collection.Num();
 	for (int32 ColumnIndex = 0; ColumnIndex < NumColumns; ColumnIndex++)
@@ -1399,8 +1403,8 @@ bool STimersView::ContextMenu_ResetColumns_CanExecute() const
 
 void STimersView::ContextMenu_ResetColumns_Execute()
 {
-	ColumnSortMode = DefaultColumnSortMode;
-	ColumnBeingSorted = DefaultColumnBeingSorted;
+	ColumnSortMode = GetDefaultColumnSortMode();
+	ColumnBeingSorted = GetDefaultColumnBeingSorted();
 
 	const int32 NumColumns = FTimersViewColumnFactory::Get().Collection.Num();
 	for (int32 ColumnIndex = 0; ColumnIndex < NumColumns; ColumnIndex++)
diff --git a/Engine/Source/Developer/TraceInsights/Private/Insights/Widgets/STimersView.h b/Engine/Source/Developer/TraceInsights/Private/Insights/Widgets/STimersView.h
--- a/Engine/Source/Developer/TraceInsights/Private/Insights/Widgets/STimersView.h
+++ b/Engine/Source/Developer/TraceInsights/Private/Insights/Widgets/STimersView.h
@@ -167,6 +167,9 @@ protected:
 	////////////////////////////////////////////////////////////////////////////////////////////////////
 	// Sorting
 
+	static const EColumnSortMode::Type GetDefaultColumnSortMode();
+	static const FName GetDefaultColumnBeingSorted();
+
 	void SortTimers();
 
 	EColumnSortMode::Type GetSortModeForColumn(const FName ColumnId) const;
@@ -309,9 +312,6 @@ protected:
 	/** Name of the column currently being sorted, NAME_None if sorting is disabled. */
 	FName ColumnBeingSorted;
 
-	static const EColumnSortMode::Type DefaultColumnSortMode;
-	static const FName DefaultColumnBeingSorted;
-
 	//////////////////////////////////////////////////
 
 	double StatsStartTime;
diff --git a/Engine/Source/Developer/TraceInsights/Private/Insights/Widgets/STimingView.cpp b/Engine/Source/Developer/TraceInsights/Private/Insights/Widgets/STimingView.cpp
--- a/Engine/Source/Developer/TraceInsights/Private/Insights/Widgets/STimingView.cpp
+++ b/Engine/Source/Developer/TraceInsights/Private/Insights/Widgets/STimingView.cpp
@@ -219,8 +219,10 @@ void STimingView::Reset()
 
 	LoadingGetEventNameFn = FLoadingTrackGetEventNameDelegate::CreateRaw(this, &STimingView::GetLoadTimeProfilerEventName);
 
-	EventAggregationStr.Reset();
-	ObjectTypeAggregationStr.Reset();
+	EventAggregationTotalCount = 0;
+	EventAggregation.Reset();
+	ObjectTypeAggregationTotalCount = 0;
+	ObjectTypeAggregation.Reset();
 
 	IoOverviewTrack = nullptr;
 	IoActivityTrack = nullptr;
@@ -523,7 +525,7 @@ void STimingView::Tick(const FGeometry& AllottedGeometry, const double InCurrent
 
 					if (!CachedTimelines.Contains(TrackId))
 					{
-						FString TrackName(ThreadInfo.Name && *ThreadInfo.Name ? ThreadInfo.Name : FString::Printf(TEXT("Thread %u")));
+						FString TrackName(ThreadInfo.Name && *ThreadInfo.Name ? ThreadInfo.Name : FString::Printf(TEXT("Thread %u"), ThreadInfo.Id));
 
 						// Create new Timing Events track for the CPU thread.
 						Track = AddTimingEventsTrack(TrackId, ETimingEventsTrackType::Cpu, TrackName, GroupName, Order);
@@ -1151,7 +1153,15 @@ int32 STimingView::OnPaint(const FPaintArgs& Args, const FGeometry& AllottedGeom
 				const float ValueX = X + ValueOffsetX;
 
 				constexpr float LineH = 14.0f;
-				constexpr float H = 10 * LineH;
+				float H = 5 * LineH;
+				if (Package)
+				{
+					H += 3 * LineH;
+				}
+				if (Export)
+				{
+					H += 2 * LineH;
+				}
 				const float MaxY = FMath::Max(0.0f, Viewport.Height - H - 18.0f);
 				float Y = FMath::Clamp<float>(MousePosition.Y + 18.0f, 0.0f, MaxY);
 
@@ -1186,7 +1196,7 @@ int32 STimingView::OnPaint(const FPaintArgs& Args, const FGeometry& AllottedGeom
 				DrawContext.DrawText(ValueX, Y, DepthStr, MainFont, ValueColor);
 				Y += LineH;
 
-				DrawContext.DrawText(X, Y, TEXT("Package Type:"), MainFont, TextColor);
+				DrawContext.DrawText(X, Y, TEXT("Package Event:"), MainFont, TextColor);
 				DrawContext.DrawText(ValueX, Y, GetName(HoveredTimingEvent.LoadingInfo.PackageEventType), MainFont, ValueColor);
 				Y += LineH;
 
@@ -1208,7 +1218,7 @@ int32 STimingView::OnPaint(const FPaintArgs& Args, const FGeometry& AllottedGeom
 				}
 
 				{
-					DrawContext.DrawText(X, Y, TEXT("Export Type:"), MainFont, TextColor);
+					DrawContext.DrawText(X, Y, TEXT("Export Event:"), MainFont, TextColor);
 					FString ExportTypeStr = FString::Printf(TEXT("%s%s"), GetName(HoveredTimingEvent.LoadingInfo.ExportEventType), Export && Export->IsAsset ? TEXT(" [asset]") : TEXT(""));
 					DrawContext.DrawText(ValueX, Y, ExportTypeStr, MainFont, ValueColor);
 					Y += LineH;
@@ -1239,36 +1249,22 @@ int32 STimingView::OnPaint(const FPaintArgs& Args, const FGeometry& AllottedGeom
 
 	if (bAssetLoadingMode)
 	{
-		const FLinearColor BackgroundColor(0.01f, 0.01f, 0.01f, 0.9f);
-		const FLinearColor TextColor(1.0f, 1.0f, 1.0f, 0.9f);
-
 		constexpr float MarginX = 8.0f;
 		constexpr float MarginY = 8.0f;
 
-		constexpr float BorderX = 4.0f;
-		constexpr float BorderY = 4.0f;
+		const float X = ViewWidth - MarginX;
 		float Y = ViewHeight;
 
-		if (ObjectTypeAggregationStr.Len() > 0)
+		if (ObjectTypeAggregation.Num() > 0)
 		{
-			FVector2D TextSize = FontMeasureService->Measure(ObjectTypeAggregationStr, MainFont);
-			const float X = ViewWidth - MarginX - TextSize.X - 2 * BorderX;
-			Y -= MarginY + TextSize.Y + 2 * BorderY;
-			DrawContext.DrawBox(X, Y, TextSize.X + 2 * BorderX, TextSize.Y + 2 * BorderY, WhiteBrush, BackgroundColor);
-			DrawContext.LayerId++;
-			DrawContext.DrawText(X + BorderX, Y + BorderY, ObjectTypeAggregationStr, MainFont, TextColor);
-			DrawContext.LayerId++;
+			Y -= MarginY;
+			Y -= DrawAssetLoadingAggregationTable(DrawContext, X, Y, TEXT("Object Type Aggregation"), ObjectTypeAggregation, ObjectTypeAggregationTotalCount);
 		}
 
-		if (EventAggregationStr.Len() > 0)
+		if (EventAggregation.Num() > 0)
 		{
-			FVector2D TextSize = FontMeasureService->Measure(EventAggregationStr, MainFont);
-			const float X = ViewWidth - MarginX - TextSize.X - 2 * BorderX;
-			Y -= MarginY + TextSize.Y + 2 * BorderY;
-			DrawContext.DrawBox(X, Y, TextSize.X + 2 * BorderX, TextSize.Y + 2 * BorderY, WhiteBrush, BackgroundColor);
-			DrawContext.LayerId++;
-			DrawContext.DrawText(X + BorderX, Y + BorderY, EventAggregationStr, MainFont, TextColor);
-			DrawContext.LayerId++;
+			Y -= MarginY;
+			Y -= DrawAssetLoadingAggregationTable(DrawContext, X, Y, TEXT("Event Aggregation"), EventAggregation, EventAggregationTotalCount);
 		}
 	}
 
@@ -1277,6 +1273,150 @@ int32 STimingView::OnPaint(const FPaintArgs& Args, const FGeometry& AllottedGeom
 
 ////////////////////////////////////////////////////////////////////////////////////////////////////
 
+float STimingView::DrawAssetLoadingAggregationTable(FDrawContext& DrawContext, float RightX, float BottomY, const TCHAR* TableName, const TArray<FAssetLoadingEventAggregationRow>& Aggregation, int32 TotalRowCount) const
+{
+	const FLinearColor BackgroundColor(0.01f, 0.01f, 0.01f, 0.9f);
+	const FLinearColor TextColorHeader(0.5f, 0.5f, 0.5f, 0.9f);
+	const FLinearColor TextColor(1.0f, 1.0f, 1.0f, 0.9f);
+
+	constexpr float BorderX = 4.0f;
+	constexpr float BorderY = 4.0f;
+
+	constexpr float LineH = 14.0f;
+
+	float TableHeight = (2 + Aggregation.Num()) * LineH;
+	if (Aggregation.Num() < TotalRowCount)
+	{
+		TableHeight += LineH; // for the "[...]" line
+	}
+
+	const TSharedRef<FSlateFontMeasure> FontMeasureService = FSlateApplication::Get().GetRenderer()->GetFontMeasureService();
+	float MaxNameWidth = 140.0f;
+	for (const FAssetLoadingEventAggregationRow& Row : Aggregation)
+	{
+		float NameWidth = FontMeasureService->Measure(Row.Name, MainFont).X;
+		if (NameWidth > MaxNameWidth)
+		{
+			MaxNameWidth = NameWidth;
+		}
+	}
+	float NameColumnWidth = MaxNameWidth + 5.0f;
+
+	constexpr float RowNumberColumnWidth = 20.0f;
+	constexpr float CountColumnWidth = 50.0f;
+	constexpr float TotalColumnWidth = 65.0f;
+	constexpr float MinColumnWidth = 60.0f;
+	constexpr float MaxColumnWidth = 60.0f;
+	constexpr float AvgColumnWidth = 60.0f;
+	constexpr float MedColumnWidth = 60.0f;
+
+	const float TableWidth = RowNumberColumnWidth
+						   + NameColumnWidth
+						   + CountColumnWidth
+						   + TotalColumnWidth
+						   + MaxColumnWidth
+						   + AvgColumnWidth
+						   + MedColumnWidth
+						   + MinColumnWidth;
+
+	float W = TableWidth + 2 * BorderX;
+	float H = TableHeight + 2 * BorderY;
+
+	float Y = BottomY - H;
+
+	DrawContext.DrawBox(RightX - W, Y, W, H, WhiteBrush, BackgroundColor);
+	DrawContext.LayerId++;
+	Y += BorderY;
+
+	const float TableX = RightX - W + BorderX;
+
+	// Table name
+	{
+		FString TableNameRow = FString::Printf(TEXT("%s (%d records, sorted by Total time)"), TableName, TotalRowCount);
+		DrawContext.DrawText(TableX, Y, TableNameRow, MainFont, TextColor);
+		Y += LineH;
+	}
+
+	// Column names
+	{
+		float X = TableX;
+
+		X += RowNumberColumnWidth;
+
+		DrawContext.DrawText(X, Y, TEXT("Name"), MainFont, TextColorHeader);
+		X += NameColumnWidth;
+
+		X += CountColumnWidth;
+		DrawContext.DrawTextAligned(HAlign_Right, X, Y, TEXT("Count"), MainFont, TextColorHeader);
+
+		X += TotalColumnWidth;
+		DrawContext.DrawTextAligned(HAlign_Right, X, Y, TEXT("Total"), MainFont, TextColorHeader);
+
+		X += MaxColumnWidth;
+		DrawContext.DrawTextAligned(HAlign_Right, X, Y, TEXT("Max [ms]"), MainFont, TextColorHeader);
+
+		X += AvgColumnWidth;
+		DrawContext.DrawTextAligned(HAlign_Right, X, Y, TEXT("Avg [ms]"), MainFont, TextColorHeader);
+
+		X += MedColumnWidth;
+		DrawContext.DrawTextAligned(HAlign_Right, X, Y, TEXT("Med [ms]"), MainFont, TextColorHeader);
+
+		X += MinColumnWidth;
+		DrawContext.DrawTextAligned(HAlign_Right, X, Y, TEXT("Min [ms]"), MainFont, TextColorHeader);
+
+		Y += LineH;
+	}
+
+	// Records
+	for (int32 Index = 0; Index < Aggregation.Num(); ++Index)
+	{
+		constexpr int32 NumDigits = 2;
+		constexpr bool bAddTimeUnit = false;
+
+		const FAssetLoadingEventAggregationRow& Row = Aggregation[Index];
+
+		float X = TableX;
+
+		X += RowNumberColumnWidth;
+		DrawContext.DrawTextAligned(HAlign_Right, X - 4.0f, Y, FString::Printf(TEXT("%d."), Index + 1), MainFont, TextColorHeader);
+
+		DrawContext.DrawText(X, Y, Row.Name, MainFont, TextColor);
+		X += NameColumnWidth;
+
+		X += CountColumnWidth;
+		DrawContext.DrawTextAligned(HAlign_Right, X, Y, FText::AsNumber(Row.Count).ToString(), MainFont, TextColor);
+
+		X += TotalColumnWidth;
+		DrawContext.DrawTextAligned(HAlign_Right, X, Y, TimeUtils::FormatTimeAuto(Row.Total), MainFont, TextColor);
+
+		X += MaxColumnWidth;
+		DrawContext.DrawTextAligned(HAlign_Right, X, Y, TimeUtils::FormatTimeMs(Row.Max, NumDigits, bAddTimeUnit), MainFont, TextColor);
+
+		X += AvgColumnWidth;
+		DrawContext.DrawTextAligned(HAlign_Right, X, Y, TimeUtils::FormatTimeMs(Row.Avg, NumDigits, bAddTimeUnit), MainFont, TextColor);
+
+		X += MedColumnWidth;
+		DrawContext.DrawTextAligned(HAlign_Right, X, Y, TimeUtils::FormatTimeMs(Row.Med, NumDigits, bAddTimeUnit), MainFont, TextColor);
+
+		X += MinColumnWidth;
+		DrawContext.DrawTextAligned(HAlign_Right, X, Y, TimeUtils::FormatTimeMs(Row.Min, NumDigits, bAddTimeUnit), MainFont, TextColor);
+
+		Y += LineH;
+	}
+
+	if (Aggregation.Num() < TotalRowCount)
+	{
+		DrawContext.DrawText(TableX, Y, TEXT("[...]"), MainFont, TextColorHeader);
+		Y += LineH;
+	}
+
+	DrawContext.LayerId++;
+
+	return H;
+}
+
+////////////////////////////////////////////////////////////////////////////////////////////////////
+
 const TCHAR* GetName(ETimingEventsTrackType Type)
 {
 	switch (Type)
@@ -2026,8 +2166,8 @@ FReply STimingView::OnMouseButtonDown(const FGeometry& MyGeometry, const FPointe
 				SelectionEndTime = SelectionStartTime;
 				LastSelectionType = ESelectionType::None;
 				//TODO: SelectionChangingEvent.Broadcast(SelectionStartTime, SelectionEndTime);
-				EventAggregationStr.Reset();
-				ObjectTypeAggregationStr.Reset();
+				EventAggregation.Reset();
+				ObjectTypeAggregation.Reset();
 			}
 
 			// Capture mouse, so we can drag outside this widget.
@@ -2115,8 +2255,8 @@ FReply STimingView::OnMouseButtonUp(const FGeometry& MyGeometry, const FPointerE
 					SelectionEndTime = SelectionStartTime = 0.0;
 					LastSelectionType = ESelectionType::None;
 					//TODO: SelectionChangedEvent.Broadcast(SelectionStartTime, SelectionEndTime);
-					EventAggregationStr.Reset();
-					ObjectTypeAggregationStr.Reset();
+					EventAggregation.Reset();
+					ObjectTypeAggregation.Reset();
 				}
 			}
 
@@ -2226,8 +2366,8 @@ FReply STimingView::OnMouseMove(const FGeometry& MyGeometry, const FPointerEvent
 			}
 			LastSelectionType = ESelectionType::TimeRange;
 			//TODO: SelectionChangingEvent.Broadcast(SelectionStartTime, SelectionEndTime);
-			EventAggregationStr.Reset();
-			ObjectTypeAggregationStr.Reset();
+			EventAggregation.Reset();
+			ObjectTypeAggregation.Reset();
 		}
 
 		Reply = FReply::Handled();
@@ -2844,16 +2984,12 @@ void STimingView::UpdateAggregatedStats()
 		TSharedPtr<const Trace::IAnalysisSession> Session = FInsightsManager::Get()->GetSession();
 		if (Session.IsValid() && Trace::ReadLoadTimeProfilerProvider(*Session.Get()))
 		{
-			auto TableToString = [](const TCHAR* TableName, Trace::ITable<Trace::FLoadTimeProfilerAggregatedStats>* Table) -> FString
+			auto ReadTable = [](Trace::ITable<Trace::FLoadTimeProfilerAggregatedStats>* Table, TArray<FAssetLoadingEventAggregationRow>& Aggregation, int32& TotalRowCount)
 			{
-				FString Str;
-
-				const uint32 TotalRowCount = static_cast<uint32>(Table->GetRowCount());
-
-				Str.Append(TableName).Append(": ").Append(FText::AsNumber(TotalRowCount).ToString()).Append(" records, sorted by Total time");
+				TotalRowCount = Table->GetRowCount();
 
-				constexpr uint32 MaxRowCount = 10; // only get first 10 records
-				const uint32 RowCount = FMath::Min<uint32>(TotalRowCount, MaxRowCount);
+				constexpr int32 MaxRowCount = 10; // only get first 10 records
+				const int32 RowCount = FMath::Min(TotalRowCount, MaxRowCount);
 
 				const Trace::ITableLayout& TableLayout = Table->GetLayout();
 				const int32 ColumnCount = TableLayout.GetColumnCount();
@@ -2864,17 +3000,17 @@ void STimingView::UpdateAggregatedStats()
 
 				struct FSortedIndexEntry
 				{
-					uint32 RowIndex;
+					int32 RowIndex;
 					double Value;
 				};
 
 				TArray<FSortedIndexEntry> SortedIndex;
 				SortedIndex.Reserve(TotalRowCount);
 
-				constexpr uint32 TotalTimeColumnIndex = 2;
+				constexpr int32 TotalTimeColumnIndex = 2;
 				ensure(TableLayout.GetColumnType(TotalTimeColumnIndex) == Trace::ETableColumnType::TableColumnType_Double);
 
-				for (uint32 RowIndex = 0; RowIndex < TotalRowCount; ++RowIndex)
+				for (int32 RowIndex = 0; RowIndex < TotalRowCount; ++RowIndex)
 				{
 					Reader->SetRowIndex(RowIndex);
 					double Value = Reader->GetValueDouble(TotalTimeColumnIndex);
@@ -2885,47 +3021,43 @@ void STimingView::UpdateAggregatedStats()
 
 				//////////////////////////////////////////////////
 
-				for (uint32 Index = 0; Index < RowCount; ++Index)
+				Aggregation.Reset();
+				Aggregation.AddDefaulted(RowCount);
+
+				for (int32 Index = 0; Index < RowCount; ++Index)
 				{
-					uint32 RowIndex = SortedIndex[Index].RowIndex;
-					Reader->SetRowIndex(RowIndex);
+					FAssetLoadingEventAggregationRow& Row = Aggregation[Index];
 
-					Str.Append("\n").Append(FText::AsNumber(Index + 1).ToString()).Append(". ");
+					int32 RowIndex = SortedIndex[Index].RowIndex;
+					Reader->SetRowIndex(RowIndex);
 
+					ensure(ColumnCount == 7);
 					for (int32 ColumnIndex = 0; ColumnIndex < ColumnCount; ++ColumnIndex)
 					{
-						Str.Append(TableLayout.GetColumnName(ColumnIndex)).Append("=");
+						ensure(TableLayout.GetColumnType(0) == Trace::ETableColumnType::TableColumnType_CString);
+						Row.Name = Reader->GetValueCString(0);
 
-						Trace::ETableColumnType ColumnType = TableLayout.GetColumnType(ColumnIndex);
-						switch (ColumnType)
-						{
-						case Trace::ETableColumnType::TableColumnType_Bool:
-							Str.Append(Reader->GetValueBool(ColumnIndex) ? TEXT("True") : TEXT("False")).Append(", ");
-							break;
-						case Trace::ETableColumnType::TableColumnType_Int:
-							Str.Append(FText::AsNumber(Reader->GetValueInt(ColumnIndex)).ToString()).Append(", ");
-							break;
-						case Trace::ETableColumnType::TableColumnType_Float:
-							Str.Append(FText::AsNumber(Reader->GetValueFloat(ColumnIndex)).ToString()).Append(", ");
-							break;
-						case Trace::ETableColumnType::TableColumnType_Double:
-							Str.Append(FText::AsNumber(Reader->GetValueDouble(ColumnIndex)).ToString()).Append(", ");
-							break;
-						case Trace::ETableColumnType::TableColumnType_CString:
-							Str.Append(Reader->GetValueCString(ColumnIndex)).Append(", ");
-							break;
-						}
-					}
-				}
+						ensure(TableLayout.GetColumnType(1) == Trace::ETableColumnType::TableColumnType_Int);
+						Row.Count = Reader->GetValueInt(1);
 
-				if (RowCount != TotalRowCount)
-				{
-					Str.Append("\n[...]");
+						ensure(TableLayout.GetColumnType(2) == Trace::ETableColumnType::TableColumnType_Double);
+						Row.Total = Reader->GetValueDouble(2);
+
+						ensure(TableLayout.GetColumnType(3) == Trace::ETableColumnType::TableColumnType_Double);
+						Row.Min = Reader->GetValueDouble(3);
+
+						ensure(TableLayout.GetColumnType(4) == Trace::ETableColumnType::TableColumnType_Double);
+						Row.Max = Reader->GetValueDouble(4);
+
+						ensure(TableLayout.GetColumnType(5) == Trace::ETableColumnType::TableColumnType_Double);
+						Row.Avg = Reader->GetValueDouble(5);
+
+						ensure(TableLayout.GetColumnType(6) == Trace::ETableColumnType::TableColumnType_Double);
+						Row.Med = Reader->GetValueDouble(6);
+					}
 				}
 
 				delete Reader;
-
-				return Str;
 			};
 
 			Trace::FAnalysisSessionReadScope SessionReadScope(*Session.Get());
@@ -2934,13 +3066,13 @@ void STimingView::UpdateAggregatedStats()
 
 			{
 				Trace::ITable<Trace::FLoadTimeProfilerAggregatedStats>* EventAggregationTable = LoadTimeProfilerProvider.CreateEventAggregation(SelectionStartTime, SelectionEndTime);
-				EventAggregationStr = TableToString(TEXT("Event Aggregation"), EventAggregationTable);
+				ReadTable(EventAggregationTable, EventAggregation, EventAggregationTotalCount);
 				delete EventAggregationTable;
 			}
 
 			{
 				Trace::ITable<Trace::FLoadTimeProfilerAggregatedStats>* ObjectTypeAggregationTable = LoadTimeProfilerProvider.CreateObjectTypeAggregation(SelectionStartTime, SelectionEndTime);
-				ObjectTypeAggregationStr = TableToString(TEXT("Object Type Aggregation"), ObjectTypeAggregationTable);
+				ReadTable(ObjectTypeAggregationTable, ObjectTypeAggregation, ObjectTypeAggregationTotalCount);
 				delete ObjectTypeAggregationTable;
 			}
 		}
diff --git a/Engine/Source/Developer/TraceInsights/Private/Insights/Widgets/STimingView.h b/Engine/Source/Developer/TraceInsights/Private/Insights/Widgets/STimingView.h
--- a/Engine/Source/Developer/TraceInsights/Private/Insights/Widgets/STimingView.h
+++ b/Engine/Source/Developer/TraceInsights/Private/Insights/Widgets/STimingView.h
@@ -44,6 +44,18 @@ class STimingView : public SCompoundWidget
 private:
 	static constexpr float MinTooltipWidth = 110.0f;
 
+protected:
+	struct FAssetLoadingEventAggregationRow
+	{
+		FString Name;
+		int32 Count;
+		double Total;
+		double Min;
+		double Max;
+		double Avg;
+		double Med;
+	};
+
 public:
 	/** Default constructor. */
 	STimingView();
@@ -276,7 +288,7 @@ protected:
 	void DrawIoOverviewTrack(FTimingViewDrawHelper& Helper, FTimingEventsTrack& Track) const;
 	void DrawIoActivityTrack(FTimingViewDrawHelper& Helper, FTimingEventsTrack& Track) const;
 
-	void DrawTimeRangeSelection(FDrawContext& DH) const;
+	void DrawTimeRangeSelection(FDrawContext& DrawContext) const;
 
 	void ShowContextMenu(const FVector2D& ScreenSpacePosition, const FPointerEvent& MouseEvent);
 
@@ -303,6 +315,8 @@ protected:
 
 	////////////////////////////////////////////////////////////////////////////////////////////////////
 
+	float DrawAssetLoadingAggregationTable(FDrawContext& DrawContext, float RightX, float BottomY, const TCHAR* TableName, const TArray<FAssetLoadingEventAggregationRow>& Aggregation, int32 TotalRowCount) const;
+
 	void UpdateAggregatedStats();
 
 	void UpdateHoveredTimingEvent(float MX, float MY);
@@ -413,8 +427,11 @@ protected:
 
 	bool bAssetLoadingMode;
 
-	FString EventAggregationStr;
-	FString ObjectTypeAggregationStr;
+	int32 EventAggregationTotalCount;
+	TArray<FAssetLoadingEventAggregationRow> EventAggregation;
+
+	int32 ObjectTypeAggregationTotalCount;
+	TArray<FAssetLoadingEventAggregationRow> ObjectTypeAggregation;
 
 	////////////////////////////////////////////////////////////
 	// File activity (I/O)
