[Insights]
  - Timing: Fixed Timers and Stats Counters tabs to correctly re-sync the list of timers / stats after tab is re-opened.
  - Timing: Removed Graph track (tab) from toolbar. Fixes UE-75603.
  - Fixed UE-75591 (Timing Insights Toolbar window status does not update when closing Windows via X).

diff --git a/Engine/Source/Developer/TraceInsights/Private/Insights/Widgets/SStatsView.cpp b/Engine/Source/Developer/TraceInsights/Private/Insights/Widgets/SStatsView.cpp
--- a/Engine/Source/Developer/TraceInsights/Private/Insights/Widgets/SStatsView.cpp
+++ b/Engine/Source/Developer/TraceInsights/Private/Insights/Widgets/SStatsView.cpp
@@ -217,6 +217,9 @@ void SStatsView::Construct(const FArguments& InArgs)
 
 	// Register ourselves with the Insights manager.
 	FInsightsManager::Get()->GetSessionChangedEvent().AddSP(this, &SStatsView::InsightsManager_OnSessionChanged);
+
+	// Update the Session (i.e. when analysis session was already started).
+	InsightsManager_OnSessionChanged();
 }
 END_SLATE_FUNCTION_BUILD_OPTIMIZATION
 
diff --git a/Engine/Source/Developer/TraceInsights/Private/Insights/Widgets/STimersView.cpp b/Engine/Source/Developer/TraceInsights/Private/Insights/Widgets/STimersView.cpp
--- a/Engine/Source/Developer/TraceInsights/Private/Insights/Widgets/STimersView.cpp
+++ b/Engine/Source/Developer/TraceInsights/Private/Insights/Widgets/STimersView.cpp
@@ -226,6 +226,9 @@ void STimersView::Construct(const FArguments& InArgs)
 
 	// Register ourselves with the Insights manager.
 	FInsightsManager::Get()->GetSessionChangedEvent().AddSP(this, &STimersView::InsightsManager_OnSessionChanged);
+
+	// Update the Session (i.e. when analysis session was already started).
+	InsightsManager_OnSessionChanged();
 }
 END_SLATE_FUNCTION_BUILD_OPTIMIZATION
 
diff --git a/Engine/Source/Developer/TraceInsights/Private/Insights/Widgets/STimingProfilerToolbar.cpp b/Engine/Source/Developer/TraceInsights/Private/Insights/Widgets/STimingProfilerToolbar.cpp
--- a/Engine/Source/Developer/TraceInsights/Private/Insights/Widgets/STimingProfilerToolbar.cpp
+++ b/Engine/Source/Developer/TraceInsights/Private/Insights/Widgets/STimingProfilerToolbar.cpp
@@ -52,7 +52,7 @@ void STimingProfilerToolbar::Construct(const FArguments& InArgs)
 			ToolbarBuilder.BeginSection("View");
 			{
 				ToolbarBuilder.AddToolBarButton(FTimingProfilerCommands::Get().ToggleFramesTrackVisibility);
-				ToolbarBuilder.AddToolBarButton(FTimingProfilerCommands::Get().ToggleGraphTrackVisibility);
+				//ToolbarBuilder.AddToolBarButton(FTimingProfilerCommands::Get().ToggleGraphTrackVisibility);
 				ToolbarBuilder.AddToolBarButton(FTimingProfilerCommands::Get().ToggleTimingViewVisibility);
 				ToolbarBuilder.AddToolBarButton(FTimingProfilerCommands::Get().ToggleTimersViewVisibility);
 				ToolbarBuilder.AddToolBarButton(FTimingProfilerCommands::Get().ToggleStatsCountersViewVisibility);
diff --git a/Engine/Source/Developer/TraceInsights/Private/Insights/Widgets/STimingProfilerWindow.cpp b/Engine/Source/Developer/TraceInsights/Private/Insights/Widgets/STimingProfilerWindow.cpp
--- a/Engine/Source/Developer/TraceInsights/Private/Insights/Widgets/STimingProfilerWindow.cpp
+++ b/Engine/Source/Developer/TraceInsights/Private/Insights/Widgets/STimingProfilerWindow.cpp
@@ -75,13 +75,23 @@ BEGIN_SLATE_FUNCTION_BUILD_OPTIMIZATION
 
 TSharedRef<SDockTab> STimingProfilerWindow::SpawnTab_Toolbar(const FSpawnTabArgs& Args)
 {
-	return SNew(SDockTab)
+	const TSharedRef<SDockTab> DockTab = SNew(SDockTab)
 		.ShouldAutosize(true)
 		.TabRole(ETabRole::PanelTab)
 		//.IsEnabled(this, &STimingProfilerWindow::IsProfilerEnabled)
 		[
 			SNew(STimingProfilerToolbar)
 		];
+
+	DockTab->SetOnTabClosed(SDockTab::FOnTabClosedCallback::CreateRaw(this, &STimingProfilerWindow::OnToolbarTabClosed));
+
+	return DockTab;
+}
+
+////////////////////////////////////////////////////////////////////////////////////////////////////
+
+void STimingProfilerWindow::OnToolbarTabClosed(TSharedRef<SDockTab> TabBeingClosed)
+{
 }
 
 ////////////////////////////////////////////////////////////////////////////////////////////////////
@@ -90,13 +100,24 @@ TSharedRef<SDockTab> STimingProfilerWindow::SpawnTab_FramesTrack(const FSpawnTab
 {
 	FTimingProfilerManager::Get()->SetFramesTrackVisible(true);
 
-	return SNew(SDockTab)
+	const TSharedRef<SDockTab> DockTab = SNew(SDockTab)
 		.ShouldAutosize(false)
 		.TabRole(ETabRole::PanelTab)
 		//.IsEnabled(this, &STimingProfilerWindow::IsProfilerEnabled)
 		[
 			SAssignNew(FrameTrack, SFrameTrack)
 		];
+
+	DockTab->SetOnTabClosed(SDockTab::FOnTabClosedCallback::CreateRaw(this, &STimingProfilerWindow::OnFramesTrackTabClosed));
+
+	return DockTab;
+}
+
+////////////////////////////////////////////////////////////////////////////////////////////////////
+
+void STimingProfilerWindow::OnFramesTrackTabClosed(TSharedRef<SDockTab> TabBeingClosed)
+{
+	FTimingProfilerManager::Get()->SetFramesTrackVisible(false);
 }
 
 ////////////////////////////////////////////////////////////////////////////////////////////////////
@@ -105,13 +126,24 @@ TSharedRef<SDockTab> STimingProfilerWindow::SpawnTab_GraphTrack(const FSpawnTabA
 {
 	FTimingProfilerManager::Get()->SetGraphTrackVisible(true);
 
-	return SNew(SDockTab)
+	const TSharedRef<SDockTab> DockTab = SNew(SDockTab)
 		.ShouldAutosize(false)
 		.TabRole(ETabRole::PanelTab)
 		//.IsEnabled(this, &STimingProfilerWindow::IsProfilerEnabled)
 		[
 			SAssignNew(GraphTrack, SGraphTrack)
 		];
+
+	DockTab->SetOnTabClosed(SDockTab::FOnTabClosedCallback::CreateRaw(this, &STimingProfilerWindow::OnGraphTrackTabClosed));
+
+	return DockTab;
+}
+
+////////////////////////////////////////////////////////////////////////////////////////////////////
+
+void STimingProfilerWindow::OnGraphTrackTabClosed(TSharedRef<SDockTab> TabBeingClosed)
+{
+	FTimingProfilerManager::Get()->SetGraphTrackVisible(false);
 }
 
 ////////////////////////////////////////////////////////////////////////////////////////////////////
@@ -120,13 +152,24 @@ TSharedRef<SDockTab> STimingProfilerWindow::SpawnTab_TimingView(const FSpawnTabA
 {
 	FTimingProfilerManager::Get()->SetTimingViewVisible(true);
 
-	return SNew(SDockTab)
+	const TSharedRef<SDockTab> DockTab = SNew(SDockTab)
 		.ShouldAutosize(false)
 		.TabRole(ETabRole::PanelTab)
 		//.IsEnabled(this, &STimingProfilerWindow::IsProfilerEnabled)
 		[
 			SAssignNew(TimingView, STimingView)
 		];
+
+	DockTab->SetOnTabClosed(SDockTab::FOnTabClosedCallback::CreateRaw(this, &STimingProfilerWindow::OnTimingViewTabClosed));
+
+	return DockTab;
+}
+
+////////////////////////////////////////////////////////////////////////////////////////////////////
+
+void STimingProfilerWindow::OnTimingViewTabClosed(TSharedRef<SDockTab> TabBeingClosed)
+{
+	FTimingProfilerManager::Get()->SetTimingViewVisible(false);
 }
 
 ////////////////////////////////////////////////////////////////////////////////////////////////////
@@ -135,13 +178,24 @@ TSharedRef<SDockTab> STimingProfilerWindow::SpawnTab_Timers(const FSpawnTabArgs&
 {
 	FTimingProfilerManager::Get()->SetTimersViewVisible(true);
 
-	return SNew(SDockTab)
+	const TSharedRef<SDockTab> DockTab = SNew(SDockTab)
 		.ShouldAutosize(false)
 		.TabRole(ETabRole::PanelTab)
 		//.IsEnabled(this, &STimingProfilerWindow::IsProfilerEnabled)
 		[
 			SAssignNew(TimersView, STimersView)
 		];
+
+	DockTab->SetOnTabClosed(SDockTab::FOnTabClosedCallback::CreateRaw(this, &STimingProfilerWindow::OnTimersTabClosed));
+
+	return DockTab;
+}
+
+////////////////////////////////////////////////////////////////////////////////////////////////////
+
+void STimingProfilerWindow::OnTimersTabClosed(TSharedRef<SDockTab> TabBeingClosed)
+{
+	FTimingProfilerManager::Get()->SetTimersViewVisible(false);
 }
 
 ////////////////////////////////////////////////////////////////////////////////////////////////////
@@ -150,13 +204,24 @@ TSharedRef<SDockTab> STimingProfilerWindow::SpawnTab_StatsCounters(const FSpawnT
 {
 	FTimingProfilerManager::Get()->SetStatsCountersViewVisible(true);
 
-	return SNew(SDockTab)
+	const TSharedRef<SDockTab> DockTab = SNew(SDockTab)
 		.ShouldAutosize(false)
 		.TabRole(ETabRole::PanelTab)
 		//.IsEnabled(this, &STimingProfilerWindow::IsProfilerEnabled)
 		[
 			SAssignNew(StatsView, SStatsView)
 		];
+
+	DockTab->SetOnTabClosed(SDockTab::FOnTabClosedCallback::CreateRaw(this, &STimingProfilerWindow::OnStatsCountersTabClosed));
+
+	return DockTab;
+}
+
+////////////////////////////////////////////////////////////////////////////////////////////////////
+
+void STimingProfilerWindow::OnStatsCountersTabClosed(TSharedRef<SDockTab> TabBeingClosed)
+{
+	FTimingProfilerManager::Get()->SetStatsCountersViewVisible(false);
 }
 
 ////////////////////////////////////////////////////////////////////////////////////////////////////
@@ -165,13 +230,24 @@ TSharedRef<SDockTab> STimingProfilerWindow::SpawnTab_LogView(const FSpawnTabArgs
 {
 	FTimingProfilerManager::Get()->SetLogViewVisible(true);
 
-	return SNew(SDockTab)
+	const TSharedRef<SDockTab> DockTab = SNew(SDockTab)
 		.ShouldAutosize(false)
 		.TabRole(ETabRole::PanelTab)
 		//.IsEnabled(this, &STimingProfilerWindow::IsProfilerEnabled)
 		[
 			SAssignNew(LogView, SLogView)
 		];
+
+	DockTab->SetOnTabClosed(SDockTab::FOnTabClosedCallback::CreateRaw(this, &STimingProfilerWindow::OnLogViewTabClosed));
+
+	return DockTab;
+}
+
+////////////////////////////////////////////////////////////////////////////////////////////////////
+
+void STimingProfilerWindow::OnLogViewTabClosed(TSharedRef<SDockTab> TabBeingClosed)
+{
+	FTimingProfilerManager::Get()->SetLogViewVisible(false);
 }
 
 ////////////////////////////////////////////////////////////////////////////////////////////////////
@@ -192,10 +268,10 @@ void STimingProfilerWindow::Construct(const FArguments& InArgs, const TSharedRef
 		.SetIcon(FSlateIcon(FInsightsStyle::GetStyleSetName(), "FramesTrack.Icon.Small"))
 		.SetGroup(AppMenuGroup);
 
-	TabManager->RegisterTabSpawner(FTimingProfilerTabs::GraphTrackID, FOnSpawnTab::CreateRaw(this, &STimingProfilerWindow::SpawnTab_GraphTrack))
-		.SetDisplayName(LOCTEXT("GraphTrackTabTitle", "Graph"))
-		.SetIcon(FSlateIcon(FInsightsStyle::GetStyleSetName(), "GraphTrack.Icon.Small"))
-		.SetGroup(AppMenuGroup);
+	//TabManager->RegisterTabSpawner(FTimingProfilerTabs::GraphTrackID, FOnSpawnTab::CreateRaw(this, &STimingProfilerWindow::SpawnTab_GraphTrack))
+	//	.SetDisplayName(LOCTEXT("GraphTrackTabTitle", "Graph"))
+	//	.SetIcon(FSlateIcon(FInsightsStyle::GetStyleSetName(), "GraphTrack.Icon.Small"))
+	//	.SetGroup(AppMenuGroup);
 
 	TabManager->RegisterTabSpawner(FTimingProfilerTabs::TimingViewID, FOnSpawnTab::CreateRaw(this, &STimingProfilerWindow::SpawnTab_TimingView))
 		.SetDisplayName(LOCTEXT("TimingViewTabTitle", "Timing View"))
@@ -284,7 +360,7 @@ void STimingProfilerWindow::Construct(const FArguments& InArgs, const TSharedRef
 
 	// Create & initialize main menu.
 	FMenuBarBuilder MenuBarBuilder = FMenuBarBuilder(TSharedPtr<FUICommandList>());
-	
+
 	MenuBarBuilder.AddPullDownMenu(
 		LOCTEXT("MenuLabel", "MENU"),
 		FText::GetEmpty(),
diff --git a/Engine/Source/Developer/TraceInsights/Private/Insights/Widgets/STimingProfilerWindow.h b/Engine/Source/Developer/TraceInsights/Private/Insights/Widgets/STimingProfilerWindow.h
--- a/Engine/Source/Developer/TraceInsights/Private/Insights/Widgets/STimingProfilerWindow.h
+++ b/Engine/Source/Developer/TraceInsights/Private/Insights/Widgets/STimingProfilerWindow.h
@@ -66,12 +66,25 @@ public:
 
 private:
 	TSharedRef<SDockTab> SpawnTab_Toolbar(const FSpawnTabArgs& Args);
+	void OnToolbarTabClosed(TSharedRef<SDockTab> TabBeingClosed);
+	
 	TSharedRef<SDockTab> SpawnTab_FramesTrack(const FSpawnTabArgs& Args);
+	void OnFramesTrackTabClosed(TSharedRef<SDockTab> TabBeingClosed);
+	
 	TSharedRef<SDockTab> SpawnTab_GraphTrack(const FSpawnTabArgs& Args);
+	void OnGraphTrackTabClosed(TSharedRef<SDockTab> TabBeingClosed);
+	
 	TSharedRef<SDockTab> SpawnTab_TimingView(const FSpawnTabArgs& Args);
+	void OnTimingViewTabClosed(TSharedRef<SDockTab> TabBeingClosed);
+	
 	TSharedRef<SDockTab> SpawnTab_Timers(const FSpawnTabArgs& Args);
+	void OnTimersTabClosed(TSharedRef<SDockTab> TabBeingClosed);
+	
 	TSharedRef<SDockTab> SpawnTab_StatsCounters(const FSpawnTabArgs& Args);
+	void OnStatsCountersTabClosed(TSharedRef<SDockTab> TabBeingClosed);
+
 	TSharedRef<SDockTab> SpawnTab_LogView(const FSpawnTabArgs& Args);
+	void OnLogViewTabClosed(TSharedRef<SDockTab> TabBeingClosed);
 
 	/**
 	 * Fill the main menu with menu items.
