[Insights] Asset Loading: Fixed crash (or corruption of the Tining view's "Tracks" drop down) when a session is re-loaded (missing re-initialisation of Timing view when session changes).

diff --git a/Engine/Source/Developer/TraceInsights/Private/Insights/InsightsManager.cpp b/Engine/Source/Developer/TraceInsights/Private/Insights/InsightsManager.cpp
--- a/Engine/Source/Developer/TraceInsights/Private/Insights/InsightsManager.cpp
+++ b/Engine/Source/Developer/TraceInsights/Private/Insights/InsightsManager.cpp
@@ -8,6 +8,7 @@
 #include "Templates/UniquePtr.h"
 
 // Insights
+#include "Insights/IoProfilerManager.h"
 #include "Insights/TimingProfilerManager.h"
 #include "Insights/Widgets/SStartPageWindow.h"
 #include "Insights/Widgets/STimingProfilerWindow.h"
@@ -140,6 +141,12 @@ void FInsightsManager::OnSessionChanged()
 		TimingProfilerManager->OnSessionChanged();
 	}
 
+	if (TSharedPtr<FIoProfilerManager> IoProfilerManager = FIoProfilerManager::Get())
+	{
+		// FIXME: make IoProfilerManager to register to SessionChangedEvent instead
+		IoProfilerManager->OnSessionChanged();
+	}
+
 	SessionChangedEvent.Broadcast();
 }
 
diff --git a/Engine/Source/Developer/TraceInsights/Private/Insights/ViewModels/TimingEventsTrack.h b/Engine/Source/Developer/TraceInsights/Private/Insights/ViewModels/TimingEventsTrack.h
--- a/Engine/Source/Developer/TraceInsights/Private/Insights/ViewModels/TimingEventsTrack.h
+++ b/Engine/Source/Developer/TraceInsights/Private/Insights/ViewModels/TimingEventsTrack.h
@@ -4,26 +4,27 @@
 
 #include "CoreMinimal.h"
 
+// Insights
 #include "Insights/ViewModels/BaseTimingTrack.h"
 
 ////////////////////////////////////////////////////////////////////////////////////////////////////
 
-static const float RealMinTimelineH = 13.0f;
+struct FTimingEventsTrackLayout
+{
+	static constexpr float RealMinTimelineH = 13.0f;
 
-static const float NormalLayoutEventH = 14.0f;
-static const float NormalLayoutEventDY = 2.0f;
-static const float NormalLayoutTimelineDY = 14.0f;
-static const float NormalLayoutMinTimelineH = 0.0f;
+	static constexpr float NormalLayoutEventH = 14.0f;
+	static constexpr float NormalLayoutEventDY = 2.0f;
+	static constexpr float NormalLayoutTimelineDY = 14.0f;
+	static constexpr float NormalLayoutMinTimelineH = 0.0f;
 
-static const float CompactLayoutEventH = 2.0f;
-static const float CompactLayoutEventDY = 1.0f;
-static const float CompactLayoutTimelineDY = 3.0f;
-static const float CompactLayoutMinTimelineH = 0.0f;
+	static constexpr float CompactLayoutEventH = 2.0f;
+	static constexpr float CompactLayoutEventDY = 1.0f;
+	static constexpr float CompactLayoutTimelineDY = 3.0f;
+	static constexpr float CompactLayoutMinTimelineH = 0.0f;
 
-////////////////////////////////////////////////////////////////////////////////////////////////////
+	//////////////////////////////////////////////////
 
-struct FTimingEventsTrackLayout
-{
 	bool bIsCompactMode;
 
 	float EventH; // height of a timing event, in Slate units
@@ -32,6 +33,8 @@ struct FTimingEventsTrackLayout
 	float MinTimelineH;
 	float TargetMinTimelineH;
 
+	//////////////////////////////////////////////////
+
 	float GetLaneY(uint32 Depth) const { return 1.0f + TimelineDY + Depth * (EventDY + EventH); }
 
 	void ForceNormalMode();
@@ -83,7 +86,7 @@ public:
 	int32 Depth; // number of lanes == Depth + 1
 	bool bIsCollapsed;
 
-	// Cached OnPaint state.
+	// TODO: Cached OnPaint state.
 	//TArray<FEventBoxInfo> Boxes;
 	//TArray<FEventBoxInfo> MergedBorderBoxes;
 	//TArray<FEventBoxInfo> Borders;
diff --git a/Engine/Source/Developer/TraceInsights/Private/Insights/ViewModels/TimingViewDrawHelper.cpp b/Engine/Source/Developer/TraceInsights/Private/Insights/ViewModels/TimingViewDrawHelper.cpp
--- a/Engine/Source/Developer/TraceInsights/Private/Insights/ViewModels/TimingViewDrawHelper.cpp
+++ b/Engine/Source/Developer/TraceInsights/Private/Insights/ViewModels/TimingViewDrawHelper.cpp
@@ -168,9 +168,9 @@ void FTimingViewDrawHelper::EndTimeline(FTimingEventsTrack& Track)
 		// 1.0f is for horizontal line between timelines
 		NewH = 1.0f + Layout.TimelineDY + Layout.EventH * (MaxDepth + 1) + Layout.EventDY * MaxDepth + Layout.TimelineDY;
 
-		if (NewH < RealMinTimelineH)
+		if (NewH < FTimingEventsTrackLayout::RealMinTimelineH)
 		{
-			NewH = RealMinTimelineH;
+			NewH = FTimingEventsTrackLayout::RealMinTimelineH;
 		}
 	}
 
diff --git a/Engine/Source/Developer/TraceInsights/Private/Insights/Widgets/STimingView.cpp b/Engine/Source/Developer/TraceInsights/Private/Insights/Widgets/STimingView.cpp
--- a/Engine/Source/Developer/TraceInsights/Private/Insights/Widgets/STimingView.cpp
+++ b/Engine/Source/Developer/TraceInsights/Private/Insights/Widgets/STimingView.cpp
@@ -512,8 +512,9 @@ void STimingView::Tick(const FGeometry& AllottedGeometry, const double InCurrent
 					}
 					else
 					{
-						bIsGroupVisible = ThreadGroups[GroupName].bIsVisible;
-						ThreadGroups[GroupName].Order = Order;
+						FThreadGroup& ThreadGroup = ThreadGroups[GroupName];
+						bIsGroupVisible = ThreadGroup.bIsVisible;
+						ThreadGroup.Order = Order;
 					}
 				}
 
@@ -3894,7 +3895,7 @@ void STimingView::ToggleAutoHideEmptyTracks()
 {
 	if (Layout.TargetMinTimelineH == 0.0f)
 	{
-		Layout.TargetMinTimelineH = RealMinTimelineH;
+		Layout.TargetMinTimelineH = FTimingEventsTrackLayout::RealMinTimelineH;
 	}
 	else
 	{
