[Insights] Fixed "-disableunity" compile error.

diff --git a/Engine/Source/Developer/TraceInsights/Private/Insights/Common/PaintUtils.h b/Engine/Source/Developer/TraceInsights/Private/Insights/Common/PaintUtils.h
--- a/Engine/Source/Developer/TraceInsights/Private/Insights/Common/PaintUtils.h
+++ b/Engine/Source/Developer/TraceInsights/Private/Insights/Common/PaintUtils.h
@@ -3,6 +3,8 @@
 #pragma once
 
 #include "CoreMinimal.h"
+#include "Fonts/FontMeasure.h"
+#include "Framework/Application/SlateApplication.h"
 #include "Rendering/DrawElements.h"
 
 #define MAKE_PAINT_GEOMETRY_PT(Geometry, X, Y)       Geometry.ToPaintGeometry(FSlateLayoutTransform(1.0f, FVector2D(X, Y)))
@@ -77,7 +79,7 @@ struct FDrawContext
 	inline void DrawTextAligned(EHorizontalAlignment HAlign, const float X, const float Y, const FString& Text, const FSlateFontInfo& Font, const FLinearColor& Color) const
 	{
 		const TSharedRef<FSlateFontMeasure> FontMeasureService = FSlateApplication::Get().GetRenderer()->GetFontMeasureService();
-		float TextWidth = FontMeasureService->Measure(Text, Font).X;
+		const float TextWidth = FontMeasureService->Measure(Text, Font).X;
 		float TextX = X;
 		if (HAlign == HAlign_Right)
 		{
