[Insights]
  - Fixed UE-76068 (When grouping Timing Insight's Timers by Name, "n" is lowercase)
  - Fixed UE-76084 (Selecting a Timer in the Asset Loading Insight changes changes the selected Timer in the Timing Insight's Timers window)

diff --git a/Engine/Source/Developer/TraceInsights/Private/Insights/Widgets/SStatsView.cpp b/Engine/Source/Developer/TraceInsights/Private/Insights/Widgets/SStatsView.cpp
--- a/Engine/Source/Developer/TraceInsights/Private/Insights/Widgets/SStatsView.cpp
+++ b/Engine/Source/Developer/TraceInsights/Private/Insights/Widgets/SStatsView.cpp
@@ -1007,7 +1007,7 @@ void SStatsView::CreateGroups()
 	{
 		for (const FStatsNodePtr& StatsNodePtr : StatsNodes)
 		{
-			const FName GroupName = *StatsNodePtr->GetName().GetPlainNameString().Left(1);
+			const FName GroupName = *StatsNodePtr->GetName().GetPlainNameString().Left(1).ToUpper();
 
 			FStatsNodePtr* GroupPtr = GroupNodeSet.Find(GroupName);
 			if (!GroupPtr)
diff --git a/Engine/Source/Developer/TraceInsights/Private/Insights/Widgets/STimersView.cpp b/Engine/Source/Developer/TraceInsights/Private/Insights/Widgets/STimersView.cpp
--- a/Engine/Source/Developer/TraceInsights/Private/Insights/Widgets/STimersView.cpp
+++ b/Engine/Source/Developer/TraceInsights/Private/Insights/Widgets/STimersView.cpp
@@ -1016,7 +1016,7 @@ void STimersView::CreateGroups()
 	{
 		for (const FTimerNodePtr& TimerNodePtr : TimerNodes)
 		{
-			const FName GroupName = *TimerNodePtr->GetName().GetPlainNameString().Left(1);
+			const FName GroupName = *TimerNodePtr->GetName().GetPlainNameString().Left(1).ToUpper();
 
 			FTimerNodePtr* GroupPtr = GroupNodeSet.Find(GroupName);
 			if (!GroupPtr)
diff --git a/Engine/Source/Developer/TraceInsights/Private/Insights/Widgets/STimingView.cpp b/Engine/Source/Developer/TraceInsights/Private/Insights/Widgets/STimingView.cpp
--- a/Engine/Source/Developer/TraceInsights/Private/Insights/Widgets/STimingView.cpp
+++ b/Engine/Source/Developer/TraceInsights/Private/Insights/Widgets/STimingView.cpp
@@ -3403,10 +3403,12 @@ bool STimingView::SearchTimingEvent(const double InStartTime,
 void STimingView::OnSelectedTimingEventChanged()
 {
 	// Select the timer node coresponding to timing event type of selected timing event.
-	if (SelectedTimingEvent.IsValid() &&
+	if (!bAssetLoadingMode &&
+		SelectedTimingEvent.IsValid() &&
 		(SelectedTimingEvent.Track->GetType() == ETimingEventsTrackType::Cpu ||
 		 SelectedTimingEvent.Track->GetType() == ETimingEventsTrackType::Gpu))
 	{
+		//TODO: make this a more generic (i.e. no hardocdings on TimingProfilerManager)
 		TSharedPtr<STimingProfilerWindow> Wnd = FTimingProfilerManager::Get()->GetProfilerWindow();
 		if (Wnd)
 		{
