[Insights] Log View: Changed the drop down for Verbosity Threshold to contain radio buttons (instead of checkboxs). Fixes UE-75564.

diff --git a/Engine/Source/Developer/TraceInsights/Private/Insights/Widgets/SLogView.cpp b/Engine/Source/Developer/TraceInsights/Private/Insights/Widgets/SLogView.cpp
--- a/Engine/Source/Developer/TraceInsights/Private/Insights/Widgets/SLogView.cpp
+++ b/Engine/Source/Developer/TraceInsights/Private/Insights/Widgets/SLogView.cpp
@@ -4,6 +4,7 @@
 
 #include "Algo/BinarySearch.h"
 #include "Async/AsyncWork.h"
+#include "EditorFontGlyphs.h"
 #include "EditorStyleSet.h"
 #include "Framework/Application/SlateApplication.h"
 #include "Framework/MultiBox/MultiBoxBuilder.h"
@@ -1104,33 +1105,34 @@ void  SLogView::CreateVerbosityThresholdMenuSection(FMenuBuilder& MenuBuilder)
 	{
 		const FVerbosityThresholdInfo& Threshold = VerbosityThresholds[Index];
 
-		/*
-		MenuBuilder.AddMenuEntry(
-			Threshold.Label,
-			Threshold.ToolTip,
-			FSlateIcon(),
-			FUIAction(FExecuteAction::CreateSP(this, &SLogView::VerbosityThreshold_Execute, Threshold.Verbosity),
-				FCanExecuteAction::CreateLambda([] { return true; }),
-				FIsActionChecked::CreateSP(this, &SLogView::VerbosityThreshold_IsChecked, Threshold.Verbosity)),
-			NAME_None,
-			EUserInterfaceActionType::ToggleButton
-		);
-		*/
-
-		const TSharedRef<SWidget> TextBlock = SNew(STextBlock)
-			.Text(Threshold.Label)
-			.ShadowColorAndOpacity(FLinearColor(0.05f, 0.05f, 0.05f, 1.0f))
-			.ShadowOffset(FVector2D(1.0f, 1.0f))
-			.ColorAndOpacity(FSlateColor(FTimeMarkerTrackBuilder::GetColorByVerbosity(Threshold.Verbosity)));
+		const TSharedRef<SWidget> TextBlock = SNew(SHorizontalBox)
+			+ SHorizontalBox::Slot()
+			.AutoWidth()
+			[
+				SNew(STextBlock)
+				.Text(Threshold.Label)
+				.ShadowColorAndOpacity(FLinearColor(0.05f, 0.05f, 0.05f, 1.0f))
+				.ShadowOffset(FVector2D(1.0f, 1.0f))
+				.ColorAndOpacity(FSlateColor(FTimeMarkerTrackBuilder::GetColorByVerbosity(Threshold.Verbosity)))
+			]
+			+ SHorizontalBox::Slot()
+			.Padding(2.0f, 0.0f)
+			.VAlign(VAlign_Center)
+			[
+				SNew(STextBlock)
+				.Font(FEditorStyle::Get().GetFontStyle("FontAwesome.9"))
+				.Text(this, &SLogView::VerbosityThreshold_GetSuffixGlyph, Threshold.Verbosity)
+				.ColorAndOpacity(this, &SLogView::VerbosityThreshold_GetSuffixColor, Threshold.Verbosity)
+			];
 
 		MenuBuilder.AddMenuEntry(
 			FUIAction(FExecuteAction::CreateSP(this, &SLogView::VerbosityThreshold_Execute, Threshold.Verbosity),
-				FCanExecuteAction::CreateLambda([] { return true; }),
+				FCanExecuteAction(),
 				FIsActionChecked::CreateSP(this, &SLogView::VerbosityThreshold_IsChecked, Threshold.Verbosity)),
 			TextBlock,
 			NAME_None,
 			Threshold.ToolTip,
-			EUserInterfaceActionType::ToggleButton
+			EUserInterfaceActionType::RadioButton
 		);
 	}
 }
@@ -1148,7 +1150,7 @@ TSharedRef<SWidget> SLogView::MakeCategoryFilterMenu()
 			LOCTEXT("ShowAllCategories_Tooltip", "Change filtering to show/hide all categories"),
 			FSlateIcon(),
 			FUIAction(FExecuteAction::CreateSP(this, &SLogView::ShowHideAllCategories_Execute),
-				FCanExecuteAction::CreateLambda([] { return true; }),
+				FCanExecuteAction(),
 				FIsActionChecked::CreateSP(this, &SLogView::ShowHideAllCategories_IsChecked)),
 			NAME_None,
 			EUserInterfaceActionType::ToggleButton
@@ -1172,17 +1174,6 @@ void SLogView::CreateCategoriesFilterMenuSection(FMenuBuilder& MenuBuilder)
 		const FString CategoryString = CategoryName.ToString();
 		const FText CategoryText(FText::AsCultureInvariant(CategoryString));
 
-		//MenuBuilder.AddMenuEntry(
-		//	FText::AsCultureInvariant(CategoryString),
-		//	FText::Format(LOCTEXT("Category_Tooltip", "Filter the Log View to show/hide category: {0}"), CategoryText),
-		//	FSlateIcon(),
-		//	FUIAction(FExecuteAction::CreateSP(this, &SLogView::ToggleCategory_Execute, CategoryName),
-		//		FCanExecuteAction::CreateLambda([] { return true; }),
-		//		FIsActionChecked::CreateSP(this, &SLogView::ToggleCategory_IsChecked, CategoryName)),
-		//	NAME_None,
-		//	EUserInterfaceActionType::ToggleButton
-		//);
-
 		const TSharedRef<SWidget> TextBlock = SNew(STextBlock)
 			.Text(FText::AsCultureInvariant(CategoryString))
 			.ShadowColorAndOpacity(FLinearColor(0.05f, 0.05f, 0.05f, 1.0f))
@@ -1191,7 +1182,7 @@ void SLogView::CreateCategoriesFilterMenuSection(FMenuBuilder& MenuBuilder)
 
 		MenuBuilder.AddMenuEntry(
 			FUIAction(FExecuteAction::CreateSP(this, &SLogView::ToggleCategory_Execute, CategoryName),
-				FCanExecuteAction::CreateLambda([] { return true; }),
+				FCanExecuteAction(),
 				FIsActionChecked::CreateSP(this, &SLogView::ToggleCategory_IsChecked, CategoryName)),
 			TextBlock,
 			NAME_None,
@@ -1205,7 +1196,7 @@ void SLogView::CreateCategoriesFilterMenuSection(FMenuBuilder& MenuBuilder)
 
 bool SLogView::VerbosityThreshold_IsChecked(ELogVerbosity::Type Verbosity) const
 {
-	return Verbosity <= Filter.GetVerbosityThreshold();
+	return Verbosity == Filter.GetVerbosityThreshold();
 }
 
 ////////////////////////////////////////////////////////////////////////////////////////////////////
@@ -1218,6 +1209,22 @@ void SLogView::VerbosityThreshold_Execute(ELogVerbosity::Type Verbosity)
 
 ////////////////////////////////////////////////////////////////////////////////////////////////////
 
+FText SLogView::VerbosityThreshold_GetSuffixGlyph(ELogVerbosity::Type Verbosity) const
+{
+	return Verbosity <= Filter.GetVerbosityThreshold() ? FEditorFontGlyphs::Check : FEditorFontGlyphs::Times;
+}
+
+////////////////////////////////////////////////////////////////////////////////////////////////////
+
+FSlateColor SLogView::VerbosityThreshold_GetSuffixColor(ELogVerbosity::Type Verbosity) const
+{
+	return Verbosity <= Filter.GetVerbosityThreshold() ?
+		FSlateColor(FLinearColor(1.0f, 1.0f, 1.0f, 1.0f)) :
+		FSlateColor(FLinearColor(1.0f, 0.0f, 0.0f, 1.0f));
+}
+
+////////////////////////////////////////////////////////////////////////////////////////////////////
+
 bool SLogView::ShowHideAllCategories_IsChecked() const
 {
 	return Filter.IsShowAllCategoriesEnabled();
diff --git a/Engine/Source/Developer/TraceInsights/Private/Insights/Widgets/SLogView.h b/Engine/Source/Developer/TraceInsights/Private/Insights/Widgets/SLogView.h
--- a/Engine/Source/Developer/TraceInsights/Private/Insights/Widgets/SLogView.h
+++ b/Engine/Source/Developer/TraceInsights/Private/Insights/Widgets/SLogView.h
@@ -93,6 +93,8 @@ protected:
 
 	bool VerbosityThreshold_IsChecked(ELogVerbosity::Type Verbosity) const;
 	void VerbosityThreshold_Execute(ELogVerbosity::Type Verbosity);
+	FText VerbosityThreshold_GetSuffixGlyph(ELogVerbosity::Type Verbosity) const;
+	FSlateColor VerbosityThreshold_GetSuffixColor(ELogVerbosity::Type Verbosity) const;
 
 	bool ShowHideAllCategories_IsChecked() const;
 	void ShowHideAllCategories_Execute();
