[Insights]
  - Graph track: Fixed draw lines for cases where a value is infinite.
  - Log View: Fixed case where there are not logs in analyzed trace session resulting in UI showing continously "Filtering... please wait...". Related to UE-76162.
  - Timing Insights: Added toggles for available trace modules (only for live sessions), in toolbar.

diff --git a/Engine/Source/Developer/TraceInsights/Private/Insights/InsightsManager.cpp b/Engine/Source/Developer/TraceInsights/Private/Insights/InsightsManager.cpp
--- a/Engine/Source/Developer/TraceInsights/Private/Insights/InsightsManager.cpp
+++ b/Engine/Source/Developer/TraceInsights/Private/Insights/InsightsManager.cpp
@@ -85,6 +85,13 @@ TSharedPtr<const Trace::IAnalysisSession> FInsightsManager::GetSession() const
 
 ////////////////////////////////////////////////////////////////////////////////////////////////////
 
+Trace::FSessionHandle FInsightsManager::GetSessionHandle() const
+{
+	return CurrentSessionHandle;
+}
+
+////////////////////////////////////////////////////////////////////////////////////////////////////
+
 const TSharedRef<FUICommandList> FInsightsManager::GetCommandList() const
 {
 	return CommandList;
@@ -127,6 +134,7 @@ void FInsightsManager::ResetSession()
 	if (Session.IsValid())
 	{
 		Session.Reset();
+		CurrentSessionHandle = 0;
 		OnSessionChanged();
 	}
 }
@@ -269,6 +277,7 @@ void FInsightsManager::LoadSession(Trace::FSessionHandle SessionHandle)
 	if (DataStream)
 	{
 		Session = AnalysisService->StartAnalysis(SessionInfo.Name, MoveTemp(DataStream));
+		CurrentSessionHandle = SessionHandle;
 		SpawnAndActivateTabs();
 		OnSessionChanged();
 	}
@@ -284,6 +293,7 @@ void FInsightsManager::LoadTraceFile(const FString& TraceFilepath)
 	if (DataStream)
 	{
 		Session = AnalysisService->StartAnalysis(*TraceFilepath, MoveTemp(DataStream));
+		CurrentSessionHandle = 0;
 		SpawnAndActivateTabs();
 		OnSessionChanged();
 	}
diff --git a/Engine/Source/Developer/TraceInsights/Private/Insights/InsightsManager.h b/Engine/Source/Developer/TraceInsights/Private/Insights/InsightsManager.h
--- a/Engine/Source/Developer/TraceInsights/Private/Insights/InsightsManager.h
+++ b/Engine/Source/Developer/TraceInsights/Private/Insights/InsightsManager.h
@@ -84,6 +84,9 @@ public:
 	/** @return an instance of the trace analysis session. */
 	TSharedPtr<const Trace::IAnalysisSession> GetSession() const;
 
+	/** @return the session handle of the trace analysis session. */
+	Trace::FSessionHandle GetSessionHandle() const;
+
 	/** @returns UI command list for the main manager. */
 	const TSharedRef<FUICommandList> GetCommandList() const;
 
@@ -187,6 +190,9 @@ private:
 	/** The trace analysis session. */
 	TSharedPtr<const Trace::IAnalysisSession> Session;
 
+	/** The session handle. */
+	Trace::FSessionHandle CurrentSessionHandle;
+
 	/** List of UI commands for this manager. This will be filled by this and corresponding classes. */
 	TSharedRef<FUICommandList> CommandList;
 
diff --git a/Engine/Source/Developer/TraceInsights/Private/Insights/ViewModels/GraphTrack.cpp b/Engine/Source/Developer/TraceInsights/Private/Insights/ViewModels/GraphTrack.cpp
--- a/Engine/Source/Developer/TraceInsights/Private/Insights/ViewModels/GraphTrack.cpp
+++ b/Engine/Source/Developer/TraceInsights/Private/Insights/ViewModels/GraphTrack.cpp
@@ -529,7 +529,8 @@ void FGraphTrackBuilder::AddConnectedLine(double Time, double Value)
 		return;
 	}
 
-	const float Y = Track.GetYForValue(Value);
+	// Transform Value to Y coordinate (local track space) and ensure Y is not infinite.
+	const float Y = FMath::Clamp<float>(Track.GetYForValue(Value), -FLT_MAX, FLT_MAX);
 
 	const float X = Viewport.TimeToSlateUnitsRounded(Viewport.RestrictEndTime(Time));
 
diff --git a/Engine/Source/Developer/TraceInsights/Private/Insights/Widgets/SLogView.cpp b/Engine/Source/Developer/TraceInsights/Private/Insights/Widgets/SLogView.cpp
--- a/Engine/Source/Developer/TraceInsights/Private/Insights/Widgets/SLogView.cpp
+++ b/Engine/Source/Developer/TraceInsights/Private/Insights/Widgets/SLogView.cpp
@@ -846,6 +846,12 @@ void SLogView::Tick(const FGeometry& AllottedGeometry, const double InCurrentTim
 			UpdateStatsText();
 		}
 	}
+	else if (bIsDirty && !FilteringAsyncTask.IsValid())
+	{
+		bIsDirty = false;
+		DirtyStopwatch.Reset();
+		UpdateStatsText();
+	}
 
 	if (FilteringAsyncTask.IsValid() &&
 		FilteringAsyncTask->IsDone())
diff --git a/Engine/Source/Developer/TraceInsights/Private/Insights/Widgets/STimingProfilerToolbar.cpp b/Engine/Source/Developer/TraceInsights/Private/Insights/Widgets/STimingProfilerToolbar.cpp
--- a/Engine/Source/Developer/TraceInsights/Private/Insights/Widgets/STimingProfilerToolbar.cpp
+++ b/Engine/Source/Developer/TraceInsights/Private/Insights/Widgets/STimingProfilerToolbar.cpp
@@ -10,6 +10,7 @@
 #include "Widgets/Docking/SDockTab.h"
 #include "Widgets/Layout/SBorder.h"
 #include "Widgets/SBoxPanel.h"
+#include "Widgets/SToolTip.h"
 
 // Insights
 #include "Insights/InsightsCommands.h"
@@ -37,18 +38,10 @@ STimingProfilerToolbar::~STimingProfilerToolbar()
 
 void STimingProfilerToolbar::Construct(const FArguments& InArgs)
 {
-	CreateCommands();
-
 	struct Local
 	{
-		static void FillToolbar1(FToolBarBuilder& ToolbarBuilder)
+		static void FillViewToolbar(FToolBarBuilder& ToolbarBuilder)
 		{
-			//ToolbarBuilder.BeginSection("Session");
-			//{
-			//	ToolbarBuilder.AddToolBarButton(FInsightsCommands::Get().InsightsManager_Live);
-			//	ToolbarBuilder.AddToolBarButton(FInsightsCommands::Get().InsightsManager_Load);
-			//}
-			//ToolbarBuilder.EndSection();
 			ToolbarBuilder.BeginSection("View");
 			{
 				ToolbarBuilder.AddToolBarButton(FTimingProfilerCommands::Get().ToggleFramesTrackVisibility);
@@ -59,16 +52,11 @@ void STimingProfilerToolbar::Construct(const FArguments& InArgs)
 				ToolbarBuilder.AddToolBarButton(FTimingProfilerCommands::Get().ToggleLogViewVisibility);
 			}
 			ToolbarBuilder.EndSection();
-			//ToolbarBuilder.BeginSection("Options");
-			//{
-			//	ToolbarBuilder.AddToolBarButton(FInsightsCommands::Get().OpenSettings);
-			//}
-			//ToolbarBuilder.EndSection();
 		}
 
-		static void FillToolbar2(FToolBarBuilder& ToolbarBuilder)
+		static void FillRightSideToolbar(FToolBarBuilder& ToolbarBuilder)
 		{
-			ToolbarBuilder.BeginSection("Options");
+			ToolbarBuilder.BeginSection("Debug");
 			{
 				ToolbarBuilder.AddToolBarButton(FInsightsCommands::Get().ToggleDebugInfo);
 			}
@@ -78,11 +66,12 @@ void STimingProfilerToolbar::Construct(const FArguments& InArgs)
 
 	TSharedPtr<FUICommandList> CommandList = FInsightsManager::Get()->GetCommandList();
 
-	FToolBarBuilder ToolbarBuilder1(CommandList.ToSharedRef(), FMultiBoxCustomization::None);
-	Local::FillToolbar1(ToolbarBuilder1);
+	FToolBarBuilder ToolbarBuilder(CommandList.ToSharedRef(), FMultiBoxCustomization::None);
+	Local::FillViewToolbar(ToolbarBuilder);
+	FillModulesToolbar(ToolbarBuilder);
 
-	FToolBarBuilder ToolbarBuilder2(CommandList.ToSharedRef(), FMultiBoxCustomization::None);
-	Local::FillToolbar2(ToolbarBuilder2);
+	FToolBarBuilder RightSideToolbarBuilder(CommandList.ToSharedRef(), FMultiBoxCustomization::None);
+	Local::FillRightSideToolbar(RightSideToolbarBuilder);
 
 	// Create the tool bar!
 	ChildSlot
@@ -100,7 +89,7 @@ void STimingProfilerToolbar::Construct(const FArguments& InArgs)
 			.BorderImage(FEditorStyle::GetBrush("NoBorder"))
 			.IsEnabled(FSlateApplication::Get().GetNormalExecutionAttribute())
 			[
-				ToolbarBuilder1.MakeWidget()
+				ToolbarBuilder.MakeWidget()
 			]
 		]
 
@@ -115,7 +104,7 @@ void STimingProfilerToolbar::Construct(const FArguments& InArgs)
 			.BorderImage(FEditorStyle::GetBrush("NoBorder"))
 			.IsEnabled(FSlateApplication::Get().GetNormalExecutionAttribute())
 			[
-				ToolbarBuilder2.MakeWidget()
+				RightSideToolbarBuilder.MakeWidget()
 			]
 		]
 	];
@@ -123,31 +112,124 @@ void STimingProfilerToolbar::Construct(const FArguments& InArgs)
 
 ////////////////////////////////////////////////////////////////////////////////////////////////////
 
-void STimingProfilerToolbar::ShowStats()
+bool STimingProfilerToolbar::ToggleModule_CanExecute(FName ModuleName) const
+{
+	TSharedRef<Trace::ISessionService> SessionService = FInsightsManager::Get()->GetSessionService();
+	if (!SessionService->IsRecorderServerRunning())
+	{
+		return false;
+	}
+
+	TSharedPtr<const Trace::IAnalysisSession> Session = FInsightsManager::Get()->GetSession();
+	if (!Session.IsValid())
+	{
+		return false;
+	}
+
+	Trace::FSessionHandle SessionHandle = FInsightsManager::Get()->GetSessionHandle();
+	if (SessionHandle == 0)
+	{
+		return false;
+	}
+
+	Trace::FSessionInfo SessionInfo;
+	SessionService->GetSessionInfo(SessionHandle, SessionInfo);
+	return SessionInfo.bIsLive;
+}
+
+////////////////////////////////////////////////////////////////////////////////////////////////////
+
+void STimingProfilerToolbar::ToggleModule_Execute(FName ModuleName)
 {
-	// do nothing
+	TSharedRef<Trace::ISessionService> SessionService = FInsightsManager::Get()->GetSessionService();
+	TSharedPtr<const Trace::IAnalysisSession> Session = FInsightsManager::Get()->GetSession();
+	Trace::FSessionHandle SessionHandle = FInsightsManager::Get()->GetSessionHandle();
+
+	bool bState = SessionService->IsModuleEnabled(SessionHandle, ModuleName);
+	return SessionService->SetModuleEnabled(SessionHandle, ModuleName, !bState);
 }
 
 ////////////////////////////////////////////////////////////////////////////////////////////////////
 
-void STimingProfilerToolbar::ShowMemory()
+bool STimingProfilerToolbar::ToggleModule_IsChecked(FName ModuleName) const
 {
-	// do nothing
+	TSharedRef<Trace::ISessionService> SessionService = FInsightsManager::Get()->GetSessionService();
+	TSharedPtr<const Trace::IAnalysisSession> Session = FInsightsManager::Get()->GetSession();
+	Trace::FSessionHandle SessionHandle = FInsightsManager::Get()->GetSessionHandle();
+
+	return SessionService->IsModuleEnabled(SessionHandle, ModuleName);
 }
 
 ////////////////////////////////////////////////////////////////////////////////////////////////////
 
-void STimingProfilerToolbar::CreateCommands()
+ECheckBoxState STimingProfilerToolbar::ToggleModule_IsChecked2(FName ModuleName) const
 {
-	//TSharedPtr<FUICommandList> ProfilerCommandList = FTimingProfilerManager::Get()->GetCommandList();
-	//const FTimingProfilerCommands& Commands = FTimingProfilerCommands::Get();
-	//
-	//// Stats command
-	//ProfilerCommandList->MapAction(Commands.StatsProfiler,
-	//	FExecuteAction::CreateRaw(this, &STimingProfilerToolbar::ShowStats),
-	//	FCanExecuteAction(),
-	//	FIsActionChecked::CreateRaw(this, &STimingProfilerToolbar::IsShowingStats)
-	//);
+	TSharedRef<Trace::ISessionService> SessionService = FInsightsManager::Get()->GetSessionService();
+	TSharedPtr<const Trace::IAnalysisSession> Session = FInsightsManager::Get()->GetSession();
+	Trace::FSessionHandle SessionHandle = FInsightsManager::Get()->GetSessionHandle();
+
+	return SessionService->IsModuleEnabled(SessionHandle, ModuleName) ? ECheckBoxState::Checked : ECheckBoxState::Unchecked;
+}
+
+////////////////////////////////////////////////////////////////////////////////////////////////////
+
+void STimingProfilerToolbar::ToggleModule_OnCheckStateChanged(ECheckBoxState NewRadioState, FName ModuleName)
+{
+	TSharedRef<Trace::ISessionService> SessionService = FInsightsManager::Get()->GetSessionService();
+	TSharedPtr<const Trace::IAnalysisSession> Session = FInsightsManager::Get()->GetSession();
+	Trace::FSessionHandle SessionHandle = FInsightsManager::Get()->GetSessionHandle();
+
+	bool bState = (NewRadioState == ECheckBoxState::Checked);
+	return SessionService->SetModuleEnabled(SessionHandle, ModuleName, bState);
+}
+
+////////////////////////////////////////////////////////////////////////////////////////////////////
+
+void STimingProfilerToolbar::FillModulesToolbar(FToolBarBuilder& ToolbarBuilder)
+{
+	ToolbarBuilder.BeginSection("Modules");
+	{
+		TSharedRef<Trace::IModuleService> ModuleService = FInsightsManager::Get()->GetModuleService();
+
+		TArray<Trace::FModuleInfo> AvailableModules;
+		ModuleService->GetAvailableModules(AvailableModules);
+
+		for (int32 ModuleIndex = 0; ModuleIndex < AvailableModules.Num(); ++ModuleIndex)
+		{
+			const Trace::FModuleInfo& Module = AvailableModules[ModuleIndex];
+
+			FText Label = FText::FromName(Module.DisplayName);
+			FText ToggleModuleToolTip = FText::Format(LOCTEXT("ToggleModuleToolTip", "Enable/disable {0} trace module (only for live sessions)."), FText::FromName(Module.DisplayName));
+
+			// TODO: Uncomment this when adding icons to toolbar.
+			//ToolbarBuilder.AddToolBarButton(
+			//	FUIAction(
+			//		FExecuteAction::CreateRaw(this, &STimingProfilerToolbar::ToggleModule_Execute, Module.Name),
+			//		FCanExecuteAction::CreateRaw(this, &STimingProfilerToolbar::ToggleModule_CanExecute, Module.Name),
+			//		FIsActionChecked::CreateRaw(this, &STimingProfilerToolbar::ToggleModule_IsChecked, Module.Name)
+			//	),
+			//	NAME_None, // ExtensionHook
+			//	Label,
+			//	ToggleModuleToolTip,
+			//	TAttribute<FSlateIcon>(), // Icon -- empty icon is not really empty :(
+			//	EUserInterfaceActionType::ToggleButton,
+			//	NAME_None); // TutorialHighlightName
+
+			// Workaround for having toogle buttons without icons.
+			ToolbarBuilder.AddWidget(SNew(SCheckBox)
+				.IsEnabled(this, &STimingProfilerToolbar::ToggleModule_CanExecute, Module.Name)
+				.IsChecked(this, &STimingProfilerToolbar::ToggleModule_IsChecked2, Module.Name)
+				.OnCheckStateChanged(this, &STimingProfilerToolbar::ToggleModule_OnCheckStateChanged, Module.Name)
+				.Content()
+				[
+					SNew(STextBlock)
+					.Text(Label)
+					.ToolTip(SNew(SToolTip).Text(ToggleModuleToolTip))
+				]
+			);
+		}
+	}
+	ToolbarBuilder.EndSection();
 }
 
 ////////////////////////////////////////////////////////////////////////////////////////////////////
diff --git a/Engine/Source/Developer/TraceInsights/Private/Insights/Widgets/STimingProfilerToolbar.h b/Engine/Source/Developer/TraceInsights/Private/Insights/Widgets/STimingProfilerToolbar.h
--- a/Engine/Source/Developer/TraceInsights/Private/Insights/Widgets/STimingProfilerToolbar.h
+++ b/Engine/Source/Developer/TraceInsights/Private/Insights/Widgets/STimingProfilerToolbar.h
@@ -8,6 +8,7 @@
 #include "Widgets/SCompoundWidget.h"
 
 class SBorder;
+class FToolBarBuilder;
 
 /** Ribbon based toolbar used as a main menu in the Profiler window. */
 class STimingProfilerToolbar : public SCompoundWidget
@@ -32,24 +33,15 @@ public:
 
 	bool IsImplemented() const { return false; }
 	bool IsShowingStats() const { return true; }
-	bool IsShowingMemory() const { return false; }
-
-protected:
-	/** Called when the list of sessions has changed. @see FProfilerManager.FOnSessionsUpdated */
-	void ProfilerManager_SessionsUpdated();
 
+private:
 	/** Create the UI commands for the toolbar */
 	void CreateCommands();
 
-	/** Shows the stats profiler */
-	void ShowStats();
-
-	/** Shows the memory profiler */
-	void ShowMemory();
-
-	/** Shows the FPSChart view */
-	void ShowFPSChart();
-
-protected:
-	TSharedPtr<SBorder> Border;
+	bool ToggleModule_CanExecute(FName ModuleName) const;
+	void ToggleModule_Execute(FName ModuleName);
+	bool ToggleModule_IsChecked(FName ModuleName) const;
+	ECheckBoxState ToggleModule_IsChecked2(FName ModuleName) const;
+	void ToggleModule_OnCheckStateChanged(ECheckBoxState NewRadioState, FName ModuleName);
+	void FillModulesToolbar(FToolBarBuilder& ToolbarBuilder);
 };
