[Insights]
  - Start Page: Added "File Size" info to sessions list and fixed sorting (by timestamp). Also added more info in tooltip of a session.
  - Timing view: Added vertical marker (gray line) for session start time (0) and end time (in Time Ruler, Graph and Timing Events tracks).
  - Frames track: Fixed auto zoom not zooming enough if session is too short.

diff --git a/Engine/Source/Developer/TraceInsights/Private/Insights/ViewModels/FrameTrackHelper.cpp b/Engine/Source/Developer/TraceInsights/Private/Insights/ViewModels/FrameTrackHelper.cpp
--- a/Engine/Source/Developer/TraceInsights/Private/Insights/ViewModels/FrameTrackHelper.cpp
+++ b/Engine/Source/Developer/TraceInsights/Private/Insights/ViewModels/FrameTrackHelper.cpp
@@ -93,8 +93,8 @@ void FFrameTrackDrawHelper::DrawBackground() const
 	const FLinearColor ValidAreaColor(0.07f, 0.07f, 0.07f, 1.0f);
 	const FLinearColor InvalidAreaColor(0.1f, 0.07f, 0.07f, 1.0f);
 
-	float X0 = Viewport.MinX - Viewport.PosX;
-	float X1 = Viewport.MaxX - Viewport.PosX;
+	const float X0 = Viewport.MinX - Viewport.PosX;
+	const float X1 = Viewport.MaxX - Viewport.PosX;
 	const float W = FMath::CeilToFloat(Viewport.Width);
 	const float H = FMath::CeilToFloat(Viewport.Height);
 
@@ -117,13 +117,13 @@ void FFrameTrackDrawHelper::DrawBackground() const
 			DrawContext.DrawBox(X1, 0.0f, W - X1, H, AreaBrush, InvalidAreaColor);
 		}
 
-		X0 = FMath::Max(X0, 0.0f);
-		X1 = FMath::Min(X1, W);
+		const float ValidX0 = FMath::Max(X0, 0.0f);
+		const float ValidX1 = FMath::Min(X1, W);
 
-		if (X1 > X0)
+		if (ValidX1 > ValidX0)
 		{
 			// Draw valid area.
-			DrawContext.DrawBox(X0, 0.0f, X1 - X0, H, AreaBrush, ValidAreaColor);
+			DrawContext.DrawBox(ValidX0, 0.0f, ValidX1 - ValidX0, H, AreaBrush, ValidAreaColor);
 		}
 	}
 
diff --git a/Engine/Source/Developer/TraceInsights/Private/Insights/ViewModels/FrameTrackViewport.h b/Engine/Source/Developer/TraceInsights/Private/Insights/ViewModels/FrameTrackViewport.h
--- a/Engine/Source/Developer/TraceInsights/Private/Insights/ViewModels/FrameTrackViewport.h
+++ b/Engine/Source/Developer/TraceInsights/Private/Insights/ViewModels/FrameTrackViewport.h
@@ -7,7 +7,7 @@
 class FFrameTrackViewport
 {
 private:
-	const float SLATE_UNITS_TOLERANCE = 0.1f;
+	static constexpr float SLATE_UNITS_TOLERANCE = 0.1f;
 
 public:
 	FFrameTrackViewport()
diff --git a/Engine/Source/Developer/TraceInsights/Private/Insights/ViewModels/GraphTrack.cpp b/Engine/Source/Developer/TraceInsights/Private/Insights/ViewModels/GraphTrack.cpp
--- a/Engine/Source/Developer/TraceInsights/Private/Insights/ViewModels/GraphTrack.cpp
+++ b/Engine/Source/Developer/TraceInsights/Private/Insights/ViewModels/GraphTrack.cpp
@@ -106,10 +106,67 @@ void FGraphTrack::UpdateStats()
 
 ////////////////////////////////////////////////////////////////////////////////////////////////////
 
-void FGraphTrack::Draw(FDrawContext& DrawContext, const FTimingTrackViewport& Viewport) const
+void FGraphTrack::DrawBackground(FDrawContext& DrawContext, const FTimingTrackViewport& Viewport) const
 {
-	DrawContext.DrawBox(0.0f, GetPosY(), Viewport.Width, GetHeight(), WhiteBrush, FLinearColor(0.1f, 0.1f, 0.1f, 1.0f));
+	const FLinearColor ValidAreaColor(0.07f, 0.07f, 0.07f, 1.0f);
+	const FLinearColor InvalidAreaColor(0.1f, 0.07f, 0.07f, 1.0f);
+	const FLinearColor StartEndMarkerColor(0.05f, 0.05f, 0.05f, 1.0f);
+
+	const float X0 = Viewport.TimeToSlateUnitsRounded(0.0);
+	const float X1 = Viewport.TimeToSlateUnitsRounded(Viewport.MaxValidTime);
+	const float W = FMath::CeilToFloat(Viewport.Width);
+
+	const float Y = GetPosY();
+	const float H = GetHeight();
+
+	if (X0 >= W || X1 <= 0.0f)
+	{
+		// Draw invalid area (entire view).
+		DrawContext.DrawBox(0.0f, Y, W, H, WhiteBrush, InvalidAreaColor);
+	}
+	else // X0 < W && X1 > 0
+	{
+		if (X0 > 0.0f)
+		{
+			// Draw invalid area (left).
+			DrawContext.DrawBox(0.0f, Y, X0, H, WhiteBrush, InvalidAreaColor);
+		}
+
+		if (X1 < W)
+		{
+			// Draw invalid area (right).
+			DrawContext.DrawBox(X1 + 1.0f, Y, W - X1 - 1.0f, H, WhiteBrush, InvalidAreaColor);
+
+			// Draw the end time marker.
+			DrawContext.DrawBox(X1, Y, 1.0f, H, WhiteBrush, StartEndMarkerColor);
+		}
+
+		float ValidX0 = FMath::Max(X0, 0.0f);
+		float ValidX1 = FMath::Min(X1, W);
+
+		if (X0 >= 0.0f)
+		{
+			// Draw the start time marker.
+			DrawContext.DrawBox(X0, Y, 1.0f, H, WhiteBrush, StartEndMarkerColor);
+
+			ValidX0 += 1.0f; // to not overlap the start time marker
+		}
+
+		if (ValidX1 > ValidX0)
+		{
+			// Draw valid area.
+			DrawContext.DrawBox(ValidX0, Y, ValidX1 - ValidX0, H, WhiteBrush, ValidAreaColor);
+		}
+	}
+
 	DrawContext.LayerId++;
+}
+
+////////////////////////////////////////////////////////////////////////////////////////////////////
+
+void FGraphTrack::Draw(FDrawContext& DrawContext, const FTimingTrackViewport& Viewport) const
+{
+	DrawBackground(DrawContext, Viewport);
 
 	//TODO: Set clipping to (0, GetPosY(), Viewport.Width, GetHeight())!
 
diff --git a/Engine/Source/Developer/TraceInsights/Private/Insights/ViewModels/GraphTrack.h b/Engine/Source/Developer/TraceInsights/Private/Insights/ViewModels/GraphTrack.h
--- a/Engine/Source/Developer/TraceInsights/Private/Insights/ViewModels/GraphTrack.h
+++ b/Engine/Source/Developer/TraceInsights/Private/Insights/ViewModels/GraphTrack.h
@@ -84,6 +84,7 @@ public:
 protected:
 	void UpdateStats();
 
+	void DrawBackground(FDrawContext& DrawContext, const FTimingTrackViewport& Viewport) const;
 	void DrawSeries(const FGraphTrackSeries& Series, FDrawContext& DrawContext, const FTimingTrackViewport& Viewport) const;
 
 	float GetYForValue(double Value) const
diff --git a/Engine/Source/Developer/TraceInsights/Private/Insights/ViewModels/TimeRulerTrack.cpp b/Engine/Source/Developer/TraceInsights/Private/Insights/ViewModels/TimeRulerTrack.cpp
--- a/Engine/Source/Developer/TraceInsights/Private/Insights/ViewModels/TimeRulerTrack.cpp
+++ b/Engine/Source/Developer/TraceInsights/Private/Insights/ViewModels/TimeRulerTrack.cpp
@@ -62,40 +62,54 @@ void FTimeRulerTrack::UpdateHoveredState(float MouseX, float MouseY, const FTimi
 
 void FTimeRulerTrack::DrawBackground(FDrawContext& DrawContext, const FTimingTrackViewport& Viewport) const
 {
+	const FLinearColor ValidAreaColor(0.07f, 0.07f, 0.07f, 1.0f);
+	const FLinearColor InvalidAreaColor(0.1f, 0.07f, 0.07f, 1.0f);
+	const FLinearColor StartEndMarkerColor(0.05f, 0.05f, 0.05f, 1.0f);
+
 	const float X0 = Viewport.TimeToSlateUnitsRounded(0.0);
 	const float X1 = Viewport.TimeToSlateUnitsRounded(Viewport.MaxValidTime);
 	const float W = FMath::CeilToFloat(Viewport.Width);
-	const float H = GetHeight();
 
-	const FLinearColor InvalidAreaColor(0.08f, 0.07f, 0.07f, 1.0f);
-	const FLinearColor ValidAreaColor(0.09f, 0.09f, 0.09f, 1.0f);
+	const float Y = GetPosY();
+	const float H = GetHeight();
 
 	if (X0 >= W || X1 <= 0.0f)
 	{
 		// Draw invalid area (entire view).
-		DrawContext.DrawBox(0.0f, 0.0f, W, H, WhiteBrush, InvalidAreaColor);
+		DrawContext.DrawBox(0.0f, Y, W, H, WhiteBrush, InvalidAreaColor);
 	}
 	else // X0 < W && X1 > 0
 	{
 		if (X0 > 0.0f)
 		{
 			// Draw invalid area (left).
-			DrawContext.DrawBox(0.0f, 0.0f, X0, H, WhiteBrush, InvalidAreaColor);
+			DrawContext.DrawBox(0.0f, Y, X0, H, WhiteBrush, InvalidAreaColor);
 		}
 
 		if (X1 < W)
 		{
 			// Draw invalid area (right).
-			DrawContext.DrawBox(X1, 0.0f, W - X1, H, WhiteBrush, InvalidAreaColor);
+			DrawContext.DrawBox(X1 + 1.0f, Y, W - X1 - 1.0f, H, WhiteBrush, InvalidAreaColor);
+
+			// Draw the end time marker.
+			DrawContext.DrawBox(X1, Y, 1.0f, H, WhiteBrush, StartEndMarkerColor);
 		}
 
 		float ValidX0 = FMath::Max(X0, 0.0f);
 		float ValidX1 = FMath::Min(X1, W);
 
+		if (X0 >= 0.0f)
+		{
+			// Draw the start time marker.
+			DrawContext.DrawBox(X0, Y, 1.0f, H, WhiteBrush, StartEndMarkerColor);
+
+			ValidX0 += 1.0f; // to not overlap the start time marker
+		}
+
 		if (ValidX1 > ValidX0)
 		{
 			// Draw valid area.
-			DrawContext.DrawBox(ValidX0, 0.0f, ValidX1 - ValidX0, H, WhiteBrush, ValidAreaColor);
+			DrawContext.DrawBox(ValidX0, Y, ValidX1 - ValidX0, H, WhiteBrush, ValidAreaColor);
 		}
 	}
 
diff --git a/Engine/Source/Developer/TraceInsights/Private/Insights/ViewModels/TimingViewDrawHelper.cpp b/Engine/Source/Developer/TraceInsights/Private/Insights/ViewModels/TimingViewDrawHelper.cpp
--- a/Engine/Source/Developer/TraceInsights/Private/Insights/ViewModels/TimingViewDrawHelper.cpp
+++ b/Engine/Source/Developer/TraceInsights/Private/Insights/ViewModels/TimingViewDrawHelper.cpp
@@ -47,9 +47,13 @@ FTimingViewDrawHelper::~FTimingViewDrawHelper()
 
 void FTimingViewDrawHelper::DrawBackground() const
 {
+	const FLinearColor StartEndMarkerColor(0.05f, 0.05f, 0.05f, 1.0f);
+
 	const float X0 = Viewport.TimeToSlateUnitsRounded(0.0);
 	const float X1 = Viewport.TimeToSlateUnitsRounded(Viewport.MaxValidTime);
 	const float W = FMath::CeilToFloat(Viewport.Width);
+
+	const float Y = 0.0f;
 	const float H = FMath::CeilToFloat(Viewport.Height);
 
 	if (X0 >= W || X1 <= 0.0f)
@@ -58,29 +62,40 @@ void FTimingViewDrawHelper::DrawBackground() const
 		ValidX1 = W;
 
 		// Draw invalid area (entire view).
-		DrawContext.DrawBox(0.0f, 0.0f, W, H, BackgroundAreaBrush, InvalidAreaColor);
+		DrawContext.DrawBox(0.0f, Y, W, H, BackgroundAreaBrush, InvalidAreaColor);
 	}
 	else // X0 < W && X1 > 0
 	{
 		if (X0 > 0.0f)
 		{
 			// Draw invalid area (left).
-			DrawContext.DrawBox(0.0f, 0.0f, X0, H, BackgroundAreaBrush, InvalidAreaColor);
+			DrawContext.DrawBox(0.0f, Y, X0, H, BackgroundAreaBrush, InvalidAreaColor);
 		}
 
 		if (X1 < W)
 		{
 			// Draw invalid area (right).
-			DrawContext.DrawBox(X1, 0.0f, W - X1, H, BackgroundAreaBrush, InvalidAreaColor);
+			DrawContext.DrawBox(X1 + 1.0f, Y, W - X1 - 1.0f, H, BackgroundAreaBrush, InvalidAreaColor);
+
+			// Draw the end time marker.
+			DrawContext.DrawBox(X1, Y, 1.0f, H, BackgroundAreaBrush, StartEndMarkerColor);
 		}
 
 		ValidX0 = FMath::Max(X0, 0.0f);
 		ValidX1 = FMath::Min(X1, W);
 
+		if (X0 >= 0.0f)
+		{
+			// Draw the start time marker.
+			DrawContext.DrawBox(X0, Y, 1.0f, H, BackgroundAreaBrush, StartEndMarkerColor);
+
+			ValidX0 += 1.0f; // to not overlap the start time marker
+		}
+
 		if (ValidX1 > ValidX0)
 		{
 			// Draw valid area.
-			DrawContext.DrawBox(ValidX0, 0.0f, ValidX1 - ValidX0, H, BackgroundAreaBrush, ValidAreaColor);
+			DrawContext.DrawBox(ValidX0, Y, ValidX1 - ValidX0, H, BackgroundAreaBrush, ValidAreaColor);
 		}
 	}
 
@@ -114,7 +129,9 @@ bool FTimingViewDrawHelper::BeginTimeline(FTimingEventsTrack& Track)
 	Track.SetPosY(TimelineTopY + Viewport.ScrollPosY);
 
 	if (Track.GetHeight() < Layout.MinTimelineH)
+	{
 		Track.SetHeight(Layout.MinTimelineH);
+	}
 
 	if (TimelineTopY + Track.GetHeight() < -1.0f)
 	{
diff --git a/Engine/Source/Developer/TraceInsights/Private/Insights/Widgets/SFrameTrack.cpp b/Engine/Source/Developer/TraceInsights/Private/Insights/Widgets/SFrameTrack.cpp
--- a/Engine/Source/Developer/TraceInsights/Private/Insights/Widgets/SFrameTrack.cpp
+++ b/Engine/Source/Developer/TraceInsights/Private/Insights/Widgets/SFrameTrack.cpp
@@ -201,9 +201,7 @@ void SFrameTrack::Tick(const FGeometry& AllottedGeometry, const double InCurrent
 					bIsStateDirty = true;
 
 					// Auto zoom out.
-					float SampleW = Viewport.GetSampleWidth();
-					int32 NumSamples = FMath::CeilToInt(Viewport.Width / SampleW);
-					if (NumFrames > NumSamples * Viewport.GetNumFramesPerSample())
+					while (Viewport.MaxX > Viewport.Width)
 					{
 						ZoomHorizontally(-0.1f, 0.0f);
 						Viewport.ScrollAtPosX(0.0f);
@@ -764,17 +762,25 @@ void SFrameTrack::ZoomHorizontally(const float Delta, const float X)
 	{
 		// frames per Slate unit
 		if (Delta > 0)
+		{
 			ScaleX = 1.0f / FMath::FloorToFloat(1.0f / ScaleX);
+		}
 		else
+		{
 			ScaleX = 1.0f / FMath::CeilToFloat(1.0f / ScaleX);
+		}
 	}
 	else
 	{
 		// Slate units per frame
 		if (Delta > 0)
+		{
 			ScaleX = FMath::CeilToFloat(ScaleX);
+		}
 		else
+		{
 			ScaleX = FMath::FloorToFloat(ScaleX);
+		}
 	}
 
 	//UE_LOG(TimingProfiler, Log, TEXT("%.2f, %.2f, %.2f"), Delta, Viewport.ScaleX, ScaleX);
@@ -849,52 +855,6 @@ void SFrameTrack::UpdateHorizontalScrollBar()
 	HorizontalScrollBar->SetState(OffsetFraction, ThumbSizeFraction);
 }
 
-////////////////////////////////////////////////////////////////////////////////////////////////////
-/*
-void SFrameTrack::AddThreadTime(int32 InFrameIndex, const TMap<uint32, float>& InThreadMS, const TSharedRef<FProfilerStatMetaData>& InStatMetaData)
-{
-	FFrameThreadTimes FrameThreadTimes;
-	FrameThreadTimes.FrameNumber = InFrameIndex;
-	FrameThreadTimes.ThreadTimes = InThreadMS;
-	RecentlyAddedFrames.Add(FrameThreadTimes);
-
-	if (!bIsActiveTimerRegistered)
-	{
-		bIsActiveTimerRegistered = true;
-		RegisterActiveTimer(0.f, FWidgetActiveTimerDelegate::CreateSP(this, &SFrameTrack::EnsureDataUpdateDuringPreview));
-	}
-
-	StatMetadata = InStatMetaData;
-}
-*/
-////////////////////////////////////////////////////////////////////////////////////////////////////
-/*
-EActiveTimerReturnType SFrameTrack::EnsureDataUpdateDuringPreview(double InCurrentTime, float InDeltaTime)
-{
-	//if (RecentlyAddedFrames.Num() > 0)
-	//{
-	//	bUpdateData = true;
-	//	return EActiveTimerReturnType::Continue;
-	//}
-
-	bIsActiveTimerRegistered = false;
-	return EActiveTimerReturnType::Stop;
-}
-*/
-////////////////////////////////////////////////////////////////////////////////////////////////////
-/*
-void SFrameTrack::MoveSelectionBox(int32 FrameIndex)
-{
-	const int32 SelectionBoxSize = SelectionBoxFrameEnd - SelectionBoxFrameStart;
-	const int32 SelectionBoxHalfSize = SelectionBoxSize / 2;
-	const int32 CenterFrameIndex = FMath::Clamp(FrameIndex - SelectionBoxHalfSize, 0, Frames.Num() - 1 - SelectionBoxSize);
-
-	// Inform other widgets that we have moved the selection box.
-	SelectionBoxFrameStart = CenterFrameIndex;
-	SelectionBoxFrameEnd = CenterFrameIndex + SelectionBoxSize;
-	SelectionBoxChangedEvent.Broadcast(SelectionBoxFrameStart, SelectionBoxFrameEnd);
-}
-*/
 ////////////////////////////////////////////////////////////////////////////////////////////////////
 
 #undef LOCTEXT_NAMESPACE
diff --git a/Engine/Source/Developer/TraceInsights/Private/Insights/Widgets/SStartPageWindow.cpp b/Engine/Source/Developer/TraceInsights/Private/Insights/Widgets/SStartPageWindow.cpp
--- a/Engine/Source/Developer/TraceInsights/Private/Insights/Widgets/SStartPageWindow.cpp
+++ b/Engine/Source/Developer/TraceInsights/Private/Insights/Widgets/SStartPageWindow.cpp
@@ -90,6 +90,17 @@ public:
 					.ToolTipText(this, &STraceSessionRow::GetTraceSessionTooltip)
 				];
 		}
+		else if (ColumnName == FName(TEXT("Size")))
+		{
+			return SNew(SBox)
+				.Padding(FMargin(4.0, 0.0))
+				[
+					SNew(STextBlock)
+					.Text(this, &STraceSessionRow::GetTraceSessionSize)
+					.ColorAndOpacity(this, &STraceSessionRow::GetColorBySize)
+					.ToolTipText(this, &STraceSessionRow::GetTraceSessionTooltip)
+				];
+		}
 		else if (ColumnName == FName(TEXT("Status")))
 		{
 			return SNew(SBox)
@@ -132,6 +143,67 @@ public:
 		}
 	}
 
+	FText GetTraceSessionTimeStamp() const
+	{
+		TSharedPtr<FTraceSession> TraceSessionPin = WeakTraceSession.Pin();
+		if (TraceSessionPin.IsValid())
+		{
+			return FText::AsDate(TraceSessionPin->TimeStamp);
+		}
+		else
+		{
+			return FText();
+		}
+	}
+
+	FText GetTraceSessionSize() const
+	{
+		TSharedPtr<FTraceSession> TraceSessionPin = WeakTraceSession.Pin();
+		if (TraceSessionPin.IsValid())
+		{
+			return FText::AsMemory(TraceSessionPin->Size);
+		}
+		else
+		{
+			return FText();
+		}
+	}
+
+	FSlateColor GetColorBySize() const
+	{
+		TSharedPtr<FTraceSession> TraceSessionPin = WeakTraceSession.Pin();
+		if (TraceSessionPin.IsValid())
+		{
+			TSharedRef<ITypedTableView<TSharedPtr<FTraceSession>>> OwnerWidget = OwnerTablePtr.Pin().ToSharedRef();
+			const TSharedPtr<FTraceSession>* MyItem = OwnerWidget->Private_ItemFromWidget(this);
+			const bool IsSelected = OwnerWidget->Private_IsItemSelected(*MyItem);
+
+			if (IsSelected)
+			{
+				return FSlateColor(FLinearColor(0.0f, 0.0f, 0.0f, 1.0f));
+			}
+			else if (TraceSessionPin->Size < 1024ULL * 1024ULL)
+			{
+				// < 1 MiB
+				return FSlateColor(FLinearColor(0.5f, 0.5f, 0.5f, 1.0f));
+			}
+			else if (TraceSessionPin->Size < 1024ULL * 1024ULL * 1024ULL)
+			{
+				// [1 MiB  .. 1 GiB)
+				return FSlateColor(FLinearColor(1.0f, 1.0f, 1.0f, 1.0f));
+			}
+			else
+			{
+				// > 1 GiB
+				return FSlateColor(FLinearColor(1.0f, 0.5f, 0.5f, 1.0f));
+			}
+		}
+		else
+		{
+			return FSlateColor(FLinearColor(0.0f, 0.0f, 0.0f, 1.0f));
+		}
+	}
+
 	FText GetTraceSessionstatus() const
 	{
 		TSharedPtr<FTraceSession> TraceSessionPin = WeakTraceSession.Pin();
@@ -150,14 +222,13 @@ public:
 		TSharedPtr<FTraceSession> TraceSessionPin = WeakTraceSession.Pin();
 		if (TraceSessionPin.IsValid())
 		{
-			if (TraceSessionPin->bIsLive)
-			{
-				return FText::Format(LOCTEXT("LiveTraceSessionTooltipFormat", "{0}\n{1}\nLIVE"), TraceSessionPin->Name, TraceSessionPin->Uri);
-			}
-			else
-			{
-				return FText::Format(LOCTEXT("TraceSessionTooltipFormat", "{0}\n{1}"), TraceSessionPin->Name, TraceSessionPin->Uri);
-			}
+			return FText::Format(LOCTEXT("TraceSessionTooltipFormat", "{0}\n{1}\nTimestamp: {2}\nFile Size: {3} bytes\nStatus: {4}"),
+				TraceSessionPin->Name,
+				TraceSessionPin->Uri,
+				FText::AsDateTime(TraceSessionPin->TimeStamp),
+				FText::AsNumber(TraceSessionPin->Size),
+				TraceSessionPin->bIsLive ? LOCTEXT("LiveTraceSessionStatus", "LIVE") : LOCTEXT("OfflineTraceSessionStatus", "Offline")
+			);
 		}
 		else
 		{
@@ -336,15 +407,19 @@ TSharedRef<SWidget> SStartPageWindow::ConstructSessionsPanel()
 				SNew(SHeaderRow)
 
 				+ SHeaderRow::Column(FName(TEXT("Name")))
-				.FillWidth(0.2f)
+				.FillWidth(0.8f)
 				.DefaultLabel(LOCTEXT("NameColumn", "Name"))
 
-				//+ SHeaderRow::Column(FName(TEXT("Uri")))
-				//.FillWidth(0.8f)
-				//.DefaultLabel(LOCTEXT("UriColumn", "URI"))
+				+ SHeaderRow::Column(FName(TEXT("Size")))
+				.FillWidth(0.2f)
+				.HAlignHeader(HAlign_Right)
+				.HAlignCell(HAlign_Right)
+				.DefaultLabel(LOCTEXT("SizeColumn", "File Size"))
 
 				+ SHeaderRow::Column(FName(TEXT("Status")))
 				.FixedWidth(60.0f)
+				.HAlignHeader(HAlign_Right)
+				.HAlignCell(HAlign_Right)
 				.DefaultLabel(LOCTEXT("StatusColumn", "Status"))
 			)
 		]
@@ -722,6 +797,7 @@ void SStartPageWindow::RefreshTraceSessionList()
 				NewSelectedTraceSession = TraceSessions.Last();
 			}
 		}
+		Algo::SortBy(TraceSessions, &FTraceSession::TimeStamp);
 
 		TraceSessionsListView->RebuildList();
 
@@ -1026,15 +1102,34 @@ TSharedRef<SWidget> SStartPageWindow::MakeSessionListMenu()
 		TArray<Trace::FSessionHandle> AvailableSessions;
 		SessionService->GetAvailableSessions(AvailableSessions);
 
+		struct FSessionInfoEx
+		{
+			FDateTime GetTimeStamp() const { return SessionInfo.TimeStamp; }
+
+			Trace::FSessionHandle SessionHandle;
+			Trace::FSessionInfo  SessionInfo;
+		};
+
+		TArray<TSharedPtr<FSessionInfoEx>> SortedAvailableSessions;
+		SortedAvailableSessions.Reserve(AvailableSessions.Num());
+
+		for (Trace::FSessionHandle SessionHandle : AvailableSessions)
+		{
+			FSessionInfoEx* SessionInfoExPtr = new FSessionInfoEx;
+			SessionInfoExPtr->SessionHandle = SessionHandle;
+			SessionService->GetSessionInfo(SessionHandle, SessionInfoExPtr->SessionInfo);
+			SortedAvailableSessions.Add(MakeShareable(SessionInfoExPtr));
+		}
+
+		Algo::SortBy(SortedAvailableSessions, &FSessionInfoEx::GetTimeStamp);
+
 		int32 SessionCountLimit = 10; // top 10
 
 		// Iterate in reverse order as we want most recent sessions first.
-		for (int32 SessionIndex = AvailableSessions.Num() - 1; SessionIndex >= 0 && SessionCountLimit > 0; --SessionIndex, --SessionCountLimit)
+		for (int32 SessionIndex = SortedAvailableSessions.Num() - 1; SessionIndex >= 0 && SessionCountLimit > 0; --SessionIndex, --SessionCountLimit)
 		{
-			Trace::FSessionHandle SessionHandle = AvailableSessions[SessionIndex];
-
-			Trace::FSessionInfo SessionInfo;
-			SessionService->GetSessionInfo(SessionHandle, SessionInfo);
+			const Trace::FSessionHandle SessionHandle = SortedAvailableSessions[SessionIndex]->SessionHandle;
+			const Trace::FSessionInfo& SessionInfo = SortedAvailableSessions[SessionIndex]->SessionInfo;
 
 			FText Label = FText::FromString(SessionInfo.Name);
 			if (SessionInfo.bIsLive)
diff --git a/Engine/Source/Developer/TraceInsights/Private/Insights/Widgets/SStartPageWindow.h b/Engine/Source/Developer/TraceInsights/Private/Insights/Widgets/SStartPageWindow.h
--- a/Engine/Source/Developer/TraceInsights/Private/Insights/Widgets/SStartPageWindow.h
+++ b/Engine/Source/Developer/TraceInsights/Private/Insights/Widgets/SStartPageWindow.h
@@ -44,12 +44,16 @@ struct FTraceSession
 	Trace::FSessionHandle SessionHandle;
 	FText Name;
 	FText Uri;
+	FDateTime TimeStamp;
+	uint64 Size;
 	bool bIsLive;
 
 	FTraceSession(const Trace::FSessionHandle InSessionHandle, const Trace::FSessionInfo& InSessionInfo)
 		: SessionHandle(InSessionHandle)
 		, Name(FText::FromString(InSessionInfo.Name))
 		, Uri(FText::FromString(InSessionInfo.Uri))
+		, TimeStamp(InSessionInfo.TimeStamp)
+		, Size(InSessionInfo.Size)
 		, bIsLive(InSessionInfo.bIsLive)
 	{}
 };
