[Insights]
  - Log View: Fixed issue where having a filter set during a live session (or during a long anaysis) results in selected item to trigger multiple scroll requests in timing view (i.e. each time new log messages are added to log view).
  - TimingView: Fixed I/O tracks to not update when session is reset (but only when the track is made visible).

diff --git a/Engine/Source/Developer/TraceInsights/Private/Insights/Widgets/SLogView.cpp b/Engine/Source/Developer/TraceInsights/Private/Insights/Widgets/SLogView.cpp
--- a/Engine/Source/Developer/TraceInsights/Private/Insights/Widgets/SLogView.cpp
+++ b/Engine/Source/Developer/TraceInsights/Private/Insights/Widgets/SLogView.cpp
@@ -811,6 +811,7 @@ void SLogView::Tick(const FGeometry& AllottedGeometry, const double InCurrentTim
 				ListView->RebuildList();
 				if (SelectedLogMessage.IsValid())
 				{
+					// Restore selection.
 					SelectedLogMessageByLogIndex(SelectedLogMessage->GetIndex());
 				}
 				bIsDirty = false;
@@ -878,6 +879,7 @@ void SLogView::Tick(const FGeometry& AllottedGeometry, const double InCurrentTim
 			ListView->RebuildList();
 			if (SelectedLogMessage.IsValid())
 			{
+				// Restore selection.
 				SelectedLogMessageByLogIndex(SelectedLogMessage->GetIndex());
 			}
 			bIsDirty = false;
@@ -926,21 +928,24 @@ void SLogView::SelectedLogMessageByLogIndex(int32 LogIndex)
 
 void SLogView::OnSelectionChanged(TSharedPtr<FLogMessage> LogMessage, ESelectInfo::Type SelectInfo)
 {
-	TSharedPtr<STimingProfilerWindow> Window = FTimingProfilerManager::Get()->GetProfilerWindow();
-	if (Window && Window->TimingView)
+	if (SelectInfo != ESelectInfo::Direct)
 	{
-		// Single item selection.
-		if (LogMessage.IsValid())
+		TSharedPtr<STimingProfilerWindow> Window = FTimingProfilerManager::Get()->GetProfilerWindow();
+		if (Window && Window->TimingView)
 		{
-			const double Time = Cache.Get(LogMessage->GetIndex()).Time;
-
-			if (FSlateApplication::Get().GetModifierKeys().IsShiftDown())
+			// Single item selection.
+			if (LogMessage.IsValid())
 			{
-				Window->TimingView->SelectToTimeMarker(Time);
-			}
-			else
-			{
-				Window->TimingView->SetAndCenterOnTimeMarker(Time);
+				const double Time = Cache.Get(LogMessage->GetIndex()).Time;
+
+				if (FSlateApplication::Get().GetModifierKeys().IsShiftDown())
+				{
+					Window->TimingView->SelectToTimeMarker(Time);
+				}
+				else
+				{
+					Window->TimingView->SetAndCenterOnTimeMarker(Time);
+				}
 			}
 		}
 	}
diff --git a/Engine/Source/Developer/TraceInsights/Private/Insights/Widgets/STimingView.cpp b/Engine/Source/Developer/TraceInsights/Private/Insights/Widgets/STimingView.cpp
--- a/Engine/Source/Developer/TraceInsights/Private/Insights/Widgets/STimingView.cpp
+++ b/Engine/Source/Developer/TraceInsights/Private/Insights/Widgets/STimingView.cpp
@@ -227,7 +227,7 @@ void STimingView::Reset()
 	IoOverviewTrack = nullptr;
 	IoActivityTrack = nullptr;
 
-	bForceIoEventsUpdate = true;
+	bForceIoEventsUpdate = false;
 	bMergeIoLanes = true;
 	AllIoEvents.Reset();
 	//AllIoEvents.Reserve(1000000);
