[Insights]
  - Timing view: Improved algorithm to layout I/O file activity events so to not overlap. "Close" events are now ignored when computing layout, but are will still displayed on same depth as Open/Read/Write corresponding events.
  - Timing view: Added tooltip for the mouse hovered event, in I/O tracks.
  - Timing view: Changed "O" key shortcut to show/hide background events for file activities, in I/O tracks.

diff --git a/Engine/Source/Developer/TraceInsights/Private/Insights/ViewModels/TimingEvent.h b/Engine/Source/Developer/TraceInsights/Private/Insights/ViewModels/TimingEvent.h
--- a/Engine/Source/Developer/TraceInsights/Private/Insights/ViewModels/TimingEvent.h
+++ b/Engine/Source/Developer/TraceInsights/Private/Insights/ViewModels/TimingEvent.h
@@ -5,38 +5,49 @@
 #include "CoreMinimal.h"
 #include "TraceServices/Model/LoadTimeProfiler.h"
 
-
-// Insights
-#include "Insights/ViewModels/TimerNode.h"
-
 class FTimingEventsTrack;
 
 struct FTimingEvent
 {
+	static const uint64 InvalidTypeId = -1;
+
 	const FTimingEventsTrack* Track;
 	uint64 TypeId;
 	uint32 Depth;
 	double StartTime;
 	double EndTime;
 	double ExclusiveTime;
+
+	// For Asset Loading events...
 	Trace::FLoadTimeProfilerCpuEvent LoadingInfo;
 
+	// For I/O events...
+	uint64 Offset;
+	uint64 Size;
+	const TCHAR* Path;
+
 	void Reset()
 	{
 		Track = nullptr;
-		TypeId = FTimerNode::InvalidId;
+		TypeId = InvalidTypeId;
 		Depth = 0;
 		StartTime = 0.0;
 		EndTime = 0.0;
 		ExclusiveTime = 0.0;
+
 		LoadingInfo.Package = nullptr;
 		LoadingInfo.Export = nullptr;
 		LoadingInfo.PackageEventType = LoadTimeProfilerPackageEventType_None;
 		LoadingInfo.ExportEventType = LoadTimeProfilerObjectEventType_None;
+
+		Offset = 0;
+		Size = 0;
+		Path = nullptr;
 	}
 
 	bool IsValidTrack() const { return Track != nullptr; }
-	bool IsValid() const { return Track != nullptr && TypeId != FTimerNode::InvalidId; }
+
+	bool IsValid() const { return Track != nullptr && TypeId != InvalidTypeId; }
 
 	double Duration() const { return EndTime - StartTime; }
 
@@ -47,7 +58,6 @@ struct FTimingEvent
 			&& Depth == Other.Depth
 			&& StartTime == Other.StartTime
 			&& EndTime == Other.EndTime;
-			//&& ExclusiveTime == Other.ExclusiveTime;
 	}
 
 	static bool AreEquals(const FTimingEvent& A, const FTimingEvent& B)
diff --git a/Engine/Source/Developer/TraceInsights/Private/Insights/Widgets/STimingView.cpp b/Engine/Source/Developer/TraceInsights/Private/Insights/Widgets/STimingView.cpp
--- a/Engine/Source/Developer/TraceInsights/Private/Insights/Widgets/STimingView.cpp
+++ b/Engine/Source/Developer/TraceInsights/Private/Insights/Widgets/STimingView.cpp
@@ -55,6 +55,9 @@ uint64 FBaseTimingTrack::IdGenerator = (1ULL << 63);
 const TCHAR* GetName(ELoadTimeProfilerPackageEventType Type);
 const TCHAR* GetName(ELoadTimeProfilerObjectEventType Type);
 
+const TCHAR* GetFileActivityTypeName(Trace::EFileActivityType Type);
+uint32 GetFileActivityTypeColor(Trace::EFileActivityType Type);
+
 ////////////////////////////////////////////////////////////////////////////////////////////////////
 
 STimingView::STimingView()
@@ -229,8 +232,10 @@ void STimingView::Reset()
 
 	bForceIoEventsUpdate = false;
 	bMergeIoLanes = true;
+	bShowFileActivityBackgroundEvents = false;
+	FileActivities.Reset();
+	FileActivityMap.Reset();
 	AllIoEvents.Reset();
-	//AllIoEvents.Reserve(1000000);
 
 	//////////////////////////////////////////////////
 
@@ -508,6 +513,7 @@ void STimingView::Tick(const FGeometry& AllottedGeometry, const double InCurrent
 				{
 					if (!ThreadGroups.Contains(GroupName))
 					{
+						UE_LOG(TimingProfiler, Log, TEXT("New CPU Thread Group (%d) : \"%s\""), ThreadGroups.Num() + 1, GroupName);
 						ThreadGroups.Add(GroupName, { GroupName, bIsGroupVisible, 0, Order });
 					}
 					else
@@ -1047,6 +1053,10 @@ int32 STimingView::OnPaint(const FPaintArgs& Args, const FGeometry& AllottedGeom
 			{
 				//TODO: ...
 			}
+			else if (SelectedTimingEvent.Track->GetType() == ETimingEventsTrackType::Io)
+			{
+				//TODO: ...
+			}
 		}
 
 		// Draw info about hovered event (like a tooltip at mouse position).
@@ -1238,6 +1248,112 @@ int32 STimingView::OnPaint(const FPaintArgs& Args, const FGeometry& AllottedGeom
 					Y += LineH;
 				}
 
+				DrawContext.LayerId++;
+			}
+			else if (HoveredTimingEvent.Track->GetType() == ETimingEventsTrackType::Io)
+			{
+				const Trace::EFileActivityType ActivityType = static_cast<Trace::EFileActivityType>(HoveredTimingEvent.TypeId & 0x0F);
+				const bool bHasFailed = ((HoveredTimingEvent.TypeId & 0xF0) != 0);
+
+				FString Name(HoveredTimingEvent.Path);
+
+				constexpr float ValueOffsetX = 50.0f;
+				constexpr float MinValueTextWidth = 220.0f;
+
+				float W = FontMeasureService->Measure(Name, MainFont).X;
+
+				if (W < ValueOffsetX + MinValueTextWidth)
+				{
+					W = ValueOffsetX + MinValueTextWidth;
+				}
+				if (W < MinTooltipWidth)
+				{
+					W = MinTooltipWidth;
+				}
+				if (TooltipWidth != W)
+				{
+					TooltipWidth = TooltipWidth * 0.75f + W * 0.25f;
+				}
+
+				const float MaxX = FMath::Max(0.0f, Viewport.Width - TooltipWidth - 21.0f);
+				const float X = FMath::Clamp<float>(MousePosition.X + 18.0f, 0.0f, MaxX);
+				const float ValueX = X + ValueOffsetX;
+
+				constexpr float LineH = 14.0f;
+				float H = 4 * LineH;
+				if (ActivityType == Trace::FileActivityType_Read || ActivityType == Trace::FileActivityType_Write)
+				{
+					H += 2 * LineH;
+				}
+
+				const float MaxY = FMath::Max(0.0f, Viewport.Height - H - 18.0f);
+				float Y = FMath::Clamp<float>(MousePosition.Y + 18.0f, 0.0f, MaxY);
+
+				const float Alpha = 1.0f - FMath::Abs(TooltipWidth - W) / W;
+				if (TooltipAlpha < Alpha)
+				{
+					TooltipAlpha = TooltipAlpha * 0.9f + Alpha * 0.1f;
+				}
+				else
+				{
+					TooltipAlpha = Alpha;
+				}
+
+				const FLinearColor BackgroundColor(0.05f, 0.05f, 0.05f, TooltipAlpha);
+				const FLinearColor NameColor(0.9f, 0.9f, 0.5f, TooltipAlpha);
+				const FLinearColor TextColor(0.6f, 0.6f, 0.6f, TooltipAlpha);
+				const FLinearColor ValueColor(1.0f, 1.0f, 1.0f, TooltipAlpha);
+
+				DrawContext.DrawBox(X - 6.0f, Y - 3.0f, TooltipWidth + 12.0f, H + 6.0f, WhiteBrush, BackgroundColor);
+				DrawContext.LayerId++;
+
+				FString TypeStr;
+				uint32 TypeColor;
+				if (bHasFailed)
+				{
+					TypeStr = TEXT("Failed ");
+					TypeStr += GetFileActivityTypeName(ActivityType);
+					TypeColor = 0xFFFF3333;
+				}
+				else
+				{
+					TypeStr = GetFileActivityTypeName(ActivityType);
+					TypeColor = GetFileActivityTypeColor(ActivityType);
+				}
+				FLinearColor TypeLinearColor = FLinearColor(FColor(TypeColor));
+				TypeLinearColor.R *= 2.0f;
+				TypeLinearColor.G *= 2.0f;
+				TypeLinearColor.B *= 2.0f;
+				DrawContext.DrawText(X, Y, TypeStr, MainFont, TypeLinearColor);
+				Y += LineH;
+
+				DrawContext.DrawText(X, Y, Name, MainFont, NameColor);
+				Y += LineH;
+
+				DrawContext.DrawText(X, Y, TEXT("Duration:"), MainFont, TextColor);
+				FString DurationStr = TimeUtils::FormatTimeAuto(HoveredTimingEvent.Duration());
+				DrawContext.DrawText(ValueX, Y, DurationStr, MainFont, ValueColor);
+				Y += LineH;
+
+				DrawContext.DrawText(X, Y, TEXT("Depth:"), MainFont, TextColor);
+				FString DepthStr = FString::Printf(TEXT("%d"), HoveredTimingEvent.Depth);
+				DrawContext.DrawText(ValueX, Y, DepthStr, MainFont, ValueColor);
+				Y += LineH;
+
+				if (ActivityType == Trace::FileActivityType_Read || ActivityType == Trace::FileActivityType_Write)
+				{
+					DrawContext.DrawText(X, Y, TEXT("Offset:"), MainFont, TextColor);
+					FString OffsetStr = FText::AsNumber(HoveredTimingEvent.Offset).ToString();
+					DrawContext.DrawText(ValueX, Y, OffsetStr, MainFont, ValueColor);
+					Y += LineH;
+
+					DrawContext.DrawText(X, Y, TEXT("Size:"), MainFont, TextColor);
+					FString SizeStr = FText::AsNumber(HoveredTimingEvent.Size).ToString();
+					SizeStr += TEXT(" bytes");
+					DrawContext.DrawText(ValueX, Y, SizeStr, MainFont, ValueColor);
+					Y += LineH;
+				}
+
 				DrawContext.LayerId++;
 			}
 		}
@@ -1698,7 +1814,11 @@ void STimingView::DrawIoOverviewTrack(FTimingViewDrawHelper& Helper, FTimingEven
 				// Ignore "Idle" and "NotClosed" events.
 				continue;
 			}
-			if (Event.EndTime <= Helper.GetViewport().StartTime)
+
+			//const double EventEndTime = Event.EndTime; // keep duration of events
+			const double EventEndTime = Event.StartTime; // make all 0 duration events
+
+			if (EventEndTime <= Helper.GetViewport().StartTime)
 			{
 				continue;
 			}
@@ -1715,23 +1835,12 @@ void STimingView::DrawIoOverviewTrack(FTimingViewDrawHelper& Helper, FTimingEven
 			{
 				FString Name = TEXT("Failed ");
 				Name += GetFileActivityTypeName(ActivityType);
-				Name += " [";
-				Name += Event.Path;
-				Name += "]";
 				Color = 0xFFAA0000;
-				Helper.AddEvent(Event.StartTime, Event.EndTime, 0, *Name, Color);
-			}
-			else if (ActivityType == Trace::FileActivityType_Open)
-			{
-				FString Name = GetFileActivityTypeName(ActivityType);
-				Name += " [";
-				Name += Event.Path;
-				Name += "]";
-				Helper.AddEvent(Event.StartTime, Event.EndTime, 0, *Name, Color);
+				Helper.AddEvent(Event.StartTime, EventEndTime, 0, *Name, Color);
 			}
 			else
 			{
-				Helper.AddEvent(Event.StartTime, Event.EndTime, 0, GetFileActivityTypeName(ActivityType), Color);
+				Helper.AddEvent(Event.StartTime, EventEndTime, 0, GetFileActivityTypeName(ActivityType), Color);
 			}
 		}
 
@@ -1747,8 +1856,13 @@ void STimingView::UpdateIo()
 	{
 		bForceIoEventsUpdate = false;
 
+		FileActivities.Reset();
+		FileActivityMap.Reset();
 		AllIoEvents.Reset();
 
+		FStopwatch Stopwatch;
+		Stopwatch.Start();
+
 		TSharedPtr<const Trace::IAnalysisSession> Session = FInsightsManager::Get()->GetSession();
 		if (Session.IsValid() && Trace::ReadFileActivityProvider(*Session.Get()))
 		{
@@ -1762,27 +1876,143 @@ void STimingView::UpdateIo()
 				Timeline.EnumerateEvents(-std::numeric_limits<double>::infinity(), +std::numeric_limits<double>::infinity(),
 					[this, &FileInfo, &Timeline](double EventStartTime, double EventEndTime, uint32 EventDepth, const Trace::FFileActivity* FileActivity)
 				{
+					TSharedPtr<FIoFileActivity> Activity;
+
+					if (!FileActivityMap.Contains(FileInfo.Id))
+					{
+						Activity = MakeShareable(new FIoFileActivity());
+
+						Activity->Id = FileInfo.Id;
+						Activity->Path = FileInfo.Path;
+						Activity->StartTime = EventStartTime;
+						Activity->EndTime = EventEndTime;
+						Activity->EventCount = 1;
+						Activity->Depth = -1;
+
+						FileActivities.Add(Activity);
+						FileActivityMap.Add(FileInfo.Id, Activity);
+					}
+					else
+					{
+						Activity = FileActivityMap[FileInfo.Id];
+
+						if (FileActivity->ActivityType != Trace::FileActivityType_Close)
+						{
+							ensure(EventStartTime >= Activity->StartTime);
+							if (EventStartTime < Activity->StartTime)
+							{
+								Activity->StartTime = EventStartTime;
+							}
+
+							if (EventEndTime > Activity->EndTime)
+							{
+								Activity->EndTime = EventEndTime;
+							}
+						}
+
+						Activity->EventCount++;
+					}
+
 					if (bMergeIoLanes)
 					{
-						EventDepth = static_cast<uint32>(FileInfo.Id); // used by MergeLanes algorithm
+						EventDepth = 0;
 					}
 					else
 					{
 						EventDepth = FileInfo.Id % 32; // simple layout
 					}
+
 					uint32 Type = ((uint32)FileActivity->ActivityType & 0x0F) | (FileActivity->Failed ? 0x80 : 0);
-					AllIoEvents.Add(FIoTimingEvent{ EventStartTime, EventEndTime, EventDepth, Type, FileInfo.Path });
+
+					AllIoEvents.Add(FIoTimingEvent{ EventStartTime, EventEndTime, EventDepth, Type, FileActivity->Offset, FileActivity->Size, Activity });
 				});
 
 				return true;
 			});
 		}
 
+		Stopwatch.Stop();
+		UE_LOG(TimingProfiler, Log, TEXT("[IO] Enumerated %s events (%s file activities) in %s."),
+			*FText::AsNumber(AllIoEvents.Num()).ToString(),
+			*FText::AsNumber(FileActivities.Num()).ToString(),
+			*TimeUtils::FormatTimeAuto(Stopwatch.GetAccumulatedTime()));
+		Stopwatch.Restart();
+
+		// Sort cached IO file activities by Start Time.
+		FileActivities.Sort([](const TSharedPtr<FIoFileActivity>& A, const TSharedPtr<FIoFileActivity>& B) { return A->StartTime < B->StartTime; });
+
 		// Sort cached IO events by Start Time.
 		AllIoEvents.Sort([](const FIoTimingEvent& A, const FIoTimingEvent& B) { return A.StartTime < B.StartTime; });
 
+		Stopwatch.Stop();
+		UE_LOG(TimingProfiler, Log, TEXT("[IO] Sorted file activities and events in %s."), *TimeUtils::FormatTimeAuto(Stopwatch.GetAccumulatedTime()));
+
 		if (bMergeIoLanes)
 		{
+			//////////////////////////////////////////////////
+			// Compute depth for file activities (avoids overlaps).
+
+			Stopwatch.Restart();
+
+			TArray<TSharedPtr<FIoFileActivity>> ActivityLanes;
+
+			for (TSharedPtr<FIoFileActivity> FileActivity : FileActivities)
+			{
+				// Find lane (avoiding overlaps with other file activities).
+				for (int32 LaneIndex = 0; LaneIndex < ActivityLanes.Num(); ++LaneIndex)
+				{
+					TSharedPtr<FIoFileActivity> Lane = ActivityLanes[LaneIndex];
+
+					if (FileActivity->StartTime >= Lane->EndTime)
+					{
+						FileActivity->Depth = LaneIndex;
+						ActivityLanes[LaneIndex] = FileActivity;
+						break;
+					}
+				}
+
+				if (FileActivity->Depth < 0)
+				{
+					const int32 MaxLanes = 10000;
+					if (ActivityLanes.Num() < MaxLanes)
+					{
+						// Add new lane.
+						FileActivity->Depth = ActivityLanes.Num();
+						ActivityLanes.Add(FileActivity);
+					}
+					else
+					{
+						int32 LaneIndex = ActivityLanes.Num() - 1;
+						FileActivity->Depth = LaneIndex;
+						TSharedPtr<FIoFileActivity> Lane = ActivityLanes[LaneIndex];
+						if (FileActivity->EndTime > Lane->EndTime)
+						{
+							ActivityLanes[LaneIndex] = FileActivity;
+						}
+					}
+				}
+			}
+
+			Stopwatch.Stop();
+			UE_LOG(TimingProfiler, Log, TEXT("[IO] Computed layout for file activities in %s."), *TimeUtils::FormatTimeAuto(Stopwatch.GetAccumulatedTime()));
+
+			//////////////////////////////////////////////////
+
+			Stopwatch.Restart();
+
+			for (FIoTimingEvent& Event : AllIoEvents)
+			{
+				Event.Depth = static_cast<uint32>(Event.FileActivity->Depth);
+			}
+
+			Stopwatch.Stop();
+			UE_LOG(TimingProfiler, Log, TEXT("[IO] Updated depth for events in %s."), *TimeUtils::FormatTimeAuto(Stopwatch.GetAccumulatedTime()));
+
+			//////////////////////////////////////////////////
+			/*
+
+			Stopwatch.Restart();
+
 			struct FIoLane
 			{
 				uint32 FileId;
@@ -1790,6 +2020,7 @@ void STimingView::UpdateIo()
 				double LastEndTime;
 				double EndTime;
 			};
+
 			TArray<FIoLane> Lanes;
 
 			TArray<FIoTimingEvent> IoEventsToAdd;
@@ -1885,6 +2116,10 @@ void STimingView::UpdateIo()
 				Event.Depth = Depth;
 			}
 
+			Stopwatch.Stop();
+			UE_LOG(TimingProfiler, Log, TEXT("[IO] Merge 1/3 in %s."), *TimeUtils::FormatTimeAuto(Stopwatch.GetAccumulatedTime()));
+			Stopwatch.Restart();
+
 			for (int32 LaneIndex = 0; LaneIndex < Lanes.Num(); ++LaneIndex)
 			{
 				FIoLane& Lane = Lanes[LaneIndex];
@@ -1899,6 +2134,10 @@ void STimingView::UpdateIo()
 				AllIoEvents.Add(Event);
 			}
 
+			Stopwatch.Stop();
+			UE_LOG(TimingProfiler, Log, TEXT("[IO] Merge 2/3 in %s."), *TimeUtils::FormatTimeAuto(Stopwatch.GetAccumulatedTime()));
+			Stopwatch.Restart();
+
 			// Sort cached IO events one more time, also by Start Time.
 			AllIoEvents.Sort([](const FIoTimingEvent& A, const FIoTimingEvent& B) { return A.StartTime < B.StartTime; });
 
@@ -1907,6 +2146,10 @@ void STimingView::UpdateIo()
 			//{
 			//	return A.Depth == B.Depth ? A.StartTime < B.StartTime : A.Depth < B.Depth;
 			//});
+
+			Stopwatch.Stop();
+			UE_LOG(TimingProfiler, Log, TEXT("[IO] Merge 3/3 (sort) in %s."), *TimeUtils::FormatTimeAuto(Stopwatch.GetAccumulatedTime()));
+			*/
 		}
 	}
 }
@@ -1915,9 +2158,30 @@ void STimingView::UpdateIo()
 
 void STimingView::DrawIoActivityTrack(FTimingViewDrawHelper& Helper, FTimingEventsTrack& Track) const
 {
+	// Draw IO track using cached events.
 	if (Helper.BeginTimeline(Track))
 	{
-		// Draw IO track using cached events.
+		// Draw IO file activity background events.
+		if (bShowFileActivityBackgroundEvents)
+		{
+			for (const TSharedPtr<FIoFileActivity> Activity : FileActivities)
+			{
+				if (Activity->EndTime <= Helper.GetViewport().StartTime)
+				{
+					continue;
+				}
+				if (Activity->StartTime >= Helper.GetViewport().EndTime)
+				{
+					break;
+				}
+
+				ensure(Activity->Depth <= 10000);
+
+				Helper.AddEvent(Activity->StartTime, Activity->EndTime, Activity->Depth, Activity->Path, 0x55333333);
+			}
+		}
+
+		// Draw IO file activity foreground events.
 		for (const FIoTimingEvent& Event : AllIoEvents)
 		{
 			if (Event.EndTime <= Helper.GetViewport().StartTime)
@@ -1929,6 +2193,8 @@ void STimingView::DrawIoActivityTrack(FTimingViewDrawHelper& Helper, FTimingEven
 				break;
 			}
 
+			ensure(Event.Depth <= 10000);
+
 			const Trace::EFileActivityType ActivityType = static_cast<Trace::EFileActivityType>(Event.Type & 0x0F);
 
 			uint32 Color = GetFileActivityTypeColor(ActivityType);
@@ -1941,20 +2207,9 @@ void STimingView::DrawIoActivityTrack(FTimingViewDrawHelper& Helper, FTimingEven
 				{
 					FString Name = TEXT("Failed ");
 					Name += GetFileActivityTypeName(ActivityType);
-					Name += " [";
-					Name += Event.Path;
-					Name += "]";
 					Color = 0xFFAA0000;
 					Helper.AddEvent(Event.StartTime, Event.EndTime, Event.Depth, *Name, Color);
 				}
-				else if (ActivityType == Trace::FileActivityType_Open)
-				{
-					FString Name = GetFileActivityTypeName(ActivityType);
-					Name += " [";
-					Name += Event.Path;
-					Name += "]";
-					Helper.AddEvent(Event.StartTime, Event.EndTime, Event.Depth, *Name, Color);
-				}
 				else
 				{
 					Helper.AddEvent(Event.StartTime, Event.EndTime, Event.Depth, GetFileActivityTypeName(ActivityType), Color);
@@ -1964,7 +2219,7 @@ void STimingView::DrawIoActivityTrack(FTimingViewDrawHelper& Helper, FTimingEven
 			{
 				FString Name = GetFileActivityTypeName(ActivityType);
 				Name += " [";
-				Name += Event.Path;
+				Name += Event.FileActivity->Path;
 				Name += "]";
 				Helper.AddEvent(Event.StartTime, Event.EndTime, Event.Depth, *Name, Color);
 			}
@@ -2742,8 +2997,9 @@ FReply STimingView::OnKeyDown(const FGeometry& MyGeometry, const FKeyEvent& InKe
 	}
 	else if (InKeyEvent.GetKey() == EKeys::O)  // debug: toggles IO merge lanes algorithm
 	{
-		bMergeIoLanes = !bMergeIoLanes;
-		bForceIoEventsUpdate = true;
+		//bMergeIoLanes = !bMergeIoLanes;
+		//bForceIoEventsUpdate = true;
+		bShowFileActivityBackgroundEvents = !bShowFileActivityBackgroundEvents;
 		return FReply::Handled();
 	}
 	else if (InKeyEvent.GetKey() == EKeys::One)
@@ -3196,7 +3452,7 @@ void STimingView::UpdateHoveredTimingEvent(float MX, float MY)
 				constexpr bool bStopAtFirstMatch = true; // get first one matching
 				constexpr bool bSearchForLargestEvent = false;
 				SearchTimingEvent(StartTime, EndTime,
-					[Depth](double, double, uint32 EventDepth, uint32)
+					[Depth](double, double, uint32 EventDepth)
 					{
 						return EventDepth == Depth;
 					},
@@ -3212,7 +3468,7 @@ void STimingView::UpdateHoveredTimingEvent(float MX, float MY)
 
 bool STimingView::SearchTimingEvent(const double InStartTime,
 									const double InEndTime,
-									TFunctionRef<bool(double, double, uint32, uint32)> InPredicate,
+									TFunctionRef<bool(double, double, uint32)> InPredicate,
 									FTimingEvent& InOutTimingEvent,
 									bool bInStopAtFirstMatch,
 									bool bInSearchForLargestEvent) const
@@ -3221,7 +3477,7 @@ bool STimingView::SearchTimingEvent(const double InStartTime,
 	{
 		const double StartTime;
 		const double EndTime;
-		TFunctionRef<bool(double, double, uint32, uint32)> Predicate;
+		TFunctionRef<bool(double, double, uint32)> Predicate;
 		FTimingEvent& TimingEvent;
 		const bool bStopAtFirstMatch;
 		const bool bSearchForLargestEvent;
@@ -3229,7 +3485,7 @@ bool STimingView::SearchTimingEvent(const double InStartTime,
 		mutable bool bContinueSearching;
 		mutable double LargestDuration;
 
-		FSearchTimingEventContext(const double InStartTime, const double InEndTime, TFunctionRef<bool(double, double, uint32, uint32)> InPredicate, FTimingEvent& InOutTimingEvent, bool bInStopAtFirstMatch, bool bInSearchForLargestEvent)
+		FSearchTimingEventContext(const double InStartTime, const double InEndTime, TFunctionRef<bool(double, double, uint32)> InPredicate, FTimingEvent& InOutTimingEvent, bool bInStopAtFirstMatch, bool bInSearchForLargestEvent)
 			: StartTime(InStartTime)
 			, EndTime(InEndTime)
 			, Predicate(InPredicate)
@@ -3244,7 +3500,7 @@ bool STimingView::SearchTimingEvent(const double InStartTime,
 
 		void CheckEvent(double EventStartTime, double EventEndTime, uint32 EventDepth, const Trace::FTimingProfilerEvent& Event)
 		{
-			if (bContinueSearching && Predicate(EventStartTime, EventEndTime, EventDepth, Event.TimerIndex))
+			if (bContinueSearching && Predicate(EventStartTime, EventEndTime, EventDepth))
 			{
 				if (!bSearchForLargestEvent || EventEndTime - EventStartTime > LargestDuration)
 				{
@@ -3263,7 +3519,7 @@ bool STimingView::SearchTimingEvent(const double InStartTime,
 
 		void CheckEvent(double EventStartTime, double EventEndTime, uint32 EventDepth, const Trace::FLoadTimeProfilerCpuEvent& Event)
 		{
-			if (bContinueSearching && Predicate(EventStartTime, EventEndTime, EventDepth, 0))
+			if (bContinueSearching && Predicate(EventStartTime, EventEndTime, EventDepth))
 			{
 				if (!bSearchForLargestEvent || EventEndTime - EventStartTime > LargestDuration)
 				{
@@ -3281,6 +3537,29 @@ bool STimingView::SearchTimingEvent(const double InStartTime,
 				}
 			}
 		}
+
+		void CheckEvent(const FIoTimingEvent& Event, uint32 EventDepth)
+		{
+			if (bContinueSearching && Predicate(Event.StartTime, Event.EndTime, EventDepth))
+			{
+				if (!bSearchForLargestEvent || Event.EndTime - Event.StartTime > LargestDuration)
+				{
+					LargestDuration = Event.EndTime - Event.StartTime;
+
+					TimingEvent.TypeId = Event.Type;
+					TimingEvent.Depth = EventDepth;
+					TimingEvent.StartTime = Event.StartTime;
+					TimingEvent.EndTime = Event.EndTime;
+
+					TimingEvent.Offset = Event.Offset;
+					TimingEvent.Size = Event.Size;
+					TimingEvent.Path = Event.FileActivity->Path;
+
+					bFound = true;
+					bContinueSearching = !bStopAtFirstMatch || bSearchForLargestEvent;
+				}
+			}
+		}
 	};
 
 	FSearchTimingEventContext Ctx(InStartTime, InEndTime, InPredicate, InOutTimingEvent, bInStopAtFirstMatch, bInSearchForLargestEvent);
@@ -3393,6 +3672,43 @@ bool STimingView::SearchTimingEvent(const double InStartTime,
 				}
 			}
 		}
+		else if (Track->GetType() == ETimingEventsTrackType::Io)
+		{
+			if (Track == IoOverviewTrack)
+			{
+				for (const FIoTimingEvent& Event : AllIoEvents)
+				{
+					if (Event.EndTime <= Ctx.StartTime)
+					{
+						continue;
+					}
+
+					if (!Ctx.bContinueSearching || Event.StartTime >= Ctx.EndTime)
+					{
+						break;
+					}
+
+					Ctx.CheckEvent(Event, 0);
+				}
+			}
+			else if (Track == IoActivityTrack)
+			{
+				for (const FIoTimingEvent& Event : AllIoEvents)
+				{
+					if (Event.EndTime <= Ctx.StartTime)
+					{
+						continue;
+					}
+
+					if (!Ctx.bContinueSearching || Event.StartTime >= Ctx.EndTime)
+					{
+						break;
+					}
+
+					Ctx.CheckEvent(Event, Event.Depth);
+				}
+			}
+		}
 	}
 
 	return Ctx.bFound;
@@ -3445,7 +3761,7 @@ void STimingView::SelectLeftTimingEvent()
 		const bool bStopAtFirstMatch = false; // get last one matching
 		const bool bSearchForLargestEvent = false;
 		if (SearchTimingEvent(0.0, StartTime,
-			[Depth, StartTime, EndTime](double EventStartTime, double EventEndTime, uint32 EventDepth, uint32 EventTypeId)
+			[Depth, StartTime, EndTime](double EventStartTime, double EventEndTime, uint32 EventDepth)
 			{
 				return EventDepth == Depth &&
 					(EventStartTime < StartTime || EventEndTime < EndTime);
@@ -3470,7 +3786,7 @@ void STimingView::SelectRightTimingEvent()
 		const bool bStopAtFirstMatch = true; // get first one matching
 		const bool bSearchForLargestEvent = false;
 		if (SearchTimingEvent(EndTime, Viewport.MaxValidTime,
-			[Depth, StartTime, EndTime](double EventStartTime, double EventEndTime, uint32 EventDepth, uint32 EventTypeId)
+			[Depth, StartTime, EndTime](double EventStartTime, double EventEndTime, uint32 EventDepth)
 			{
 				return EventDepth == Depth &&
 					(EventStartTime > StartTime || EventEndTime > EndTime);
@@ -3496,7 +3812,7 @@ void STimingView::SelectUpTimingEvent()
 		const bool bStopAtFirstMatch = true; // get first one matching
 		const bool bSearchForLargestEvent = false;
 		if (SearchTimingEvent(StartTime, EndTime,
-			[Depth, StartTime, EndTime](double EventStartTime, double EventEndTime, uint32 EventDepth, uint32 EventTypeId)
+			[Depth, StartTime, EndTime](double EventStartTime, double EventEndTime, uint32 EventDepth)
 			{
 				return EventDepth == Depth
 					&& EventStartTime <= EndTime
@@ -3522,7 +3838,7 @@ void STimingView::SelectDownTimingEvent()
 		const bool bStopAtFirstMatch = false; // check all timing events
 		const bool bSearchForLargestEvent = true; // get largest timing event
 		if (SearchTimingEvent(StartTime, EndTime,
-			[Depth, StartTime, EndTime](double EventStartTime, double EventEndTime, uint32 EventDepth, uint32 EventTypeId)
+			[Depth, StartTime, EndTime](double EventStartTime, double EventEndTime, uint32 EventDepth)
 			{
 				return EventDepth == Depth
 					&& EventStartTime <= EndTime
diff --git a/Engine/Source/Developer/TraceInsights/Private/Insights/Widgets/STimingView.h b/Engine/Source/Developer/TraceInsights/Private/Insights/Widgets/STimingView.h
--- a/Engine/Source/Developer/TraceInsights/Private/Insights/Widgets/STimingView.h
+++ b/Engine/Source/Developer/TraceInsights/Private/Insights/Widgets/STimingView.h
@@ -323,7 +323,7 @@ protected:
 
 	bool SearchTimingEvent(const double InStartTime,
 						   const double InEndTime,
-						   TFunctionRef<bool(double, double, uint32, uint32)> InPredicate,
+						   TFunctionRef<bool(double, double, uint32)> InPredicate,
 						   FTimingEvent& InOutTimingEvent,
 						   bool bInStopAtFirstMatch,
 						   bool bInSearchForLargestEvent) const;
@@ -442,14 +442,30 @@ protected:
 	bool bForceIoEventsUpdate;
 
 	bool bMergeIoLanes;
+	bool bShowFileActivityBackgroundEvents;
+
+	struct FIoFileActivity
+	{
+		uint64 Id;
+		const TCHAR* Path;
+		double StartTime;
+		double EndTime;
+		int32 EventCount;
+		int32 Depth;
+	};
+
+	TArray<TSharedPtr<FIoFileActivity>> FileActivities;
+	TMap<uint64, TSharedPtr<FIoFileActivity>> FileActivityMap;
 
 	struct FIoTimingEvent
 	{
 		double StartTime;
 		double EndTime;
 		uint32 Depth;
-		uint32 Type; // Trace::EFileActivityType + Failed
-		const TCHAR* Path;
+		uint32 Type; // Trace::EFileActivityType + "Failed" flag
+		uint64 Offset;
+		uint64 Size;
+		TSharedPtr<FIoFileActivity> FileActivity;
 	};
 
 	/** All IO events, cached. */
