[Insights]
  - Frames track: Fixed auto zooming to not zoom anymore if view is not in initial position (ex. if user has scrolled or zoomed the view). It fixes UE-76529: Frames View does not stay zoomed in when profiling a Live Trace Session.

diff --git a/Engine/Source/Developer/TraceInsights/Private/Insights/Widgets/SFrameTrack.cpp b/Engine/Source/Developer/TraceInsights/Private/Insights/Widgets/SFrameTrack.cpp
--- a/Engine/Source/Developer/TraceInsights/Private/Insights/Widgets/SFrameTrack.cpp
+++ b/Engine/Source/Developer/TraceInsights/Private/Insights/Widgets/SFrameTrack.cpp
@@ -200,11 +200,14 @@ void SFrameTrack::Tick(const FGeometry& AllottedGeometry, const double InCurrent
 					UpdateHorizontalScrollBar();
 					bIsStateDirty = true;
 
-					// Auto zoom out.
-					while (Viewport.MaxX > Viewport.Width)
+					if (Viewport.PosX == 0.0f) // only if the view is unchanged
 					{
-						ZoomHorizontally(-0.1f, 0.0f);
-						Viewport.ScrollAtPosX(0.0f);
+						// Auto zoom out (until entire session time range fits into view).
+						while (Viewport.MaxX > Viewport.Width)
+						{
+							ZoomHorizontally(-0.1f, 0.0f);
+							Viewport.ScrollAtPosX(0.0f);
+						}
 					}
 				}
 			}
