[Insights] Frames track: Added context menu with toggle options for Game Frames, Rendering Frames and Auto Zoom.

diff --git a/Engine/Source/Developer/TraceInsights/Private/Insights/Widgets/SFrameTrack.cpp b/Engine/Source/Developer/TraceInsights/Private/Insights/Widgets/SFrameTrack.cpp
--- a/Engine/Source/Developer/TraceInsights/Private/Insights/Widgets/SFrameTrack.cpp
+++ b/Engine/Source/Developer/TraceInsights/Private/Insights/Widgets/SFrameTrack.cpp
@@ -59,6 +59,10 @@ void SFrameTrack::Reset()
 	TimelinesOrder.Add(2);
 	bIsStateDirty = true;
 
+	bShowGameFrames = true;
+	bShowRenderingFrames = true;
+	bIsAutoZoomEnabled = true;
+
 	AnalysisSyncNextTimestamp = 0;
 
 	MousePosition = FVector2D::ZeroVector;
@@ -197,18 +201,26 @@ void SFrameTrack::Tick(const FGeometry& AllottedGeometry, const double InCurrent
 				if (NumFrames > Viewport.MaxIndex)
 				{
 					Viewport.SetMinMaxIndexInterval(0, NumFrames);
-					UpdateHorizontalScrollBar();
-					bIsStateDirty = true;
 
-					if (Viewport.PosX == 0.0f) // only if the view is unchanged
+					if (bIsAutoZoomEnabled)
 					{
-						// Auto zoom out (until entire session time range fits into view).
-						while (Viewport.MaxX > Viewport.Width)
+						if (Viewport.PosX == 0.0f) // only if the view is unchanged
 						{
-							ZoomHorizontally(-0.1f, 0.0f);
-							Viewport.ScrollAtPosX(0.0f);
+							// Auto zoom out (until entire session time range fits into view).
+							while (Viewport.MaxX > Viewport.Width)
+							{
+								ZoomHorizontally(-0.1f, 0.0f);
+								Viewport.ScrollAtPosX(0.0f);
+							}
+						}
+						else
+						{
+							bIsAutoZoomEnabled = false;
 						}
 					}
+
+					UpdateHorizontalScrollBar();
+					bIsStateDirty = true;
 				}
 			}
 		}
@@ -281,6 +293,12 @@ FSampleRef SFrameTrack::GetSampleAtMousePosition(float X, float Y)
 		for (int32 TimelineIndex = TimelinesOrder.Num() - 1; TimelineIndex >= 0; --TimelineIndex)
 		{
 			int32 TimelineId = TimelinesOrder[TimelineIndex];
+
+			if ((!bShowRenderingFrames && TimelineId == 1) || (!bShowGameFrames && TimelineId == 0))
+			{
+				continue;
+			}
+
 			if (CachedTimelines.Contains(TimelineId))
 			{
 				const FFrameTrackTimeline& CachedTimeline = CachedTimelines[TimelineId];
@@ -365,11 +383,22 @@ int32 SFrameTrack::OnPaint(const FPaintArgs& Args, const FGeometry& AllottedGeom
 		Helper.DrawBackground();
 
 		// Draw frames, for each visible timeline.
-		for (const TPair<uint64, FFrameTrackTimeline>& KeyValuePair : CachedTimelines)
+		for (int32 TimelineIndex = 0; TimelineIndex < TimelinesOrder.Num(); ++TimelineIndex)
 		{
-			const FFrameTrackTimeline& Timeline = KeyValuePair.Value;
-			Helper.DrawCached(Timeline);
+			int32 TimelineId = TimelinesOrder[TimelineIndex];
+
+			if ((!bShowRenderingFrames && TimelineId == 1) || (!bShowGameFrames && TimelineId == 0))
+			{
+				continue;
+			}
+
+			if (CachedTimelines.Contains(TimelineId))
+			{
+				const FFrameTrackTimeline& Timeline = CachedTimelines[TimelineId];
+				Helper.DrawCached(Timeline);
+			}
 		}
+
 		NumDrawSamples = Helper.GetNumDrawSamples();
 
 		// Highlight the mouse hovered sample.
@@ -380,7 +409,7 @@ int32 SFrameTrack::OnPaint(const FPaintArgs& Args, const FGeometry& AllottedGeom
 			FString Text = FString::Format(TEXT("{0} frame {1} ({2})"),
 				{
 					(HoveredSample.Timeline->Id == 0) ? TEXT("Game") :
-					(HoveredSample.Timeline->Id == 1) ? TEXT("Render") : TEXT("Misc"),
+					(HoveredSample.Timeline->Id == 1) ? TEXT("Rendering") : TEXT("Misc"),
 					FText::AsNumber(HoveredSample.Sample->LargestFrameIndex).ToString(),
 					TimeUtils::FormatTimeAuto(HoveredSample.Sample->LargestFrameDuration)
 				});
@@ -679,9 +708,13 @@ FReply SFrameTrack::OnMouseMove(const FGeometry& MyGeometry, const FPointerEvent
 	{
 		HoveredSample = GetSampleAtMousePosition(MousePosition.X, MousePosition.Y);
 		if (!HoveredSample.IsValid())
+		{
 			HoveredSample = GetSampleAtMousePosition(MousePosition.X - 1.0f, MousePosition.Y);
+		}
 		if (!HoveredSample.IsValid())
+		{
 			HoveredSample = GetSampleAtMousePosition(MousePosition.X + 1.0f, MousePosition.Y);
+		}
 	}
 
 	return Reply;
@@ -821,14 +854,167 @@ FCursorReply SFrameTrack::OnCursorQuery(const FGeometry& MyGeometry, const FPoin
 
 void SFrameTrack::ShowContextMenu(const FVector2D& ScreenSpacePosition)
 {
+	const bool bShouldCloseWindowAfterMenuSelection = true;
+	FMenuBuilder MenuBuilder(bShouldCloseWindowAfterMenuSelection, NULL);
 
+	MenuBuilder.BeginSection("Series", LOCTEXT("ContextMenu_Header_Series", "Series"));
+	{
+		struct FLocal
+		{
+			static bool ReturnFalse()
+			{
+				return false;
+			}
+		};
+
+		FUIAction Action_ShowGameFrames
+		(
+			FExecuteAction::CreateSP(this, &SFrameTrack::ContextMenu_ShowGameFrames_Execute),
+			FCanExecuteAction::CreateSP(this, &SFrameTrack::ContextMenu_ShowGameFrames_CanExecute),
+			FIsActionChecked::CreateSP(this, &SFrameTrack::ContextMenu_ShowGameFrames_IsChecked)
+		);
+		MenuBuilder.AddMenuEntry
+		(
+			LOCTEXT("ContextMenu_ShowGameFrames", "Game Frames"),
+			LOCTEXT("ContextMenu_ShowGameFrames_Desc", "Show/hide the Game frames"),
+			FSlateIcon(),
+			Action_ShowGameFrames,
+			NAME_None,
+			EUserInterfaceActionType::ToggleButton
+		);
+
+		FUIAction Action_ShowRenderingFrames
+		(
+			FExecuteAction::CreateSP(this, &SFrameTrack::ContextMenu_ShowRenderingFrames_Execute),
+			FCanExecuteAction::CreateSP(this, &SFrameTrack::ContextMenu_ShowRenderingFrames_CanExecute),
+			FIsActionChecked::CreateSP(this, &SFrameTrack::ContextMenu_ShowRenderingFrames_IsChecked)
+		);
+		MenuBuilder.AddMenuEntry
+		(
+			LOCTEXT("ContextMenu_ShowRenderingFrames", "Rendering Frames"),
+			LOCTEXT("ContextMenu_ShowRenderingFrames_Desc", "Show/hide the Rendering frames"),
+			FSlateIcon(),
+			Action_ShowRenderingFrames,
+			NAME_None,
+			EUserInterfaceActionType::ToggleButton
+		);
+	}
+	MenuBuilder.EndSection();
+
+	MenuBuilder.BeginSection("Misc");
+	{
+		FUIAction Action_AutoZoom
+		(
+			FExecuteAction::CreateSP(this, &SFrameTrack::ContextMenu_AutoZoom_Execute),
+			FCanExecuteAction::CreateSP(this, &SFrameTrack::ContextMenu_AutoZoom_CanExecute),
+			FIsActionChecked::CreateSP(this, &SFrameTrack::ContextMenu_AutoZoom_IsChecked)
+		);
+		MenuBuilder.AddMenuEntry
+		(
+			LOCTEXT("ContextMenu_AutoZoom", "Auto Zoom"),
+			LOCTEXT("ContextMenu_AutoZoom_Desc", "Enable auto zoom. Makes entire session time range to fit into view."),
+			FSlateIcon(),
+			Action_AutoZoom,
+			NAME_None,
+			EUserInterfaceActionType::ToggleButton
+		);
+	}
+	MenuBuilder.EndSection();
+
+	TSharedRef<SWidget> MenuWidget = MenuBuilder.MakeWidget();
+
+	FWidgetPath EventPath;
+	FSlateApplication::Get().PushMenu(SharedThis(this), EventPath, MenuWidget, ScreenSpacePosition, FPopupTransitionEffect::ContextMenu);
 }
 
 ////////////////////////////////////////////////////////////////////////////////////////////////////
 
-void SFrameTrack::BindCommands()
+void SFrameTrack::ContextMenu_ShowGameFrames_Execute()
+{
+	bShowGameFrames = !bShowGameFrames;
+}
+
+////////////////////////////////////////////////////////////////////////////////////////////////////
+
+bool SFrameTrack::ContextMenu_ShowGameFrames_CanExecute()
+{
+	return true;
+}
+
+////////////////////////////////////////////////////////////////////////////////////////////////////
+
+bool SFrameTrack::ContextMenu_ShowGameFrames_IsChecked()
 {
+	return bShowGameFrames;
+}
+
+////////////////////////////////////////////////////////////////////////////////////////////////////
+
+void SFrameTrack::ContextMenu_ShowRenderingFrames_Execute()
+{
+	bShowRenderingFrames = !bShowRenderingFrames;
+}
 
+////////////////////////////////////////////////////////////////////////////////////////////////////
+
+bool SFrameTrack::ContextMenu_ShowRenderingFrames_CanExecute()
+{
+	return true;
+}
+
+////////////////////////////////////////////////////////////////////////////////////////////////////
+
+bool SFrameTrack::ContextMenu_ShowRenderingFrames_IsChecked()
+{
+	return bShowRenderingFrames;
+}
+
+////////////////////////////////////////////////////////////////////////////////////////////////////
+
+void SFrameTrack::ContextMenu_AutoZoom_Execute()
+{
+	bIsAutoZoomEnabled = !(bIsAutoZoomEnabled && Viewport.PosX == 0.0f);
+	if (bIsAutoZoomEnabled)
+	{
+		Viewport.ScrollAtPosX(0.0f);
+
+		// Auto zoom in.
+		while (Viewport.MaxX < Viewport.Width)
+		{
+			ZoomHorizontally(+0.1f, 0.0f);
+			Viewport.ScrollAtPosX(0.0f);
+		}
+
+		// Auto zoom out (until entire session time range fits into view).
+		while (Viewport.MaxX > Viewport.Width)
+		{
+			ZoomHorizontally(-0.1f, 0.0f);
+			Viewport.ScrollAtPosX(0.0f);
+		}
+
+		UpdateHorizontalScrollBar();
+		bIsStateDirty = true;
+	}
+}
+
+////////////////////////////////////////////////////////////////////////////////////////////////////
+
+bool SFrameTrack::ContextMenu_AutoZoom_CanExecute()
+{
+	return true;
+}
+
+////////////////////////////////////////////////////////////////////////////////////////////////////
+
+bool SFrameTrack::ContextMenu_AutoZoom_IsChecked()
+{
+	return bIsAutoZoomEnabled && Viewport.PosX == 0.0f;
+}
+
+////////////////////////////////////////////////////////////////////////////////////////////////////
+
+void SFrameTrack::BindCommands()
+{
 }
 
 ////////////////////////////////////////////////////////////////////////////////////////////////////
diff --git a/Engine/Source/Developer/TraceInsights/Private/Insights/Widgets/SFrameTrack.h b/Engine/Source/Developer/TraceInsights/Private/Insights/Widgets/SFrameTrack.h
--- a/Engine/Source/Developer/TraceInsights/Private/Insights/Widgets/SFrameTrack.h
+++ b/Engine/Source/Developer/TraceInsights/Private/Insights/Widgets/SFrameTrack.h
@@ -104,6 +104,18 @@ protected:
 
 	void ShowContextMenu(const FVector2D& ScreenSpacePosition);
 
+	void ContextMenu_ShowGameFrames_Execute();
+	bool ContextMenu_ShowGameFrames_CanExecute();
+	bool ContextMenu_ShowGameFrames_IsChecked();
+
+	void ContextMenu_ShowRenderingFrames_Execute();
+	bool ContextMenu_ShowRenderingFrames_CanExecute();
+	bool ContextMenu_ShowRenderingFrames_IsChecked();
+
+	void ContextMenu_AutoZoom_Execute();
+	bool ContextMenu_AutoZoom_CanExecute();
+	bool ContextMenu_AutoZoom_IsChecked();
+
 	/** Binds our UI commands to delegates. */
 	void BindCommands();
 
@@ -140,6 +152,10 @@ protected:
 
 	bool bIsStateDirty;
 
+	bool bShowGameFrames;
+	bool bShowRenderingFrames;
+	bool bIsAutoZoomEnabled;
+
 	uint64 AnalysisSyncNextTimestamp;
 
 	//////////////////////////////////////////////////
