[Insights] Log view: Fixed log view showing invalid messages when a new session is loaded (due to message cache not being reset).

diff --git a/Engine/Source/Developer/TraceInsights/Private/Insights/ViewModels/LogMessage.cpp b/Engine/Source/Developer/TraceInsights/Private/Insights/ViewModels/LogMessage.cpp
--- a/Engine/Source/Developer/TraceInsights/Private/Insights/ViewModels/LogMessage.cpp
+++ b/Engine/Source/Developer/TraceInsights/Private/Insights/ViewModels/LogMessage.cpp
@@ -120,8 +120,21 @@ FLogMessageCache::FLogMessageCache()
 
 ////////////////////////////////////////////////////////////////////////////////////////////////////
 
+void FLogMessageCache::SetSession(TSharedPtr<const Trace::IAnalysisSession> InSession)
+{
+	if (Session != InSession)
+	{
+		Reset();
+		Session = InSession;
+	}
+}
+
+////////////////////////////////////////////////////////////////////////////////////////////////////
+
 void FLogMessageCache::Reset()
 {
+	FScopeLock Lock(&CriticalSection);
+	Map.Reset();
 }
 
 ////////////////////////////////////////////////////////////////////////////////////////////////////
diff --git a/Engine/Source/Developer/TraceInsights/Private/Insights/ViewModels/LogMessage.h b/Engine/Source/Developer/TraceInsights/Private/Insights/ViewModels/LogMessage.h
--- a/Engine/Source/Developer/TraceInsights/Private/Insights/ViewModels/LogMessage.h
+++ b/Engine/Source/Developer/TraceInsights/Private/Insights/ViewModels/LogMessage.h
@@ -47,7 +47,7 @@ class FLogMessageCache
 public:
 	FLogMessageCache();
 
-	void SetSession(TSharedPtr<const Trace::IAnalysisSession> InSession) { Session = InSession; }
+	void SetSession(TSharedPtr<const Trace::IAnalysisSession> InSession);
 	void Reset();
 
 	FLogMessageRecord& Get(uint64 Index);
