[Insights]
  - Graph track: Added support for horizontal panning and zooming with mouse.
  - Graph track: Added time range selection (Ctrl + Shift + mouse drag or mouse drag on time ruler area).
  - Graph track: Added context menu with on/off toggles: "Points", "Points have Border", "Lines", "Lines have Duration", "Bars".
  - Graph: Reduced size of points from 7px to 5px.
  - Graph: Fixed position of lines (being offseted with 1px).
  - Timing: Re-enabled the "Graph" toolbar option that opens a test graph track.
  - Time ruler: Fixed position of current time info when selecting (when the selected time interval being just 1-2 pixels).
  - Code refactoring to reuse some draw code (background, selection) between TimingView and GraphTrack. Added FDrawHelpers class.

diff --git a/Engine/Source/Developer/TraceInsights/Private/Insights/ViewModels/DrawHelpers.cpp b/Engine/Source/Developer/TraceInsights/Private/Insights/ViewModels/DrawHelpers.cpp
new file mode 100644
--- /dev/null
+++ b/Engine/Source/Developer/TraceInsights/Private/Insights/ViewModels/DrawHelpers.cpp
@@ -0,0 +1,256 @@
+// Copyright 1998-2019 Epic Games, Inc. All Rights Reserved.
+
+#include "DrawHelpers.h"
+
+#include "Brushes/SlateBorderBrush.h"
+#include "Brushes/SlateColorBrush.h"
+#include "Fonts/FontMeasure.h"
+#include "Fonts/SlateFontInfo.h"
+
+// Insights
+#include "Insights/Common/PaintUtils.h"
+#include "Insights/Common/TimeUtils.h"
+#include "Insights/ViewModels/TimingTrackViewport.h"
+
+////////////////////////////////////////////////////////////////////////////////////////////////////
+
+void FDrawHelpers::DrawBackground(
+	const FDrawContext& DrawContext,
+	const FTimingTrackViewport& Viewport,
+	const FSlateBrush* BackgroundAreaBrush,
+	const FLinearColor& ValidAreaColor,
+	const FLinearColor& InvalidAreaColor,
+	const FLinearColor& EdgeColor,
+	const float X0,
+	const float Y,
+	const float W,
+	const float H,
+	float& OutValidAreaX,
+	float& OutValidAreaW)
+{
+	//  <------- W ------->
+	//  X0    X1    X2     X3
+	//  ++++++|*****|++++++
+	//  ++++++|*****|++++++
+	//  ++++++|*****|++++++
+
+	const float X1 = Viewport.TimeToSlateUnitsRounded(Viewport.MinValidTime);
+	const float X2 = Viewport.TimeToSlateUnitsRounded(Viewport.MaxValidTime);
+	const float X3 = X0 + W;
+
+	if (X1 >= X3 || X2 <= X0)
+	{
+		OutValidAreaX = X0;
+		OutValidAreaW = W;
+
+		// Draw invalid area (entire view).
+		DrawContext.DrawBox(X0, Y, W, H, BackgroundAreaBrush, InvalidAreaColor);
+	}
+	else // X1 < X3 && X2 > X0
+	{
+		if (X1 > X0)
+		{
+			// Draw invalid area (left).
+			DrawContext.DrawBox(X0, Y, X1 - X0, H, BackgroundAreaBrush, InvalidAreaColor);
+		}
+
+		if (X2 < X3)
+		{
+			// Draw invalid area (right).
+			DrawContext.DrawBox(X2 + 1.0f, Y, X3 - X2 - 1.0f, H, BackgroundAreaBrush, InvalidAreaColor);
+
+			// Draw the right edge (end time).
+			DrawContext.DrawBox(X2, Y, 1.0f, H, BackgroundAreaBrush, EdgeColor);
+		}
+
+		float ValidAreaX = FMath::Max(X1, X0);
+		float ValidAreaW = FMath::Min(X2, X3) - ValidAreaX;
+
+		if (X1 >= X0)
+		{
+			// Draw the left edge (start time).
+			DrawContext.DrawBox(X1, Y, 1.0f, H, BackgroundAreaBrush, EdgeColor);
+
+			// Adjust valid area to not overlap the left edge.
+			ValidAreaX += 1.0f;
+			ValidAreaW -= 1.0f;
+		}
+
+		if (ValidAreaW > 0.0f)
+		{
+			// Draw valid area.
+			DrawContext.DrawBox(ValidAreaX, Y, ValidAreaW, H, BackgroundAreaBrush, ValidAreaColor);
+		}
+
+		OutValidAreaX = ValidAreaX;
+		OutValidAreaW = ValidAreaW;
+	}
+
+	DrawContext.LayerId++;
+}
+
+////////////////////////////////////////////////////////////////////////////////////////////////////
+
+void FDrawHelpers::DrawTimeRangeSelection(
+	const FDrawContext& DrawContext,
+	const FTimingTrackViewport& Viewport,
+	const double StartTime,
+	const double EndTime,
+	const FSlateBrush* Brush,
+	const FSlateFontInfo& Font)
+{
+	if (EndTime > StartTime)
+	{
+		float SelectionX1 = Viewport.TimeToSlateUnitsRounded(StartTime);
+		float SelectionX2 = Viewport.TimeToSlateUnitsRounded(EndTime);
+
+		if (SelectionX1 <= Viewport.Width &&
+			SelectionX2 >= 0)
+		{
+			float ClipLeft = 0.0f;
+			float ClipRight = 0.0f;
+			if (SelectionX1 < 0.0f)
+			{
+				ClipLeft = -SelectionX1;
+				SelectionX1 = 0.0f;
+			}
+			if (SelectionX2 > Viewport.Width)
+			{
+				ClipRight = SelectionX2 - Viewport.Width;
+				SelectionX2 = Viewport.Width;
+			}
+
+			// Draw selection area.
+			DrawContext.DrawBox(SelectionX1, 0.0f, SelectionX2 - SelectionX1, Viewport.Height, Brush, FLinearColor(0.25f, 0.5f, 1.0f, 0.25f));
+			DrawContext.LayerId++;
+
+			const FColor ArrowFillColor(32, 64, 128, 255);
+			const FLinearColor ArrowColor(ArrowFillColor);
+
+			if (SelectionX1 > 0.0f)
+			{
+				// Draw left side (vertical line).
+				DrawContext.DrawBox(SelectionX1 - 1.0f, 0.0f, 1.0f, Viewport.Height, Brush, ArrowColor);
+			}
+
+			if (SelectionX2 < Viewport.Width)
+			{
+				// Draw right side (vertical line).
+				DrawContext.DrawBox(SelectionX2, 0.0f, 1.0f, Viewport.Height, Brush, ArrowColor);
+			}
+
+			DrawContext.LayerId++;
+
+			const float ArrowSize = 6.0f;
+			const float ArrowY = 6.0f;
+
+			if (SelectionX2 - SelectionX1 > 2 * ArrowSize)
+			{
+				// Draw horizontal line.
+				float HorizLineX1 = SelectionX1;
+				if (ClipLeft == 0.0f)
+				{
+					HorizLineX1 += 1.0f;
+				}
+				float HorizLineX2 = SelectionX2;
+				if (ClipRight == 0.0f)
+				{
+					HorizLineX2 -= 1.0f;
+				}
+				DrawContext.DrawBox(HorizLineX1, ArrowY - 1.0f, HorizLineX2 - HorizLineX1, 3.0f, Brush, ArrowColor);
+
+				if (ClipLeft < ArrowSize)
+				{
+					// Draw left arrow.
+					for (float AH = 0.0f; AH < ArrowSize; AH += 1.0f)
+					{
+						DrawContext.DrawBox(SelectionX1 - ClipLeft + AH, ArrowY - AH, 1.0f, 2.0f * AH + 1.0f, Brush, ArrowColor);
+					}
+				}
+
+				if (ClipRight < ArrowSize)
+				{
+					// Draw right arrow.
+					for (float AH = 0.0f; AH < ArrowSize; AH += 1.0f)
+					{
+						DrawContext.DrawBox(SelectionX2 + ClipRight - AH - 1.0f, ArrowY - AH, 1.0f, 2.0f * AH + 1.0f, Brush, ArrowColor);
+					}
+				}
+
+				DrawContext.LayerId++;
+
+#if 0
+				//im: This should be a more efficeint way top draw the arrows, but it renders them with artifacts (missing vertical lines; shader bug?)!
+
+				const FSlateBrush* MyBrush = WhiteBrush;
+				FSlateShaderResourceProxy* ResourceProxy = FSlateDataPayload::ResourceManager->GetShaderResource(*MyBrush);
+				FSlateResourceHandle ResourceHandle = FSlateApplication::Get().GetRenderer()->GetResourceHandle(*MyBrush);
+
+				FVector2D AtlasOffset = ResourceProxy ? ResourceProxy->StartUV : FVector2D(0.0f, 0.0f);
+				FVector2D AtlasUVSize = ResourceProxy ? ResourceProxy->SizeUV : FVector2D(1.0f, 1.0f);
+
+				const FVector2D Pos = AllottedGeometry.GetAbsolutePosition() + FVector2D(0.0f, 40.0f);
+				const float Scale = AllottedGeometry.Scale;
+
+				FSlateRenderTransform RenderTransform;
+
+				TArray<FSlateVertex> Verts;
+				Verts.Reserve(6);
+				Verts.Add(FSlateVertex::Make<ESlateVertexRounding::Disabled>(RenderTransform, Pos + FVector2D(0.5f + SelectionX1 + ArrowSize, 0.5f + ArrowY + ArrowSize) * Scale, AtlasOffset + FVector2D(0.0f, 1.0f) * AtlasUVSize, ArrowFillColor));
+				Verts.Add(FSlateVertex::Make<ESlateVertexRounding::Disabled>(RenderTransform, Pos + FVector2D(0.5f + SelectionX1, 0.5f + ArrowY) * Scale, AtlasOffset + FVector2D(1.0f, 0.5f) * AtlasUVSize, ArrowFillColor));
+				Verts.Add(FSlateVertex::Make<ESlateVertexRounding::Disabled>(RenderTransform, Pos + FVector2D(0.5f + SelectionX1 + ArrowSize, 0.5f + ArrowY - ArrowSize) * Scale, AtlasOffset + FVector2D(0.0f, 0.0f) * AtlasUVSize, ArrowFillColor));
+				Verts.Add(FSlateVertex::Make<ESlateVertexRounding::Disabled>(RenderTransform, Pos + FVector2D(0.5f + SelectionX2 - ArrowSize, 0.5f + ArrowY - ArrowSize) * Scale, AtlasOffset + FVector2D(0.0f, 0.0f) * AtlasUVSize, ArrowFillColor));
+				Verts.Add(FSlateVertex::Make<ESlateVertexRounding::Disabled>(RenderTransform, Pos + FVector2D(0.5f + SelectionX2, 0.5f + ArrowY) * Scale, AtlasOffset + FVector2D(1.0f, 0.5f) * AtlasUVSize, ArrowFillColor));
+				Verts.Add(FSlateVertex::Make<ESlateVertexRounding::Disabled>(RenderTransform, Pos + FVector2D(0.5f + SelectionX2 - ArrowSize, 0.5f + ArrowY + ArrowSize) * Scale, AtlasOffset + FVector2D(0.0f, 1.0f) * AtlasUVSize, ArrowFillColor));
+
+				TArray<SlateIndex> Indices;
+				Indices.Reserve(6);
+				if (ClipLeft < ArrowSize)
+				{
+					Indices.Add(0);
+					Indices.Add(1);
+					Indices.Add(2);
+				}
+				if (ClipRight < ArrowSize)
+				{
+					Indices.Add(3);
+					Indices.Add(4);
+					Indices.Add(5);
+				}
+
+				FSlateDrawElement::MakeCustomVerts(
+					OutDrawElements,
+					LayerId,
+					ResourceHandle,
+					Verts,
+					Indices,
+					nullptr,
+					0,
+					0,
+					ESlateDrawEffect::PreMultipliedAlpha);
+
+				DrawContext.LayerId++;
+#endif
+			}
+
+			//////////////////////////////////////////////////
+			// Draw duration for selected time interval.
+
+			const double Duration = EndTime - StartTime;
+			const FString Text = TimeUtils::FormatTimeAuto(Duration);
+
+			const TSharedRef<FSlateFontMeasure> FontMeasureService = FSlateApplication::Get().GetRenderer()->GetFontMeasureService();
+			const float TextWidth = FontMeasureService->Measure(Text, Font).X;
+
+			const float CenterX = (SelectionX1 + SelectionX2) / 2.0f;
+
+			DrawContext.DrawBox(CenterX - TextWidth / 2 - 2.0, ArrowY - 6.0f, TextWidth + 4.0f, 13.0f, Brush, ArrowColor);
+			DrawContext.LayerId++;
+
+			DrawContext.DrawText(CenterX - TextWidth / 2, ArrowY - 6.0f, Text, Font, FLinearColor::White);
+			DrawContext.LayerId++;
+		}
+	}
+}
+
+////////////////////////////////////////////////////////////////////////////////////////////////////
diff --git a/Engine/Source/Developer/TraceInsights/Private/Insights/ViewModels/DrawHelpers.h b/Engine/Source/Developer/TraceInsights/Private/Insights/ViewModels/DrawHelpers.h
new file mode 100644
--- /dev/null
+++ b/Engine/Source/Developer/TraceInsights/Private/Insights/ViewModels/DrawHelpers.h
@@ -0,0 +1,48 @@
+// Copyright 1998-2019 Epic Games, Inc. All Rights Reserved.
+
+#pragma once
+
+#include "CoreMinimal.h"
+#include "Fonts/SlateFontInfo.h"
+#include "Math/Color.h"
+
+////////////////////////////////////////////////////////////////////////////////////////////////////
+
+struct FDrawContext;
+struct FSlateBrush;
+
+class FTimingTrackViewport;
+
+////////////////////////////////////////////////////////////////////////////////////////////////////
+
+class FDrawHelpers
+{
+public:
+	/**
+	 * Draw background.
+	 */
+	static void DrawBackground(const FDrawContext& DrawContext,
+							   const FTimingTrackViewport& Viewport,
+							   const FSlateBrush* BackgroundAreaBrush,
+							   const FLinearColor& ValidAreaColor,
+							   const FLinearColor& InvalidAreaColor,
+							   const FLinearColor& EdgeColor,
+							   const float X,
+							   const float Y,
+							   const float W,
+							   const float H,
+							   float& OutValidAreaX,
+							   float& OutValidAreaW);
+
+	/**
+	 * Draw time range selection.
+	 */
+	static void DrawTimeRangeSelection(const FDrawContext& DrawContext,
+									   const FTimingTrackViewport& Viewport,
+									   const double StartTime,
+									   const double EndTime,
+									   const FSlateBrush* Brush,
+									   const FSlateFontInfo& Font);
+};
+
+////////////////////////////////////////////////////////////////////////////////////////////////////
diff --git a/Engine/Source/Developer/TraceInsights/Private/Insights/ViewModels/GraphTrack.cpp b/Engine/Source/Developer/TraceInsights/Private/Insights/ViewModels/GraphTrack.cpp
--- a/Engine/Source/Developer/TraceInsights/Private/Insights/ViewModels/GraphTrack.cpp
+++ b/Engine/Source/Developer/TraceInsights/Private/Insights/ViewModels/GraphTrack.cpp
@@ -14,6 +14,7 @@
 #include "Insights/Common/PaintUtils.h"
 #include "Insights/InsightsManager.h"
 #include "Insights/TimingProfilerManager.h"
+#include "Insights/ViewModels/DrawHelpers.h"
 #include "Insights/ViewModels/TimingTrackViewport.h"
 
 #define LOCTEXT_NAMESPACE "GraphTrack"
@@ -110,56 +111,16 @@ void FGraphTrack::DrawBackground(FDrawContext& DrawContext, const FTimingTrackVi
 {
 	const FLinearColor ValidAreaColor(0.07f, 0.07f, 0.07f, 1.0f);
 	const FLinearColor InvalidAreaColor(0.1f, 0.07f, 0.07f, 1.0f);
-	const FLinearColor StartEndMarkerColor(0.05f, 0.05f, 0.05f, 1.0f);
+	const FLinearColor EdgeColor(0.05f, 0.05f, 0.05f, 1.0f);
 
-	const float X0 = Viewport.TimeToSlateUnitsRounded(0.0);
-	const float X1 = Viewport.TimeToSlateUnitsRounded(Viewport.MaxValidTime);
+	const float X = 0.0f;
 	const float W = FMath::CeilToFloat(Viewport.Width);
-
 	const float Y = GetPosY();
 	const float H = GetHeight();
 
-	if (X0 >= W || X1 <= 0.0f)
-	{
-		// Draw invalid area (entire view).
-		DrawContext.DrawBox(0.0f, Y, W, H, WhiteBrush, InvalidAreaColor);
-	}
-	else // X0 < W && X1 > 0
-	{
-		if (X0 > 0.0f)
-		{
-			// Draw invalid area (left).
-			DrawContext.DrawBox(0.0f, Y, X0, H, WhiteBrush, InvalidAreaColor);
-		}
-
-		if (X1 < W)
-		{
-			// Draw invalid area (right).
-			DrawContext.DrawBox(X1 + 1.0f, Y, W - X1 - 1.0f, H, WhiteBrush, InvalidAreaColor);
-
-			// Draw the end time marker.
-			DrawContext.DrawBox(X1, Y, 1.0f, H, WhiteBrush, StartEndMarkerColor);
-		}
-
-		float ValidX0 = FMath::Max(X0, 0.0f);
-		float ValidX1 = FMath::Min(X1, W);
+	float ValidAreaX, ValidAreaW;
 
-		if (X0 >= 0.0f)
-		{
-			// Draw the start time marker.
-			DrawContext.DrawBox(X0, Y, 1.0f, H, WhiteBrush, StartEndMarkerColor);
-
-			ValidX0 += 1.0f; // to not overlap the start time marker
-		}
-
-		if (ValidX1 > ValidX0)
-		{
-			// Draw valid area.
-			DrawContext.DrawBox(ValidX0, Y, ValidX1 - ValidX0, H, WhiteBrush, ValidAreaColor);
-		}
-	}
-
-	DrawContext.LayerId++;
+	FDrawHelpers::DrawBackground(DrawContext, Viewport, WhiteBrush, ValidAreaColor, InvalidAreaColor, EdgeColor, X, Y, W, H, ValidAreaX, ValidAreaW);
 }
 
 ////////////////////////////////////////////////////////////////////////////////////////////////////
@@ -175,6 +136,8 @@ void FGraphTrack::Draw(FDrawContext& DrawContext, const FTimingTrackViewport& Vi
 		DrawSeries(Series, DrawContext, Viewport);
 	}
 
+	//DrawContext.DrawBox(0.0f, GetPosY(), Viewport.Width, 1.0f, WhiteBrush, FLinearColor(0.05f, 0.05f, 0.05f, 1.0f));
+
 	// Draw baseline (Value == 0).
 	DrawContext.DrawBox(0.0f, GetPosY() + BaselineY - 1.0f, Viewport.Width, 1.0f, WhiteBrush, FLinearColor(0.05f, 0.05f, 0.05f, 1.0f));
 	DrawContext.LayerId++;
@@ -198,7 +161,7 @@ void FGraphTrack::DrawSeries(const FGraphTrackSeries& Series, FDrawContext& Draw
 	if (bDrawLines)
 	{
 		FPaintGeometry Geo = DrawContext.Geometry.ToPaintGeometry();
-		Geo.AppendTransform(FSlateLayoutTransform(FVector2D(0, GetPosY())));
+		Geo.AppendTransform(FSlateLayoutTransform(FVector2D(0.5f, 0.5f + GetPosY())));
 		FSlateDrawElement::MakeLines(DrawContext.ElementList, DrawContext.LayerId, Geo, Series.LinePoints, DrawContext.DrawEffects, Series.Color, false, 1.0f);
 		DrawContext.LayerId++;
 	}
@@ -261,7 +224,7 @@ void FGraphTrack::DrawSeries(const FGraphTrackSeries& Series, FDrawContext& Draw
 			const FVector2D& Pt = Series.Points[Index];
 			const float PtX = Pt.X - PtSize / 2.0f;
 			const float PtY = GetPosY() + Pt.Y - PtSize / 2.0f;
-			DrawContext.DrawBox(PtX, PtY, PtSize, PtSize, PointBrush, Series.Color);
+			DrawContext.DrawBox(PtX, PtY, PtSize, PtSize, WhiteBrush, Series.Color);
 			//DrawContext.DrawRotatedBox(PtX, PtY, PtSize, PtSize, WhiteBrush, Series.Color, Angle, RotationPoint);
 		}
 		DrawContext.LayerId++;
diff --git a/Engine/Source/Developer/TraceInsights/Private/Insights/ViewModels/GraphTrack.h b/Engine/Source/Developer/TraceInsights/Private/Insights/ViewModels/GraphTrack.h
--- a/Engine/Source/Developer/TraceInsights/Private/Insights/ViewModels/GraphTrack.h
+++ b/Engine/Source/Developer/TraceInsights/Private/Insights/ViewModels/GraphTrack.h
@@ -60,7 +60,7 @@ class FGraphTrack : public FBaseTimingTrack
 
 private:
 	// Visual size of points (in pixels).
-	static constexpr float PointVisualSize = 7.0f;
+	static constexpr float PointVisualSize = 5.0f;
 
 	// Size of points (in pixels) used in reduction algorithm.
 	static constexpr double PointSizeX = 3.0f;
@@ -101,12 +101,14 @@ protected:
 	const FSlateBrush* BorderBrush;
 	const FSlateFontInfo Font;
 
+public:
 	bool bDrawPoints;
 	bool bDrawPointsWithBorder;
 	bool bDrawLines;
 	bool bDrawLinesWithDuration;
 	bool bDrawBoxes;
 
+protected:
 	float BaselineY; // Y position (in viewport local space) of the baseline (with Value == 0); in pixels (Slate units)
 	double ScaleY; // scale between Value units and viewport units; in pixels (Slate units) / Value unit
 
diff --git a/Engine/Source/Developer/TraceInsights/Private/Insights/ViewModels/TimeRulerTrack.cpp b/Engine/Source/Developer/TraceInsights/Private/Insights/ViewModels/TimeRulerTrack.cpp
--- a/Engine/Source/Developer/TraceInsights/Private/Insights/ViewModels/TimeRulerTrack.cpp
+++ b/Engine/Source/Developer/TraceInsights/Private/Insights/ViewModels/TimeRulerTrack.cpp
@@ -207,13 +207,14 @@ void FTimeRulerTrack::Draw(FDrawContext& DrawContext, const FTimingTrackViewport
 		float W = CrtMousePosTextWidth + 4.0f;
 		if (bIsSelecting && SelectionStartTime < SelectionEndTime)
 		{
+			// While selecting, display the current time on either left or right side of the selected time range (i.e. to not overlap the selection arrows).
 			float SelectionX1 = Viewport.TimeToSlateUnitsRounded(SelectionStartTime);
 			float SelectionX2 = Viewport.TimeToSlateUnitsRounded(SelectionEndTime);
-			if (X - SelectionX1 > 1.0f)
+			if (FMath::Abs(X - SelectionX1) > FMath::Abs(SelectionX2 - X))
 			{
 				X = SelectionX2 + W / 2;
 			}
-			if (SelectionX2 - X > 1.0f)
+			else
 			{
 				X = SelectionX1 - W / 2;
 			}
diff --git a/Engine/Source/Developer/TraceInsights/Private/Insights/ViewModels/TimingViewDrawHelper.cpp b/Engine/Source/Developer/TraceInsights/Private/Insights/ViewModels/TimingViewDrawHelper.cpp
--- a/Engine/Source/Developer/TraceInsights/Private/Insights/ViewModels/TimingViewDrawHelper.cpp
+++ b/Engine/Source/Developer/TraceInsights/Private/Insights/ViewModels/TimingViewDrawHelper.cpp
@@ -15,6 +15,7 @@
 // Insights
 #include "Insights/Common/PaintUtils.h"
 #include "Insights/Common/TimeUtils.h"
+#include "Insights/ViewModels/DrawHelpers.h"
 #include "Insights/ViewModels/MarkersTimingTrack.h"
 #include "Insights/ViewModels/TimingEventsTrack.h"
 #include "Insights/ViewModels/TimingTrackViewport.h"
@@ -31,6 +32,7 @@ FTimingViewDrawHelper::FTimingViewDrawHelper(const FDrawContext& InDC, const FTi
 	, BackgroundAreaBrush(WhiteBrush)
 	, ValidAreaColor(0.07f, 0.07f, 0.07f, 1.0f)
 	, InvalidAreaColor(0.1f, 0.07f, 0.07f, 1.0f)
+	, EdgeColor(0.05f, 0.05f, 0.05f, 1.0f)
 	, EventFont(FCoreStyle::GetDefaultFontStyle("Regular", 8))
 	, Stats()
 {
@@ -47,59 +49,11 @@ FTimingViewDrawHelper::~FTimingViewDrawHelper()
 
 void FTimingViewDrawHelper::DrawBackground() const
 {
-	const FLinearColor StartEndMarkerColor(0.05f, 0.05f, 0.05f, 1.0f);
-
-	const float X0 = Viewport.TimeToSlateUnitsRounded(0.0);
-	const float X1 = Viewport.TimeToSlateUnitsRounded(Viewport.MaxValidTime);
+	const float X = 0.0f;
 	const float W = FMath::CeilToFloat(Viewport.Width);
-
 	const float Y = 0.0f;
 	const float H = FMath::CeilToFloat(Viewport.Height);
-
-	if (X0 >= W || X1 <= 0.0f)
-	{
-		ValidX0 = 0.0f;
-		ValidX1 = W;
-
-		// Draw invalid area (entire view).
-		DrawContext.DrawBox(0.0f, Y, W, H, BackgroundAreaBrush, InvalidAreaColor);
-	}
-	else // X0 < W && X1 > 0
-	{
-		if (X0 > 0.0f)
-		{
-			// Draw invalid area (left).
-			DrawContext.DrawBox(0.0f, Y, X0, H, BackgroundAreaBrush, InvalidAreaColor);
-		}
-
-		if (X1 < W)
-		{
-			// Draw invalid area (right).
-			DrawContext.DrawBox(X1 + 1.0f, Y, W - X1 - 1.0f, H, BackgroundAreaBrush, InvalidAreaColor);
-
-			// Draw the end time marker.
-			DrawContext.DrawBox(X1, Y, 1.0f, H, BackgroundAreaBrush, StartEndMarkerColor);
-		}
-
-		ValidX0 = FMath::Max(X0, 0.0f);
-		ValidX1 = FMath::Min(X1, W);
-
-		if (X0 >= 0.0f)
-		{
-			// Draw the start time marker.
-			DrawContext.DrawBox(X0, Y, 1.0f, H, BackgroundAreaBrush, StartEndMarkerColor);
-
-			ValidX0 += 1.0f; // to not overlap the start time marker
-		}
-
-		if (ValidX1 > ValidX0)
-		{
-			// Draw valid area.
-			DrawContext.DrawBox(ValidX0, Y, ValidX1 - ValidX0, H, BackgroundAreaBrush, ValidAreaColor);
-		}
-	}
-
-	DrawContext.LayerId++;
+	FDrawHelpers::DrawBackground(DrawContext, Viewport, BackgroundAreaBrush, ValidAreaColor, InvalidAreaColor, EdgeColor, X, Y, W, H, ValidAreaX, ValidAreaW);
 }
 
 ////////////////////////////////////////////////////////////////////////////////////////////////////
@@ -111,12 +65,12 @@ void FTimingViewDrawHelper::BeginTimelines()
 
 	float Y = TimelineTopY + Viewport.TopOffset;
 
-	if (Y > 0.0f && ValidX1 > ValidX0)
+	if (Y > 0.0f && ValidAreaW > 0.0f)
 	{
 		Y = FMath::Min(Y, Viewport.Height);
 
 		// Draw invalid area (top).
-		DrawContext.DrawBox(DrawContext.LayerId + ToInt32(EDrawLayer::TimelineHeader), ValidX0, 0.0f, ValidX1 - ValidX0, Y, BackgroundAreaBrush, InvalidAreaColor);
+		DrawContext.DrawBox(DrawContext.LayerId + ToInt32(EDrawLayer::TimelineHeader), ValidAreaX, 0.0f, ValidAreaW, Y, BackgroundAreaBrush, InvalidAreaColor);
 	}
 }
 
@@ -222,7 +176,6 @@ void FTimingViewDrawHelper::EndTimeline(FTimingEventsTrack& Track)
 	{
 		const float Y = Viewport.TopOffset + TimelineTopY;
 
-		FLinearColor Color(0.05f, 0.05f, 0.05f, 1.0f);
 		FLinearColor TextColor(1.0f, 1.0f, 1.0f, 1.0f);
 
 		if (Track.IsSelected())
@@ -242,13 +195,13 @@ void FTimingViewDrawHelper::EndTimeline(FTimingEventsTrack& Track)
 		}
 
 		// Draw a horizontal line between timelines.
-		DrawContext.DrawBox(DrawContext.LayerId + ToInt32(EDrawLayer::TimelineHeader), 0, Y, Viewport.Width, 1.0f, WhiteBrush, Color);
+		DrawContext.DrawBox(DrawContext.LayerId + ToInt32(EDrawLayer::TimelineHeader), 0, Y, Viewport.Width, 1.0f, WhiteBrush, EdgeColor);
 
 		// Draw name of timeline.
 		//const FString Name = FString::Printf(TEXT("%d Y: %g, H: %g, VO: %g, VY: %g"), TimelineIndex, Track.Y, Track.H, Viewport.TopOffset, Viewport.ScrollPosY); //debug
 		const TSharedRef<FSlateFontMeasure> FontMeasureService = FSlateApplication::Get().GetRenderer()->GetFontMeasureService();
 		float NameWidth = FontMeasureService->Measure(Track.GetName(), EventFont).X;
-		DrawContext.DrawBox(DrawContext.LayerId + ToInt32(EDrawLayer::TimelineHeader), 0.0f, Y + 1.0f, NameWidth + 4.0f, 12.0f, WhiteBrush, Color);
+		DrawContext.DrawBox(DrawContext.LayerId + ToInt32(EDrawLayer::TimelineHeader), 0.0f, Y + 1.0f, NameWidth + 4.0f, 12.0f, WhiteBrush, EdgeColor);
 		DrawContext.DrawText(DrawContext.LayerId + ToInt32(EDrawLayer::TimelineText), 2.0f, Y, Track.GetName(), EventFont, TextColor);
 	}
 
@@ -507,15 +460,15 @@ void FTimingViewDrawHelper::EndTimelines()
 	float Y = Viewport.TopOffset + TimelineTopY;
 
 	// Draw a last horizontal line.
-	DrawContext.DrawBox(DrawContext.LayerId + ToInt32(EDrawLayer::TimelineHeader), 0, Y, Viewport.Width, 1.0f, WhiteBrush, FLinearColor(0.05f, 0.05f, 0.05f, 1.0f));
+	DrawContext.DrawBox(DrawContext.LayerId + ToInt32(EDrawLayer::TimelineHeader), 0, Y, Viewport.Width, 1.0f, WhiteBrush, EdgeColor);
 	Y += 1.0f;
 
-	if (Y < Viewport.Height && ValidX1 > ValidX0)
+	if (Y < Viewport.Height && ValidAreaW > 0.0f)
 	{
 		Y = FMath::Max(Y, 0.0f);
 
 		// Draw invalid area (bottom).
-		DrawContext.DrawBox(DrawContext.LayerId + ToInt32(EDrawLayer::TimelineHeader), ValidX0, Y, ValidX1 - ValidX0, Viewport.Height - Y, BackgroundAreaBrush, InvalidAreaColor);
+		DrawContext.DrawBox(DrawContext.LayerId + ToInt32(EDrawLayer::TimelineHeader), ValidAreaX, Y, ValidAreaW, Viewport.Height - Y, BackgroundAreaBrush, InvalidAreaColor);
 	}
 
 	DrawContext.LayerId += ToInt32(EDrawLayer::Count);
diff --git a/Engine/Source/Developer/TraceInsights/Private/Insights/ViewModels/TimingViewDrawHelper.h b/Engine/Source/Developer/TraceInsights/Private/Insights/ViewModels/TimingViewDrawHelper.h
--- a/Engine/Source/Developer/TraceInsights/Private/Insights/ViewModels/TimingViewDrawHelper.h
+++ b/Engine/Source/Developer/TraceInsights/Private/Insights/ViewModels/TimingViewDrawHelper.h
@@ -134,13 +134,14 @@ private:
 	const FSlateBrush* BackgroundAreaBrush;
 	const FLinearColor ValidAreaColor;
 	const FLinearColor InvalidAreaColor;
+	const FLinearColor EdgeColor;
 	const FSlateFontInfo EventFont;
 
 	//static const FSlateColorBrush SolidWhiteBrush = FSlateColorBrush(FColorList::White);
 	//static const FSlateBrush BorderBrush = FSlateBorderBrush(NAME_None, FMargin(1.0f));
 
-	mutable float ValidX0;
-	mutable float ValidX1;
+	mutable float ValidAreaX;
+	mutable float ValidAreaW;
 
 	float TimelineTopY;
 	float TimelineY;
diff --git a/Engine/Source/Developer/TraceInsights/Private/Insights/Widgets/SFrameTrack.cpp b/Engine/Source/Developer/TraceInsights/Private/Insights/Widgets/SFrameTrack.cpp
--- a/Engine/Source/Developer/TraceInsights/Private/Insights/Widgets/SFrameTrack.cpp
+++ b/Engine/Source/Developer/TraceInsights/Private/Insights/Widgets/SFrameTrack.cpp
@@ -54,9 +54,9 @@ void SFrameTrack::Reset()
 
 	CachedTimelines.Reset();
 	TimelinesOrder.Reset();
-	TimelinesOrder.Add(0);
-	TimelinesOrder.Add(1);
-	TimelinesOrder.Add(2);
+	TimelinesOrder.Add(TraceFrameType_Game);
+	TimelinesOrder.Add(TraceFrameType_Rendering);
+
 	bIsStateDirty = true;
 
 	bShowGameFrames = true;
@@ -294,7 +294,8 @@ FSampleRef SFrameTrack::GetSampleAtMousePosition(float X, float Y)
 		{
 			int32 TimelineId = TimelinesOrder[TimelineIndex];
 
-			if ((!bShowRenderingFrames && TimelineId == 1) || (!bShowGameFrames && TimelineId == 0))
+			if ((TimelineId == TraceFrameType_Rendering && !bShowRenderingFrames) ||
+				(TimelineId == TraceFrameType_Game && !bShowGameFrames))
 			{
 				continue;
 			}
@@ -387,7 +388,8 @@ int32 SFrameTrack::OnPaint(const FPaintArgs& Args, const FGeometry& AllottedGeom
 		{
 			int32 TimelineId = TimelinesOrder[TimelineIndex];
 
-			if ((!bShowRenderingFrames && TimelineId == 1) || (!bShowGameFrames && TimelineId == 0))
+			if ((TimelineId == TraceFrameType_Rendering && !bShowRenderingFrames) ||
+				(TimelineId == TraceFrameType_Game && !bShowGameFrames))
 			{
 				continue;
 			}
@@ -610,9 +612,6 @@ FReply SFrameTrack::OnMouseButtonUp(const FGeometry& MyGeometry, const FPointerE
 	{
 		if (bIsLMB_Pressed)
 		{
-			// Release the mouse.
-			Reply = FReply::Handled().ReleaseMouseCapture();
-
 			if (bIsScrolling)
 			{
 				bIsScrolling = false;
@@ -624,15 +623,15 @@ FReply SFrameTrack::OnMouseButtonUp(const FGeometry& MyGeometry, const FPointerE
 			}
 
 			bIsLMB_Pressed = false;
+
+			// Release the mouse.
+			Reply = FReply::Handled().ReleaseMouseCapture();
 		}
 	}
 	else if (MouseEvent.GetEffectingButton() == EKeys::RightMouseButton)
 	{
 		if (bIsRMB_Pressed)
 		{
-			// Release mouse as we no longer scroll.
-			Reply = FReply::Handled().ReleaseMouseCapture();
-
 			if (bIsScrolling)
 			{
 				bIsScrolling = false;
@@ -640,11 +639,13 @@ FReply SFrameTrack::OnMouseButtonUp(const FGeometry& MyGeometry, const FPointerE
 			}
 			else if (bIsValidForMouseClick)
 			{
-				ShowContextMenu(MouseEvent.GetScreenSpacePosition());
-				Reply = FReply::Handled();
+				ShowContextMenu(MouseEvent);
 			}
 
 			bIsRMB_Pressed = false;
+
+			// Release mouse as we no longer scroll.
+			Reply = FReply::Handled().ReleaseMouseCapture();
 		}
 	}
 
@@ -852,7 +853,7 @@ FCursorReply SFrameTrack::OnCursorQuery(const FGeometry& MyGeometry, const FPoin
 
 ////////////////////////////////////////////////////////////////////////////////////////////////////
 
-void SFrameTrack::ShowContextMenu(const FVector2D& ScreenSpacePosition)
+void SFrameTrack::ShowContextMenu(const FPointerEvent& MouseEvent)
 {
 	const bool bShouldCloseWindowAfterMenuSelection = true;
 	FMenuBuilder MenuBuilder(bShouldCloseWindowAfterMenuSelection, NULL);
@@ -923,7 +924,8 @@ void SFrameTrack::ShowContextMenu(const FVector2D& ScreenSpacePosition)
 
 	TSharedRef<SWidget> MenuWidget = MenuBuilder.MakeWidget();
 
-	FWidgetPath EventPath;
+	FWidgetPath EventPath = MouseEvent.GetEventPath() != nullptr ? *MouseEvent.GetEventPath() : FWidgetPath();
+	const FVector2D ScreenSpacePosition = MouseEvent.GetScreenSpacePosition();
 	FSlateApplication::Get().PushMenu(SharedThis(this), EventPath, MenuWidget, ScreenSpacePosition, FPopupTransitionEffect::ContextMenu);
 }
 
diff --git a/Engine/Source/Developer/TraceInsights/Private/Insights/Widgets/SFrameTrack.h b/Engine/Source/Developer/TraceInsights/Private/Insights/Widgets/SFrameTrack.h
--- a/Engine/Source/Developer/TraceInsights/Private/Insights/Widgets/SFrameTrack.h
+++ b/Engine/Source/Developer/TraceInsights/Private/Insights/Widgets/SFrameTrack.h
@@ -39,11 +39,9 @@ struct FSampleRef
  */
 class SFrameTrack : public SCompoundWidget
 {
-	enum
-	{
-		/** Number of pixels. */
-		MOUSE_SNAP_DISTANCE = 4,
-	};
+public:
+	/** Number of pixels. */
+	static constexpr float MOUSE_SNAP_DISTANCE = 2.0f;
 
 	enum class EFrameTrackCursor
 	{
@@ -102,7 +100,7 @@ protected:
 	FSampleRef GetSampleAtMousePosition(float X, float Y);
 	void SelectFrameAtMousePosition(float X, float Y);
 
-	void ShowContextMenu(const FVector2D& ScreenSpacePosition);
+	void ShowContextMenu(const FPointerEvent& MouseEvent);
 
 	void ContextMenu_ShowGameFrames_Execute();
 	bool ContextMenu_ShowGameFrames_CanExecute();
@@ -163,7 +161,7 @@ protected:
 	TSharedPtr<SScrollBar> HorizontalScrollBar;
 
 	//////////////////////////////////////////////////
-	// Panning, Zooming and Selection behaviors
+	// Panning and Zooming behaviors
 
 	/** The current mouse position. */
 	FVector2D MousePosition;
@@ -183,6 +181,7 @@ protected:
 	bool bIsScrolling;
 
 	//////////////////////////////////////////////////
+	// Selection
 
 	int32 SelectionStartFrameIndex;
 	int32 SelectionEndFrameIndex;
diff --git a/Engine/Source/Developer/TraceInsights/Private/Insights/Widgets/SGraphTrack.cpp b/Engine/Source/Developer/TraceInsights/Private/Insights/Widgets/SGraphTrack.cpp
--- a/Engine/Source/Developer/TraceInsights/Private/Insights/Widgets/SGraphTrack.cpp
+++ b/Engine/Source/Developer/TraceInsights/Private/Insights/Widgets/SGraphTrack.cpp
@@ -26,7 +26,9 @@
 
 // Insights
 #include "Insights/Common/PaintUtils.h"
+#include "Insights/Common/TimeUtils.h"
 #include "Insights/TimingProfilerManager.h"
+#include "Insights/ViewModels/DrawHelpers.h"
 
 ////////////////////////////////////////////////////////////////////////////////////////////////////
 
@@ -39,6 +41,8 @@
 SGraphTrack::SGraphTrack()
 	: TimeRulerTrack(0)
 	, GraphTrack(1)
+	, WhiteBrush(FCoreStyle::Get().GetBrush("WhiteBrush"))
+	, MainFont(FCoreStyle::GetDefaultFontStyle("Regular", 8))
 {
 	Reset();
 }
@@ -55,7 +59,33 @@ void SGraphTrack::Reset()
 {
 	TimeRulerTrack.Reset();
 	GraphTrack.Reset();
-	Viewport.ScaleX = (5 * 20) / 0.1; // 100ms between major tick marks
+
+	Viewport.Reset();
+	Viewport.MaxValidTime = 84.0 * 60.0;
+	//Viewport.ScaleX = (5 * 20) / 0.1; // 100ms between major tick marks
+
+	bIsViewportDirty = true;
+	bIsVerticalViewportDirty = true;
+
+	MousePosition = FVector2D::ZeroVector;
+
+	MousePositionOnButtonDown = FVector2D::ZeroVector;
+	ViewportStartTimeOnButtonDown = 0.0;
+	ViewportScrollPosYOnButtonDown = 0.0f;
+
+	MousePositionOnButtonUp = FVector2D::ZeroVector;
+
+	bIsLMB_Pressed = false;
+	bIsRMB_Pressed = false;
+
+	bIsDragging = false;
+
+	bIsPanning = false;
+	PanningMode = EPanningMode::None;
+
+	bIsSelecting = false;
+	SelectionStartTime = 0.0;
+	SelectionEndTime = 0.0;
 }
 
 ////////////////////////////////////////////////////////////////////////////////////////////////////
@@ -77,12 +107,10 @@ void SGraphTrack::Tick(const FGeometry& AllottedGeometry, const double InCurrent
 
 	bool bIsGraphDirty = false;
 
-	if (Viewport.UpdateSize(TrackWidth, TrackHeight))
-	{
-		bIsGraphDirty = true;
-	}
-
-	if (GraphTrack.GetHeight() != TrackHeight)
+	if (Viewport.UpdateSize(TrackWidth, TrackHeight) ||
+		bIsViewportDirty ||
+		bIsVerticalViewportDirty ||
+		GraphTrack.GetHeight() != TrackHeight)
 	{
 		bIsGraphDirty = true;
 	}
@@ -106,7 +134,55 @@ int32 SGraphTrack::OnPaint(const FPaintArgs& Args, const FGeometry& AllottedGeom
 	FDrawContext DrawContext(AllottedGeometry, MyCullingRect, InWidgetStyle, DrawEffects, OutDrawElements, LayerId);
 
 	GraphTrack.Draw(DrawContext, Viewport);
-	TimeRulerTrack.Draw(DrawContext, Viewport);
+
+	TimeRulerTrack.Draw(DrawContext, Viewport, MousePosition, bIsSelecting, SelectionStartTime, SelectionEndTime);
+	DrawContext.DrawBox(0.0f, TimeRulerTrack.GetHeight(), Viewport.Width, 1.0f, WhiteBrush, FLinearColor(0.05f, 0.05f, 0.05f, 1.0f));
+
+	FDrawHelpers::DrawTimeRangeSelection(DrawContext, Viewport, SelectionStartTime, SelectionEndTime, WhiteBrush, MainFont);
+
+	const bool bShouldDisplayDebugInfo = FInsightsManager::Get()->IsDebugInfoEnabled();
+	if (bShouldDisplayDebugInfo)
+	{
+		const FSlateFontInfo& SummaryFont = MainFont;
+
+		const TSharedRef<FSlateFontMeasure> FontMeasureService = FSlateApplication::Get().GetRenderer()->GetFontMeasureService();
+		const float MaxFontCharHeight = FontMeasureService->Measure(TEXT("!"), SummaryFont).Y;
+		const float DbgDY = MaxFontCharHeight;
+
+		const float DbgW = 320.0f;
+		const float DbgH = DbgDY * 2 + 3.0f;
+		const float DbgX = Viewport.Width - DbgW - 20.0f;
+		float DbgY = Viewport.TopOffset + 10.0f;
+
+		DrawContext.LayerId++;
+
+		DrawContext.DrawBox(DbgX - 2.0f, DbgY - 2.0f, DbgW, DbgH, WhiteBrush, FLinearColor(1.0f, 1.0f, 1.0f, 0.9f));
+		DrawContext.LayerId++;
+
+		FLinearColor DbgTextColor(0.0f, 0.0f, 0.0f, 0.9f);
+
+		//////////////////////////////////////////////////
+		// Display viewport's horizontal info.
+
+		DrawContext.DrawText
+		(
+			DbgX, DbgY,
+			FString::Printf(TEXT("SX: %g, ST: %g, ET: %s"), Viewport.ScaleX, Viewport.StartTime, *TimeUtils::FormatTimeAuto(Viewport.MaxValidTime)),
+			SummaryFont, DbgTextColor
+		);
+		DbgY += DbgDY;
+
+		//////////////////////////////////////////////////
+		// Display viewport's vertical info.
+
+		DrawContext.DrawText
+		(
+			DbgX, DbgY,
+			FString::Printf(TEXT("Y: %.2f, H: %g, VH: %g"), Viewport.ScrollPosY, Viewport.ScrollHeight, Viewport.Height),
+			SummaryFont, DbgTextColor
+		);
+		DbgY += DbgDY;
+	}
 
 	return SCompoundWidget::OnPaint(Args, AllottedGeometry, MyCullingRect, OutDrawElements, LayerId, InWidgetStyle, bParentEnabled && IsEnabled());
 }
@@ -115,48 +191,293 @@ int32 SGraphTrack::OnPaint(const FPaintArgs& Args, const FGeometry& AllottedGeom
 
 FReply SGraphTrack::OnMouseButtonDown(const FGeometry& MyGeometry, const FPointerEvent& MouseEvent)
 {
-	//...
+	FReply Reply = FReply::Unhandled();
 
-	return FReply::Unhandled();
+	MousePositionOnButtonDown = MyGeometry.AbsoluteToLocal(MouseEvent.GetScreenSpacePosition());
+	MousePosition = MousePositionOnButtonDown;
+
+	bool bStartPanning = false;
+	bool bStartSelecting = false;
+
+	if (MouseEvent.GetEffectingButton() == EKeys::LeftMouseButton)
+	{
+		if (!bIsRMB_Pressed)
+		{
+			bIsLMB_Pressed = true;
+
+			if ((MousePositionOnButtonDown.Y < TimeRulerTrack.GetHeight() ||
+				(MouseEvent.GetModifierKeys().IsControlDown() && MouseEvent.GetModifierKeys().IsShiftDown())))
+			{
+				bStartSelecting = true;
+			}
+			else
+			{
+				bStartPanning = true;
+			}
+
+			// Capture mouse, so we can drag outside this widget.
+			Reply = FReply::Handled().CaptureMouse(SharedThis(this));
+		}
+	}
+	else if (MouseEvent.GetEffectingButton() == EKeys::RightMouseButton)
+	{
+		if (!bIsLMB_Pressed)
+		{
+			bIsRMB_Pressed = true;
+
+			if ((MousePositionOnButtonDown.Y < TimeRulerTrack.GetHeight() ||
+				(MouseEvent.GetModifierKeys().IsControlDown() && MouseEvent.GetModifierKeys().IsShiftDown())))
+			{
+				bStartSelecting = true;
+			}
+			else
+			{
+				bStartPanning = true;
+			}
+
+			// Capture mouse, so we can drag outside this widget.
+			Reply = FReply::Handled().CaptureMouse(SharedThis(this));
+		}
+	}
+
+	if (bStartPanning)
+	{
+		bIsPanning = true;
+		bIsDragging = false;
+
+		ViewportStartTimeOnButtonDown = Viewport.StartTime;
+		ViewportScrollPosYOnButtonDown = Viewport.ScrollPosY;
+
+		if (MouseEvent.GetModifierKeys().IsControlDown())
+		{
+			// Allow panning only horizontally.
+			PanningMode = EPanningMode::Horizontal;
+		}
+		else if (MouseEvent.GetModifierKeys().IsShiftDown())
+		{
+			// Allow panning only vertically.
+			PanningMode = EPanningMode::Vertical;
+		}
+		else
+		{
+			// Allow panning both horizontally and vertically.
+			PanningMode = EPanningMode::HorizontalAndVertical;
+		}
+	}
+	else if (bStartSelecting)
+	{
+		bIsSelecting = true;
+		bIsDragging = false;
+
+		SelectionStartTime = Viewport.SlateUnitsToTime(MousePositionOnButtonDown.X);
+		SelectionEndTime = SelectionStartTime;
+		//TODO: SelectionChangingEvent.Broadcast(SelectionStartTime, SelectionEndTime);
+	}
+
+	return Reply;
 }
 
 ////////////////////////////////////////////////////////////////////////////////////////////////////
 
 FReply SGraphTrack::OnMouseButtonUp(const FGeometry& MyGeometry, const FPointerEvent& MouseEvent)
 {
-	//...
+	FReply Reply = FReply::Unhandled();
 
-	return FReply::Unhandled();
+	MousePositionOnButtonUp = MyGeometry.AbsoluteToLocal(MouseEvent.GetScreenSpacePosition());
+	MousePosition = MousePositionOnButtonUp;
+
+	const bool bIsValidForMouseClick = MousePositionOnButtonUp.Equals(MousePositionOnButtonDown, MOUSE_SNAP_DISTANCE);
+
+	if (MouseEvent.GetEffectingButton() == EKeys::LeftMouseButton)
+	{
+		if (bIsLMB_Pressed)
+		{
+			if (bIsPanning)
+			{
+				PanningMode = EPanningMode::None;
+
+				bIsPanning = false;
+			}
+			else if (bIsSelecting)
+			{
+				//TODO: SelectionChangedEvent.Broadcast(SelectionStartTime, SelectionEndTime);
+
+				bIsSelecting = false;
+			}
+
+			if (bIsValidForMouseClick)
+			{
+				// Select single event...
+
+				// When clicking on an empty space...
+				//if (!SelectedTimingEvent.IsValid())
+				{
+					// ...reset selection.
+					SelectionEndTime = SelectionStartTime = 0.0;
+					//TODO: SelectionChangedEvent.Broadcast(SelectionStartTime, SelectionEndTime);
+				}
+			}
+
+			bIsDragging = false;
+
+			// Release mouse as we no longer drag.
+			Reply = FReply::Handled().ReleaseMouseCapture();
+
+			bIsLMB_Pressed = false;
+		}
+	}
+	else if (MouseEvent.GetEffectingButton() == EKeys::RightMouseButton)
+	{
+		if (bIsRMB_Pressed)
+		{
+			if (bIsPanning)
+			{
+				PanningMode = EPanningMode::None;
+
+				bIsPanning = false;
+			}
+			else if (bIsSelecting)
+			{
+				//TODO: SelectionChangedEvent.Broadcast(SelectionStartTime, SelectionEndTime);
+
+				bIsSelecting = false;
+			}
+
+			if (bIsValidForMouseClick)
+			{
+				ShowContextMenu(MouseEvent);
+			}
+
+			bIsDragging = false;
+
+			// Release mouse as we no longer drag.
+			Reply = FReply::Handled().ReleaseMouseCapture();
+
+			bIsRMB_Pressed = false;
+		}
+	}
+
+	return Reply;
 }
 
 ////////////////////////////////////////////////////////////////////////////////////////////////////
 
 FReply SGraphTrack::OnMouseMove(const FGeometry& MyGeometry, const FPointerEvent& MouseEvent)
 {
-	//...
+	FReply Reply = FReply::Unhandled();
 
-	return FReply::Unhandled();
+	MousePosition = MyGeometry.AbsoluteToLocal(MouseEvent.GetScreenSpacePosition());
+
+	if (bIsPanning)
+	{
+		if (HasMouseCapture() && !MouseEvent.GetCursorDelta().IsZero())
+		{
+			bIsDragging = true;
+
+			if ((int32)PanningMode & (int32)EPanningMode::Horizontal)
+			{
+				const double StartTime = ViewportStartTimeOnButtonDown + static_cast<double>(MousePositionOnButtonDown.X - MousePosition.X) / Viewport.ScaleX;
+				ScrollAtTime(StartTime);
+			}
+
+			if ((int32)PanningMode & (int32)EPanningMode::Vertical)
+			{
+				const float ScrollPosY = ViewportScrollPosYOnButtonDown + (MousePositionOnButtonDown.Y - MousePosition.Y);
+				ScrollAtPosY(ScrollPosY);
+			}
+		}
+
+		Reply = FReply::Handled();
+	}
+	else if (bIsSelecting)
+	{
+		if (HasMouseCapture() && !MouseEvent.GetCursorDelta().IsZero())
+		{
+			bIsDragging = true;
+
+			SelectionStartTime = Viewport.SlateUnitsToTime(MousePositionOnButtonDown.X);
+			SelectionEndTime = Viewport.SlateUnitsToTime(MousePosition.X);
+			if (SelectionStartTime > SelectionEndTime)
+			{
+				double Temp = SelectionStartTime;
+				SelectionStartTime = SelectionEndTime;
+				SelectionEndTime = Temp;
+			}
+			//TODO: SelectionChangingEvent.Broadcast(SelectionStartTime, SelectionEndTime);
+		}
+
+		Reply = FReply::Handled();
+	}
+
+	return Reply;
 }
 
 ////////////////////////////////////////////////////////////////////////////////////////////////////
 
 void SGraphTrack::OnMouseEnter(const FGeometry& MyGeometry, const FPointerEvent& MouseEvent)
 {
-	//...
 }
 
 ////////////////////////////////////////////////////////////////////////////////////////////////////
 
 void SGraphTrack::OnMouseLeave(const FPointerEvent& MouseEvent)
 {
-	//...
+	if (!HasMouseCapture())
+	{
+		// No longer dragging (unless we have mouse capture).
+		bIsDragging = false;
+		bIsPanning = false;
+		bIsSelecting = false;
+
+		bIsLMB_Pressed = false;
+		bIsRMB_Pressed = false;
+
+		MousePosition = FVector2D::ZeroVector;
+	}
 }
 
 ////////////////////////////////////////////////////////////////////////////////////////////////////
 
 FReply SGraphTrack::OnMouseWheel(const FGeometry& MyGeometry, const FPointerEvent& MouseEvent)
 {
-	//...
+	if (MouseEvent.GetModifierKeys().IsShiftDown())
+	{
+		// Scroll vertically.
+		constexpr float ScrollSpeedY = 16.0f * 3;
+		const float ScrollPosY = Viewport.ScrollPosY - ScrollSpeedY * MouseEvent.GetWheelDelta();
+		ScrollAtPosY(ScrollPosY);
+	}
+	else if (MouseEvent.GetModifierKeys().IsControlDown())
+	{
+		// Scroll horizontally.
+		const double ScrollSpeedX = Viewport.GetDurationForViewportDX(16.0 * 3);
+		ScrollAtTime(Viewport.StartTime - ScrollSpeedX * MouseEvent.GetWheelDelta());
+	}
+	else
+	{
+		// Zoom in/out horizontally.
+		const double Delta = MouseEvent.GetWheelDelta();
+		constexpr double ZoomStep = 0.25; // as percent
+		double ScaleX;
+
+		if (Delta > 0)
+		{
+			ScaleX = Viewport.ScaleX * FMath::Pow(1.0 + ZoomStep, Delta);
+		}
+		else
+		{
+			ScaleX = Viewport.ScaleX * FMath::Pow(1.0 / (1.0 + ZoomStep), -Delta);
+		}
+
+		//UE_LOG(TimingProfiler, Log, TEXT("%.2f, %.2f, %.2f"), Delta, Viewport.ScaleX, ScaleX);
+		MousePosition = MyGeometry.AbsoluteToLocal(MouseEvent.GetScreenSpacePosition());
+
+		if (Viewport.ZoomWithFixedX(ScaleX, MousePosition.X))
+		{
+			//UpdateHorizontalScrollBar();
+			bIsViewportDirty = true;
+		}
+	}
 
 	return FReply::Handled();
 }
@@ -210,11 +531,252 @@ FReply SGraphTrack::OnDrop(const FGeometry& MyGeometry, const FDragDropEvent& Dr
 
 FCursorReply SGraphTrack::OnCursorQuery(const FGeometry& MyGeometry, const FPointerEvent& CursorEvent) const
 {
-	//...
+	if (bIsPanning)
+	{
+		if (bIsDragging)
+		{
+			//return FCursorReply::Cursor(EMouseCursor::GrabHandClosed);
+			return FCursorReply::Cursor(EMouseCursor::GrabHand);
+		}
+	}
+	else if (bIsSelecting)
+	{
+		if (bIsDragging)
+		{
+			return FCursorReply::Cursor(EMouseCursor::ResizeLeftRight);
+		}
+	}
 
 	return FCursorReply::Unhandled();
 }
 
 ////////////////////////////////////////////////////////////////////////////////////////////////////
 
+void SGraphTrack::ScrollAtPosY(float ScrollPosY)
+{
+	if (Viewport.ScrollPosY != ScrollPosY)
+	{
+		Viewport.ScrollPosY = ScrollPosY;
+
+		//UpdateVerticalScrollBar();
+		bIsVerticalViewportDirty = true;
+	}
+}
+
+////////////////////////////////////////////////////////////////////////////////////////////////////
+
+void SGraphTrack::ScrollAtTime(double StartTime)
+{
+	if (Viewport.ScrollAtTime(StartTime))
+	{
+		//UpdateHorizontalScrollBar();
+		bIsViewportDirty = true;
+	}
+}
+
+////////////////////////////////////////////////////////////////////////////////////////////////////
+
+void SGraphTrack::ShowContextMenu(const FPointerEvent& MouseEvent)
+{
+	const bool bShouldCloseWindowAfterMenuSelection = true;
+	FMenuBuilder MenuBuilder(bShouldCloseWindowAfterMenuSelection, NULL);
+
+	MenuBuilder.BeginSection("Misc");
+	{
+		FUIAction Action_ShowPoints
+		(
+			FExecuteAction::CreateSP(this, &SGraphTrack::ContextMenu_ShowPoints_Execute),
+			FCanExecuteAction::CreateSP(this, &SGraphTrack::ContextMenu_ShowPoints_CanExecute),
+			FIsActionChecked::CreateSP(this, &SGraphTrack::ContextMenu_ShowPoints_IsChecked)
+		);
+		MenuBuilder.AddMenuEntry
+		(
+			LOCTEXT("ContextMenu_ShowPoints", "Points"),
+			LOCTEXT("ContextMenu_ShowPoints_Desc", "Show points."),
+			FSlateIcon(),
+			Action_ShowPoints,
+			NAME_None,
+			EUserInterfaceActionType::ToggleButton
+		);
+
+		FUIAction Action_ShowPointsWithBorder
+		(
+			FExecuteAction::CreateSP(this, &SGraphTrack::ContextMenu_ShowPointsWithBorder_Execute),
+			FCanExecuteAction::CreateSP(this, &SGraphTrack::ContextMenu_ShowPointsWithBorder_CanExecute),
+			FIsActionChecked::CreateSP(this, &SGraphTrack::ContextMenu_ShowPointsWithBorder_IsChecked)
+		);
+		MenuBuilder.AddMenuEntry
+		(
+			LOCTEXT("ContextMenu_ShowPointsWithBorder", "Points have Border"),
+			LOCTEXT("ContextMenu_ShowPointsWithBorder_Desc", "Show border around points."),
+			FSlateIcon(),
+			Action_ShowPointsWithBorder,
+			NAME_None,
+			EUserInterfaceActionType::ToggleButton
+		);
+
+		FUIAction Action_ShowLines
+		(
+			FExecuteAction::CreateSP(this, &SGraphTrack::ContextMenu_ShowLines_Execute),
+			FCanExecuteAction::CreateSP(this, &SGraphTrack::ContextMenu_ShowLines_CanExecute),
+			FIsActionChecked::CreateSP(this, &SGraphTrack::ContextMenu_ShowLines_IsChecked)
+		);
+		MenuBuilder.AddMenuEntry
+		(
+			LOCTEXT("ContextMenu_ShowLines", "Lines"),
+			LOCTEXT("ContextMenu_ShowLines_Desc", "Show connected lines. Each event is a single point in time."),
+			FSlateIcon(),
+			Action_ShowLines,
+			NAME_None,
+			EUserInterfaceActionType::ToggleButton
+		);
+
+		FUIAction Action_ShowLinesWithDuration
+		(
+			FExecuteAction::CreateSP(this, &SGraphTrack::ContextMenu_ShowLinesWithDuration_Execute),
+			FCanExecuteAction::CreateSP(this, &SGraphTrack::ContextMenu_ShowLinesWithDuration_CanExecute),
+			FIsActionChecked::CreateSP(this, &SGraphTrack::ContextMenu_ShowLinesWithDuration_IsChecked)
+		);
+		MenuBuilder.AddMenuEntry
+		(
+			LOCTEXT("ContextMenu_ShowLinesWithDuration", "Lines have Duration"),
+			LOCTEXT("ContextMenu_ShowLinesWithDuration_Desc", "Show connected lines. Each event is a horizontal line (duration of the timing event)."),
+			FSlateIcon(),
+			Action_ShowLinesWithDuration,
+			NAME_None,
+			EUserInterfaceActionType::ToggleButton
+		);
+
+		FUIAction Action_ShowBars
+		(
+			FExecuteAction::CreateSP(this, &SGraphTrack::ContextMenu_ShowBars_Execute),
+			FCanExecuteAction::CreateSP(this, &SGraphTrack::ContextMenu_ShowBars_CanExecute),
+			FIsActionChecked::CreateSP(this, &SGraphTrack::ContextMenu_ShowBars_IsChecked)
+		);
+		MenuBuilder.AddMenuEntry
+		(
+			LOCTEXT("ContextMenu_ShowBars", "Bars"),
+			LOCTEXT("ContextMenu_ShowBars_Desc", "Show bars."),
+			FSlateIcon(),
+			Action_ShowBars,
+			NAME_None,
+			EUserInterfaceActionType::ToggleButton
+		);
+	}
+	MenuBuilder.EndSection();
+
+	TSharedRef<SWidget> MenuWidget = MenuBuilder.MakeWidget();
+
+	FWidgetPath EventPath = MouseEvent.GetEventPath() != nullptr ? *MouseEvent.GetEventPath() : FWidgetPath();
+	const FVector2D ScreenSpacePosition = MouseEvent.GetScreenSpacePosition();
+	FSlateApplication::Get().PushMenu(SharedThis(this), EventPath, MenuWidget, ScreenSpacePosition, FPopupTransitionEffect::ContextMenu);
+}
+
+////////////////////////////////////////////////////////////////////////////////////////////////////
+
+void SGraphTrack::ContextMenu_ShowPoints_Execute()
+{
+	GraphTrack.bDrawPoints = !GraphTrack.bDrawPoints;
+}
+
+////////////////////////////////////////////////////////////////////////////////////////////////////
+
+bool SGraphTrack::ContextMenu_ShowPoints_CanExecute()
+{
+	return true;
+}
+
+////////////////////////////////////////////////////////////////////////////////////////////////////
+
+bool SGraphTrack::ContextMenu_ShowPoints_IsChecked()
+{
+	return GraphTrack.bDrawPoints;
+}
+
+////////////////////////////////////////////////////////////////////////////////////////////////////
+
+void SGraphTrack::ContextMenu_ShowPointsWithBorder_Execute()
+{
+	GraphTrack.bDrawPointsWithBorder = !GraphTrack.bDrawPointsWithBorder;
+}
+
+////////////////////////////////////////////////////////////////////////////////////////////////////
+
+bool SGraphTrack::ContextMenu_ShowPointsWithBorder_CanExecute()
+{
+	return GraphTrack.bDrawPoints;
+}
+
+////////////////////////////////////////////////////////////////////////////////////////////////////
+
+bool SGraphTrack::ContextMenu_ShowPointsWithBorder_IsChecked()
+{
+	return GraphTrack.bDrawPointsWithBorder;
+}
+
+////////////////////////////////////////////////////////////////////////////////////////////////////
+
+void SGraphTrack::ContextMenu_ShowLines_Execute()
+{
+	GraphTrack.bDrawLines = !GraphTrack.bDrawLines;
+}
+
+////////////////////////////////////////////////////////////////////////////////////////////////////
+
+bool SGraphTrack::ContextMenu_ShowLines_CanExecute()
+{
+	return true;
+}
+
+////////////////////////////////////////////////////////////////////////////////////////////////////
+
+bool SGraphTrack::ContextMenu_ShowLines_IsChecked()
+{
+	return GraphTrack.bDrawLines;
+}
+
+////////////////////////////////////////////////////////////////////////////////////////////////////
+
+void SGraphTrack::ContextMenu_ShowLinesWithDuration_Execute()
+{
+	GraphTrack.bDrawLinesWithDuration = !GraphTrack.bDrawLinesWithDuration;
+}
+
+////////////////////////////////////////////////////////////////////////////////////////////////////
+
+bool SGraphTrack::ContextMenu_ShowLinesWithDuration_CanExecute()
+{
+	return GraphTrack.bDrawLines;
+}
+
+////////////////////////////////////////////////////////////////////////////////////////////////////
+
+bool SGraphTrack::ContextMenu_ShowLinesWithDuration_IsChecked()
+{
+	return GraphTrack.bDrawLinesWithDuration;
+}
+
+////////////////////////////////////////////////////////////////////////////////////////////////////
+
+void SGraphTrack::ContextMenu_ShowBars_Execute()
+{
+	GraphTrack.bDrawBoxes = !GraphTrack.bDrawBoxes;
+}
+
+////////////////////////////////////////////////////////////////////////////////////////////////////
+
+bool SGraphTrack::ContextMenu_ShowBars_CanExecute()
+{
+	return true;
+}
+
+////////////////////////////////////////////////////////////////////////////////////////////////////
+
+bool SGraphTrack::ContextMenu_ShowBars_IsChecked()
+{
+	return GraphTrack.bDrawBoxes;
+}
+
+////////////////////////////////////////////////////////////////////////////////////////////////////
+
 #undef LOCTEXT_NAMESPACE
diff --git a/Engine/Source/Developer/TraceInsights/Private/Insights/Widgets/SGraphTrack.h b/Engine/Source/Developer/TraceInsights/Private/Insights/Widgets/SGraphTrack.h
--- a/Engine/Source/Developer/TraceInsights/Private/Insights/Widgets/SGraphTrack.h
+++ b/Engine/Source/Developer/TraceInsights/Private/Insights/Widgets/SGraphTrack.h
@@ -21,6 +21,9 @@ class FSlateWindowElementList;
 /** A custom widget used to display graphs. */
 class SGraphTrack : public SCompoundWidget
 {
+public:
+	static constexpr float MOUSE_SNAP_DISTANCE = 2.0f;
+
 public:
 	/** Default constructor. */
 	SGraphTrack();
@@ -180,8 +183,89 @@ protected:
 
 	void UpdateGraphTrack();
 
+	void ScrollAtPosY(float ScrollPosY);
+	void ScrollAtTime(double StartTime);
+
+	void ShowContextMenu(const FPointerEvent& MouseEvent);
+
+	void ContextMenu_ShowPoints_Execute();
+	bool ContextMenu_ShowPoints_CanExecute();
+	bool ContextMenu_ShowPoints_IsChecked();
+
+	void ContextMenu_ShowPointsWithBorder_Execute();
+	bool ContextMenu_ShowPointsWithBorder_CanExecute();
+	bool ContextMenu_ShowPointsWithBorder_IsChecked();
+
+	void ContextMenu_ShowLines_Execute();
+	bool ContextMenu_ShowLines_CanExecute();
+	bool ContextMenu_ShowLines_IsChecked();
+
+	void ContextMenu_ShowLinesWithDuration_Execute();
+	bool ContextMenu_ShowLinesWithDuration_CanExecute();
+	bool ContextMenu_ShowLinesWithDuration_IsChecked();
+
+	void ContextMenu_ShowBars_Execute();
+	bool ContextMenu_ShowBars_CanExecute();
+	bool ContextMenu_ShowBars_IsChecked();
+
 protected:
 	FTimeRulerTrack TimeRulerTrack;
 	FRandomGraphTrack GraphTrack;
+	//FFramesGraphTrack GraphTrack;
 	FTimingTrackViewport Viewport;
+
+	bool bIsViewportDirty;
+	bool bIsVerticalViewportDirty;
+
+	////////////////////////////////////////////////////////////
+
+	/** The current mouse position. */
+	FVector2D MousePosition;
+
+	/** Mouse position during the call on mouse button down. */
+	FVector2D MousePositionOnButtonDown;
+	double ViewportStartTimeOnButtonDown;
+	float ViewportScrollPosYOnButtonDown;
+
+	/** Mouse position during the call on mouse button up. */
+	FVector2D MousePositionOnButtonUp;
+
+	bool bIsLMB_Pressed;
+	bool bIsRMB_Pressed;
+
+	bool bIsDragging;
+
+	////////////////////////////////////////////////////////////
+	// Panning
+
+	/**
+	 * True, if the user is currently interactively panning the view (horizontally and/or vertically),
+	 * either by holding the right mouse button and dragging
+	 * or by holding spacebar pressed and dragging with left mouse button.
+	 */
+	bool bIsPanning;
+
+	/** How to pan. */
+	enum class EPanningMode : uint8
+	{
+		None = 0,
+		Horizontal = 0x01,
+		Vertical = 0x02,
+		HorizontalAndVertical = Horizontal | Vertical,
+	};
+	EPanningMode PanningMode;
+
+	////////////////////////////////////////////////////////////
+	// Selection
+
+	/** True, if the user is currently changing the selection (by holding the left mouse button and dragging). */
+	bool bIsSelecting;
+
+	double SelectionStartTime;
+	double SelectionEndTime;
+
+	////////////////////////////////////////////////////////////
+
+	const FSlateBrush* WhiteBrush;
+	const FSlateFontInfo MainFont;
 };
diff --git a/Engine/Source/Developer/TraceInsights/Private/Insights/Widgets/STimingProfilerToolbar.cpp b/Engine/Source/Developer/TraceInsights/Private/Insights/Widgets/STimingProfilerToolbar.cpp
--- a/Engine/Source/Developer/TraceInsights/Private/Insights/Widgets/STimingProfilerToolbar.cpp
+++ b/Engine/Source/Developer/TraceInsights/Private/Insights/Widgets/STimingProfilerToolbar.cpp
@@ -46,7 +46,7 @@ void STimingProfilerToolbar::Construct(const FArguments& InArgs)
 			ToolbarBuilder.BeginSection("View");
 			{
 				ToolbarBuilder.AddToolBarButton(FTimingProfilerCommands::Get().ToggleFramesTrackVisibility);
-				//ToolbarBuilder.AddToolBarButton(FTimingProfilerCommands::Get().ToggleGraphTrackVisibility);
+				ToolbarBuilder.AddToolBarButton(FTimingProfilerCommands::Get().ToggleGraphTrackVisibility);
 				ToolbarBuilder.AddToolBarButton(FTimingProfilerCommands::Get().ToggleTimingViewVisibility);
 				ToolbarBuilder.AddToolBarButton(FTimingProfilerCommands::Get().ToggleTimersViewVisibility);
 				ToolbarBuilder.AddToolBarButton(FTimingProfilerCommands::Get().ToggleStatsCountersViewVisibility);
diff --git a/Engine/Source/Developer/TraceInsights/Private/Insights/Widgets/STimingProfilerWindow.cpp b/Engine/Source/Developer/TraceInsights/Private/Insights/Widgets/STimingProfilerWindow.cpp
--- a/Engine/Source/Developer/TraceInsights/Private/Insights/Widgets/STimingProfilerWindow.cpp
+++ b/Engine/Source/Developer/TraceInsights/Private/Insights/Widgets/STimingProfilerWindow.cpp
@@ -268,10 +268,10 @@ void STimingProfilerWindow::Construct(const FArguments& InArgs, const TSharedRef
 		.SetIcon(FSlateIcon(FInsightsStyle::GetStyleSetName(), "FramesTrack.Icon.Small"))
 		.SetGroup(AppMenuGroup);
 
-	//TabManager->RegisterTabSpawner(FTimingProfilerTabs::GraphTrackID, FOnSpawnTab::CreateRaw(this, &STimingProfilerWindow::SpawnTab_GraphTrack))
-	//	.SetDisplayName(LOCTEXT("GraphTrackTabTitle", "Graph"))
-	//	.SetIcon(FSlateIcon(FInsightsStyle::GetStyleSetName(), "GraphTrack.Icon.Small"))
-	//	.SetGroup(AppMenuGroup);
+	TabManager->RegisterTabSpawner(FTimingProfilerTabs::GraphTrackID, FOnSpawnTab::CreateRaw(this, &STimingProfilerWindow::SpawnTab_GraphTrack))
+		.SetDisplayName(LOCTEXT("GraphTrackTabTitle", "Graph"))
+		.SetIcon(FSlateIcon(FInsightsStyle::GetStyleSetName(), "GraphTrack.Icon.Small"))
+		.SetGroup(AppMenuGroup);
 
 	TabManager->RegisterTabSpawner(FTimingProfilerTabs::TimingViewID, FOnSpawnTab::CreateRaw(this, &STimingProfilerWindow::SpawnTab_TimingView))
 		.SetDisplayName(LOCTEXT("TimingViewTabTitle", "Timing View"))
diff --git a/Engine/Source/Developer/TraceInsights/Private/Insights/Widgets/STimingView.cpp b/Engine/Source/Developer/TraceInsights/Private/Insights/Widgets/STimingView.cpp
--- a/Engine/Source/Developer/TraceInsights/Private/Insights/Widgets/STimingView.cpp
+++ b/Engine/Source/Developer/TraceInsights/Private/Insights/Widgets/STimingView.cpp
@@ -37,6 +37,7 @@
 #include "Insights/TimingProfilerCommon.h"
 #include "Insights/TimingProfilerManager.h"
 #include "Insights/ViewModels/BaseTimingTrack.h"
+#include "Insights/ViewModels/DrawHelpers.h"
 #include "Insights/ViewModels/TimingViewDrawHelper.h"
 #include "Insights/Widgets/SStatsView.h"
 #include "Insights/Widgets/STimersView.h"
@@ -265,7 +266,6 @@ void STimingView::Reset()
 	PanningMode = EPanningMode::None;
 
 	bIsSelecting = false;
-
 	SelectionStartTime = 0.0;
 	SelectionEndTime = 0.0;
 
@@ -377,12 +377,7 @@ void STimingView::Tick(const FGeometry& AllottedGeometry, const double InCurrent
 			}
 		}
 
-		if (Viewport.ScrollPosY != ScrollPosY)
-		{
-			Viewport.ScrollPosY = ScrollPosY;
-			UpdateVerticalScrollBar();
-			bIsVerticalViewportDirty = true;
-		}
+		ScrollAtPosY(ScrollPosY);
 
 		//////////////////////////////////////////////////
 		// Elastic snap to horizontal time limits.
@@ -833,7 +828,7 @@ int32 STimingView::OnPaint(const FPaintArgs& Args, const FGeometry& AllottedGeom
 	//////////////////////////////////////////////////
 
 	// Draw the time range selection.
-	DrawTimeRangeSelection(DrawContext);
+	FDrawHelpers::DrawTimeRangeSelection(DrawContext, Viewport, SelectionStartTime, SelectionEndTime, WhiteBrush, MainFont);
 
 	//////////////////////////////////////////////////
 
@@ -2233,170 +2228,15 @@ void STimingView::DrawIoActivityTrack(FTimingViewDrawHelper& Helper, FTimingEven
 // end of "I/O - File Activity" prototype
 ////////////////////////////////////////////////////////////////////////////////////////////////////
 
-void STimingView::DrawTimeRangeSelection(FDrawContext& DrawContext) const
-{
-	if (SelectionEndTime > SelectionStartTime)
-	{
-		float SelectionX1 = Viewport.TimeToSlateUnitsRounded(SelectionStartTime);
-		float SelectionX2 = Viewport.TimeToSlateUnitsRounded(SelectionEndTime);
-
-		if (SelectionX1 <= Viewport.Width &&
-			SelectionX2 >= 0)
-		{
-			float ClipLeft = 0.0f;
-			float ClipRight = 0.0f;
-			if (SelectionX1 < 0.0f)
-			{
-				ClipLeft = -SelectionX1;
-				SelectionX1 = 0.0f;
-			}
-			if (SelectionX2 > Viewport.Width)
-			{
-				ClipRight = SelectionX2 - Viewport.Width;
-				SelectionX2 = Viewport.Width;
-			}
-
-			// Draw selection area.
-			DrawContext.DrawBox(SelectionX1, 0.0f, SelectionX2 - SelectionX1, Viewport.Height, WhiteBrush, FLinearColor(0.25f, 0.5f, 1.0f, 0.25f));
-			DrawContext.LayerId++;
-
-			FColor ArrowFillColor(32, 64, 128, 255);
-			FLinearColor ArrowColor(ArrowFillColor);
-
-			if (SelectionX1 > 0.0f)
-			{
-				// Draw left side (vertical line).
-				DrawContext.DrawBox(SelectionX1 - 1.0f, 0.0f, 1.0f, Viewport.Height, WhiteBrush, ArrowColor);
-			}
-
-			if (SelectionX2 < Viewport.Width)
-			{
-				// Draw right side (vertical line).
-				DrawContext.DrawBox(SelectionX2, 0.0f, 1.0f, Viewport.Height, WhiteBrush, ArrowColor);
-			}
-
-			DrawContext.LayerId++;
-
-			const float ArrowSize = 6.0f;
-			const float ArrowY = 6.0f;
-
-			if (SelectionX2 - SelectionX1 > 2 * ArrowSize)
-			{
-				// Draw horizontal line.
-				float HorizLineX1 = SelectionX1;
-				if (ClipLeft == 0.0f)
-				{
-					HorizLineX1 += 1.0f;
-				}
-				float HorizLineX2 = SelectionX2;
-				if (ClipRight == 0.0f)
-				{
-					HorizLineX2 -= 1.0f;
-				}
-				DrawContext.DrawBox(HorizLineX1, ArrowY - 1.0f, HorizLineX2 - HorizLineX1, 3.0f, WhiteBrush, ArrowColor);
-
-				if (ClipLeft < ArrowSize)
-				{
-					// Draw left arrow.
-					for (float AH = 0.0f; AH < ArrowSize; AH += 1.0f)
-					{
-						DrawContext.DrawBox(SelectionX1 - ClipLeft + AH, ArrowY - AH, 1.0f, 2.0f * AH + 1.0f, WhiteBrush, ArrowColor);
-					}
-				}
-
-				if (ClipRight < ArrowSize)
-				{
-					// Draw right arrow.
-					for (float AH = 0.0f; AH < ArrowSize; AH += 1.0f)
-					{
-						DrawContext.DrawBox(SelectionX2 + ClipRight - AH - 1.0f, ArrowY - AH, 1.0f, 2.0f * AH + 1.0f, WhiteBrush, ArrowColor);
-					}
-				}
-
-				DrawContext.LayerId++;
-
-#if 0
-				//im: This should be a more efficeint way top draw the arrows, but it renders them with artifacts (missing vertical lines; shader bug?)!
-
-				const FSlateBrush* MyBrush = WhiteBrush;
-				FSlateShaderResourceProxy* ResourceProxy = FSlateDataPayload::ResourceManager->GetShaderResource(*MyBrush);
-				FSlateResourceHandle ResourceHandle = FSlateApplication::Get().GetRenderer()->GetResourceHandle(*MyBrush);
-
-				FVector2D AtlasOffset = ResourceProxy ? ResourceProxy->StartUV : FVector2D(0.0f, 0.0f);
-				FVector2D AtlasUVSize = ResourceProxy ? ResourceProxy->SizeUV : FVector2D(1.0f, 1.0f);
-
-				const FVector2D Pos = AllottedGeometry.GetAbsolutePosition() + FVector2D(0.0f, 40.0f);
-				const float Scale = AllottedGeometry.Scale;
-
-				FSlateRenderTransform RenderTransform;
-
-				TArray<FSlateVertex> Verts;
-				Verts.Reserve(6);
-				Verts.Add(FSlateVertex::Make<ESlateVertexRounding::Disabled>(RenderTransform, Pos + FVector2D(0.5f + SelectionX1 + ArrowSize, 0.5f + ArrowY + ArrowSize) * Scale, AtlasOffset + FVector2D(0.0f, 1.0f) * AtlasUVSize, ArrowFillColor));
-				Verts.Add(FSlateVertex::Make<ESlateVertexRounding::Disabled>(RenderTransform, Pos + FVector2D(0.5f + SelectionX1,             0.5f + ArrowY            ) * Scale, AtlasOffset + FVector2D(1.0f, 0.5f) * AtlasUVSize, ArrowFillColor));
-				Verts.Add(FSlateVertex::Make<ESlateVertexRounding::Disabled>(RenderTransform, Pos + FVector2D(0.5f + SelectionX1 + ArrowSize, 0.5f + ArrowY - ArrowSize) * Scale, AtlasOffset + FVector2D(0.0f, 0.0f) * AtlasUVSize, ArrowFillColor));
-				Verts.Add(FSlateVertex::Make<ESlateVertexRounding::Disabled>(RenderTransform, Pos + FVector2D(0.5f + SelectionX2 - ArrowSize, 0.5f + ArrowY - ArrowSize) * Scale, AtlasOffset + FVector2D(0.0f, 0.0f) * AtlasUVSize, ArrowFillColor));
-				Verts.Add(FSlateVertex::Make<ESlateVertexRounding::Disabled>(RenderTransform, Pos + FVector2D(0.5f + SelectionX2,             0.5f + ArrowY            ) * Scale, AtlasOffset + FVector2D(1.0f, 0.5f) * AtlasUVSize, ArrowFillColor));
-				Verts.Add(FSlateVertex::Make<ESlateVertexRounding::Disabled>(RenderTransform, Pos + FVector2D(0.5f + SelectionX2 - ArrowSize, 0.5f + ArrowY + ArrowSize) * Scale, AtlasOffset + FVector2D(0.0f, 1.0f) * AtlasUVSize, ArrowFillColor));
-
-				TArray<SlateIndex> Indices;
-				Indices.Reserve(6);
-				if (ClipLeft < ArrowSize)
-				{
-					Indices.Add(0);
-					Indices.Add(1);
-					Indices.Add(2);
-				}
-				if (ClipRight < ArrowSize)
-				{
-					Indices.Add(3);
-					Indices.Add(4);
-					Indices.Add(5);
-				}
-
-				FSlateDrawElement::MakeCustomVerts(
-					OutDrawElements,
-					LayerId,
-					ResourceHandle,
-					Verts,
-					Indices,
-					nullptr,
-					0,
-					0,
-					ESlateDrawEffect::PreMultipliedAlpha);
-
-				DrawContext.LayerId++;
-#endif
-			}
-
-			//////////////////////////////////////////////////
-			// Draw duration for selected time interval.
-
-			double Duration = SelectionEndTime - SelectionStartTime;
-			FString Text = TimeUtils::FormatTimeAuto(Duration);
-
-			const TSharedRef<FSlateFontMeasure> FontMeasureService = FSlateApplication::Get().GetRenderer()->GetFontMeasureService();
-			const float TextWidth = FontMeasureService->Measure(Text, MainFont).X;
-
-			const float CenterX = (SelectionX1 + SelectionX2) / 2.0f;
-
-			DrawContext.DrawBox(CenterX - TextWidth / 2 - 2.0, ArrowY - 6.0f, TextWidth + 4.0f, 13.0f, WhiteBrush, ArrowColor);
-			DrawContext.LayerId++;
-
-			DrawContext.DrawText(CenterX - TextWidth / 2, ArrowY - 6.0f, Text, MainFont, FLinearColor::White);
-			DrawContext.LayerId++;
-		}
-	}
-}
-
-////////////////////////////////////////////////////////////////////////////////////////////////////
-
 FReply STimingView::OnMouseButtonDown(const FGeometry& MyGeometry, const FPointerEvent& MouseEvent)
 {
 	FReply Reply = FReply::Unhandled();
+
 	MousePositionOnButtonDown = MyGeometry.AbsoluteToLocal(MouseEvent.GetScreenSpacePosition());
+	MousePosition = MousePositionOnButtonDown;
 
 	bool bStartPanning = false;
+	bool bStartSelecting = false;
 
 	if (MouseEvent.GetEffectingButton() == EKeys::LeftMouseButton)
 	{
@@ -2409,21 +2249,15 @@ FReply STimingView::OnMouseButtonDown(const FGeometry& MyGeometry, const FPointe
 		{
 			bIsLMB_Pressed = true;
 
-			if (bIsSpaceBarKeyPressed || MousePositionOnButtonDown.Y > TimeRulerTrack.GetHeight())
+			if (!bIsSpaceBarKeyPressed &&
+				(MousePositionOnButtonDown.Y < TimeRulerTrack.GetHeight() ||
+				 (MouseEvent.GetModifierKeys().IsControlDown() && MouseEvent.GetModifierKeys().IsShiftDown())))
 			{
-				bStartPanning = true;
+				bStartSelecting = true;
 			}
 			else
 			{
-				bIsSelecting = true;
-				bIsDragging = false;
-
-				SelectionStartTime = Viewport.SlateUnitsToTime(MousePositionOnButtonDown.X);
-				SelectionEndTime = SelectionStartTime;
-				LastSelectionType = ESelectionType::None;
-				//TODO: SelectionChangingEvent.Broadcast(SelectionStartTime, SelectionEndTime);
-				EventAggregation.Reset();
-				ObjectTypeAggregation.Reset();
+				bStartPanning = true;
 			}
 
 			// Capture mouse, so we can drag outside this widget.
@@ -2436,7 +2270,16 @@ FReply STimingView::OnMouseButtonDown(const FGeometry& MyGeometry, const FPointe
 		{
 			bIsRMB_Pressed = true;
 
-			bStartPanning = true;
+			if (!bIsSpaceBarKeyPressed &&
+				(MousePositionOnButtonDown.Y < TimeRulerTrack.GetHeight() ||
+				(MouseEvent.GetModifierKeys().IsControlDown() && MouseEvent.GetModifierKeys().IsShiftDown())))
+			{
+				bStartSelecting = true;
+			}
+			else
+			{
+				bStartPanning = true;
+			}
 
 			// Capture mouse, so we can drag outside this widget.
 			Reply = FReply::Handled().CaptureMouse(SharedThis(this));
@@ -2467,6 +2310,18 @@ FReply STimingView::OnMouseButtonDown(const FGeometry& MyGeometry, const FPointe
 			PanningMode = EPanningMode::HorizontalAndVertical;
 		}
 	}
+	else if (bStartSelecting)
+	{
+		bIsSelecting = true;
+		bIsDragging = false;
+
+		SelectionStartTime = Viewport.SlateUnitsToTime(MousePositionOnButtonDown.X);
+		SelectionEndTime = SelectionStartTime;
+		LastSelectionType = ESelectionType::None;
+		//TODO: SelectionChangingEvent.Broadcast(SelectionStartTime, SelectionEndTime);
+		EventAggregation.Reset();
+		ObjectTypeAggregation.Reset();
+	}
 
 	return Reply;
 }
@@ -2476,7 +2331,9 @@ FReply STimingView::OnMouseButtonDown(const FGeometry& MyGeometry, const FPointe
 FReply STimingView::OnMouseButtonUp(const FGeometry& MyGeometry, const FPointerEvent& MouseEvent)
 {
 	FReply Reply = FReply::Unhandled();
+
 	MousePositionOnButtonUp = MyGeometry.AbsoluteToLocal(MouseEvent.GetScreenSpacePosition());
+	MousePosition = MousePositionOnButtonUp;
 
 	const bool bIsValidForMouseClick = MousePositionOnButtonUp.Equals(MousePositionOnButtonDown, 2.0f);
 
@@ -2534,11 +2391,17 @@ FReply STimingView::OnMouseButtonUp(const FGeometry& MyGeometry, const FPointerE
 
 				bIsPanning = false;
 			}
+			else if (bIsSelecting)
+			{
+				//TODO: SelectionChangedEvent.Broadcast(SelectionStartTime, SelectionEndTime);
+				UpdateAggregatedStats();
+
+				bIsSelecting = false;
+			}
 
-			if (!bIsDragging && !bIsSpaceBarKeyPressed && bIsValidForMouseClick)
+			if (bIsValidForMouseClick)
 			{
 				ShowContextMenu(MouseEvent.GetScreenSpacePosition(), MouseEvent);
-				Reply = FReply::Handled();
 			}
 
 			bIsDragging = false;
@@ -2593,14 +2456,14 @@ FReply STimingView::OnMouseMove(const FGeometry& MyGeometry, const FPointerEvent
 
 			if ((int32)PanningMode & (int32)EPanningMode::Horizontal)
 			{
-				ScrollAtTime(ViewportStartTimeOnButtonDown + static_cast<double>(MousePositionOnButtonDown.X - MousePosition.X) / Viewport.ScaleX);
+				const double StartTime = ViewportStartTimeOnButtonDown + static_cast<double>(MousePositionOnButtonDown.X - MousePosition.X) / Viewport.ScaleX;
+				ScrollAtTime(StartTime);
 			}
 
 			if ((int32)PanningMode & (int32)EPanningMode::Vertical)
 			{
-				Viewport.ScrollPosY = ViewportScrollPosYOnButtonDown + (MousePositionOnButtonDown.Y - MousePosition.Y);
-				UpdateVerticalScrollBar();
-				bIsVerticalViewportDirty = true;
+				const float ScrollPosY = ViewportScrollPosYOnButtonDown + (MousePositionOnButtonDown.Y - MousePosition.Y);
+				ScrollAtPosY(ScrollPosY);
 			}
 		}
 
@@ -2638,14 +2501,13 @@ FReply STimingView::OnMouseMove(const FGeometry& MyGeometry, const FPointerEvent
 		UpdateHoveredTimingEvent(MousePosition.X, MousePosition.Y);
 	}
 
-	return Reply; // SAssetViewItem::CreateToolTipWidget
+	return Reply;
 }
 
 ////////////////////////////////////////////////////////////////////////////////////////////////////
 
 void STimingView::OnMouseEnter(const FGeometry& MyGeometry, const FPointerEvent& MouseEvent)
 {
-
 }
 
 ////////////////////////////////////////////////////////////////////////////////////////////////////
@@ -2674,9 +2536,8 @@ FReply STimingView::OnMouseWheel(const FGeometry& MyGeometry, const FPointerEven
 	{
 		// Scroll vertically.
 		constexpr float ScrollSpeedY = 16.0f * 3;
-		Viewport.ScrollPosY -= ScrollSpeedY * MouseEvent.GetWheelDelta();
-		UpdateVerticalScrollBar();
-		bIsVerticalViewportDirty = true;
+		const float ScrollPosY = Viewport.ScrollPosY - ScrollSpeedY * MouseEvent.GetWheelDelta();
+		ScrollAtPosY(ScrollPosY);
 	}
 	else if (MouseEvent.GetModifierKeys().IsControlDown())
 	{
@@ -2921,10 +2782,7 @@ FReply STimingView::OnKeyDown(const FGeometry& MyGeometry, const FKeyEvent& InKe
 		if (InKeyEvent.GetModifierKeys().IsControlDown())
 		{
 			// Scroll Up
-			Viewport.ScrollPosY -= 16.0 * 3;
-			//Viewport.EnforceVerticalScrollLimits(1.0);
-			UpdateVerticalScrollBar();
-			bIsVerticalViewportDirty = true;
+			ScrollAtPosY(Viewport.ScrollPosY - 16.0 * 3);
 		}
 		else
 		{
@@ -2937,10 +2795,7 @@ FReply STimingView::OnKeyDown(const FGeometry& MyGeometry, const FKeyEvent& InKe
 		if (InKeyEvent.GetModifierKeys().IsControlDown())
 		{
 			// Scroll Down
-			Viewport.ScrollPosY += 16.0 * 3;
-			//Viewport.EnforceVerticalScrollLimits(1.0);
-			UpdateVerticalScrollBar();
-			bIsVerticalViewportDirty = true;
+			ScrollAtPosY(Viewport.ScrollPosY + 16.0 * 3);
 		}
 		else
 		{
@@ -3159,6 +3014,19 @@ void STimingView::UpdateVerticalScrollBar()
 
 ////////////////////////////////////////////////////////////////////////////////////////////////////
 
+void STimingView::ScrollAtPosY(float ScrollPosY)
+{
+	if (Viewport.ScrollPosY != ScrollPosY)
+	{
+		Viewport.ScrollPosY = ScrollPosY;
+
+		UpdateVerticalScrollBar();
+		bIsVerticalViewportDirty = true;
+	}
+}
+
+////////////////////////////////////////////////////////////////////////////////////////////////////
+
 void STimingView::ScrollAtTime(double StartTime)
 {
 	if (Viewport.ScrollAtTime(StartTime))
@@ -4227,9 +4095,7 @@ void STimingView::ToggleAutoHideEmptyTracks()
 		CachedTimelineKV.Value->SetHeight(0.0f);
 	}
 
-	Viewport.ScrollPosY = 0.0f;
-	UpdateVerticalScrollBar();
-	bIsVerticalViewportDirty = true;
+	ScrollAtPosY(0.0f);
 }
 
 ////////////////////////////////////////////////////////////////////////////////////////////////////
diff --git a/Engine/Source/Developer/TraceInsights/Private/Insights/Widgets/STimingView.h b/Engine/Source/Developer/TraceInsights/Private/Insights/Widgets/STimingView.h
--- a/Engine/Source/Developer/TraceInsights/Private/Insights/Widgets/STimingView.h
+++ b/Engine/Source/Developer/TraceInsights/Private/Insights/Widgets/STimingView.h
@@ -3,6 +3,7 @@
 #pragma once
 
 #include "CoreMinimal.h"
+#include "Delegates/Delegate.h"
 #include "Framework/MultiBox/MultiBoxBuilder.h"
 #include "Input/CursorReply.h"
 #include "Input/Reply.h"
@@ -252,6 +253,7 @@ public:
 	bool IsTimeSelected(double Time) const { return Time >= SelectionStartTime && Time < SelectionEndTime; }
 	bool IsTimeSelectedInclusive(double Time) const { return Time >= SelectionStartTime && Time <= SelectionEndTime; }
 
+	void ScrollAtPosY(float ScrollPosY);
 	void ScrollAtTime(double StartTime);
 	void CenterOnTimeInterval(double IntervalStartTime, double IntervalDuration);
 	void BringIntoView(double StartTime, double EndTime);
