[Insights]
  - Graph track:  Added name, tooltip and visibility to each series.
  - Graph track: Added list of series to context menu (to show/hide each series).
  - Graph: Renamed FFramesGraphTrack to FTimingGraphTrack and made it more more generic (Frames, Timers, StatsCounters).
  - Graph: Added support for Stats Counters series (FTimingGraphTrack).
  - Graph: Added auto vertical zoom for the Stats Counters series (FTimingGraphTrack).
  - Graph: Added polygon rendering (i.e. fill area under series). Added "Show Polygon" toggle in ctx menu.
  - Graph: Added more test/debug series to FRandomGraphTrack.
  - Graph: Fixes drawing points (bottom points being hidden).
  - Moved code re context menu from SGraphTrack widget to FGraphTrack view model (to be reused in STimingView).
  - Timing view: Added context menu to tracks (for now, only the Graph track shows a valid ctx menu; the cpu/gpu tracks only show a dummy menu with name of thread track and group).
  - Timing view: Added "Graph Track - G" visibility toggle option in Tracks drop down menu.
  - Stats Counters view: Double click on a counter will show/hide a series in the Timing view's main graph track. Added colored rectangle in front of name for counters having a series added in Timing view's main graph track. Color is generated pseudo-randomly (based on counter id).
  - Fixed FCounterProvider::ReadCounter API in TraceServices.

diff --git a/Engine/Source/Developer/TraceInsights/Private/Insights/ViewModels/BaseTimingTrack.h b/Engine/Source/Developer/TraceInsights/Private/Insights/ViewModels/BaseTimingTrack.h
--- a/Engine/Source/Developer/TraceInsights/Private/Insights/ViewModels/BaseTimingTrack.h
+++ b/Engine/Source/Developer/TraceInsights/Private/Insights/ViewModels/BaseTimingTrack.h
@@ -4,10 +4,12 @@
 
 #include "CoreMinimal.h"
 #include "Misc/EnumClassFlags.h"
+#include "Templates/SharedPointer.h"
 
 ////////////////////////////////////////////////////////////////////////////////////////////////////
 
 class FTimingTrackViewport;
+class FMenuBuilder;
 
 ////////////////////////////////////////////////////////////////////////////////////////////////////
 
@@ -15,15 +17,16 @@ enum class ETimingTrackFlags : uint32
 {
 	None = 0,
 	IsVisible = (1 << 0),
-	IsSelected = (1 << 1),
-	IsHovered = (1 << 2),
-	IsHeaderHovered = (1 << 3),
+	IsDirty = (1 << 1),
+	IsSelected = (1 << 2),
+	IsHovered = (1 << 3),
+	IsHeaderHovered = (1 << 4),
 };
 ENUM_CLASS_FLAGS(ETimingTrackFlags);
 
 ////////////////////////////////////////////////////////////////////////////////////////////////////
 
-class FBaseTimingTrack
+class FBaseTimingTrack : public TSharedFromThis<FBaseTimingTrack>
 {
 	friend class FTimingViewDrawHelper;
 
@@ -32,7 +35,7 @@ protected:
 		: Id(InId)
 		, PosY(0.0f)
 		, Height(0.0f)
-		, Flags(ETimingTrackFlags::IsVisible)
+		, Flags(ETimingTrackFlags::IsVisible | ETimingTrackFlags::IsDirty)
 	{}
 
 	virtual ~FBaseTimingTrack()
@@ -44,7 +47,7 @@ public:
 		Id = 0;
 		PosY = 0.0f;
 		Height = 0.0f;
-		Flags = ETimingTrackFlags::IsVisible;
+		Flags = ETimingTrackFlags::IsVisible | ETimingTrackFlags::IsDirty;
 	}
 
 	uint64 GetId() const { return Id; }
@@ -62,6 +65,11 @@ public:
 	void SetVisibilityFlag(bool bIsVisible) { bIsVisible ? Show() : Hide(); }
 	virtual void OnVisibilityChanged() {}
 
+	bool IsDirty() const { return EnumHasAnyFlags(Flags, ETimingTrackFlags::IsDirty); }
+	void SetDirtyFlag() { Flags |= ETimingTrackFlags::IsDirty; OnDirtyFlagChanged(); }
+	void ClearDirtyFlag() { Flags &= ~ETimingTrackFlags::IsDirty; OnDirtyFlagChanged(); }
+	virtual void OnDirtyFlagChanged() {}
+
 	bool IsSelected() const { return EnumHasAnyFlags(Flags, ETimingTrackFlags::IsSelected); }
 	void Select() { Flags |= ETimingTrackFlags::IsSelected; OnSelectedFlagChanged(); }
 	void Unselect() { Flags &= ~ETimingTrackFlags::IsSelected; OnSelectedFlagChanged(); }
@@ -80,6 +88,8 @@ public:
 
 	static uint64 GenerateId() { return IdGenerator++; }
 
+	virtual void BuildContextMenu(FMenuBuilder& MenuBuilder) {}
+
 private:
 	uint64 Id;
 
diff --git a/Engine/Source/Developer/TraceInsights/Private/Insights/ViewModels/GraphTrack.cpp b/Engine/Source/Developer/TraceInsights/Private/Insights/ViewModels/GraphTrack.cpp
--- a/Engine/Source/Developer/TraceInsights/Private/Insights/ViewModels/GraphTrack.cpp
+++ b/Engine/Source/Developer/TraceInsights/Private/Insights/ViewModels/GraphTrack.cpp
@@ -2,9 +2,12 @@
 
 #include "GraphTrack.h"
 
+#include <limits>
+
 #include "Brushes/SlateColorBrush.h"
 #include "Brushes/SlateBorderBrush.h"
 #include "EditorStyleSet.h"
+#include "Framework/MultiBox/MultiBoxBuilder.h"
 #include "Rendering/DrawElements.h"
 #include "Styling/CoreStyle.h"
 #include "TraceServices/AnalysisService.h"
@@ -24,7 +27,10 @@
 ////////////////////////////////////////////////////////////////////////////////////////////////////
 
 FGraphTrackSeries::FGraphTrackSeries()
-	: Color(0.0f, 0.5f, 1.0f, 1.0f)
+	: Name()
+	, Description()
+	, bIsVisible(true)
+	, Color(0.0f, 0.5f, 1.0f, 1.0f)
 	, BorderColor(0.3f, 0.8f, 1.0f, 1.0f)
 	//, Points()
 	//, LinePoints()
@@ -53,7 +59,8 @@ FGraphTrack::FGraphTrack(uint64 InTrackId)
 	, bDrawPoints(true)
 	, bDrawPointsWithBorder(true)
 	, bDrawLines(true)
-	, bDrawLinesWithDuration(false)
+	, bDrawPolygon(true)
+	, bUseEventDuration(true)
 	, bDrawBoxes(true)
 	, BaselineY(0.0f)
 	, ScaleY(1.0)
@@ -97,11 +104,11 @@ void FGraphTrack::UpdateStats()
 	NumDrawLines = 0;
 	NumDrawBoxes = 0;
 
-	for (const FGraphTrackSeries& Series : AllSeries)
+	for (const TSharedPtr<FGraphTrackSeries>& Series : AllSeries)
 	{
-		NumDrawPoints += Series.Points.Num();
-		NumDrawLines += Series.LinePoints.Num() / 2;
-		NumDrawBoxes += Series.Boxes.Num();
+		NumDrawPoints += Series->Points.Num();
+		NumDrawLines += Series->LinePoints.Num() / 2;
+		NumDrawBoxes += Series->Boxes.Num();
 	}
 }
 
@@ -131,16 +138,20 @@ void FGraphTrack::Draw(FDrawContext& DrawContext, const FTimingTrackViewport& Vi
 
 	//TODO: Set clipping to (0, GetPosY(), Viewport.Width, GetHeight())!
 
-	for (const FGraphTrackSeries& Series : AllSeries)
-	{
-		DrawSeries(Series, DrawContext, Viewport);
-	}
-
+	// Draw top line.
 	//DrawContext.DrawBox(0.0f, GetPosY(), Viewport.Width, 1.0f, WhiteBrush, FLinearColor(0.05f, 0.05f, 0.05f, 1.0f));
 
 	// Draw baseline (Value == 0).
 	DrawContext.DrawBox(0.0f, GetPosY() + BaselineY - 1.0f, Viewport.Width, 1.0f, WhiteBrush, FLinearColor(0.05f, 0.05f, 0.05f, 1.0f));
 	DrawContext.LayerId++;
+
+	for (const TSharedPtr<FGraphTrackSeries>& Series : AllSeries)
+	{
+		if (Series->IsVisible())
+		{
+			DrawSeries(*Series, DrawContext, Viewport);
+		}
+	}
 }
 
 ////////////////////////////////////////////////////////////////////////////////////////////////////
@@ -158,6 +169,70 @@ void FGraphTrack::DrawSeries(const FGraphTrackSeries& Series, FDrawContext& Draw
 		DrawContext.LayerId++;
 	}
 
+	if (bDrawPolygon)
+	{
+		FSlateShaderResourceProxy* ResourceProxy = FSlateDataPayload::ResourceManager->GetShaderResource(*WhiteBrush);
+		FSlateResourceHandle ResourceHandle = FSlateApplication::Get().GetRenderer()->GetResourceHandle(*WhiteBrush);
+
+		FVector2D AtlasOffset = ResourceProxy ? ResourceProxy->StartUV : FVector2D(0.f, 0.f);
+		FVector2D AtlasUVSize = ResourceProxy ? ResourceProxy->SizeUV : FVector2D(1.f, 1.f);
+
+		const FVector2D Pos = DrawContext.Geometry.GetAbsolutePosition();
+		const FVector2D Size = DrawContext.Geometry.GetLocalSize();
+		const float Scale = DrawContext.Geometry.Scale;
+
+		FSlateRenderTransform RenderTransform;
+		FColor FillColor(Series.Color.R * 255, Series.Color.G * 255, Series.Color.B * 255, 42);
+
+		TArray<SlateIndex> Indices;
+		TArray<FSlateVertex> Verts;
+
+		Indices.Reserve(Series.LinePoints.Num() * 6);
+		Verts.Reserve(Series.LinePoints.Num() * 2);
+
+		const float TopY = GetPosY();
+		const float BottomY = (TopY + GetHeight()) / Size.Y;
+
+		for (const FVector2D& LinePoint : Series.LinePoints)
+		{
+			// Add verts top->bottom
+			FVector2D UV(LinePoint.X / Size.X, (TopY + LinePoint.Y) / Size.Y);
+			Verts.Add(FSlateVertex::Make<ESlateVertexRounding::Disabled>(RenderTransform, (Pos + UV * Size * Scale), AtlasOffset + UV * AtlasUVSize, FillColor));
+			UV.Y = BottomY;
+			Verts.Add(FSlateVertex::Make<ESlateVertexRounding::Disabled>(RenderTransform, (Pos + UV * Size * Scale), AtlasOffset + FVector2D(UV.X, 0.5f) * AtlasUVSize, FillColor));
+
+			if (Verts.Num() >= 4)
+			{
+				int32 Index0 = Verts.Num() - 4;
+				int32 Index1 = Verts.Num() - 3;
+				int32 Index2 = Verts.Num() - 2;
+				int32 Index3 = Verts.Num() - 1;
+
+				Indices.Add(Index0);
+				Indices.Add(Index1);
+				Indices.Add(Index2);
+
+				Indices.Add(Index1);
+				Indices.Add(Index2);
+				Indices.Add(Index3);
+			}
+		}
+
+		if (Indices.Num())
+		{
+			FSlateDrawElement::MakeCustomVerts(
+				DrawContext.ElementList,
+				DrawContext.LayerId,
+				ResourceHandle,
+				Verts,
+				Indices,
+				nullptr,
+				0,
+				0, ESlateDrawEffect::PreMultipliedAlpha);
+			DrawContext.LayerId++;
+		}
+	}
+
 	if (bDrawLines)
 	{
 		FPaintGeometry Geo = DrawContext.Geometry.ToPaintGeometry();
@@ -233,6 +308,288 @@ void FGraphTrack::DrawSeries(const FGraphTrackSeries& Series, FDrawContext& Draw
 	}
 }
 
+////////////////////////////////////////////////////////////////////////////////////////////////////
+
+void FGraphTrack::BuildContextMenu(FMenuBuilder& MenuBuilder)
+{
+	MenuBuilder.BeginSection("Misc");
+	{
+		FUIAction Action_ShowPoints
+		(
+			FExecuteAction::CreateSP(this, &FGraphTrack::ContextMenu_ShowPoints_Execute),
+			FCanExecuteAction::CreateSP(this, &FGraphTrack::ContextMenu_ShowPoints_CanExecute),
+			FIsActionChecked::CreateSP(this, &FGraphTrack::ContextMenu_ShowPoints_IsChecked)
+		);
+		MenuBuilder.AddMenuEntry
+		(
+			LOCTEXT("ContextMenu_ShowPoints", "Show Points"),
+			LOCTEXT("ContextMenu_ShowPoints_Desc", "Show points."),
+			FSlateIcon(),
+			Action_ShowPoints,
+			NAME_None,
+			EUserInterfaceActionType::ToggleButton
+		);
+
+		FUIAction Action_ShowPointsWithBorder
+		(
+			FExecuteAction::CreateSP(this, &FGraphTrack::ContextMenu_ShowPointsWithBorder_Execute),
+			FCanExecuteAction::CreateSP(this, &FGraphTrack::ContextMenu_ShowPointsWithBorder_CanExecute),
+			FIsActionChecked::CreateSP(this, &FGraphTrack::ContextMenu_ShowPointsWithBorder_IsChecked)
+		);
+		MenuBuilder.AddMenuEntry
+		(
+			LOCTEXT("ContextMenu_ShowPointsWithBorder", "Show Points with Border"),
+			LOCTEXT("ContextMenu_ShowPointsWithBorder_Desc", "Show border around points."),
+			FSlateIcon(),
+			Action_ShowPointsWithBorder,
+			NAME_None,
+			EUserInterfaceActionType::ToggleButton
+		);
+
+		FUIAction Action_ShowLines
+		(
+			FExecuteAction::CreateSP(this, &FGraphTrack::ContextMenu_ShowLines_Execute),
+			FCanExecuteAction::CreateSP(this, &FGraphTrack::ContextMenu_ShowLines_CanExecute),
+			FIsActionChecked::CreateSP(this, &FGraphTrack::ContextMenu_ShowLines_IsChecked)
+		);
+		MenuBuilder.AddMenuEntry
+		(
+			LOCTEXT("ContextMenu_ShowLines", "Show Connected Lines"),
+			LOCTEXT("ContextMenu_ShowLines_Desc", "Show connected lines. Each event is a single point in time."),
+			FSlateIcon(),
+			Action_ShowLines,
+			NAME_None,
+			EUserInterfaceActionType::ToggleButton
+		);
+
+		FUIAction Action_ShowPolygon
+		(
+			FExecuteAction::CreateSP(this, &FGraphTrack::ContextMenu_ShowPolygon_Execute),
+			FCanExecuteAction::CreateSP(this, &FGraphTrack::ContextMenu_ShowPolygon_CanExecute),
+			FIsActionChecked::CreateSP(this, &FGraphTrack::ContextMenu_ShowPolygon_IsChecked)
+		);
+		MenuBuilder.AddMenuEntry
+		(
+			LOCTEXT("ContextMenu_ShowPolygon", "Show Polygon"),
+			LOCTEXT("ContextMenu_ShowPolygon_Desc", "Show filled polygon under the graph series."),
+			FSlateIcon(),
+			Action_ShowPolygon,
+			NAME_None,
+			EUserInterfaceActionType::ToggleButton
+		);
+
+		FUIAction Action_UseEventDuration
+		(
+			FExecuteAction::CreateSP(this, &FGraphTrack::ContextMenu_UseEventDuration_Execute),
+			FCanExecuteAction::CreateSP(this, &FGraphTrack::ContextMenu_UseEventDuration_CanExecute),
+			FIsActionChecked::CreateSP(this, &FGraphTrack::ContextMenu_UseEventDuration_IsChecked)
+		);
+		MenuBuilder.AddMenuEntry
+		(
+			LOCTEXT("ContextMenu_UseEventDuration", "Use Event Duration"),
+			LOCTEXT("ContextMenu_UseEventDuration_Desc", "Use duration of timing events (for Connected Lines and Polygon)."),
+			FSlateIcon(),
+			Action_UseEventDuration,
+			NAME_None,
+			EUserInterfaceActionType::ToggleButton
+		);
+
+		FUIAction Action_ShowBars
+		(
+			FExecuteAction::CreateSP(this, &FGraphTrack::ContextMenu_ShowBars_Execute),
+			FCanExecuteAction::CreateSP(this, &FGraphTrack::ContextMenu_ShowBars_CanExecute),
+			FIsActionChecked::CreateSP(this, &FGraphTrack::ContextMenu_ShowBars_IsChecked)
+		);
+		MenuBuilder.AddMenuEntry
+		(
+			LOCTEXT("ContextMenu_ShowBars", "Show Bars"),
+			LOCTEXT("ContextMenu_ShowBars_Desc", "Show bars. Width of bars corresponds to duration of timing events."),
+			FSlateIcon(),
+			Action_ShowBars,
+			NAME_None,
+			EUserInterfaceActionType::ToggleButton
+		);
+	}
+	MenuBuilder.EndSection();
+
+	MenuBuilder.BeginSection("Series", LOCTEXT("ContextMenu_Header_Series", "Series"));
+	{
+		for (const TSharedPtr<FGraphTrackSeries>& Series : AllSeries)
+		{
+			FUIAction Action_ShowSeries
+			(
+				FExecuteAction::CreateSP(this, &FGraphTrack::ContextMenu_ShowSeries_Execute, Series.Get()),
+				FCanExecuteAction::CreateSP(this, &FGraphTrack::ContextMenu_ShowSeries_CanExecute, Series.Get()),
+				FIsActionChecked::CreateSP(this, &FGraphTrack::ContextMenu_ShowSeries_IsChecked, Series.Get())
+			);
+			MenuBuilder.AddMenuEntry
+			(
+				Series->GetName(),
+				Series->GetDescription(),
+				FSlateIcon(),
+				Action_ShowSeries,
+				NAME_None,
+				EUserInterfaceActionType::ToggleButton
+			);
+		}
+	}
+	MenuBuilder.EndSection();
+}
+
+////////////////////////////////////////////////////////////////////////////////////////////////////
+
+void FGraphTrack::ContextMenu_ShowPoints_Execute()
+{
+	bDrawPoints = !bDrawPoints;
+	SetDirtyFlag();
+}
+
+////////////////////////////////////////////////////////////////////////////////////////////////////
+
+bool FGraphTrack::ContextMenu_ShowPoints_CanExecute()
+{
+	return true;
+}
+
+////////////////////////////////////////////////////////////////////////////////////////////////////
+
+bool FGraphTrack::ContextMenu_ShowPoints_IsChecked()
+{
+	return bDrawPoints;
+}
+
+////////////////////////////////////////////////////////////////////////////////////////////////////
+
+void FGraphTrack::ContextMenu_ShowPointsWithBorder_Execute()
+{
+	bDrawPointsWithBorder = !bDrawPointsWithBorder;
+	SetDirtyFlag();
+}
+
+////////////////////////////////////////////////////////////////////////////////////////////////////
+
+bool FGraphTrack::ContextMenu_ShowPointsWithBorder_CanExecute()
+{
+	return bDrawPoints;
+}
+
+////////////////////////////////////////////////////////////////////////////////////////////////////
+
+bool FGraphTrack::ContextMenu_ShowPointsWithBorder_IsChecked()
+{
+	return bDrawPointsWithBorder;
+}
+
+////////////////////////////////////////////////////////////////////////////////////////////////////
+
+void FGraphTrack::ContextMenu_ShowLines_Execute()
+{
+	bDrawLines = !bDrawLines;
+	SetDirtyFlag();
+}
+
+////////////////////////////////////////////////////////////////////////////////////////////////////
+
+bool FGraphTrack::ContextMenu_ShowLines_CanExecute()
+{
+	return true;
+}
+
+////////////////////////////////////////////////////////////////////////////////////////////////////
+
+bool FGraphTrack::ContextMenu_ShowLines_IsChecked()
+{
+	return bDrawLines;
+}
+
+////////////////////////////////////////////////////////////////////////////////////////////////////
+
+void FGraphTrack::ContextMenu_ShowPolygon_Execute()
+{
+	bDrawPolygon = !bDrawPolygon;
+	SetDirtyFlag();
+}
+
+////////////////////////////////////////////////////////////////////////////////////////////////////
+
+bool FGraphTrack::ContextMenu_ShowPolygon_CanExecute()
+{
+	return true;
+}
+
+////////////////////////////////////////////////////////////////////////////////////////////////////
+
+bool FGraphTrack::ContextMenu_ShowPolygon_IsChecked()
+{
+	return bDrawPolygon;
+}
+
+////////////////////////////////////////////////////////////////////////////////////////////////////
+
+void FGraphTrack::ContextMenu_UseEventDuration_Execute()
+{
+	bUseEventDuration = !bUseEventDuration;
+	SetDirtyFlag();
+}
+
+////////////////////////////////////////////////////////////////////////////////////////////////////
+
+bool FGraphTrack::ContextMenu_UseEventDuration_CanExecute()
+{
+	return bDrawLines || bDrawPolygon;
+}
+
+////////////////////////////////////////////////////////////////////////////////////////////////////
+
+bool FGraphTrack::ContextMenu_UseEventDuration_IsChecked()
+{
+	return bUseEventDuration;
+}
+
+////////////////////////////////////////////////////////////////////////////////////////////////////
+
+void FGraphTrack::ContextMenu_ShowBars_Execute()
+{
+	bDrawBoxes = !bDrawBoxes;
+	SetDirtyFlag();
+}
+
+////////////////////////////////////////////////////////////////////////////////////////////////////
+
+bool FGraphTrack::ContextMenu_ShowBars_CanExecute()
+{
+	return true;
+}
+
+////////////////////////////////////////////////////////////////////////////////////////////////////
+
+bool FGraphTrack::ContextMenu_ShowBars_IsChecked()
+{
+	return bDrawBoxes;
+}
+
+////////////////////////////////////////////////////////////////////////////////////////////////////
+
+void FGraphTrack::ContextMenu_ShowSeries_Execute(FGraphTrackSeries* Series)
+{
+	Series->SetVisibility(!Series->IsVisible());
+	SetDirtyFlag();
+}
+
+////////////////////////////////////////////////////////////////////////////////////////////////////
+
+bool FGraphTrack::ContextMenu_ShowSeries_CanExecute(FGraphTrackSeries* Series)
+{
+	return true;
+}
+
+////////////////////////////////////////////////////////////////////////////////////////////////////
+
+bool FGraphTrack::ContextMenu_ShowSeries_IsChecked(FGraphTrackSeries* Series)
+{
+	return Series->IsVisible();
+}
+
 ////////////////////////////////////////////////////////////////////////////////////////////////////
 // FRandomGraphTrack
 ////////////////////////////////////////////////////////////////////////////////////////////////////
@@ -243,12 +600,30 @@ FRandomGraphTrack::FRandomGraphTrack(uint64 InTrackId)
 	bDrawPoints = true;
 	bDrawPointsWithBorder = true;
 	bDrawLines = true;
-	bDrawLinesWithDuration = false;
+	bDrawPolygon = false;
+	bUseEventDuration = false;
 	bDrawBoxes = false;
 
-	AllSeries.AddDefaulted();
-	FGraphTrackSeries& Series = AllSeries[0];
-	Series.SetColor(FLinearColor(0.0f, 0.5f, 1.0f, 1.0f), FLinearColor(0.3f, 0.8f, 1.0f, 1.0f));
+	TSharedPtr<FGraphTrackSeries> Series0 = MakeShareable(new FGraphTrackSeries());
+	Series0->SetName(TEXT("Random Blue"));
+	Series0->SetDescription(TEXT("Random series; for debuging purposes"));
+	Series0->SetColor(FLinearColor(0.1f, 0.5f, 1.0f, 1.0f), FLinearColor(0.4f, 0.8f, 1.0f, 1.0f));
+	Series0->SetVisibility(true);
+	AllSeries.Add(Series0);
+
+	TSharedPtr<FGraphTrackSeries> Series1 = MakeShareable(new FGraphTrackSeries());
+	Series1->SetName(TEXT("Random Yellow"));
+	Series1->SetDescription(TEXT("Random series; for debuging purposes"));
+	Series1->SetColor(FLinearColor(0.9f, 0.9f, 0.1f, 1.0f), FLinearColor(1.0f, 1.0f, 0.4f, 1.0f));
+	Series1->SetVisibility(false);
+	AllSeries.Add(Series1);
+
+	TSharedPtr<FGraphTrackSeries> Series2 = MakeShareable(new FGraphTrackSeries());
+	Series2->SetName(TEXT("Random Red"));
+	Series2->SetDescription(TEXT("Random series; for debuging purposes"));
+	Series2->SetColor(FLinearColor(1.0f, 0.1f, 0.2f, 1.0f), FLinearColor(1.0f, 0.4f, 0.5f, 1.0f));
+	Series2->SetVisibility(false);
+	AllSeries.Add(Series2);
 }
 
 ////////////////////////////////////////////////////////////////////////////////////////////////////
@@ -267,15 +642,17 @@ void FRandomGraphTrack::Update(const FTimingTrackViewport& Viewport)
 	BaselineY = GetHeight();
 	ScaleY = 1.0;
 
-	ensure(AllSeries.Num() == 1);
-	GenerateSeries(AllSeries[0], Viewport, 1000000);
+	ensure(AllSeries.Num() == 3);
+	GenerateSeries(*AllSeries[0], Viewport, 1000000, 0);
+	GenerateSeries(*AllSeries[1], Viewport, 1000000, 1);
+	GenerateSeries(*AllSeries[2], Viewport, 1000000, 2);
 
 	UpdateStats();
 }
 
 ////////////////////////////////////////////////////////////////////////////////////////////////////
 
-void FRandomGraphTrack::GenerateSeries(FGraphTrackSeries& Series, const FTimingTrackViewport& Viewport, const int32 EventCount)
+void FRandomGraphTrack::GenerateSeries(FGraphTrackSeries& Series, const FTimingTrackViewport& Viewport, const int32 EventCount, int32 Seed)
 {
 	//////////////////////////////////////////////////
 	// Generate random events.
@@ -295,7 +672,7 @@ void FRandomGraphTrack::GenerateSeries(FGraphTrackSeries& Series, const FTimingT
 	TArray<FGraphEvent> Events;
 	Events.Reserve(EventCount);
 
-	FRandomStream RandomStream(0);
+	FRandomStream RandomStream(Seed);
 	double NextT = 0.0;
 	for (int32 Index = 0; Index < EventCount; ++Index)
 	{
@@ -339,36 +716,107 @@ void FRandomGraphTrack::GenerateSeries(FGraphTrackSeries& Series, const FTimingT
 }
 
 ////////////////////////////////////////////////////////////////////////////////////////////////////
-// FFramesGraphTrack
+// FTimingGraphTrack
 ////////////////////////////////////////////////////////////////////////////////////////////////////
 
-FFramesGraphTrack::FFramesGraphTrack(uint64 InTrackId)
+FTimingGraphTrack::FTimingGraphTrack(uint64 InTrackId)
 	: FGraphTrack(InTrackId)
 {
 	bDrawPoints = true;
 	bDrawPointsWithBorder = true;
 	bDrawLines = true;
-	bDrawLinesWithDuration = true;
+	bDrawPolygon = true;
+	bUseEventDuration = true;
 	bDrawBoxes = false;
 
-	AllSeries.AddDefaulted(2);
+	TSharedPtr<FTimingGraphSeries> GameFramesSeries = MakeShareable(new FTimingGraphSeries());
+	GameFramesSeries->SetName(TEXT("Game Frames"));
+	GameFramesSeries->SetDescription(TEXT("Duration of Game frames"));
+	GameFramesSeries->SetColor(FLinearColor(0.3f, 0.3f, 1.0f, 1.0f), FLinearColor(1.0f, 1.0f, 1.0f, 1.0f));
+	GameFramesSeries->Type = FTimingGraphSeries::ESeriesType::Frame;
+	GameFramesSeries->Id = static_cast<uint32>(TraceFrameType_Game);
+	AllSeries.Add(GameFramesSeries);
+
+	TSharedPtr<FTimingGraphSeries> RenderingFramesSeries = MakeShareable(new FTimingGraphSeries());
+	RenderingFramesSeries->SetName(TEXT("Rendering Frames"));
+	RenderingFramesSeries->SetDescription(TEXT("Duration of Rendering frames"));
+	RenderingFramesSeries->SetColor(FLinearColor(1.0f, 0.3f, 0.3f, 1.0f), FLinearColor(1.0f, 1.0f, 1.0f, 1.0f));
+	RenderingFramesSeries->Type = FTimingGraphSeries::ESeriesType::Frame;
+	RenderingFramesSeries->Id = static_cast<uint32>(TraceFrameType_Rendering);
+	AllSeries.Add(RenderingFramesSeries);
+}
+
+////////////////////////////////////////////////////////////////////////////////////////////////////
+
+TSharedPtr<FTimingGraphSeries> FTimingGraphTrack::GetStatsCounterSeries(uint32 CounterId)
+{
+	TSharedPtr<FGraphTrackSeries>* Ptr = AllSeries.FindByPredicate([=](const TSharedPtr<FGraphTrackSeries>& Series)
+	{
+		const TSharedPtr<FTimingGraphSeries> TimingSeries = StaticCastSharedPtr<FTimingGraphSeries>(Series);
+		return TimingSeries->Type == FTimingGraphSeries::ESeriesType::StatsCounter && TimingSeries->Id == CounterId;
+	});
+	return (Ptr != nullptr) ? StaticCastSharedPtr<FTimingGraphSeries>(*Ptr) : nullptr;
+}
+
+////////////////////////////////////////////////////////////////////////////////////////////////////
+
+void FTimingGraphTrack::AddStatsCounterSeries(uint32 CounterId, FLinearColor Color, double ValueOffset, double ValueScale)
+{
+	TSharedPtr<FTimingGraphSeries> Series = MakeShareable(new FTimingGraphSeries());
+
+	const TCHAR* CounterName = nullptr;
+	bool bIsFloatingPoint = false;
+
+	TSharedPtr<const Trace::IAnalysisSession> Session = FInsightsManager::Get()->GetSession();
+	if (Session.IsValid())
+	{
+		Trace::FAnalysisSessionReadScope SessionReadScope(*Session.Get());
+		const Trace::ICounterProvider& CountersProvider = Trace::ReadCounterProvider(*Session.Get());
+		if (CounterId < CountersProvider.GetCounterCount())
+		{
+			CountersProvider.ReadCounter(CounterId, [&](const Trace::ICounter& Counter)
+			{
+				CounterName = Counter.GetName();
+				bIsFloatingPoint = Counter.IsFloatingPoint();
+			});
+		}
+	}
+
+	Series->SetName(CounterName != nullptr ? CounterName : TEXT("<StatsCounter>"));
+	Series->SetDescription(TEXT("Stats counter series"));
 
-	FGraphTrackSeries& GameFramesSeries = AllSeries[0];
-	GameFramesSeries.SetColor(FLinearColor(0.3f, 0.3f, 1.0f, 1.0f), FLinearColor(1.0f, 1.0f, 1.0f, 1.0f));
+	FLinearColor BorderColor(Color.R + 0.4f, Color.G + 0.4f, Color.B + 0.4f, 1.0f);
+	Series->SetColor(Color, BorderColor);
 
-	FGraphTrackSeries& RenderFramesSeries = AllSeries[1];
-	RenderFramesSeries.SetColor(FLinearColor(1.0f, 0.3f, 0.3f, 1.0f), FLinearColor(1.0f, 1.0f, 1.0f, 1.0f));
+	Series->Type = FTimingGraphSeries::ESeriesType::StatsCounter;
+	Series->Id = CounterId;
+	Series->bIsFloatingPoint = bIsFloatingPoint;
+	Series->ValueOffset = ValueOffset;
+	Series->ValueScale = ValueScale;
+
+	AllSeries.Add(Series);
+}
+
+////////////////////////////////////////////////////////////////////////////////////////////////////
+
+void FTimingGraphTrack::RemoveStatsCounterSeries(uint32 CounterId)
+{
+	AllSeries.RemoveAll([=](const TSharedPtr<FGraphTrackSeries>& Series)
+	{
+		const TSharedPtr<FTimingGraphSeries> TimingSeries = StaticCastSharedPtr<FTimingGraphSeries>(Series);
+		return TimingSeries->Type == FTimingGraphSeries::ESeriesType::StatsCounter && TimingSeries->Id == CounterId;
+	});
 }
 
 ////////////////////////////////////////////////////////////////////////////////////////////////////
 
-FFramesGraphTrack::~FFramesGraphTrack()
+FTimingGraphTrack::~FTimingGraphTrack()
 {
 }
 
 ////////////////////////////////////////////////////////////////////////////////////////////////////
 
-void FFramesGraphTrack::Update(const FTimingTrackViewport& Viewport)
+void FTimingGraphTrack::Update(const FTimingTrackViewport& Viewport)
 {
 	NumAddedEvents = 0;
 
@@ -376,33 +824,131 @@ void FFramesGraphTrack::Update(const FTimingTrackViewport& Viewport)
 	BaselineY = GetHeight();
 	ScaleY = 200.0 / 0.1; // 200px = 100ms
 
-	ensure(AllSeries.Num() == 2);
-	UpdateSeries(AllSeries[0], Viewport, TraceFrameType_Game);
-	UpdateSeries(AllSeries[1], Viewport, TraceFrameType_Rendering);
+	for (TSharedPtr<FGraphTrackSeries>& Series : AllSeries)
+	{
+		if (Series->IsVisible())
+		{
+			TSharedPtr<FTimingGraphSeries> TimingSeries = StaticCastSharedPtr<FTimingGraphSeries>(Series);
+			switch (TimingSeries->Type)
+			{
+				case FTimingGraphSeries::ESeriesType::Frame:
+					UpdateFrameSeries(*TimingSeries, Viewport);
+					break;
+
+				case FTimingGraphSeries::ESeriesType::Timer:
+					UpdateTimerSeries(*TimingSeries, Viewport);
+					break;
+
+				case FTimingGraphSeries::ESeriesType::StatsCounter:
+					UpdateStatsCounterSeries(*TimingSeries, Viewport);
+					break;
+			}
+		}
+	}
 
 	UpdateStats();
 }
 
 ////////////////////////////////////////////////////////////////////////////////////////////////////
 
-void FFramesGraphTrack::UpdateSeries(FGraphTrackSeries& Series, const FTimingTrackViewport& Viewport, ETraceFrameType FrameType)
+void FTimingGraphTrack::UpdateFrameSeries(FTimingGraphSeries& Series, const FTimingTrackViewport& Viewport)
 {
 	FGraphTrackBuilder Builder(*this, Series, Viewport);
-
 	TSharedPtr<const Trace::IAnalysisSession> Session = FInsightsManager::Get()->GetSession();
 	if (Session.IsValid())
 	{
 		Trace::FAnalysisSessionReadScope SessionReadScope(*Session.Get());
 		const Trace::IFrameProvider& FramesProvider = ReadFrameProvider(*Session.Get());
+		ETraceFrameType FrameType = static_cast<ETraceFrameType>(Series.Id);
 		uint64 FrameCount = FramesProvider.GetFrameCount(FrameType);
 		FramesProvider.EnumerateFrames(FrameType, 0, FrameCount - 1, [&Builder](const Trace::FFrame& Frame)
 		{
+			//TODO: add a "frame converter" (i.e. to fps, miliseconds or seconds)
 			const double Duration = Frame.EndTime - Frame.StartTime;
 			Builder.AddEvent(Frame.StartTime, Duration, Duration);
 		});
 	}
 }
 
+////////////////////////////////////////////////////////////////////////////////////////////////////
+
+void FTimingGraphTrack::UpdateTimerSeries(FTimingGraphSeries& Series, const FTimingTrackViewport& Viewport)
+{
+	//TODO
+}
+
+////////////////////////////////////////////////////////////////////////////////////////////////////
+
+void FTimingGraphTrack::UpdateStatsCounterSeries(FTimingGraphSeries& Series, const FTimingTrackViewport& Viewport)
+{
+	FGraphTrackBuilder Builder(*this, Series, Viewport);
+	TSharedPtr<const Trace::IAnalysisSession> Session = FInsightsManager::Get()->GetSession();
+	if (Session.IsValid())
+	{
+		Trace::FAnalysisSessionReadScope SessionReadScope(*Session.Get());
+		const Trace::ICounterProvider& CountersProvider = Trace::ReadCounterProvider(*Session.Get());
+		CountersProvider.ReadCounter(Series.Id, [this, &Viewport, &Builder, &Series](const Trace::ICounter& Counter)
+		{
+			//const double ValueOffset = Series.ValueOffset;
+			//const double ValueScale = Series.ValueScale;
+			double MinValue = std::numeric_limits<double>::infinity();
+			double MaxValue = -std::numeric_limits<double>::infinity();
+
+			if (Counter.IsFloatingPoint())
+			{
+				Counter.EnumerateFloatValues(Viewport.StartTime, Viewport.EndTime, [&Builder, &MinValue, &MaxValue](double Time, double Value)
+				{
+					if (Value < MinValue)
+					{
+						MinValue = Value;
+					}
+					if (Value > MaxValue)
+					{
+						MaxValue = Value;
+					}
+				});
+			}
+			else
+			{
+				Counter.EnumerateValues(Viewport.StartTime, Viewport.EndTime, [&Builder, &MinValue, &MaxValue](double Time, int64 IntValue)
+				{
+					const double Value = static_cast<double>(IntValue);
+					if (Value < MinValue)
+					{
+						MinValue = Value;
+					}
+					if (Value > MaxValue)
+					{
+						MaxValue = Value;
+					}
+				});
+			}
+
+			const double ValueOffset = (MinValue != std::numeric_limits<double>::infinity()) ? -MinValue : 0.0;
+
+			const double AdjustedHeight = GetHeight() - 3.0f;
+			const double ValueScale = (MinValue < MaxValue) ? AdjustedHeight / ScaleY / (MaxValue - MinValue) : 1.0;
+
+			if (Counter.IsFloatingPoint())
+			{
+				Counter.EnumerateFloatValues(Viewport.StartTime, Viewport.EndTime, [&Builder, ValueOffset, ValueScale](double Time, double Value)
+				{
+					//TODO: add a "value converter"
+					Builder.AddEvent(Time, 0.0, (Value + ValueOffset) * ValueScale);
+				});
+			}
+			else
+			{
+				Counter.EnumerateValues(Viewport.StartTime, Viewport.EndTime, [&Builder, ValueOffset, ValueScale](double Time, int64 IntValue)
+				{
+					//TODO: add a "value converter"
+					Builder.AddEvent(Time, 0.0, (static_cast<double>(IntValue) + ValueOffset) * ValueScale);
+				});
+			}
+		});
+	}
+}
+
 ////////////////////////////////////////////////////////////////////////////////////////////////////
 // FGraphTrackBuilder
 ////////////////////////////////////////////////////////////////////////////////////////////////////
@@ -437,11 +983,11 @@ void FGraphTrackBuilder::AddEvent(double Time, double Duration, double Value)
 		AddPoint(Time, Value);
 	}
 
-	if (Track.bDrawLines)
+	if (Track.bDrawLines || Track.bDrawPolygon)
 	{
 		AddConnectedLine(Time, Value);
 
-		if (Track.bDrawLinesWithDuration)
+		if (Track.bUseEventDuration && Duration != 0.0)
 		{
 			AddConnectedLine(Time + Duration, Value);
 		}
@@ -464,7 +1010,7 @@ void FGraphTrackBuilder::BeginPoints()
 	PointsCurrentX = -DBL_MAX;
 
 	PointsAtCurrentX.Reset();
-	int32 MaxPointsPerLineScan = FMath::CeilToInt(Track.GetHeight() / FGraphTrack::PointSizeY);
+	int32 MaxPointsPerLineScan = FMath::CeilToInt(Track.GetHeight() / FGraphTrack::PointSizeY) + 1;
 	if (MaxPointsPerLineScan > 0)
 	{
 		PointsAtCurrentX.AddDefaulted(MaxPointsPerLineScan);
diff --git a/Engine/Source/Developer/TraceInsights/Private/Insights/ViewModels/GraphTrack.h b/Engine/Source/Developer/TraceInsights/Private/Insights/ViewModels/GraphTrack.h
--- a/Engine/Source/Developer/TraceInsights/Private/Insights/ViewModels/GraphTrack.h
+++ b/Engine/Source/Developer/TraceInsights/Private/Insights/ViewModels/GraphTrack.h
@@ -10,6 +10,7 @@
 
 struct FDrawContext;
 struct FSlateBrush;
+class FMenuBuilder;
 class FTimingTrackViewport;
 
 ////////////////////////////////////////////////////////////////////////////////////////////////////
@@ -32,24 +33,61 @@ public:
 	FGraphTrackSeries();
 	~FGraphTrackSeries();
 
+	const FText& GetName() const { return Name; }
+	void SetName(const TCHAR* InName) { Name = FText::FromString(InName); }
+	void SetName(const FString& InName) { Name = FText::FromString(InName); }
+	void SetName(const FText& InName) { Name = InName; }
+
+	const FText& GetDescription() const { return Description; }
+	void SetDescription(const TCHAR* InDescription) { Description = FText::FromString(InDescription); }
+	void SetDescription(const FString& InDescription) { Description = FText::FromString(InDescription); }
+	void SetDescription(const FText& InDescription) { Description = InDescription; }
+
+	bool IsVisible() const { return bIsVisible; }
+	void SetVisibility(bool bOnOff) { bIsVisible = bOnOff; }
+
+	const FLinearColor& GetColor() const { return Color; }
+	const FLinearColor& GetBorderColor() const { return BorderColor; }
+
 	void SetColor(FLinearColor InColor, FLinearColor InBorderColor)
 	{
 		Color = InColor;
 		BorderColor = InBorderColor;
 	}
 
-protected:
-	//FText Name;
-	//FText Description;
+private:
+	FText Name;
+	FText Description;
+
+	bool bIsVisible;
 
 	FLinearColor Color;
 	FLinearColor BorderColor;
 
+protected:
 	TArray<FVector2D> Points;
 	TArray<FVector2D> LinePoints;
 	TArray<FGraphBox> Boxes;
+};
+
+////////////////////////////////////////////////////////////////////////////////////////////////////
+
+class FTimingGraphSeries : public FGraphTrackSeries
+{
+public:
+	enum class ESeriesType
+	{
+		Frame,
+		Timer,
+		StatsCounter
+	};
 
-	//bool bIsVisible;
+public:
+	ESeriesType Type;
+	uint32 Id; // frame type, timer id or stats counter id
+	bool bIsFloatingPoint; // for stats counters
+	double ValueOffset; // offset added to Y values (used to pan graph vertically)
+	double ValueScale; // scale of Y values (used to scale graph vertically)
 };
 
 ////////////////////////////////////////////////////////////////////////////////////////////////////
@@ -70,12 +108,17 @@ public:
 	FGraphTrack(uint64 InTrackId);
 	virtual ~FGraphTrack();
 
-	virtual void UpdateHoveredState(float MouseX, float MouseY, const FTimingTrackViewport& Viewport);
+	virtual void UpdateHoveredState(float MouseX, float MouseY, const FTimingTrackViewport& Viewport) override;
 
 	virtual void Update(const FTimingTrackViewport& Viewport) override = 0;
 
 	void Draw(FDrawContext& DrawContext, const FTimingTrackViewport& Viewport) const;
 
+	virtual void BuildContextMenu(FMenuBuilder& MenuBuilder) override;
+
+	//virtual int GetDebugLineCount() override;
+	//virtual void BuildDebugLines(FString& OutStr) override;
+
 	int32 GetNumAddedEvents() const { return NumAddedEvents; }
 	int32 GetNumDrawPoints() const { return NumDrawPoints; }
 	int32 GetNumDrawLines() const { return NumDrawLines; }
@@ -92,8 +135,37 @@ protected:
 		return BaselineY - static_cast<float>(Value * ScaleY); // TODO: vertical zooming and panning
 	}
 
+private:
+	void ContextMenu_ShowPoints_Execute();
+	bool ContextMenu_ShowPoints_CanExecute();
+	bool ContextMenu_ShowPoints_IsChecked();
+
+	void ContextMenu_ShowPointsWithBorder_Execute();
+	bool ContextMenu_ShowPointsWithBorder_CanExecute();
+	bool ContextMenu_ShowPointsWithBorder_IsChecked();
+
+	void ContextMenu_ShowLines_Execute();
+	bool ContextMenu_ShowLines_CanExecute();
+	bool ContextMenu_ShowLines_IsChecked();
+
+	void ContextMenu_ShowPolygon_Execute();
+	bool ContextMenu_ShowPolygon_CanExecute();
+	bool ContextMenu_ShowPolygon_IsChecked();
+
+	void ContextMenu_UseEventDuration_Execute();
+	bool ContextMenu_UseEventDuration_CanExecute();
+	bool ContextMenu_UseEventDuration_IsChecked();
+
+	void ContextMenu_ShowBars_Execute();
+	bool ContextMenu_ShowBars_CanExecute();
+	bool ContextMenu_ShowBars_IsChecked();
+
+	void ContextMenu_ShowSeries_Execute(FGraphTrackSeries* Series);
+	bool ContextMenu_ShowSeries_CanExecute(FGraphTrackSeries* Series);
+	bool ContextMenu_ShowSeries_IsChecked(FGraphTrackSeries* Series);
+
 protected:
-	TArray<FGraphTrackSeries> AllSeries;
+	TArray<TSharedPtr<FGraphTrackSeries>> AllSeries;
 
 	// Slate resources
 	const FSlateBrush* WhiteBrush;
@@ -101,14 +173,13 @@ protected:
 	const FSlateBrush* BorderBrush;
 	const FSlateFontInfo Font;
 
-public:
 	bool bDrawPoints;
 	bool bDrawPointsWithBorder;
 	bool bDrawLines;
-	bool bDrawLinesWithDuration;
+	bool bDrawPolygon;
+	bool bUseEventDuration;
 	bool bDrawBoxes;
 
-protected:
 	float BaselineY; // Y position (in viewport local space) of the baseline (with Value == 0); in pixels (Slate units)
 	double ScaleY; // scale between Value units and viewport units; in pixels (Slate units) / Value unit
 
@@ -130,21 +201,27 @@ public:
 	virtual void Update(const FTimingTrackViewport& Viewport) override;
 
 protected:
-	void GenerateSeries(FGraphTrackSeries& Series, const FTimingTrackViewport& Viewport, const int32 EventCount);
+	void GenerateSeries(FGraphTrackSeries& Series, const FTimingTrackViewport& Viewport, const int32 EventCount, int32 Seed);
 };
 
 ////////////////////////////////////////////////////////////////////////////////////////////////////
 
-class FFramesGraphTrack : public FGraphTrack
+class FTimingGraphTrack : public FGraphTrack
 {
 public:
-	FFramesGraphTrack(uint64 InTrackId);
-	virtual ~FFramesGraphTrack();
+	FTimingGraphTrack(uint64 InTrackId);
+	virtual ~FTimingGraphTrack();
 
 	virtual void Update(const FTimingTrackViewport& Viewport) override;
 
+	TSharedPtr<FTimingGraphSeries> GetStatsCounterSeries(uint32 CounterId);
+	void AddStatsCounterSeries(uint32 CounterId, FLinearColor Color, double ValueOffset = 0.0, double ValueScale = 1.0);
+	void RemoveStatsCounterSeries(uint32 CounterId);
+
 protected:
-	void UpdateSeries(FGraphTrackSeries& Series, const FTimingTrackViewport& Viewport, ETraceFrameType FrameType);
+	void UpdateFrameSeries(FTimingGraphSeries& Series, const FTimingTrackViewport& Viewport);
+	void UpdateTimerSeries(FTimingGraphSeries& Series, const FTimingTrackViewport& Viewport);
+	void UpdateStatsCounterSeries(FTimingGraphSeries& Series, const FTimingTrackViewport& Viewport);
 };
 
 ////////////////////////////////////////////////////////////////////////////////////////////////////
diff --git a/Engine/Source/Developer/TraceInsights/Private/Insights/ViewModels/MarkersTimingTrack.cpp b/Engine/Source/Developer/TraceInsights/Private/Insights/ViewModels/MarkersTimingTrack.cpp
--- a/Engine/Source/Developer/TraceInsights/Private/Insights/ViewModels/MarkersTimingTrack.cpp
+++ b/Engine/Source/Developer/TraceInsights/Private/Insights/ViewModels/MarkersTimingTrack.cpp
@@ -28,7 +28,6 @@ FMarkersTimingTrack::FMarkersTimingTrack(uint64 InTrackId)
 	//, TimeMarkerTexts()
 	, bIsCollapsed(false)
 	, bUseOnlyBookmarks(true)
-	, bIsDirty(true)
 	, TargetHoveredAnimPercent(0.0f)
 	, CurrentHoveredAnimPercent(0.0f)
 	, NumLogMessages(0)
@@ -56,7 +55,6 @@ void FMarkersTimingTrack::Reset()
 
 	bIsCollapsed = false;
 	bUseOnlyBookmarks = true;
-	bIsDirty = true;
 
 	TargetHoveredAnimPercent = 0.0f;
 	CurrentHoveredAnimPercent = 0.0f;
diff --git a/Engine/Source/Developer/TraceInsights/Private/Insights/ViewModels/MarkersTimingTrack.h b/Engine/Source/Developer/TraceInsights/Private/Insights/ViewModels/MarkersTimingTrack.h
--- a/Engine/Source/Developer/TraceInsights/Private/Insights/ViewModels/MarkersTimingTrack.h
+++ b/Engine/Source/Developer/TraceInsights/Private/Insights/ViewModels/MarkersTimingTrack.h
@@ -59,10 +59,6 @@ public:
 		UpdateHeight();
 	}
 
-	bool IsDirty() const { return bIsDirty; }
-	void SetDirtyFlag() { bIsDirty = true; }
-	void ClearDirtyFlag() { bIsDirty = false; }
-
 	// Stats
 	int32 GetNumLogMessages() const { return NumLogMessages; }
 	int32 GetNumBoxes() const { return TimeMarkerBoxes.Num(); }
@@ -91,7 +87,6 @@ private:
 
 	bool bIsCollapsed; // If false, the vertical lines will extend to entire viewport height; otherwise will be limited to this track's height.
 	bool bUseOnlyBookmarks; // If true, uses only bookmarks; otherwise it uses all log messages.
-	bool bIsDirty; // Source list of bookmarks/markers has changed. Cached Boxes and Texts are dirty.
 
 	float TargetHoveredAnimPercent; // [0.0 .. 1.0], 0.0 = hidden, 1.0 = visible
 	float CurrentHoveredAnimPercent;
diff --git a/Engine/Source/Developer/TraceInsights/Private/Insights/ViewModels/StatsNode.h b/Engine/Source/Developer/TraceInsights/Private/Insights/ViewModels/StatsNode.h
--- a/Engine/Source/Developer/TraceInsights/Private/Insights/ViewModels/StatsNode.h
+++ b/Engine/Source/Developer/TraceInsights/Private/Insights/ViewModels/StatsNode.h
@@ -89,8 +89,15 @@ public:
 		, Name(InName)
 		, MetaGroupName(InMetaGroupName)
 		, Type(InType)
+		, bIsAddedToGraph(false)
 		, bForceExpandGroupNode(false)
 	{
+		const uint32 HashColor = Id * 0x2c2c57ed;
+		Color.R = ((HashColor >> 16) & 0xFF) / 255.0f;
+		Color.G = ((HashColor >> 8) & 0xFF) / 255.0f;
+		Color.B = ((HashColor) & 0xFF) / 255.0f;
+		Color.A = 1.0;
+
 		ResetAggregatedStats();
 	}
 
@@ -99,6 +106,8 @@ public:
 		: Id(0)
 		, Name(InGroupName)
 		, Type(EStatsNodeType::Group)
+		, Color(0.0, 0.0, 0.0, 1.0)
+		, bIsAddedToGraph(false)
 		, bForceExpandGroupNode(false)
 	{
 		ResetAggregatedStats();
@@ -149,6 +158,14 @@ public:
 		return Type == EStatsNodeType::Group;
 	}
 
+	/**
+	 * @return color of the node. Used when showing a graph series for a stats counter.
+	 */
+	FLinearColor GetColor() const
+	{
+		return Color;
+	}
+
 	/**
 	 * @return the aggregated stats of this counter (if counter is a "float number" type).
 	 */
@@ -231,6 +248,16 @@ public:
 		return false; // TODO
 	}
 
+	bool IsAddedToGraph() const
+	{
+		return bIsAddedToGraph;
+	}
+
+	void SetAddedToGraphFlag(bool bOnOff)
+	{
+		bIsAddedToGraph = bOnOff;
+	}
+
 public:
 	/** Sorts children using the specified class instance. */
 	template<typename TSortingClass>
@@ -268,9 +295,14 @@ protected:
 	/** The name of the group that this timer belongs to, based on the timer's metadata; only valid for timer nodes. */
 	const FName MetaGroupName;
 
-	/** Holds the type of this timer; for the group, this is Group. */
+	/** Holds the type of this node; for the group, this is Group. */
 	const EStatsNodeType Type;
 
+	/** Color of the node. */
+	FLinearColor Color;
+
+	bool bIsAddedToGraph;
+
 	/** Aggregated stats (double). */
 	FAggregatedStats AggregatedStats;
 
diff --git a/Engine/Source/Developer/TraceInsights/Private/Insights/Widgets/SFrameTrack.cpp b/Engine/Source/Developer/TraceInsights/Private/Insights/Widgets/SFrameTrack.cpp
--- a/Engine/Source/Developer/TraceInsights/Private/Insights/Widgets/SFrameTrack.cpp
+++ b/Engine/Source/Developer/TraceInsights/Private/Insights/Widgets/SFrameTrack.cpp
@@ -195,7 +195,7 @@ void SFrameTrack::Tick(const FGeometry& AllottedGeometry, const double InCurrent
 			for (int32 FrameType = 0; FrameType < TraceFrameType_Count; ++FrameType)
 			{
 				FFrameTrackTimeline& CachedTimeline = CachedTimelines.FindOrAdd(FrameType);
-				CachedTimeline.Id = FrameType;
+				CachedTimeline.Id = static_cast<uint64>(FrameType);
 
 				int32 NumFrames = FramesProvider.GetFrameCount(static_cast<ETraceFrameType>(FrameType));
 				if (NumFrames > Viewport.MaxIndex)
@@ -258,7 +258,7 @@ void SFrameTrack::UpdateState()
 		for (int32 FrameType = 0; FrameType < TraceFrameType_Count; ++FrameType)
 		{
 			FFrameTrackTimeline& Timeline = CachedTimelines.FindOrAdd(FrameType);
-			Timeline.Id = FrameType;
+			Timeline.Id = static_cast<uint64>(FrameType);
 
 			FFrameTrackTimelineBuilder Builder(Timeline, Viewport);
 
@@ -292,7 +292,7 @@ FSampleRef SFrameTrack::GetSampleAtMousePosition(float X, float Y)
 		// Search in reverse paint order.
 		for (int32 TimelineIndex = TimelinesOrder.Num() - 1; TimelineIndex >= 0; --TimelineIndex)
 		{
-			int32 TimelineId = TimelinesOrder[TimelineIndex];
+			uint64 TimelineId = TimelinesOrder[TimelineIndex];
 
 			if ((TimelineId == TraceFrameType_Rendering && !bShowRenderingFrames) ||
 				(TimelineId == TraceFrameType_Game && !bShowGameFrames))
@@ -302,11 +302,13 @@ FSampleRef SFrameTrack::GetSampleAtMousePosition(float X, float Y)
 
 			if (CachedTimelines.Contains(TimelineId))
 			{
-				const FFrameTrackTimeline& CachedTimeline = CachedTimelines[TimelineId];
-				if (CachedTimeline.NumAggregatedFrames > 0 &&
-					SampleIndex < CachedTimeline.Samples.Num())
+				const FFrameTrackTimeline& Timeline = CachedTimelines[TimelineId];
+				ensure(Timeline.Id == TimelineId);
+
+				if (Timeline.NumAggregatedFrames > 0 &&
+					SampleIndex < Timeline.Samples.Num())
 				{
-					const FFrameTrackSample& Sample = CachedTimeline.Samples[SampleIndex];
+					const FFrameTrackSample& Sample = Timeline.Samples[SampleIndex];
 					if (Sample.NumFrames > 0)
 					{
 						float SampleY1;
@@ -320,7 +322,7 @@ FSampleRef SFrameTrack::GetSampleAtMousePosition(float X, float Y)
 						}
 						if (Y >= SampleY1 && Y <= SampleY2)
 						{
-							return FSampleRef(&CachedTimeline, &Sample);
+							return FSampleRef(&Timeline, &Sample);
 						}
 					}
 				}
@@ -386,7 +388,7 @@ int32 SFrameTrack::OnPaint(const FPaintArgs& Args, const FGeometry& AllottedGeom
 		// Draw frames, for each visible timeline.
 		for (int32 TimelineIndex = 0; TimelineIndex < TimelinesOrder.Num(); ++TimelineIndex)
 		{
-			int32 TimelineId = TimelinesOrder[TimelineIndex];
+			uint64 TimelineId = TimelinesOrder[TimelineIndex];
 
 			if ((TimelineId == TraceFrameType_Rendering && !bShowRenderingFrames) ||
 				(TimelineId == TraceFrameType_Game && !bShowGameFrames))
@@ -397,6 +399,8 @@ int32 SFrameTrack::OnPaint(const FPaintArgs& Args, const FGeometry& AllottedGeom
 			if (CachedTimelines.Contains(TimelineId))
 			{
 				const FFrameTrackTimeline& Timeline = CachedTimelines[TimelineId];
+				ensure(Timeline.Id == TimelineId);
+
 				Helper.DrawCached(Timeline);
 			}
 		}
@@ -410,8 +414,8 @@ int32 SFrameTrack::OnPaint(const FPaintArgs& Args, const FGeometry& AllottedGeom
 
 			FString Text = FString::Format(TEXT("{0} frame {1} ({2})"),
 				{
-					(HoveredSample.Timeline->Id == 0) ? TEXT("Game") :
-					(HoveredSample.Timeline->Id == 1) ? TEXT("Rendering") : TEXT("Misc"),
+					(HoveredSample.Timeline->Id == TraceFrameType_Game) ? TEXT("Game") :
+					(HoveredSample.Timeline->Id == TraceFrameType_Rendering) ? TEXT("Rendering") : TEXT("Misc"),
 					FText::AsNumber(HoveredSample.Sample->LargestFrameIndex).ToString(),
 					TimeUtils::FormatTimeAuto(HoveredSample.Sample->LargestFrameDuration)
 				});
@@ -440,12 +444,21 @@ int32 SFrameTrack::OnPaint(const FPaintArgs& Args, const FGeometry& AllottedGeom
 		}
 
 		TSharedPtr<STimingProfilerWindow> Window = FTimingProfilerManager::Get()->GetProfilerWindow();
-		if (Window && Window->TimingView && CachedTimelines.Num() > 0)
+		if (Window && Window->TimingView)
 		{
-			// Highlight the area corresponding to viewport of Timing View.
-			const double StartTime = Window->TimingView->GetViewport().StartTime;
-			const double EndTime = Window->TimingView->GetViewport().EndTime;
-			Helper.DrawHighlightedInterval(CachedTimelines[0], StartTime, EndTime);
+			const FFrameTrackTimeline* TimelinePtr = nullptr;
+			for (const TPair<uint64, FFrameTrackTimeline>& KeyValuePair : CachedTimelines)
+			{
+				TimelinePtr = &KeyValuePair.Value;
+				break; // stop at first enumerated timeline
+			}
+			if (TimelinePtr)
+			{
+				// Highlight the area corresponding to viewport of Timing View.
+				const double StartTime = Window->TimingView->GetViewport().StartTime;
+				const double EndTime = Window->TimingView->GetViewport().EndTime;
+				Helper.DrawHighlightedInterval(*TimelinePtr, StartTime, EndTime);
+			}
 		}
 
 		Stopwatch.Stop();
diff --git a/Engine/Source/Developer/TraceInsights/Private/Insights/Widgets/SFrameTrack.h b/Engine/Source/Developer/TraceInsights/Private/Insights/Widgets/SFrameTrack.h
--- a/Engine/Source/Developer/TraceInsights/Private/Insights/Widgets/SFrameTrack.h
+++ b/Engine/Source/Developer/TraceInsights/Private/Insights/Widgets/SFrameTrack.h
@@ -146,7 +146,7 @@ protected:
 
 	/** Cached info for timelines. */
 	TMap<uint64, FFrameTrackTimeline> CachedTimelines;
-	TArray<int32> TimelinesOrder;
+	TArray<uint64> TimelinesOrder;
 
 	bool bIsStateDirty;
 
diff --git a/Engine/Source/Developer/TraceInsights/Private/Insights/Widgets/SGraphTrack.cpp b/Engine/Source/Developer/TraceInsights/Private/Insights/Widgets/SGraphTrack.cpp
--- a/Engine/Source/Developer/TraceInsights/Private/Insights/Widgets/SGraphTrack.cpp
+++ b/Engine/Source/Developer/TraceInsights/Private/Insights/Widgets/SGraphTrack.cpp
@@ -39,8 +39,8 @@
 ////////////////////////////////////////////////////////////////////////////////////////////////////
 
 SGraphTrack::SGraphTrack()
-	: TimeRulerTrack(0)
-	, GraphTrack(1)
+	: TimeRulerTrack(MakeShareable(new FTimeRulerTrack(0)))
+	, GraphTrack(MakeShareable(new FRandomGraphTrack(1)))
 	, WhiteBrush(FCoreStyle::Get().GetBrush("WhiteBrush"))
 	, MainFont(FCoreStyle::GetDefaultFontStyle("Regular", 8))
 {
@@ -57,8 +57,8 @@ SGraphTrack::~SGraphTrack()
 
 void SGraphTrack::Reset()
 {
-	TimeRulerTrack.Reset();
-	GraphTrack.Reset();
+	TimeRulerTrack->Reset();
+	GraphTrack->Reset();
 
 	Viewport.Reset();
 	Viewport.MaxValidTime = 84.0 * 60.0;
@@ -110,7 +110,7 @@ void SGraphTrack::Tick(const FGeometry& AllottedGeometry, const double InCurrent
 	if (Viewport.UpdateSize(TrackWidth, TrackHeight) ||
 		bIsViewportDirty ||
 		bIsVerticalViewportDirty ||
-		GraphTrack.GetHeight() != TrackHeight)
+		GraphTrack->GetHeight() != TrackHeight)
 	{
 		bIsGraphDirty = true;
 	}
@@ -119,9 +119,9 @@ void SGraphTrack::Tick(const FGeometry& AllottedGeometry, const double InCurrent
 	{
 		bIsGraphDirty = false;
 
-		GraphTrack.SetPosY(TimeRulerTrack.GetHeight());
-		GraphTrack.SetHeight(Viewport.Height - TimeRulerTrack.GetHeight());
-		GraphTrack.Update(Viewport);
+		GraphTrack->SetPosY(TimeRulerTrack->GetHeight());
+		GraphTrack->SetHeight(Viewport.Height - TimeRulerTrack->GetHeight());
+		GraphTrack->Update(Viewport);
 	}
 }
 
@@ -133,10 +133,10 @@ int32 SGraphTrack::OnPaint(const FPaintArgs& Args, const FGeometry& AllottedGeom
 	const ESlateDrawEffect DrawEffects = bEnabled ? ESlateDrawEffect::None : ESlateDrawEffect::DisabledEffect;
 	FDrawContext DrawContext(AllottedGeometry, MyCullingRect, InWidgetStyle, DrawEffects, OutDrawElements, LayerId);
 
-	GraphTrack.Draw(DrawContext, Viewport);
+	GraphTrack->Draw(DrawContext, Viewport);
 
-	TimeRulerTrack.Draw(DrawContext, Viewport, MousePosition, bIsSelecting, SelectionStartTime, SelectionEndTime);
-	DrawContext.DrawBox(0.0f, TimeRulerTrack.GetHeight(), Viewport.Width, 1.0f, WhiteBrush, FLinearColor(0.05f, 0.05f, 0.05f, 1.0f));
+	TimeRulerTrack->Draw(DrawContext, Viewport, MousePosition, bIsSelecting, SelectionStartTime, SelectionEndTime);
+	DrawContext.DrawBox(0.0f, TimeRulerTrack->GetHeight(), Viewport.Width, 1.0f, WhiteBrush, FLinearColor(0.05f, 0.05f, 0.05f, 1.0f));
 
 	FDrawHelpers::DrawTimeRangeSelection(DrawContext, Viewport, SelectionStartTime, SelectionEndTime, WhiteBrush, MainFont);
 
@@ -205,7 +205,7 @@ FReply SGraphTrack::OnMouseButtonDown(const FGeometry& MyGeometry, const FPointe
 		{
 			bIsLMB_Pressed = true;
 
-			if ((MousePositionOnButtonDown.Y < TimeRulerTrack.GetHeight() ||
+			if ((MousePositionOnButtonDown.Y < TimeRulerTrack->GetHeight() ||
 				(MouseEvent.GetModifierKeys().IsControlDown() && MouseEvent.GetModifierKeys().IsShiftDown())))
 			{
 				bStartSelecting = true;
@@ -225,7 +225,7 @@ FReply SGraphTrack::OnMouseButtonDown(const FGeometry& MyGeometry, const FPointe
 		{
 			bIsRMB_Pressed = true;
 
-			if ((MousePositionOnButtonDown.Y < TimeRulerTrack.GetHeight() ||
+			if ((MousePositionOnButtonDown.Y < TimeRulerTrack->GetHeight() ||
 				(MouseEvent.GetModifierKeys().IsControlDown() && MouseEvent.GetModifierKeys().IsShiftDown())))
 			{
 				bStartSelecting = true;
@@ -581,89 +581,7 @@ void SGraphTrack::ShowContextMenu(const FPointerEvent& MouseEvent)
 	const bool bShouldCloseWindowAfterMenuSelection = true;
 	FMenuBuilder MenuBuilder(bShouldCloseWindowAfterMenuSelection, NULL);
 
-	MenuBuilder.BeginSection("Misc");
-	{
-		FUIAction Action_ShowPoints
-		(
-			FExecuteAction::CreateSP(this, &SGraphTrack::ContextMenu_ShowPoints_Execute),
-			FCanExecuteAction::CreateSP(this, &SGraphTrack::ContextMenu_ShowPoints_CanExecute),
-			FIsActionChecked::CreateSP(this, &SGraphTrack::ContextMenu_ShowPoints_IsChecked)
-		);
-		MenuBuilder.AddMenuEntry
-		(
-			LOCTEXT("ContextMenu_ShowPoints", "Points"),
-			LOCTEXT("ContextMenu_ShowPoints_Desc", "Show points."),
-			FSlateIcon(),
-			Action_ShowPoints,
-			NAME_None,
-			EUserInterfaceActionType::ToggleButton
-		);
-
-		FUIAction Action_ShowPointsWithBorder
-		(
-			FExecuteAction::CreateSP(this, &SGraphTrack::ContextMenu_ShowPointsWithBorder_Execute),
-			FCanExecuteAction::CreateSP(this, &SGraphTrack::ContextMenu_ShowPointsWithBorder_CanExecute),
-			FIsActionChecked::CreateSP(this, &SGraphTrack::ContextMenu_ShowPointsWithBorder_IsChecked)
-		);
-		MenuBuilder.AddMenuEntry
-		(
-			LOCTEXT("ContextMenu_ShowPointsWithBorder", "Points have Border"),
-			LOCTEXT("ContextMenu_ShowPointsWithBorder_Desc", "Show border around points."),
-			FSlateIcon(),
-			Action_ShowPointsWithBorder,
-			NAME_None,
-			EUserInterfaceActionType::ToggleButton
-		);
-
-		FUIAction Action_ShowLines
-		(
-			FExecuteAction::CreateSP(this, &SGraphTrack::ContextMenu_ShowLines_Execute),
-			FCanExecuteAction::CreateSP(this, &SGraphTrack::ContextMenu_ShowLines_CanExecute),
-			FIsActionChecked::CreateSP(this, &SGraphTrack::ContextMenu_ShowLines_IsChecked)
-		);
-		MenuBuilder.AddMenuEntry
-		(
-			LOCTEXT("ContextMenu_ShowLines", "Lines"),
-			LOCTEXT("ContextMenu_ShowLines_Desc", "Show connected lines. Each event is a single point in time."),
-			FSlateIcon(),
-			Action_ShowLines,
-			NAME_None,
-			EUserInterfaceActionType::ToggleButton
-		);
-
-		FUIAction Action_ShowLinesWithDuration
-		(
-			FExecuteAction::CreateSP(this, &SGraphTrack::ContextMenu_ShowLinesWithDuration_Execute),
-			FCanExecuteAction::CreateSP(this, &SGraphTrack::ContextMenu_ShowLinesWithDuration_CanExecute),
-			FIsActionChecked::CreateSP(this, &SGraphTrack::ContextMenu_ShowLinesWithDuration_IsChecked)
-		);
-		MenuBuilder.AddMenuEntry
-		(
-			LOCTEXT("ContextMenu_ShowLinesWithDuration", "Lines have Duration"),
-			LOCTEXT("ContextMenu_ShowLinesWithDuration_Desc", "Show connected lines. Each event is a horizontal line (duration of the timing event)."),
-			FSlateIcon(),
-			Action_ShowLinesWithDuration,
-			NAME_None,
-			EUserInterfaceActionType::ToggleButton
-		);
-
-		FUIAction Action_ShowBars
-		(
-			FExecuteAction::CreateSP(this, &SGraphTrack::ContextMenu_ShowBars_Execute),
-			FCanExecuteAction::CreateSP(this, &SGraphTrack::ContextMenu_ShowBars_CanExecute),
-			FIsActionChecked::CreateSP(this, &SGraphTrack::ContextMenu_ShowBars_IsChecked)
-		);
-		MenuBuilder.AddMenuEntry
-		(
-			LOCTEXT("ContextMenu_ShowBars", "Bars"),
-			LOCTEXT("ContextMenu_ShowBars_Desc", "Show bars."),
-			FSlateIcon(),
-			Action_ShowBars,
-			NAME_None,
-			EUserInterfaceActionType::ToggleButton
-		);
-	}
-	MenuBuilder.EndSection();
+	GraphTrack->BuildContextMenu(MenuBuilder);
 
 	TSharedRef<SWidget> MenuWidget = MenuBuilder.MakeWidget();
 
@@ -674,109 +592,4 @@ void SGraphTrack::ShowContextMenu(const FPointerEvent& MouseEvent)
 
 ////////////////////////////////////////////////////////////////////////////////////////////////////
 
-void SGraphTrack::ContextMenu_ShowPoints_Execute()
-{
-	GraphTrack.bDrawPoints = !GraphTrack.bDrawPoints;
-}
-
-////////////////////////////////////////////////////////////////////////////////////////////////////
-
-bool SGraphTrack::ContextMenu_ShowPoints_CanExecute()
-{
-	return true;
-}
-
-////////////////////////////////////////////////////////////////////////////////////////////////////
-
-bool SGraphTrack::ContextMenu_ShowPoints_IsChecked()
-{
-	return GraphTrack.bDrawPoints;
-}
-
-////////////////////////////////////////////////////////////////////////////////////////////////////
-
-void SGraphTrack::ContextMenu_ShowPointsWithBorder_Execute()
-{
-	GraphTrack.bDrawPointsWithBorder = !GraphTrack.bDrawPointsWithBorder;
-}
-
-////////////////////////////////////////////////////////////////////////////////////////////////////
-
-bool SGraphTrack::ContextMenu_ShowPointsWithBorder_CanExecute()
-{
-	return GraphTrack.bDrawPoints;
-}
-
-////////////////////////////////////////////////////////////////////////////////////////////////////
-
-bool SGraphTrack::ContextMenu_ShowPointsWithBorder_IsChecked()
-{
-	return GraphTrack.bDrawPointsWithBorder;
-}
-
-////////////////////////////////////////////////////////////////////////////////////////////////////
-
-void SGraphTrack::ContextMenu_ShowLines_Execute()
-{
-	GraphTrack.bDrawLines = !GraphTrack.bDrawLines;
-}
-
-////////////////////////////////////////////////////////////////////////////////////////////////////
-
-bool SGraphTrack::ContextMenu_ShowLines_CanExecute()
-{
-	return true;
-}
-
-////////////////////////////////////////////////////////////////////////////////////////////////////
-
-bool SGraphTrack::ContextMenu_ShowLines_IsChecked()
-{
-	return GraphTrack.bDrawLines;
-}
-
-////////////////////////////////////////////////////////////////////////////////////////////////////
-
-void SGraphTrack::ContextMenu_ShowLinesWithDuration_Execute()
-{
-	GraphTrack.bDrawLinesWithDuration = !GraphTrack.bDrawLinesWithDuration;
-}
-
-////////////////////////////////////////////////////////////////////////////////////////////////////
-
-bool SGraphTrack::ContextMenu_ShowLinesWithDuration_CanExecute()
-{
-	return GraphTrack.bDrawLines;
-}
-
-////////////////////////////////////////////////////////////////////////////////////////////////////
-
-bool SGraphTrack::ContextMenu_ShowLinesWithDuration_IsChecked()
-{
-	return GraphTrack.bDrawLinesWithDuration;
-}
-
-////////////////////////////////////////////////////////////////////////////////////////////////////
-
-void SGraphTrack::ContextMenu_ShowBars_Execute()
-{
-	GraphTrack.bDrawBoxes = !GraphTrack.bDrawBoxes;
-}
-
-////////////////////////////////////////////////////////////////////////////////////////////////////
-
-bool SGraphTrack::ContextMenu_ShowBars_CanExecute()
-{
-	return true;
-}
-
-////////////////////////////////////////////////////////////////////////////////////////////////////
-
-bool SGraphTrack::ContextMenu_ShowBars_IsChecked()
-{
-	return GraphTrack.bDrawBoxes;
-}
-
-////////////////////////////////////////////////////////////////////////////////////////////////////
-
 #undef LOCTEXT_NAMESPACE
diff --git a/Engine/Source/Developer/TraceInsights/Private/Insights/Widgets/SGraphTrack.h b/Engine/Source/Developer/TraceInsights/Private/Insights/Widgets/SGraphTrack.h
--- a/Engine/Source/Developer/TraceInsights/Private/Insights/Widgets/SGraphTrack.h
+++ b/Engine/Source/Developer/TraceInsights/Private/Insights/Widgets/SGraphTrack.h
@@ -188,30 +188,9 @@ protected:
 
 	void ShowContextMenu(const FPointerEvent& MouseEvent);
 
-	void ContextMenu_ShowPoints_Execute();
-	bool ContextMenu_ShowPoints_CanExecute();
-	bool ContextMenu_ShowPoints_IsChecked();
-
-	void ContextMenu_ShowPointsWithBorder_Execute();
-	bool ContextMenu_ShowPointsWithBorder_CanExecute();
-	bool ContextMenu_ShowPointsWithBorder_IsChecked();
-
-	void ContextMenu_ShowLines_Execute();
-	bool ContextMenu_ShowLines_CanExecute();
-	bool ContextMenu_ShowLines_IsChecked();
-
-	void ContextMenu_ShowLinesWithDuration_Execute();
-	bool ContextMenu_ShowLinesWithDuration_CanExecute();
-	bool ContextMenu_ShowLinesWithDuration_IsChecked();
-
-	void ContextMenu_ShowBars_Execute();
-	bool ContextMenu_ShowBars_CanExecute();
-	bool ContextMenu_ShowBars_IsChecked();
-
 protected:
-	FTimeRulerTrack TimeRulerTrack;
-	FRandomGraphTrack GraphTrack;
-	//FFramesGraphTrack GraphTrack;
+	TSharedPtr<FTimeRulerTrack> TimeRulerTrack;
+	TSharedPtr<FGraphTrack> GraphTrack;
 	FTimingTrackViewport Viewport;
 
 	bool bIsViewportDirty;
diff --git a/Engine/Source/Developer/TraceInsights/Private/Insights/Widgets/SStatsTableCell.cpp b/Engine/Source/Developer/TraceInsights/Private/Insights/Widgets/SStatsTableCell.cpp
--- a/Engine/Source/Developer/TraceInsights/Private/Insights/Widgets/SStatsTableCell.cpp
+++ b/Engine/Source/Developer/TraceInsights/Private/Insights/Widgets/SStatsTableCell.cpp
@@ -6,6 +6,7 @@
 #include "SlateOptMacros.h"
 #include "Widgets/Images/SImage.h"
 #include "Widgets/Input/SButton.h"
+#include "Widgets/Layout/SBox.h"
 #include "Widgets/SBoxPanel.h"
 #include "Widgets/Text/STextBlock.h"
 #include "Widgets/Views/SExpanderArrow.h"
@@ -65,7 +66,7 @@ TSharedRef<SWidget> SStatsTableCell::GenerateWidgetForNameColumn(const FArgument
 			SNew(SExpanderArrow, TableRow)
 		]
 
-		// Event info icon + tooltip.
+		// Info icon + tooltip.
 		+ SHorizontalBox::Slot()
 		.AutoWidth()
 		.HAlign(HAlign_Center)
@@ -77,6 +78,25 @@ TSharedRef<SWidget> SStatsTableCell::GenerateWidgetForNameColumn(const FArgument
 			.ToolTip(SStatsViewTooltip::GetTableCellTooltip(StatsNodePtr))
 		]
 
+		// Color box.
+		+ SHorizontalBox::Slot()
+		.AutoWidth()
+		.HAlign(HAlign_Center)
+		.VAlign(VAlign_Center)
+		[
+			SNew(SBox)
+			.Visibility(this, &SStatsTableCell::GetColorBoxVisibility)
+			.WidthOverride(14.0f)
+			.HeightOverride(14.0f)
+			[
+				SNew(SBorder)
+				.BorderImage(FCoreStyle::Get().GetBrush("WhiteBrush"))
+				.BorderBackgroundColor(this, &SStatsTableCell::GetStatsBoxColorAndOpacity)
+				.HAlign(HAlign_Fill)
+				.VAlign(VAlign_Fill)
+			]
+		]
+
 		// Name
 		+ SHorizontalBox::Slot()
 		.AutoWidth()
diff --git a/Engine/Source/Developer/TraceInsights/Private/Insights/Widgets/SStatsTableCell.h b/Engine/Source/Developer/TraceInsights/Private/Insights/Widgets/SStatsTableCell.h
--- a/Engine/Source/Developer/TraceInsights/Private/Insights/Widgets/SStatsTableCell.h
+++ b/Engine/Source/Developer/TraceInsights/Private/Insights/Widgets/SStatsTableCell.h
@@ -88,6 +88,11 @@ protected:
 		SetHoveredTableCellDelegate.ExecuteIfBound(NAME_None, nullptr);
 	}
 
+	EVisibility GetColorBoxVisibility() const
+	{
+		return StatsNodePtr->IsAddedToGraph() ? EVisibility::Visible : EVisibility::Collapsed;
+	}
+
 	EVisibility GetHintIconVisibility() const
 	{
 		return IsHovered() ? EVisibility::Visible : EVisibility::Hidden;
@@ -121,6 +126,11 @@ protected:
 		return TextColor;
 	}
 
+	FSlateColor GetStatsBoxColorAndOpacity() const
+	{
+		return StatsNodePtr->GetColor();
+	}
+
 	FLinearColor GetShadowColorAndOpacity() const
 	{
 		const FLinearColor ShadowColor =
diff --git a/Engine/Source/Developer/TraceInsights/Private/Insights/Widgets/SStatsView.cpp b/Engine/Source/Developer/TraceInsights/Private/Insights/Widgets/SStatsView.cpp
--- a/Engine/Source/Developer/TraceInsights/Private/Insights/Widgets/SStatsView.cpp
+++ b/Engine/Source/Developer/TraceInsights/Private/Insights/Widgets/SStatsView.cpp
@@ -18,10 +18,13 @@
 // Insights
 #include "Insights/TimingProfilerCommon.h"
 #include "Insights/TimingProfilerManager.h"
+#include "Insights/ViewModels/GraphTrack.h"
 #include "Insights/ViewModels/StatsNodeHelper.h"
 #include "Insights/ViewModels/StatsViewColumnFactory.h"
 #include "Insights/Widgets/SStatsViewTooltip.h"
 #include "Insights/Widgets/SStatsTableRow.h"
+#include "Insights/Widgets/STimingProfilerWindow.h"
+#include "Insights/Widgets/STimingView.h"
 
 #define LOCTEXT_NAMESPACE "SStatsView"
 
@@ -816,20 +819,69 @@ void SStatsView::TreeView_OnGetChildren(FStatsNodePtr InParent, TArray<FStatsNod
 
 ////////////////////////////////////////////////////////////////////////////////////////////////////
 
+void SStatsView::UpdateStatsNode(FStatsNodePtr StatsNode)
+{
+	bool bAddedToGraphFlag = false;
+
+	if (!StatsNode->IsGroup())
+	{
+		TSharedPtr<STimingProfilerWindow> Wnd = FTimingProfilerManager::Get()->GetProfilerWindow();
+		TSharedPtr<STimingView> TimingView = Wnd.IsValid() ? Wnd->TimingView : nullptr;
+		if (TimingView.IsValid())
+		{
+			TSharedPtr<FTimingGraphTrack> GraphTrack = TimingView->GetMainTimingGraphTrack();
+			if (GraphTrack.IsValid())
+			{
+				uint32 StatsCounterId = static_cast<uint32>(StatsNode->GetId());
+				TSharedPtr<FTimingGraphSeries> Series = GraphTrack->GetStatsCounterSeries(StatsCounterId);
+				bAddedToGraphFlag = Series.IsValid();
+			}
+		}
+	}
+
+	StatsNode->SetAddedToGraphFlag(bAddedToGraphFlag);
+}
+
+////////////////////////////////////////////////////////////////////////////////////////////////////
+
 void SStatsView::TreeView_OnMouseButtonDoubleClick(FStatsNodePtr StatsNode)
 {
 	if (!StatsNode->IsGroup())
 	{
-		//im:TODO: const bool bIsStatTracked = FTimingProfilerManager::Get()->IsStatTracked(StatsNode->GetId());
-	//	if (!bIsStatTracked)
+		TSharedPtr<STimingProfilerWindow> Wnd = FTimingProfilerManager::Get()->GetProfilerWindow();
+		TSharedPtr<STimingView> TimingView = Wnd.IsValid() ? Wnd->TimingView : nullptr;
+		if (TimingView.IsValid())
+		{
+			TSharedPtr<FTimingGraphTrack> GraphTrack = TimingView->GetMainTimingGraphTrack();
+			if (GraphTrack.IsValid())
+			{
+				uint32 StatsCounterId = static_cast<uint32>(StatsNode->GetId());
+				TSharedPtr<FTimingGraphSeries> Series = GraphTrack->GetStatsCounterSeries(StatsCounterId);
+				if (Series.IsValid())
+				{
+					GraphTrack->RemoveStatsCounterSeries(StatsCounterId);
+					GraphTrack->SetDirtyFlag();
+					StatsNode->SetAddedToGraphFlag(false);
+				}
+				else
+				{
+					GraphTrack->AddStatsCounterSeries(StatsCounterId, StatsNode->GetColor(), 0.0, 1.0);
+					GraphTrack->SetDirtyFlag();
+					StatsNode->SetAddedToGraphFlag(true);
+				}
+			}
+		}
+
+		//im:TODO: const bool bIsStatsCounterTracked = FTimingProfilerManager::Get()->IsStatsCounterTracked(StatsNode->GetId());
+	//	if (!bIsStatsCounterTracked)
 	//	{
-	//		// Add a new graph.
-	//		FTimingProfilerManager::Get()->TrackStat(StatsNode->GetId());
+	//		// Add a new graph series.
+	//		FTimingProfilerManager::Get()->TrackStatsCounter(StatsNode->GetId());
 	//	}
 	//	else
 	//	{
-	//		// Remove a graph
-	//		FTimingProfilerManager::Get()->UntrackStat(StatsNode->GetId());
+	//		// Remove the corresponding graph series.
+	//		FTimingProfilerManager::Get()->UntrackStatsCounter(StatsNode->GetId());
 	//	}
 	}
 	else
@@ -1433,6 +1485,7 @@ void SStatsView::RebuildTree(bool bResync)
 							Counter.IsFloatingPoint() ? TEXT("float") : TEXT("int64"));
 				EStatsNodeType Type = Counter.IsFloatingPoint() ? EStatsNodeType::Float : EStatsNodeType::Int64;
 				FStatsNodePtr StatsNodePtr = MakeShareable(new FStatsNode(Counter.GetId(), Name, Group, Type));
+				UpdateStatsNode(StatsNodePtr);
 				StatsNodes.Add(StatsNodePtr);
 				//StatsNodesMap.Add(Name, StatsNodePtr);
 				StatsNodesIdMap.Add(Counter.GetId(), StatsNodePtr);
diff --git a/Engine/Source/Developer/TraceInsights/Private/Insights/Widgets/SStatsView.h b/Engine/Source/Developer/TraceInsights/Private/Insights/Widgets/SStatsView.h
--- a/Engine/Source/Developer/TraceInsights/Private/Insights/Widgets/SStatsView.h
+++ b/Engine/Source/Developer/TraceInsights/Private/Insights/Widgets/SStatsView.h
@@ -240,6 +240,10 @@ protected:
 	bool ContextMenu_ResetColumns_CanExecute() const;
 	void ContextMenu_ResetColumns_Execute();
 
+	////////////////////////////////////////////////////////////////////////////////////////////////////
+
+	void UpdateStatsNode(FStatsNodePtr StatsNode);
+
 protected:
 	/** A weak pointer to the profiler session used to populate this widget. */
 	TSharedPtr<const Trace::IAnalysisSession>/*Weak*/ Session;
diff --git a/Engine/Source/Developer/TraceInsights/Private/Insights/Widgets/STimingView.cpp b/Engine/Source/Developer/TraceInsights/Private/Insights/Widgets/STimingView.cpp
--- a/Engine/Source/Developer/TraceInsights/Private/Insights/Widgets/STimingView.cpp
+++ b/Engine/Source/Developer/TraceInsights/Private/Insights/Widgets/STimingView.cpp
@@ -65,7 +65,7 @@ STimingView::STimingView()
 	: bAssetLoadingMode(false)
 	, TimeRulerTrack(FBaseTimingTrack::GenerateId())
 	, MarkersTrack(FBaseTimingTrack::GenerateId())
-	, GraphTrack(FBaseTimingTrack::GenerateId())
+	, GraphTrack(MakeShareable(new FTimingGraphTrack(FBaseTimingTrack::GenerateId())))
 	, TooltipWidth(MinTooltipWidth)
 	, TooltipAlpha(0.0f)
 	, WhiteBrush(FCoreStyle::Get().GetBrush("WhiteBrush"))
@@ -243,8 +243,8 @@ void STimingView::Reset()
 	TimeRulerTrack.Reset();
 	MarkersTrack.Reset();
 
-	GraphTrack.Reset();
-	GraphTrack.SetVisibilityFlag(false);
+	GraphTrack->Reset();
+	GraphTrack->SetVisibilityFlag(false);
 
 	//////////////////////////////////////////////////
 
@@ -336,11 +336,11 @@ void STimingView::Tick(const FGeometry& AllottedGeometry, const double InCurrent
 		MarkersTrack.SetPosY(TopOffset);
 		TopOffset += MarkersTrack.GetHeight();
 	}
-	if (GraphTrack.IsVisible())
+	if (GraphTrack->IsVisible())
 	{
-		GraphTrack.SetPosY(TopOffset);
-		GraphTrack.SetHeight(200.0f);
-		TopOffset += GraphTrack.GetHeight();
+		GraphTrack->SetPosY(TopOffset);
+		GraphTrack->SetHeight(200.0f);
+		TopOffset += GraphTrack->GetHeight();
 	}
 	Viewport.TopOffset = TopOffset;
 
@@ -671,14 +671,11 @@ void STimingView::Tick(const FGeometry& AllottedGeometry, const double InCurrent
 	{
 		bAreTimingEventsTracksDirty = false;
 
+		GraphTrack->SetDirtyFlag();
+
 		FStopwatch Stopwatch;
 		Stopwatch.Start();
 
-		if (GraphTrack.IsVisible())
-		{
-			GraphTrack.Update(Viewport);
-		}
-
 		for (int32 TrackIndex = 0; TrackIndex < TimingEventsTracks.Num(); ++TrackIndex)
 		{
 			TimingEventsTracks[TrackIndex]->Update(Viewport);
@@ -687,6 +684,12 @@ void STimingView::Tick(const FGeometry& AllottedGeometry, const double InCurrent
 		Stopwatch.Stop();
 		TimelineCacheUpdateDurationHistory.AddValue(Stopwatch.AccumulatedTime);
 	}
+
+	if (GraphTrack->IsVisible() && GraphTrack->IsDirty())
+	{
+		GraphTrack->ClearDirtyFlag();
+		GraphTrack->Update(Viewport);
+	}
 }
 
 ////////////////////////////////////////////////////////////////////////////////////////////////////
@@ -777,9 +780,9 @@ int32 STimingView::OnPaint(const FPaintArgs& Args, const FGeometry& AllottedGeom
 		Helper.EndTimelines();
 	}
 
-	if (GraphTrack.IsVisible())
+	if (GraphTrack->IsVisible())
 	{
-		GraphTrack.Draw(DrawContext, Viewport);
+		GraphTrack->Draw(DrawContext, Viewport);
 	}
 
 	if (!FTimingEvent::AreEquals(SelectedTimingEvent, HoveredTimingEvent))
@@ -970,10 +973,10 @@ int32 STimingView::OnPaint(const FPaintArgs& Args, const FGeometry& AllottedGeom
 			DbgX, DbgY,
 			FString::Format(TEXT("{0} events : {1} points, {2} lines, {3} boxes"),
 				{
-					FText::AsNumber(GraphTrack.GetNumAddedEvents()).ToString(),
-					FText::AsNumber(GraphTrack.GetNumDrawPoints()).ToString(),
-					FText::AsNumber(GraphTrack.GetNumDrawLines()).ToString(),
-					FText::AsNumber(GraphTrack.GetNumDrawBoxes()).ToString(),
+					FText::AsNumber(GraphTrack->GetNumAddedEvents()).ToString(),
+					FText::AsNumber(GraphTrack->GetNumDrawPoints()).ToString(),
+					FText::AsNumber(GraphTrack->GetNumDrawLines()).ToString(),
+					FText::AsNumber(GraphTrack->GetNumDrawBoxes()).ToString(),
 				}),
 				SummaryFont, DbgTextColor
 				);
@@ -1873,6 +1876,11 @@ void STimingView::UpdateIo()
 				{
 					TSharedPtr<FIoFileActivity> Activity;
 
+					//if (EventEndTime == std::numeric_limits<double>::infinity())
+					//{
+					//	EventEndTime = EventStartTime;
+					//}
+
 					if (!FileActivityMap.Contains(FileInfo.Id))
 					{
 						Activity = MakeShareable(new FIoFileActivity());
@@ -2401,7 +2409,7 @@ FReply STimingView::OnMouseButtonUp(const FGeometry& MyGeometry, const FPointerE
 
 			if (bIsValidForMouseClick)
 			{
-				ShowContextMenu(MouseEvent.GetScreenSpacePosition(), MouseEvent);
+				ShowContextMenu(MouseEvent);
 			}
 
 			bIsDragging = false;
@@ -2826,8 +2834,7 @@ FReply STimingView::OnKeyDown(const FGeometry& MyGeometry, const FKeyEvent& InKe
 	}
 	else if (InKeyEvent.GetKey() == EKeys::G) // debug: toggles Graph track on/off
 	{
-		GraphTrack.ToggleVisibility();
-		bAreTimingEventsTracksDirty = true;
+		ShowHideGraphTrack_Execute();
 		return FReply::Handled();
 	}
 	else if (InKeyEvent.GetKey() == EKeys::Y) // debug: toggles GPU track on/off
@@ -2902,30 +2909,89 @@ FReply STimingView::OnKeyUp(const FGeometry& MyGeometry, const FKeyEvent& InKeyE
 
 ////////////////////////////////////////////////////////////////////////////////////////////////////
 
-void STimingView::ShowContextMenu(const FVector2D& ScreenSpacePosition, const FPointerEvent& MouseEvent)
+void STimingView::ShowContextMenu(const FPointerEvent& MouseEvent)
 {
-	/* TODO
-	TSharedPtr<FUICommandList> ProfilerCommandList = FTimingProfilerManager::Get()->GetCommandList();
-	const FTimingProfilerCommands& ProfilerCommands = FTimingProfilerManager::GetCommands();
-	const FTimingViewCommands& Commands = FTimingViewCommands::Get();
+	//TSharedPtr<FUICommandList> ProfilerCommandList = FTimingProfilerManager::Get()->GetCommandList();
+	//const FTimingProfilerCommands& ProfilerCommands = FTimingProfilerManager::GetCommands();
+	//const FTimingViewCommands& Commands = FTimingViewCommands::Get();
 
 	const bool bShouldCloseWindowAfterMenuSelection = true;
-	FMenuBuilder MenuBuilder(bShouldCloseWindowAfterMenuSelection, ProfilerCommandList);
+	//FMenuBuilder MenuBuilder(bShouldCloseWindowAfterMenuSelection, ProfilerCommandList);
+	FMenuBuilder MenuBuilder(bShouldCloseWindowAfterMenuSelection, NULL);
 
-	MenuBuilder.BeginSection(TEXT("Misc"), LOCTEXT("Miscellaneous", "Miscellaneous"));
+	if (GraphTrack->IsVisible() &&
+		MousePosition.Y >= GraphTrack->GetPosY() &&
+		MousePosition.Y < GraphTrack->GetPosY() + GraphTrack->GetHeight())
 	{
-		MenuBuilder.AddMenuEntry(Commands.ShowAllGpuTracks);
-		MenuBuilder.AddMenuEntry(Commands.ShowAllCpuTracks);
-		MenuBuilder.AddMenuEntry(ProfilerCommands.ToggleTimersViewVisibility);
-		MenuBuilder.AddMenuEntry(ProfilerCommands.ToggleStatsCountersViewVisibility);
+		//GraphTrack->BuildContextMenu(MenuBuilder);
+		MenuBuilder.AddSubMenu
+		(
+			LOCTEXT("ContextMenu_Header_GraphTrack", "Graph Track"),
+			TAttribute<FText>(), // no tooltip
+			FNewMenuDelegate::CreateSP(GraphTrack.Get(), &FBaseTimingTrack::BuildContextMenu),
+			false,
+			FSlateIcon()
+		);
 	}
-	MenuBuilder.EndSection();
+	else
+	{
+		MenuBuilder.BeginSection(TEXT("Empty"));
+		{
+			struct FLocal
+			{
+				static bool ReturnFalse()
+				{
+					return false;
+				}
+			};
+
+			FUIAction DummyUIAction;
+			DummyUIAction.CanExecuteAction = FCanExecuteAction::CreateStatic(&FLocal::ReturnFalse);
+
+			FText Title;
+			if (HoveredTimingEvent.Track != nullptr)
+			{
+				if (HoveredTimingEvent.Track->GetGroupName() != nullptr)
+				{
+					Title = FText::Format(LOCTEXT("TrackTitleGroupFmt", "{0} (Group: {1})"), FText::FromString(HoveredTimingEvent.Track->GetName()), FText::FromString(HoveredTimingEvent.Track->GetGroupName()));
+				}
+				else
+				{
+					Title = FText::FromString(HoveredTimingEvent.Track->GetName());
+				}
+			}
+			else
+			{
+				Title = LOCTEXT("ContextMenu_NA", "N/A");
+			}
+
+			MenuBuilder.AddMenuEntry
+			(
+				Title,
+				LOCTEXT("ContextMenu_NA_Desc", "No actions available."),
+				FSlateIcon(),
+				DummyUIAction,
+				NAME_None,
+				EUserInterfaceActionType::Button
+			);
+		}
+		MenuBuilder.EndSection();
+	}
+
+	//MenuBuilder.BeginSection(TEXT("Misc"), LOCTEXT("Miscellaneous", "Miscellaneous"));
+	//{
+	//	MenuBuilder.AddMenuEntry(Commands.ShowAllGpuTracks);
+	//	MenuBuilder.AddMenuEntry(Commands.ShowAllCpuTracks);
+	//	MenuBuilder.AddMenuEntry(ProfilerCommands.ToggleTimersViewVisibility);
+	//	MenuBuilder.AddMenuEntry(ProfilerCommands.ToggleStatsCountersViewVisibility);
+	//}
+	//MenuBuilder.EndSection();
 
 	TSharedRef<SWidget> MenuWidget = MenuBuilder.MakeWidget();
 
 	FWidgetPath EventPath = MouseEvent.GetEventPath() != nullptr ? *MouseEvent.GetEventPath() : FWidgetPath();
+	const FVector2D ScreenSpacePosition = MouseEvent.GetScreenSpacePosition();
 	FSlateApplication::Get().PushMenu(SharedThis(this), EventPath, MenuWidget, ScreenSpacePosition, FPopupTransitionEffect::ContextMenu);
-	*/
 }
 
 ////////////////////////////////////////////////////////////////////////////////////////////////////
@@ -3823,6 +3889,17 @@ TSharedRef<SWidget> STimingView::MakeTracksFilterMenu()
 	{
 		const FTimingViewCommands& Commands = FTimingViewCommands::Get();
 
+		MenuBuilder.AddMenuEntry(
+			LOCTEXT("ShowGraphTrack", "Graph Track - G"),
+			LOCTEXT("ShowGraphTrack_Tooltip", "Show/hide the Graph track"),
+			FSlateIcon(),
+			FUIAction(FExecuteAction::CreateSP(this, &STimingView::ShowHideGraphTrack_Execute),
+				FCanExecuteAction(),
+				FIsActionChecked::CreateSP(this, &STimingView::ShowHideGraphTrack_IsChecked)),
+			NAME_None,
+			EUserInterfaceActionType::ToggleButton
+		);
+
 		//TODO: MenuBuilder.AddMenuEntry(Commands.ShowAllGpuTracks);
 		MenuBuilder.AddMenuEntry(
 			LOCTEXT("ShowAllGpuTracks", "GPU Track - Y"),
@@ -3954,6 +4031,21 @@ void STimingView::CreateTracksMenu(FMenuBuilder& MenuBuilder)
 
 ////////////////////////////////////////////////////////////////////////////////////////////////////
 
+bool STimingView::ShowHideGraphTrack_IsChecked() const
+{
+	return GraphTrack->IsVisible();
+}
+
+////////////////////////////////////////////////////////////////////////////////////////////////////
+
+void STimingView::ShowHideGraphTrack_Execute()
+{
+	GraphTrack->ToggleVisibility();
+	bAreTimingEventsTracksDirty = true;
+}
+
+////////////////////////////////////////////////////////////////////////////////////////////////////
+
 bool STimingView::ShowHideAllCpuTracks_IsChecked() const
 {
 	return bShowHideAllCpuTracks;
diff --git a/Engine/Source/Developer/TraceInsights/Private/Insights/Widgets/STimingView.h b/Engine/Source/Developer/TraceInsights/Private/Insights/Widgets/STimingView.h
--- a/Engine/Source/Developer/TraceInsights/Private/Insights/Widgets/STimingView.h
+++ b/Engine/Source/Developer/TraceInsights/Private/Insights/Widgets/STimingView.h
@@ -88,6 +88,9 @@ public:
 	void CreateThreadGroupsMenu(FMenuBuilder& MenuBuilder);
 	void CreateTracksMenu(FMenuBuilder& MenuBuilder);
 
+	bool ShowHideGraphTrack_IsChecked() const;
+	void ShowHideGraphTrack_Execute();
+
 	bool ShowHideAllCpuTracks_IsChecked() const;
 	void ShowHideAllCpuTracks_Execute();
 
@@ -269,6 +272,8 @@ public:
 	bool IsGpuTrackVisible() const;
 	bool IsCpuTrackVisible(uint32 ThreadId) const;
 
+	TSharedPtr<FTimingGraphTrack> GetMainTimingGraphTrack() { return StaticCastSharedPtr<FTimingGraphTrack>(GraphTrack); }
+
 protected:
 	virtual FVector2D ComputeDesiredSize(float) const override
 	{
@@ -292,7 +297,7 @@ protected:
 
 	void DrawTimeRangeSelection(FDrawContext& DrawContext) const;
 
-	void ShowContextMenu(const FVector2D& ScreenSpacePosition, const FPointerEvent& MouseEvent);
+	void ShowContextMenu(const FPointerEvent& MouseEvent);
 
 	/** Binds our UI commands to delegates. */
 	void BindCommands();
@@ -482,7 +487,7 @@ protected:
 	FMarkersTimingTrack MarkersTrack;
 
 	/** A graph track for frame times. */
-	FFramesGraphTrack GraphTrack;
+	TSharedPtr<FGraphTrack> GraphTrack;
 
 	////////////////////////////////////////////////////////////
 
diff --git a/Engine/Source/Developer/TraceServices/Private/Model/Counters.cpp b/Engine/Source/Developer/TraceServices/Private/Model/Counters.cpp
--- a/Engine/Source/Developer/TraceServices/Private/Model/Counters.cpp
+++ b/Engine/Source/Developer/TraceServices/Private/Model/Counters.cpp
@@ -147,7 +147,7 @@ void FCounterProvider::EnumerateCounters(TFunctionRef<void(const ICounter &)> Ca
 
 bool FCounterProvider::ReadCounter(uint32 CounterId, TFunctionRef<void(const ICounter &)> Callback) const
 {
-	if (CounterId < uint32(Counters.Num()))
+	if (CounterId >= uint32(Counters.Num()))
 	{
 		return false;
 	}
