[Insights]
 - Moved GraphTrack functionality re Frames, Timers and Stats Counter series in separate TimingGraphTrack.cpp/.h file (out of GraphTrack.cpp/.h)
 - Separated the specific TimingEventsTrack implementations (Cpu/Gpu tracks, Loading tracks, FileActivity tracks) in own classes (see FileActivityTimingTrack.h/.cpp, LoadingTimingTrack.h/.cpp and THreadTimingTrack.h/.cpp) out of STimingView class
  - Increased version to 0.08

diff --git a/Engine/Source/Developer/TraceInsights/Private/Insights/TimingProfilerManager.cpp b/Engine/Source/Developer/TraceInsights/Private/Insights/TimingProfilerManager.cpp
--- a/Engine/Source/Developer/TraceInsights/Private/Insights/TimingProfilerManager.cpp
+++ b/Engine/Source/Developer/TraceInsights/Private/Insights/TimingProfilerManager.cpp
@@ -244,4 +244,24 @@ void FTimingProfilerManager::ShowHideLogView(const bool bLogViewVisibleState)
 
 ////////////////////////////////////////////////////////////////////////////////////////////////////
 
+const FTimerNodePtr FTimingProfilerManager::GetTimerNode(uint64 TypeId) const
+{
+	const FTimerNodePtr* TimerNodePtrPtr = nullptr;
+	TSharedPtr<STimingProfilerWindow> Wnd = GetProfilerWindow();
+	if (Wnd.IsValid() && Wnd->TimersView)
+	{
+		TimerNodePtrPtr = Wnd->TimersView->GetTimerNode(TypeId);
+		if (TimerNodePtrPtr == nullptr)
+		{
+			// List of timers in TimersView not up to date?
+			// Refresh and try again.
+			Wnd->TimersView->RebuildTree(false);
+			TimerNodePtrPtr = Wnd->TimersView->GetTimerNode(TypeId);
+		}
+	}
+	return TimerNodePtrPtr ? *TimerNodePtrPtr : nullptr;
+}
+
+////////////////////////////////////////////////////////////////////////////////////////////////////
+
 #undef LOCTEXT_NAMESPACE
diff --git a/Engine/Source/Developer/TraceInsights/Private/Insights/TimingProfilerManager.h b/Engine/Source/Developer/TraceInsights/Private/Insights/TimingProfilerManager.h
--- a/Engine/Source/Developer/TraceInsights/Private/Insights/TimingProfilerManager.h
+++ b/Engine/Source/Developer/TraceInsights/Private/Insights/TimingProfilerManager.h
@@ -9,6 +9,7 @@
 // Insights
 #include "Insights/InsightsManager.h"
 #include "Insights/TimingProfilerCommands.h"
+#include "Insights/ViewModels/TimerNode.h"
 
 class STimingProfilerWindow;
 
@@ -131,6 +132,8 @@ public:
 
 	void OnSessionChanged();
 
+	const FTimerNodePtr GetTimerNode(uint64 TypeId) const;
+
 protected:
 	/** Updates this manager, done through FCoreTicker. */
 	bool Tick(float DeltaTime);
diff --git a/Engine/Source/Developer/TraceInsights/Private/Insights/ViewModels/DrawHelpers.cpp b/Engine/Source/Developer/TraceInsights/Private/Insights/ViewModels/DrawHelpers.cpp
--- a/Engine/Source/Developer/TraceInsights/Private/Insights/ViewModels/DrawHelpers.cpp
+++ b/Engine/Source/Developer/TraceInsights/Private/Insights/ViewModels/DrawHelpers.cpp
@@ -183,8 +183,8 @@ void FDrawHelpers::DrawTimeRangeSelection(
 				//im: This should be a more efficeint way top draw the arrows, but it renders them with artifacts (missing vertical lines; shader bug?)!
 
 				const FSlateBrush* MyBrush = WhiteBrush;
-				FSlateShaderResourceProxy* ResourceProxy = FSlateDataPayload::ResourceManager->GetShaderResource(*MyBrush);
 				FSlateResourceHandle ResourceHandle = FSlateApplication::Get().GetRenderer()->GetResourceHandle(*MyBrush);
+				const FSlateShaderResourceProxy* ResourceProxy = ResourceHandle.GetResourceProxy();
 
 				FVector2D AtlasOffset = ResourceProxy ? ResourceProxy->StartUV : FVector2D(0.0f, 0.0f);
 				FVector2D AtlasUVSize = ResourceProxy ? ResourceProxy->SizeUV : FVector2D(1.0f, 1.0f);
diff --git a/Engine/Source/Developer/TraceInsights/Private/Insights/ViewModels/FileActivityTimingTrack.cpp b/Engine/Source/Developer/TraceInsights/Private/Insights/ViewModels/FileActivityTimingTrack.cpp
new file mode 100644
--- /dev/null
+++ b/Engine/Source/Developer/TraceInsights/Private/Insights/ViewModels/FileActivityTimingTrack.cpp
@@ -0,0 +1,704 @@
+// Copyright 1998-2019 Epic Games, Inc. All Rights Reserved.
+
+#include "FileActivityTimingTrack.h"
+
+#include "Fonts/FontMeasure.h"
+#include "Styling/SlateBrush.h"
+#include "TraceServices/AnalysisService.h"
+#include "TraceServices/SessionService.h"
+
+// Insights
+#include "Insights/Common/Stopwatch.h"
+#include "Insights/Common/PaintUtils.h"
+#include "Insights/Common/TimeUtils.h"
+#include "Insights/InsightsManager.h"
+#include "Insights/TimingProfilerCommon.h"
+#include "Insights/ViewModels/TimingEvent.h"
+#include "Insights/ViewModels/TimingTrackViewport.h"
+#include "Insights/ViewModels/TimingViewDrawHelper.h"
+
+#include <limits>
+
+#define LOCTEXT_NAMESPACE "FileActivityTimingTrack"
+
+// The FileActivity (I/O) timelines are just prototypes for now.
+// Below code will be removed once the functionality is moved in analyzer.
+
+////////////////////////////////////////////////////////////////////////////////////////////////////
+
+const TCHAR* GetFileActivityTypeName(Trace::EFileActivityType Type)
+{
+	static_assert(Trace::FileActivityType_Open == 0, "Trace::EFileActivityType enum has changed!?");
+	static_assert(Trace::FileActivityType_Close == 1, "Trace::EFileActivityType enum has changed!?");
+	static_assert(Trace::FileActivityType_Read == 2, "Trace::EFileActivityType enum has changed!?");
+	static_assert(Trace::FileActivityType_Write == 3, "Trace::EFileActivityType enum has changed!?");
+	static_assert(Trace::FileActivityType_Count == 4, "Trace::EFileActivityType enum has changed!?");
+	static const TCHAR* GFileActivityTypeNames[] =
+	{
+		TEXT("Open"),
+		TEXT("Close"),
+		TEXT("Read"),
+		TEXT("Write"),
+		TEXT("Idle"), // virtual events added for cases where Close event is more than 1s away from last Open/Read/Write event.
+		TEXT("NotClosed") // virtual events added when an Open activity never closes
+	};
+	return GFileActivityTypeNames[Type];
+}
+
+////////////////////////////////////////////////////////////////////////////////////////////////////
+
+uint32 GetFileActivityTypeColor(Trace::EFileActivityType Type)
+{
+	static const uint32 GFileActivityTypeColors[] =
+	{
+		0xFFCCAA33, // open
+		0xFF33AACC, // close
+		0xFF33AA33, // read
+		0xFFDD33CC, // write
+		0x55333333, // idle
+		0x55553333, // close
+	};
+	return GFileActivityTypeColors[Type];
+}
+
+////////////////////////////////////////////////////////////////////////////////////////////////////
+// FFileActivitySharedState
+////////////////////////////////////////////////////////////////////////////////////////////////////
+
+void FFileActivitySharedState::Reset()
+{
+	bForceIoEventsUpdate = false;
+
+	bMergeIoLanes = true;
+	bShowFileActivityBackgroundEvents = false;
+
+	FileActivities.Reset();
+	FileActivityMap.Reset();
+
+	AllIoEvents.Reset();
+}
+
+////////////////////////////////////////////////////////////////////////////////////////////////////
+
+void FFileActivitySharedState::Update()
+{
+	if (bForceIoEventsUpdate)
+	{
+		bForceIoEventsUpdate = false;
+
+		FileActivities.Reset();
+		FileActivityMap.Reset();
+		AllIoEvents.Reset();
+
+		FStopwatch Stopwatch;
+		Stopwatch.Start();
+
+		TSharedPtr<const Trace::IAnalysisSession> Session = FInsightsManager::Get()->GetSession();
+		if (Session.IsValid() && Trace::ReadFileActivityProvider(*Session.Get()))
+		{
+			Trace::FAnalysisSessionReadScope SessionReadScope(*Session.Get());
+
+			const Trace::IFileActivityProvider& FileActivityProvider = *Trace::ReadFileActivityProvider(*Session.Get());
+
+			// Enumerate all IO events and cache them.
+			FileActivityProvider.EnumerateFileActivity([this](const Trace::FFileInfo& FileInfo, const Trace::IFileActivityProvider::Timeline& Timeline)
+			{
+				Timeline.EnumerateEvents(-std::numeric_limits<double>::infinity(), +std::numeric_limits<double>::infinity(),
+					[this, &FileInfo, &Timeline](double EventStartTime, double EventEndTime, uint32 EventDepth, const Trace::FFileActivity* FileActivity)
+				{
+					TSharedPtr<FIoFileActivity> Activity;
+
+					//if (EventEndTime == std::numeric_limits<double>::infinity())
+					//{
+					//	EventEndTime = EventStartTime;
+					//}
+
+					if (!FileActivityMap.Contains(FileInfo.Id))
+					{
+						Activity = MakeShareable(new FIoFileActivity());
+
+						Activity->Id = FileInfo.Id;
+						Activity->Path = FileInfo.Path;
+						Activity->StartTime = EventStartTime;
+						Activity->EndTime = EventEndTime;
+						Activity->EventCount = 1;
+						Activity->Depth = -1;
+
+						FileActivities.Add(Activity);
+						FileActivityMap.Add(FileInfo.Id, Activity);
+					}
+					else
+					{
+						Activity = FileActivityMap[FileInfo.Id];
+
+						if (FileActivity->ActivityType != Trace::FileActivityType_Close)
+						{
+							ensure(EventStartTime >= Activity->StartTime);
+							if (EventStartTime < Activity->StartTime)
+							{
+								Activity->StartTime = EventStartTime;
+							}
+
+							if (EventEndTime > Activity->EndTime)
+							{
+								Activity->EndTime = EventEndTime;
+							}
+						}
+
+						Activity->EventCount++;
+					}
+
+					if (bMergeIoLanes)
+					{
+						EventDepth = 0;
+					}
+					else
+					{
+						EventDepth = FileInfo.Id % 32; // simple layout
+					}
+
+					uint32 Type = ((uint32)FileActivity->ActivityType & 0x0F) | (FileActivity->Failed ? 0x80 : 0);
+
+					AllIoEvents.Add(FIoTimingEvent{ EventStartTime, EventEndTime, EventDepth, Type, FileActivity->Offset, FileActivity->Size, Activity });
+				});
+
+				return true;
+			});
+		}
+
+		Stopwatch.Stop();
+		UE_LOG(TimingProfiler, Log, TEXT("[IO] Enumerated %s events (%s file activities) in %s."),
+			*FText::AsNumber(AllIoEvents.Num()).ToString(),
+			*FText::AsNumber(FileActivities.Num()).ToString(),
+			*TimeUtils::FormatTimeAuto(Stopwatch.GetAccumulatedTime()));
+		Stopwatch.Restart();
+
+		// Sort cached IO file activities by Start Time.
+		FileActivities.Sort([](const TSharedPtr<FIoFileActivity>& A, const TSharedPtr<FIoFileActivity>& B) { return A->StartTime < B->StartTime; });
+
+		// Sort cached IO events by Start Time.
+		AllIoEvents.Sort([](const FIoTimingEvent& A, const FIoTimingEvent& B) { return A.StartTime < B.StartTime; });
+
+		Stopwatch.Stop();
+		UE_LOG(TimingProfiler, Log, TEXT("[IO] Sorted file activities and events in %s."), *TimeUtils::FormatTimeAuto(Stopwatch.GetAccumulatedTime()));
+
+		if (bMergeIoLanes)
+		{
+			//////////////////////////////////////////////////
+			// Compute depth for file activities (avoids overlaps).
+
+			Stopwatch.Restart();
+
+			TArray<TSharedPtr<FIoFileActivity>> ActivityLanes;
+
+			for (TSharedPtr<FIoFileActivity> FileActivity : FileActivities)
+			{
+				// Find lane (avoiding overlaps with other file activities).
+				for (int32 LaneIndex = 0; LaneIndex < ActivityLanes.Num(); ++LaneIndex)
+				{
+					TSharedPtr<FIoFileActivity> Lane = ActivityLanes[LaneIndex];
+
+					if (FileActivity->StartTime >= Lane->EndTime)
+					{
+						FileActivity->Depth = LaneIndex;
+						ActivityLanes[LaneIndex] = FileActivity;
+						break;
+					}
+				}
+
+				if (FileActivity->Depth < 0)
+				{
+					const int32 MaxLanes = 10000;
+					if (ActivityLanes.Num() < MaxLanes)
+					{
+						// Add new lane.
+						FileActivity->Depth = ActivityLanes.Num();
+						ActivityLanes.Add(FileActivity);
+					}
+					else
+					{
+						int32 LaneIndex = ActivityLanes.Num() - 1;
+						FileActivity->Depth = LaneIndex;
+						TSharedPtr<FIoFileActivity> Lane = ActivityLanes[LaneIndex];
+						if (FileActivity->EndTime > Lane->EndTime)
+						{
+							ActivityLanes[LaneIndex] = FileActivity;
+						}
+					}
+				}
+			}
+
+			Stopwatch.Stop();
+			UE_LOG(TimingProfiler, Log, TEXT("[IO] Computed layout for file activities in %s."), *TimeUtils::FormatTimeAuto(Stopwatch.GetAccumulatedTime()));
+
+			//////////////////////////////////////////////////
+
+			Stopwatch.Restart();
+
+			for (FIoTimingEvent& Event : AllIoEvents)
+			{
+				Event.Depth = static_cast<uint32>(Event.FileActivity->Depth);
+			}
+
+			Stopwatch.Stop();
+			UE_LOG(TimingProfiler, Log, TEXT("[IO] Updated depth for events in %s."), *TimeUtils::FormatTimeAuto(Stopwatch.GetAccumulatedTime()));
+
+			//////////////////////////////////////////////////
+			/*
+
+			Stopwatch.Restart();
+
+			struct FIoLane
+			{
+				uint32 FileId;
+				const TCHAR* Path;
+				double LastEndTime;
+				double EndTime;
+			};
+
+			TArray<FIoLane> Lanes;
+
+			TArray<FIoTimingEvent> IoEventsToAdd;
+
+			for (FIoTimingEvent& Event : AllIoEvents)
+			{
+				uint64 TimelineId = Event.Depth;
+
+				int32 Depth = -1;
+
+				bool bIsCloseEvent = false;
+
+				const Trace::EFileActivityType ActivityType = static_cast<Trace::EFileActivityType>(Event.Type & 0x0F);
+
+				if (ActivityType == Trace::FileActivityType_Open)
+				{
+					// Find lane (avoiding overlaps with other opened files).
+					for (int32 LaneIndex = 0; LaneIndex < Lanes.Num(); ++LaneIndex)
+					{
+						FIoLane& Lane = Lanes[LaneIndex];
+						if (Event.StartTime >= Lane.EndTime)
+						{
+							Depth = LaneIndex;
+							break;
+						}
+					}
+
+					if (Depth < 0)
+					{
+						// Add new lane.
+						Depth = Lanes.AddDefaulted();
+					}
+
+					const bool bHasFailed = ((Event.Type & 0xF0) != 0);
+
+					FIoLane& Lane = Lanes[Depth];
+					Lane.FileId = TimelineId;
+					Lane.Path = Event.Path;
+					Lane.LastEndTime = Event.EndTime;
+					Lane.EndTime = bHasFailed ? Event.EndTime : std::numeric_limits<double>::infinity();
+				}
+				else if (ActivityType == Trace::FileActivityType_Close)
+				{
+					bIsCloseEvent = true;
+
+					// Find lane with same id.
+					for (int32 LaneIndex = 0; LaneIndex < Lanes.Num(); ++LaneIndex)
+					{
+						FIoLane& Lane = Lanes[LaneIndex];
+						if (Lane.FileId == TimelineId)
+						{
+							// Adds an Idle event, but only if has passed at least 1s since last open/read/write activity.
+							if (Event.StartTime - Lane.LastEndTime > 1.0)
+							{
+								IoEventsToAdd.Add(FIoTimingEvent{ Lane.LastEndTime, Event.StartTime, static_cast<uint32>(LaneIndex), Trace::FileActivityType_Count, Event.Path });
+							}
+							Lane.LastEndTime = Event.EndTime;
+							Lane.EndTime = Event.EndTime;
+							Depth = LaneIndex;
+							break;
+						}
+					}
+					ensure(Depth >= 0);
+				}
+				else
+				{
+					// All other events should be inside the virtual Open-Close event.
+					// Find lane with same id.
+					for (int32 LaneIndex = 0; LaneIndex < Lanes.Num(); ++LaneIndex)
+					{
+						FIoLane& Lane = Lanes[LaneIndex];
+						if (Lane.FileId == TimelineId)
+						{
+							Lane.LastEndTime = Event.EndTime;
+							Depth = LaneIndex;
+							break;
+						}
+					}
+					ensure(Depth >= 0);
+				}
+
+				if (Depth < 0) // just in case
+				{
+					// Add new lane.
+					Depth = Lanes.AddDefaulted();
+					FIoLane& Lane = Lanes[Depth];
+					Lane.FileId = TimelineId;
+					Lane.Path = Event.Path;
+					Lane.LastEndTime = Event.EndTime;
+					Lane.EndTime = bIsCloseEvent ? Event.EndTime : std::numeric_limits<double>::infinity();
+				}
+
+				Event.Depth = Depth;
+			}
+
+			Stopwatch.Stop();
+			UE_LOG(TimingProfiler, Log, TEXT("[IO] Merge 1/3 in %s."), *TimeUtils::FormatTimeAuto(Stopwatch.GetAccumulatedTime()));
+			Stopwatch.Restart();
+
+			for (int32 LaneIndex = 0; LaneIndex < Lanes.Num(); ++LaneIndex)
+			{
+				FIoLane& Lane = Lanes[LaneIndex];
+				if (Lane.EndTime == std::numeric_limits<double>::infinity())
+				{
+					IoEventsToAdd.Add(FIoTimingEvent{ Lane.LastEndTime, Lane.EndTime, static_cast<uint32>(LaneIndex), Trace::FileActivityType_Count + 1, Lane.Path });
+				}
+			}
+
+			for (const FIoTimingEvent& Event : IoEventsToAdd)
+			{
+				AllIoEvents.Add(Event);
+			}
+
+			Stopwatch.Stop();
+			UE_LOG(TimingProfiler, Log, TEXT("[IO] Merge 2/3 in %s."), *TimeUtils::FormatTimeAuto(Stopwatch.GetAccumulatedTime()));
+			Stopwatch.Restart();
+
+			// Sort cached IO events one more time, also by Start Time.
+			AllIoEvents.Sort([](const FIoTimingEvent& A, const FIoTimingEvent& B) { return A.StartTime < B.StartTime; });
+
+			//// Sort cached IO events again, by Depth and then by Start Time.
+			//AllIoEvents.Sort([](const FIoTimingEvent& A, const FIoTimingEvent& B)
+			//{
+			//	return A.Depth == B.Depth ? A.StartTime < B.StartTime : A.Depth < B.Depth;
+			//});
+
+			Stopwatch.Stop();
+			UE_LOG(TimingProfiler, Log, TEXT("[IO] Merge 3/3 (sort) in %s."), *TimeUtils::FormatTimeAuto(Stopwatch.GetAccumulatedTime()));
+			*/
+		}
+	}
+}
+
+////////////////////////////////////////////////////////////////////////////////////////////////////
+// FFileActivityTimingTrack
+////////////////////////////////////////////////////////////////////////////////////////////////////
+
+void FFileActivityTimingTrack::DrawTooltip(FTimingViewTooltip& Tooltip, const FVector2D& MousePosition, const FTimingEvent& HoveredTimingEvent, const FTimingTrackViewport& Viewport, const FDrawContext& DrawContext, const FSlateBrush* WhiteBrush, const FSlateFontInfo& Font) const
+{
+	const Trace::EFileActivityType ActivityType = static_cast<Trace::EFileActivityType>(HoveredTimingEvent.TypeId & 0x0F);
+	const bool bHasFailed = ((HoveredTimingEvent.TypeId & 0xF0) != 0);
+
+	FString FilePath(HoveredTimingEvent.Path);
+
+	constexpr float ValueOffsetX = 50.0f;
+	constexpr float MinValueTextWidth = 220.0f;
+
+	// Compute desired tooltip width.
+	const TSharedRef<FSlateFontMeasure> FontMeasureService = FSlateApplication::Get().GetRenderer()->GetFontMeasureService();
+	const float MinWidth1 = FontMeasureService->Measure(FilePath, Font).X + 2 * FTimingViewTooltip::BorderX;
+	const float MinWidth2 = ValueOffsetX + MinValueTextWidth + 2 * FTimingViewTooltip::BorderX;
+	const float DesiredTooltipWidth = FMath::Max3(MinWidth1, MinWidth2, FTimingViewTooltip::MinWidth);
+
+	constexpr float LineH = 14.0f;
+
+	// Compute desired tooltip height.
+	const int32 LineCount = (ActivityType == Trace::FileActivityType_Read || ActivityType == Trace::FileActivityType_Write) ? 6 : 4;
+	const float MinHeight = LineCount * LineH + 2 * FTimingViewTooltip::BorderY;
+	const float DesiredTooltipHeight = FMath::Max(MinHeight, FTimingViewTooltip::MinHeight);
+
+	Tooltip.Update(MousePosition, DesiredTooltipWidth, DesiredTooltipHeight, Viewport.Width, Viewport.Height);
+
+	const FLinearColor BackgroundColor(0.05f, 0.05f, 0.05f, Tooltip.Opacity);
+	const FLinearColor NameColor(0.9f, 0.9f, 0.5f, Tooltip.Opacity);
+	const FLinearColor TextColor(0.6f, 0.6f, 0.6f, Tooltip.Opacity);
+	const FLinearColor ValueColor(1.0f, 1.0f, 1.0f, Tooltip.Opacity);
+
+	DrawContext.DrawBox(Tooltip.PosX, Tooltip.PosY, Tooltip.Width, Tooltip.Height, WhiteBrush, BackgroundColor);
+	DrawContext.LayerId++;
+
+	const float X = Tooltip.PosX + FTimingViewTooltip::BorderX;
+	const float ValueX = X + ValueOffsetX;
+	float Y = Tooltip.PosY + FTimingViewTooltip::BorderY;
+
+	FString TypeStr;
+	uint32 TypeColor;
+	if (bHasFailed)
+	{
+		TypeStr = TEXT("Failed ");
+		TypeStr += GetFileActivityTypeName(ActivityType);
+		TypeColor = 0xFFFF3333;
+	}
+	else
+	{
+		TypeStr = GetFileActivityTypeName(ActivityType);
+		TypeColor = GetFileActivityTypeColor(ActivityType);
+	}
+	FLinearColor TypeLinearColor = FLinearColor(FColor(TypeColor));
+	TypeLinearColor.R *= 2.0f;
+	TypeLinearColor.G *= 2.0f;
+	TypeLinearColor.B *= 2.0f;
+	DrawContext.DrawText(X, Y, TypeStr, Font, TypeLinearColor);
+	Y += LineH;
+
+	DrawContext.DrawText(X, Y, FilePath, Font, NameColor);
+	Y += LineH;
+
+	DrawContext.DrawText(X, Y, TEXT("Duration:"), Font, TextColor);
+	FString DurationStr = TimeUtils::FormatTimeAuto(HoveredTimingEvent.Duration());
+	DrawContext.DrawText(ValueX, Y, DurationStr, Font, ValueColor);
+	Y += LineH;
+
+	DrawContext.DrawText(X, Y, TEXT("Depth:"), Font, TextColor);
+	FString DepthStr = FString::Printf(TEXT("%d"), HoveredTimingEvent.Depth);
+	DrawContext.DrawText(ValueX, Y, DepthStr, Font, ValueColor);
+	Y += LineH;
+
+	if (ActivityType == Trace::FileActivityType_Read || ActivityType == Trace::FileActivityType_Write)
+	{
+		DrawContext.DrawText(X, Y, TEXT("Offset:"), Font, TextColor);
+		FString OffsetStr = FText::AsNumber(HoveredTimingEvent.Offset).ToString();
+		DrawContext.DrawText(ValueX, Y, OffsetStr, Font, ValueColor);
+		Y += LineH;
+
+		DrawContext.DrawText(X, Y, TEXT("Size:"), Font, TextColor);
+		FString SizeStr = FText::AsNumber(HoveredTimingEvent.Size).ToString();
+		SizeStr += TEXT(" bytes");
+		DrawContext.DrawText(ValueX, Y, SizeStr, Font, ValueColor);
+		Y += LineH;
+	}
+
+	DrawContext.LayerId++;
+}
+
+////////////////////////////////////////////////////////////////////////////////////////////////////
+
+bool FFileActivityTimingTrack::SearchEvent(const double InStartTime,
+										   const double InEndTime,
+										   TFunctionRef<bool(double, double, uint32)> InPredicate,
+										   FTimingEvent& InOutTimingEvent,
+										   bool bInStopAtFirstMatch,
+										   bool bInSearchForLargestEvent,
+										   bool bIgnoreEventDepth) const
+{
+	struct FSearchTimingEventContext
+	{
+		const double StartTime;
+		const double EndTime;
+		TFunctionRef<bool(double, double, uint32)> Predicate;
+		FTimingEvent& TimingEvent;
+		const bool bStopAtFirstMatch;
+		const bool bSearchForLargestEvent;
+		mutable bool bFound;
+		mutable bool bContinueSearching;
+		mutable double LargestDuration;
+
+		FSearchTimingEventContext(const double InStartTime, const double InEndTime, TFunctionRef<bool(double, double, uint32)> InPredicate, FTimingEvent& InOutTimingEvent, bool bInStopAtFirstMatch, bool bInSearchForLargestEvent)
+			: StartTime(InStartTime)
+			, EndTime(InEndTime)
+			, Predicate(InPredicate)
+			, TimingEvent(InOutTimingEvent)
+			, bStopAtFirstMatch(bInStopAtFirstMatch)
+			, bSearchForLargestEvent(bInSearchForLargestEvent)
+			, bFound(false)
+			, bContinueSearching(true)
+			, LargestDuration(-1.0)
+		{
+		}
+
+		void CheckEvent(const FFileActivitySharedState::FIoTimingEvent& Event, uint32 EventDepth)
+		{
+			if (bContinueSearching && Predicate(Event.StartTime, Event.EndTime, EventDepth))
+			{
+				if (!bSearchForLargestEvent || Event.EndTime - Event.StartTime > LargestDuration)
+				{
+					LargestDuration = Event.EndTime - Event.StartTime;
+
+					TimingEvent.TypeId = Event.Type;
+					TimingEvent.Depth = EventDepth;
+					TimingEvent.StartTime = Event.StartTime;
+					TimingEvent.EndTime = Event.EndTime;
+
+					TimingEvent.Offset = Event.Offset;
+					TimingEvent.Size = Event.Size;
+					TimingEvent.Path = Event.FileActivity->Path;
+
+					bFound = true;
+					bContinueSearching = !bStopAtFirstMatch || bSearchForLargestEvent;
+				}
+			}
+		}
+	};
+
+	FSearchTimingEventContext Ctx(InStartTime, InEndTime, InPredicate, InOutTimingEvent, bInStopAtFirstMatch, bInSearchForLargestEvent);
+
+	TSharedPtr<const Trace::IAnalysisSession> Session = FInsightsManager::Get()->GetSession();
+	if (Session.IsValid())
+	{
+		Trace::FAnalysisSessionReadScope SessionReadScope(*Session.Get());
+
+		const FTimingEventsTrack* Track = Ctx.TimingEvent.Track;
+
+		for (const FFileActivitySharedState::FIoTimingEvent& Event : State->GetAllEvents())
+		{
+			if (Event.EndTime <= Ctx.StartTime)
+			{
+				continue;
+			}
+
+			if (!Ctx.bContinueSearching || Event.StartTime >= Ctx.EndTime)
+			{
+				break;
+			}
+
+			Ctx.CheckEvent(Event, bIgnoreEventDepth ? 0 : Event.Depth);
+		}
+	}
+
+	return Ctx.bFound;
+}
+////////////////////////////////////////////////////////////////////////////////////////////////////
+// FOverviewFileActivityTimingTrack
+////////////////////////////////////////////////////////////////////////////////////////////////////
+
+void FOverviewFileActivityTimingTrack::Draw(FTimingViewDrawHelper& Helper) const
+{
+	FTimingEventsTrack& Track = *const_cast<FOverviewFileActivityTimingTrack*>(this);
+
+	if (Helper.BeginTimeline(Track))
+	{
+		// Draw the IO Overiew track using cached events (as those are sorted by Start Time).
+		for (const FFileActivitySharedState::FIoTimingEvent& Event : State->AllIoEvents)
+		{
+			const Trace::EFileActivityType ActivityType = static_cast<Trace::EFileActivityType>(Event.Type & 0x0F);
+
+			if (ActivityType >= Trace::FileActivityType_Count)
+			{
+				// Ignore "Idle" and "NotClosed" events.
+				continue;
+			}
+
+			//const double EventEndTime = Event.EndTime; // keep duration of events
+			const double EventEndTime = Event.StartTime; // make all 0 duration events
+
+			if (EventEndTime <= Helper.GetViewport().StartTime)
+			{
+				continue;
+			}
+			if (Event.StartTime >= Helper.GetViewport().EndTime)
+			{
+				break;
+			}
+
+			uint32 Color = GetFileActivityTypeColor(ActivityType);
+
+			const bool bHasFailed = ((Event.Type & 0xF0) != 0);
+
+			if (bHasFailed)
+			{
+				FString EventName = TEXT("Failed ");
+				EventName += GetFileActivityTypeName(ActivityType);
+				Color = 0xFFAA0000;
+				Helper.AddEvent(Event.StartTime, EventEndTime, 0, *EventName, Color);
+			}
+			else
+			{
+				Helper.AddEvent(Event.StartTime, EventEndTime, 0, GetFileActivityTypeName(ActivityType), Color);
+			}
+		}
+
+		Helper.EndTimeline(Track);
+	}
+}
+
+////////////////////////////////////////////////////////////////////////////////////////////////////
+// FDetailedFileActivityTimingTrack
+////////////////////////////////////////////////////////////////////////////////////////////////////
+
+void FDetailedFileActivityTimingTrack::Draw(FTimingViewDrawHelper& Helper) const
+{
+	FTimingEventsTrack& Track = *const_cast<FDetailedFileActivityTimingTrack*>(this);
+
+	// Draw IO track using cached events.
+	if (Helper.BeginTimeline(Track))
+	{
+		// Draw IO file activity background events.
+		if (State->bShowFileActivityBackgroundEvents)
+		{
+			for (const TSharedPtr<FFileActivitySharedState::FIoFileActivity> Activity : State->FileActivities)
+			{
+				if (Activity->EndTime <= Helper.GetViewport().StartTime)
+				{
+					continue;
+				}
+				if (Activity->StartTime >= Helper.GetViewport().EndTime)
+				{
+					break;
+				}
+
+				ensure(Activity->Depth <= 10000);
+
+				Helper.AddEvent(Activity->StartTime, Activity->EndTime, Activity->Depth, Activity->Path, 0x55333333);
+			}
+		}
+
+		// Draw IO file activity foreground events.
+		for (const FFileActivitySharedState::FIoTimingEvent& Event : State->AllIoEvents)
+		{
+			if (Event.EndTime <= Helper.GetViewport().StartTime)
+			{
+				continue;
+			}
+			if (Event.StartTime >= Helper.GetViewport().EndTime)
+			{
+				break;
+			}
+
+			ensure(Event.Depth <= 10000);
+
+			const Trace::EFileActivityType ActivityType = static_cast<Trace::EFileActivityType>(Event.Type & 0x0F);
+
+			uint32 Color = GetFileActivityTypeColor(ActivityType);
+
+			if (ActivityType < Trace::FileActivityType_Count)
+			{
+				const bool bHasFailed = ((Event.Type & 0xF0) != 0);
+
+				if (bHasFailed)
+				{
+					FString EventName = TEXT("Failed ");
+					EventName += GetFileActivityTypeName(ActivityType);
+					Color = 0xFFAA0000;
+					Helper.AddEvent(Event.StartTime, Event.EndTime, Event.Depth, *EventName, Color);
+				}
+				else
+				{
+					Helper.AddEvent(Event.StartTime, Event.EndTime, Event.Depth, GetFileActivityTypeName(ActivityType), Color);
+				}
+			}
+			else
+			{
+				FString EventName = GetFileActivityTypeName(ActivityType);
+				EventName += " [";
+				EventName += Event.FileActivity->Path;
+				EventName += "]";
+				Helper.AddEvent(Event.StartTime, Event.EndTime, Event.Depth, *EventName, Color);
+			}
+		}
+
+		Helper.EndTimeline(Track);
+	}
+}
+
+////////////////////////////////////////////////////////////////////////////////////////////////////
+
+#undef LOCTEXT_NAMESPACE
diff --git a/Engine/Source/Developer/TraceInsights/Private/Insights/ViewModels/FileActivityTimingTrack.h b/Engine/Source/Developer/TraceInsights/Private/Insights/ViewModels/FileActivityTimingTrack.h
new file mode 100644
--- /dev/null
+++ b/Engine/Source/Developer/TraceInsights/Private/Insights/ViewModels/FileActivityTimingTrack.h
@@ -0,0 +1,123 @@
+// Copyright 1998-2019 Epic Games, Inc. All Rights Reserved.
+
+#pragma once
+
+#include "CoreMinimal.h"
+
+// Insights
+#include "Insights/ViewModels/TimingEventsTrack.h"
+
+////////////////////////////////////////////////////////////////////////////////////////////////////
+
+class FFileActivitySharedState
+{
+	friend class FOverviewFileActivityTimingTrack;
+	friend class FDetailedFileActivityTimingTrack;
+
+public:
+	struct FIoFileActivity
+	{
+		uint64 Id;
+		const TCHAR* Path;
+		double StartTime;
+		double EndTime;
+		int32 EventCount;
+		int32 Depth;
+	};
+
+	struct FIoTimingEvent
+	{
+		double StartTime;
+		double EndTime;
+		uint32 Depth;
+		uint32 Type; // Trace::EFileActivityType + "Failed" flag
+		uint64 Offset;
+		uint64 Size;
+		TSharedPtr<FIoFileActivity> FileActivity;
+	};
+
+public:
+	FFileActivitySharedState() {}
+
+	const TArray<FIoTimingEvent>& GetAllEvents() const { return AllIoEvents; }
+
+	~FFileActivitySharedState() {}
+
+	void Reset();
+	void Update();
+
+	void RequestUpdate() { bForceIoEventsUpdate = true; }
+	void ToggleMergeLanes() { bMergeIoLanes = !bMergeIoLanes; }
+	void ToggleBackgroundEvents() { bShowFileActivityBackgroundEvents = !bShowFileActivityBackgroundEvents; }
+
+private:
+	bool bForceIoEventsUpdate;
+
+	bool bMergeIoLanes;
+	bool bShowFileActivityBackgroundEvents;
+
+	TArray<TSharedPtr<FIoFileActivity>> FileActivities;
+	TMap<uint64, TSharedPtr<FIoFileActivity>> FileActivityMap;
+
+	/** All IO events, cached. */
+	TArray<FIoTimingEvent> AllIoEvents;
+};
+
+////////////////////////////////////////////////////////////////////////////////////////////////////
+
+class FFileActivityTimingTrack : public FTimingEventsTrack
+{
+public:
+	FFileActivityTimingTrack(uint64 InTrackId, const FName SubType, const FString& InName, TSharedPtr<FFileActivitySharedState> InState)
+		: FTimingEventsTrack(InTrackId, FName(TEXT("FileActivity")), SubType, InName)
+		, State(InState)
+	{
+	}
+	virtual ~FFileActivityTimingTrack() {}
+
+	virtual void DrawTooltip(FTimingViewTooltip& Tooltip, const FVector2D& MousePosition, const FTimingEvent& HoveredTimingEvent, const FTimingTrackViewport& Viewport, const FDrawContext& DrawContext, const FSlateBrush* WhiteBrush, const FSlateFontInfo& Font) const override;
+
+protected:
+	bool SearchEvent(const double InStartTime, const double InEndTime, TFunctionRef<bool(double, double, uint32)> InPredicate, FTimingEvent& InOutTimingEvent, bool bInStopAtFirstMatch, bool bInSearchForLargestEvent, bool bIgnoreEventDepth) const;
+
+protected:
+	TSharedPtr<FFileActivitySharedState> State;
+};
+
+////////////////////////////////////////////////////////////////////////////////////////////////////
+
+class FOverviewFileActivityTimingTrack : public FFileActivityTimingTrack
+{
+public:
+	FOverviewFileActivityTimingTrack(uint64 InTrackId, TSharedPtr<FFileActivitySharedState> InState)
+		: FFileActivityTimingTrack(InTrackId, FName(TEXT("Overview")), TEXT("I/O Overview"), InState)
+	{
+	}
+
+	virtual void Draw(FTimingViewDrawHelper& Helper) const override;
+	virtual bool SearchTimingEvent(const double InStartTime, const double InEndTime, TFunctionRef<bool(double, double, uint32)> InPredicate, FTimingEvent& InOutTimingEvent, bool bInStopAtFirstMatch, bool bInSearchForLargestEvent) const override
+	{
+		constexpr bool bIgnoreEventDepth = true;
+		return SearchEvent(InStartTime, InEndTime, InPredicate, InOutTimingEvent, bInStopAtFirstMatch, bInSearchForLargestEvent, bIgnoreEventDepth);
+	}
+};
+
+////////////////////////////////////////////////////////////////////////////////////////////////////
+
+class FDetailedFileActivityTimingTrack : public FFileActivityTimingTrack
+{
+public:
+	FDetailedFileActivityTimingTrack(uint64 InTrackId, TSharedPtr<FFileActivitySharedState> InState)
+		: FFileActivityTimingTrack(InTrackId, FName(TEXT("Detailed")), TEXT("I/O Activity"), InState)
+	{
+	}
+
+	virtual void Draw(FTimingViewDrawHelper& Helper) const override;
+	virtual bool SearchTimingEvent(const double InStartTime, const double InEndTime, TFunctionRef<bool(double, double, uint32)> InPredicate, FTimingEvent& InOutTimingEvent, bool bInStopAtFirstMatch, bool bInSearchForLargestEvent) const override
+	{
+		constexpr bool bIgnoreEventDepth = false;
+		return SearchEvent(InStartTime, InEndTime, InPredicate, InOutTimingEvent, bInStopAtFirstMatch, bInSearchForLargestEvent, bIgnoreEventDepth);
+	}
+};
+
+////////////////////////////////////////////////////////////////////////////////////////////////////
diff --git a/Engine/Source/Developer/TraceInsights/Private/Insights/ViewModels/GraphTrack.cpp b/Engine/Source/Developer/TraceInsights/Private/Insights/ViewModels/GraphTrack.cpp
--- a/Engine/Source/Developer/TraceInsights/Private/Insights/ViewModels/GraphTrack.cpp
+++ b/Engine/Source/Developer/TraceInsights/Private/Insights/ViewModels/GraphTrack.cpp
@@ -2,31 +2,25 @@
 
 #include "GraphTrack.h"
 
-#include <limits>
-
 #include "Brushes/SlateColorBrush.h"
 #include "Brushes/SlateBorderBrush.h"
 #include "EditorStyleSet.h"
 #include "Framework/MultiBox/MultiBoxBuilder.h"
 #include "Rendering/DrawElements.h"
 #include "Styling/CoreStyle.h"
-#include "TraceServices/AnalysisService.h"
-#include "TraceServices/SessionService.h"
 
 // Insights
 #include "Insights/Common/PaintUtils.h"
-#include "Insights/InsightsManager.h"
-#include "Insights/TimingProfilerManager.h"
 #include "Insights/ViewModels/DrawHelpers.h"
 #include "Insights/ViewModels/TimingTrackViewport.h"
 
 #define LOCTEXT_NAMESPACE "GraphTrack"
 
 ////////////////////////////////////////////////////////////////////////////////////////////////////
-// FGraphTrackSeries
+// FGraphSeries
 ////////////////////////////////////////////////////////////////////////////////////////////////////
 
-FGraphTrackSeries::FGraphTrackSeries()
+FGraphSeries::FGraphSeries()
 	: Name()
 	, Description()
 	, bIsVisible(true)
@@ -40,7 +34,7 @@ FGraphTrackSeries::FGraphTrackSeries()
 
 ////////////////////////////////////////////////////////////////////////////////////////////////////
 
-FGraphTrackSeries::~FGraphTrackSeries()
+FGraphSeries::~FGraphSeries()
 {
 }
 
@@ -104,7 +98,7 @@ void FGraphTrack::UpdateStats()
 	NumDrawLines = 0;
 	NumDrawBoxes = 0;
 
-	for (const TSharedPtr<FGraphTrackSeries>& Series : AllSeries)
+	for (const TSharedPtr<FGraphSeries>& Series : AllSeries)
 	{
 		NumDrawPoints += Series->Points.Num();
 		NumDrawLines += Series->LinePoints.Num() / 2;
@@ -145,7 +139,7 @@ void FGraphTrack::Draw(FDrawContext& DrawContext, const FTimingTrackViewport& Vi
 	DrawContext.DrawBox(0.0f, GetPosY() + BaselineY - 1.0f, Viewport.Width, 1.0f, WhiteBrush, FLinearColor(0.05f, 0.05f, 0.05f, 1.0f));
 	DrawContext.LayerId++;
 
-	for (const TSharedPtr<FGraphTrackSeries>& Series : AllSeries)
+	for (const TSharedPtr<FGraphSeries>& Series : AllSeries)
 	{
 		if (Series->IsVisible())
 		{
@@ -156,7 +150,7 @@ void FGraphTrack::Draw(FDrawContext& DrawContext, const FTimingTrackViewport& Vi
 
 ////////////////////////////////////////////////////////////////////////////////////////////////////
 
-void FGraphTrack::DrawSeries(const FGraphTrackSeries& Series, FDrawContext& DrawContext, const FTimingTrackViewport& Viewport) const
+void FGraphTrack::DrawSeries(const FGraphSeries& Series, FDrawContext& DrawContext, const FTimingTrackViewport& Viewport) const
 {
 	if (bDrawBoxes)
 	{
@@ -171,8 +165,8 @@ void FGraphTrack::DrawSeries(const FGraphTrackSeries& Series, FDrawContext& Draw
 
 	if (bDrawPolygon)
 	{
-		FSlateShaderResourceProxy* ResourceProxy = FSlateDataPayload::ResourceManager->GetShaderResource(*WhiteBrush);
 		FSlateResourceHandle ResourceHandle = FSlateApplication::Get().GetRenderer()->GetResourceHandle(*WhiteBrush);
+		const FSlateShaderResourceProxy* ResourceProxy = ResourceHandle.GetResourceProxy();
 
 		FVector2D AtlasOffset = ResourceProxy ? ResourceProxy->StartUV : FVector2D(0.f, 0.f);
 		FVector2D AtlasUVSize = ResourceProxy ? ResourceProxy->SizeUV : FVector2D(1.f, 1.f);
@@ -414,7 +408,7 @@ void FGraphTrack::BuildContextMenu(FMenuBuilder& MenuBuilder)
 
 	MenuBuilder.BeginSection("Series", LOCTEXT("ContextMenu_Header_Series", "Series"));
 	{
-		for (const TSharedPtr<FGraphTrackSeries>& Series : AllSeries)
+		for (const TSharedPtr<FGraphSeries>& Series : AllSeries)
 		{
 			FUIAction Action_ShowSeries
 			(
@@ -570,7 +564,7 @@ bool FGraphTrack::ContextMenu_ShowBars_IsChecked()
 
 ////////////////////////////////////////////////////////////////////////////////////////////////////
 
-void FGraphTrack::ContextMenu_ShowSeries_Execute(FGraphTrackSeries* Series)
+void FGraphTrack::ContextMenu_ShowSeries_Execute(FGraphSeries* Series)
 {
 	Series->SetVisibility(!Series->IsVisible());
 	SetDirtyFlag();
@@ -578,14 +572,14 @@ void FGraphTrack::ContextMenu_ShowSeries_Execute(FGraphTrackSeries* Series)
 
 ////////////////////////////////////////////////////////////////////////////////////////////////////
 
-bool FGraphTrack::ContextMenu_ShowSeries_CanExecute(FGraphTrackSeries* Series)
+bool FGraphTrack::ContextMenu_ShowSeries_CanExecute(FGraphSeries* Series)
 {
 	return true;
 }
 
 ////////////////////////////////////////////////////////////////////////////////////////////////////
 
-bool FGraphTrack::ContextMenu_ShowSeries_IsChecked(FGraphTrackSeries* Series)
+bool FGraphTrack::ContextMenu_ShowSeries_IsChecked(FGraphSeries* Series)
 {
 	return Series->IsVisible();
 }
@@ -604,21 +598,21 @@ FRandomGraphTrack::FRandomGraphTrack(uint64 InTrackId)
 	bUseEventDuration = false;
 	bDrawBoxes = false;
 
-	TSharedPtr<FGraphTrackSeries> Series0 = MakeShareable(new FGraphTrackSeries());
+	TSharedPtr<FGraphSeries> Series0 = MakeShareable(new FGraphSeries());
 	Series0->SetName(TEXT("Random Blue"));
 	Series0->SetDescription(TEXT("Random series; for debuging purposes"));
 	Series0->SetColor(FLinearColor(0.1f, 0.5f, 1.0f, 1.0f), FLinearColor(0.4f, 0.8f, 1.0f, 1.0f));
 	Series0->SetVisibility(true);
 	AllSeries.Add(Series0);
 
-	TSharedPtr<FGraphTrackSeries> Series1 = MakeShareable(new FGraphTrackSeries());
+	TSharedPtr<FGraphSeries> Series1 = MakeShareable(new FGraphSeries());
 	Series1->SetName(TEXT("Random Yellow"));
 	Series1->SetDescription(TEXT("Random series; for debuging purposes"));
 	Series1->SetColor(FLinearColor(0.9f, 0.9f, 0.1f, 1.0f), FLinearColor(1.0f, 1.0f, 0.4f, 1.0f));
 	Series1->SetVisibility(false);
 	AllSeries.Add(Series1);
 
-	TSharedPtr<FGraphTrackSeries> Series2 = MakeShareable(new FGraphTrackSeries());
+	TSharedPtr<FGraphSeries> Series2 = MakeShareable(new FGraphSeries());
 	Series2->SetName(TEXT("Random Red"));
 	Series2->SetDescription(TEXT("Random series; for debuging purposes"));
 	Series2->SetColor(FLinearColor(1.0f, 0.1f, 0.2f, 1.0f), FLinearColor(1.0f, 0.4f, 0.5f, 1.0f));
@@ -652,7 +646,7 @@ void FRandomGraphTrack::Update(const FTimingTrackViewport& Viewport)
 
 ////////////////////////////////////////////////////////////////////////////////////////////////////
 
-void FRandomGraphTrack::GenerateSeries(FGraphTrackSeries& Series, const FTimingTrackViewport& Viewport, const int32 EventCount, int32 Seed)
+void FRandomGraphTrack::GenerateSeries(FGraphSeries& Series, const FTimingTrackViewport& Viewport, const int32 EventCount, int32 Seed)
 {
 	//////////////////////////////////////////////////
 	// Generate random events.
@@ -715,245 +709,11 @@ void FRandomGraphTrack::GenerateSeries(FGraphTrackSeries& Series, const FTimingT
 	}
 }
 
-////////////////////////////////////////////////////////////////////////////////////////////////////
-// FTimingGraphTrack
-////////////////////////////////////////////////////////////////////////////////////////////////////
-
-FTimingGraphTrack::FTimingGraphTrack(uint64 InTrackId)
-	: FGraphTrack(InTrackId)
-{
-	bDrawPoints = true;
-	bDrawPointsWithBorder = true;
-	bDrawLines = true;
-	bDrawPolygon = true;
-	bUseEventDuration = true;
-	bDrawBoxes = false;
-
-	TSharedPtr<FTimingGraphSeries> GameFramesSeries = MakeShareable(new FTimingGraphSeries());
-	GameFramesSeries->SetName(TEXT("Game Frames"));
-	GameFramesSeries->SetDescription(TEXT("Duration of Game frames"));
-	GameFramesSeries->SetColor(FLinearColor(0.3f, 0.3f, 1.0f, 1.0f), FLinearColor(1.0f, 1.0f, 1.0f, 1.0f));
-	GameFramesSeries->Type = FTimingGraphSeries::ESeriesType::Frame;
-	GameFramesSeries->Id = static_cast<uint32>(TraceFrameType_Game);
-	AllSeries.Add(GameFramesSeries);
-
-	TSharedPtr<FTimingGraphSeries> RenderingFramesSeries = MakeShareable(new FTimingGraphSeries());
-	RenderingFramesSeries->SetName(TEXT("Rendering Frames"));
-	RenderingFramesSeries->SetDescription(TEXT("Duration of Rendering frames"));
-	RenderingFramesSeries->SetColor(FLinearColor(1.0f, 0.3f, 0.3f, 1.0f), FLinearColor(1.0f, 1.0f, 1.0f, 1.0f));
-	RenderingFramesSeries->Type = FTimingGraphSeries::ESeriesType::Frame;
-	RenderingFramesSeries->Id = static_cast<uint32>(TraceFrameType_Rendering);
-	AllSeries.Add(RenderingFramesSeries);
-}
-
-////////////////////////////////////////////////////////////////////////////////////////////////////
-
-TSharedPtr<FTimingGraphSeries> FTimingGraphTrack::GetStatsCounterSeries(uint32 CounterId)
-{
-	TSharedPtr<FGraphTrackSeries>* Ptr = AllSeries.FindByPredicate([=](const TSharedPtr<FGraphTrackSeries>& Series)
-	{
-		const TSharedPtr<FTimingGraphSeries> TimingSeries = StaticCastSharedPtr<FTimingGraphSeries>(Series);
-		return TimingSeries->Type == FTimingGraphSeries::ESeriesType::StatsCounter && TimingSeries->Id == CounterId;
-	});
-	return (Ptr != nullptr) ? StaticCastSharedPtr<FTimingGraphSeries>(*Ptr) : nullptr;
-}
-
-////////////////////////////////////////////////////////////////////////////////////////////////////
-
-void FTimingGraphTrack::AddStatsCounterSeries(uint32 CounterId, FLinearColor Color, double ValueOffset, double ValueScale)
-{
-	TSharedPtr<FTimingGraphSeries> Series = MakeShareable(new FTimingGraphSeries());
-
-	const TCHAR* CounterName = nullptr;
-	bool bIsFloatingPoint = false;
-
-	TSharedPtr<const Trace::IAnalysisSession> Session = FInsightsManager::Get()->GetSession();
-	if (Session.IsValid())
-	{
-		Trace::FAnalysisSessionReadScope SessionReadScope(*Session.Get());
-		const Trace::ICounterProvider& CountersProvider = Trace::ReadCounterProvider(*Session.Get());
-		if (CounterId < CountersProvider.GetCounterCount())
-		{
-			CountersProvider.ReadCounter(CounterId, [&](const Trace::ICounter& Counter)
-			{
-				CounterName = Counter.GetName();
-				bIsFloatingPoint = Counter.IsFloatingPoint();
-			});
-		}
-	}
-
-	Series->SetName(CounterName != nullptr ? CounterName : TEXT("<StatsCounter>"));
-	Series->SetDescription(TEXT("Stats counter series"));
-
-	FLinearColor BorderColor(Color.R + 0.4f, Color.G + 0.4f, Color.B + 0.4f, 1.0f);
-	Series->SetColor(Color, BorderColor);
-
-	Series->Type = FTimingGraphSeries::ESeriesType::StatsCounter;
-	Series->Id = CounterId;
-	Series->bIsFloatingPoint = bIsFloatingPoint;
-	Series->ValueOffset = ValueOffset;
-	Series->ValueScale = ValueScale;
-
-	AllSeries.Add(Series);
-}
-
-////////////////////////////////////////////////////////////////////////////////////////////////////
-
-void FTimingGraphTrack::RemoveStatsCounterSeries(uint32 CounterId)
-{
-	AllSeries.RemoveAll([=](const TSharedPtr<FGraphTrackSeries>& Series)
-	{
-		const TSharedPtr<FTimingGraphSeries> TimingSeries = StaticCastSharedPtr<FTimingGraphSeries>(Series);
-		return TimingSeries->Type == FTimingGraphSeries::ESeriesType::StatsCounter && TimingSeries->Id == CounterId;
-	});
-}
-
-////////////////////////////////////////////////////////////////////////////////////////////////////
-
-FTimingGraphTrack::~FTimingGraphTrack()
-{
-}
-
-////////////////////////////////////////////////////////////////////////////////////////////////////
-
-void FTimingGraphTrack::Update(const FTimingTrackViewport& Viewport)
-{
-	NumAddedEvents = 0;
-
-	// TODO: Vertical panning and zooming needs to be moved out in a Viewport like controller.
-	BaselineY = GetHeight();
-	ScaleY = 200.0 / 0.1; // 200px = 100ms
-
-	for (TSharedPtr<FGraphTrackSeries>& Series : AllSeries)
-	{
-		if (Series->IsVisible())
-		{
-			TSharedPtr<FTimingGraphSeries> TimingSeries = StaticCastSharedPtr<FTimingGraphSeries>(Series);
-			switch (TimingSeries->Type)
-			{
-				case FTimingGraphSeries::ESeriesType::Frame:
-					UpdateFrameSeries(*TimingSeries, Viewport);
-					break;
-
-				case FTimingGraphSeries::ESeriesType::Timer:
-					UpdateTimerSeries(*TimingSeries, Viewport);
-					break;
-
-				case FTimingGraphSeries::ESeriesType::StatsCounter:
-					UpdateStatsCounterSeries(*TimingSeries, Viewport);
-					break;
-			}
-		}
-	}
-
-	UpdateStats();
-}
-
-////////////////////////////////////////////////////////////////////////////////////////////////////
-
-void FTimingGraphTrack::UpdateFrameSeries(FTimingGraphSeries& Series, const FTimingTrackViewport& Viewport)
-{
-	FGraphTrackBuilder Builder(*this, Series, Viewport);
-	TSharedPtr<const Trace::IAnalysisSession> Session = FInsightsManager::Get()->GetSession();
-	if (Session.IsValid())
-	{
-		Trace::FAnalysisSessionReadScope SessionReadScope(*Session.Get());
-		const Trace::IFrameProvider& FramesProvider = ReadFrameProvider(*Session.Get());
-		ETraceFrameType FrameType = static_cast<ETraceFrameType>(Series.Id);
-		uint64 FrameCount = FramesProvider.GetFrameCount(FrameType);
-		FramesProvider.EnumerateFrames(FrameType, 0, FrameCount - 1, [&Builder](const Trace::FFrame& Frame)
-		{
-			//TODO: add a "frame converter" (i.e. to fps, miliseconds or seconds)
-			const double Duration = Frame.EndTime - Frame.StartTime;
-			Builder.AddEvent(Frame.StartTime, Duration, Duration);
-		});
-	}
-}
-
-////////////////////////////////////////////////////////////////////////////////////////////////////
-
-void FTimingGraphTrack::UpdateTimerSeries(FTimingGraphSeries& Series, const FTimingTrackViewport& Viewport)
-{
-	//TODO
-}
-
-////////////////////////////////////////////////////////////////////////////////////////////////////
-
-void FTimingGraphTrack::UpdateStatsCounterSeries(FTimingGraphSeries& Series, const FTimingTrackViewport& Viewport)
-{
-	FGraphTrackBuilder Builder(*this, Series, Viewport);
-	TSharedPtr<const Trace::IAnalysisSession> Session = FInsightsManager::Get()->GetSession();
-	if (Session.IsValid())
-	{
-		Trace::FAnalysisSessionReadScope SessionReadScope(*Session.Get());
-		const Trace::ICounterProvider& CountersProvider = Trace::ReadCounterProvider(*Session.Get());
-		CountersProvider.ReadCounter(Series.Id, [this, &Viewport, &Builder, &Series](const Trace::ICounter& Counter)
-		{
-			//const double ValueOffset = Series.ValueOffset;
-			//const double ValueScale = Series.ValueScale;
-			double MinValue = std::numeric_limits<double>::infinity();
-			double MaxValue = -std::numeric_limits<double>::infinity();
-
-			if (Counter.IsFloatingPoint())
-			{
-				Counter.EnumerateFloatValues(Viewport.StartTime, Viewport.EndTime, [&Builder, &MinValue, &MaxValue](double Time, double Value)
-				{
-					if (Value < MinValue)
-					{
-						MinValue = Value;
-					}
-					if (Value > MaxValue)
-					{
-						MaxValue = Value;
-					}
-				});
-			}
-			else
-			{
-				Counter.EnumerateValues(Viewport.StartTime, Viewport.EndTime, [&Builder, &MinValue, &MaxValue](double Time, int64 IntValue)
-				{
-					const double Value = static_cast<double>(IntValue);
-					if (Value < MinValue)
-					{
-						MinValue = Value;
-					}
-					if (Value > MaxValue)
-					{
-						MaxValue = Value;
-					}
-				});
-			}
-
-			const double ValueOffset = (MinValue != std::numeric_limits<double>::infinity()) ? -MinValue : 0.0;
-
-			const double AdjustedHeight = GetHeight() - 3.0f;
-			const double ValueScale = (MinValue < MaxValue) ? AdjustedHeight / ScaleY / (MaxValue - MinValue) : 1.0;
-
-			if (Counter.IsFloatingPoint())
-			{
-				Counter.EnumerateFloatValues(Viewport.StartTime, Viewport.EndTime, [&Builder, ValueOffset, ValueScale](double Time, double Value)
-				{
-					//TODO: add a "value converter"
-					Builder.AddEvent(Time, 0.0, (Value + ValueOffset) * ValueScale);
-				});
-			}
-			else
-			{
-				Counter.EnumerateValues(Viewport.StartTime, Viewport.EndTime, [&Builder, ValueOffset, ValueScale](double Time, int64 IntValue)
-				{
-					//TODO: add a "value converter"
-					Builder.AddEvent(Time, 0.0, (static_cast<double>(IntValue) + ValueOffset) * ValueScale);
-				});
-			}
-		});
-	}
-}
-
 ////////////////////////////////////////////////////////////////////////////////////////////////////
 // FGraphTrackBuilder
 ////////////////////////////////////////////////////////////////////////////////////////////////////
 
-FGraphTrackBuilder::FGraphTrackBuilder(FGraphTrack& InTrack, FGraphTrackSeries& InSeries, const FTimingTrackViewport& InViewport)
+FGraphTrackBuilder::FGraphTrackBuilder(FGraphTrack& InTrack, FGraphSeries& InSeries, const FTimingTrackViewport& InViewport)
 	: Track(InTrack)
 	, Series(InSeries)
 	, Viewport(InViewport)
diff --git a/Engine/Source/Developer/TraceInsights/Private/Insights/ViewModels/GraphTrack.h b/Engine/Source/Developer/TraceInsights/Private/Insights/ViewModels/GraphTrack.h
--- a/Engine/Source/Developer/TraceInsights/Private/Insights/ViewModels/GraphTrack.h
+++ b/Engine/Source/Developer/TraceInsights/Private/Insights/ViewModels/GraphTrack.h
@@ -24,14 +24,14 @@ struct FGraphBox
 
 ////////////////////////////////////////////////////////////////////////////////////////////////////
 
-class FGraphTrackSeries
+class FGraphSeries
 {
 	friend class FGraphTrack;
 	friend class FGraphTrackBuilder;
 
 public:
-	FGraphTrackSeries();
-	~FGraphTrackSeries();
+	FGraphSeries();
+	~FGraphSeries();
 
 	const FText& GetName() const { return Name; }
 	void SetName(const TCHAR* InName) { Name = FText::FromString(InName); }
@@ -72,26 +72,6 @@ protected:
 
 ////////////////////////////////////////////////////////////////////////////////////////////////////
 
-class FTimingGraphSeries : public FGraphTrackSeries
-{
-public:
-	enum class ESeriesType
-	{
-		Frame,
-		Timer,
-		StatsCounter
-	};
-
-public:
-	ESeriesType Type;
-	uint32 Id; // frame type, timer id or stats counter id
-	bool bIsFloatingPoint; // for stats counters
-	double ValueOffset; // offset added to Y values (used to pan graph vertically)
-	double ValueScale; // scale of Y values (used to scale graph vertically)
-};
-
-////////////////////////////////////////////////////////////////////////////////////////////////////
-
 class FGraphTrack : public FBaseTimingTrack
 {
 	friend class FGraphTrackBuilder;
@@ -128,7 +108,7 @@ protected:
 	void UpdateStats();
 
 	void DrawBackground(FDrawContext& DrawContext, const FTimingTrackViewport& Viewport) const;
-	void DrawSeries(const FGraphTrackSeries& Series, FDrawContext& DrawContext, const FTimingTrackViewport& Viewport) const;
+	void DrawSeries(const FGraphSeries& Series, FDrawContext& DrawContext, const FTimingTrackViewport& Viewport) const;
 
 	float GetYForValue(double Value) const
 	{
@@ -160,12 +140,12 @@ private:
 	bool ContextMenu_ShowBars_CanExecute();
 	bool ContextMenu_ShowBars_IsChecked();
 
-	void ContextMenu_ShowSeries_Execute(FGraphTrackSeries* Series);
-	bool ContextMenu_ShowSeries_CanExecute(FGraphTrackSeries* Series);
-	bool ContextMenu_ShowSeries_IsChecked(FGraphTrackSeries* Series);
+	void ContextMenu_ShowSeries_Execute(FGraphSeries* Series);
+	bool ContextMenu_ShowSeries_CanExecute(FGraphSeries* Series);
+	bool ContextMenu_ShowSeries_IsChecked(FGraphSeries* Series);
 
 protected:
-	TArray<TSharedPtr<FGraphTrackSeries>> AllSeries;
+	TArray<TSharedPtr<FGraphSeries>> AllSeries;
 
 	// Slate resources
 	const FSlateBrush* WhiteBrush;
@@ -201,27 +181,7 @@ public:
 	virtual void Update(const FTimingTrackViewport& Viewport) override;
 
 protected:
-	void GenerateSeries(FGraphTrackSeries& Series, const FTimingTrackViewport& Viewport, const int32 EventCount, int32 Seed);
-};
-
-////////////////////////////////////////////////////////////////////////////////////////////////////
-
-class FTimingGraphTrack : public FGraphTrack
-{
-public:
-	FTimingGraphTrack(uint64 InTrackId);
-	virtual ~FTimingGraphTrack();
-
-	virtual void Update(const FTimingTrackViewport& Viewport) override;
-
-	TSharedPtr<FTimingGraphSeries> GetStatsCounterSeries(uint32 CounterId);
-	void AddStatsCounterSeries(uint32 CounterId, FLinearColor Color, double ValueOffset = 0.0, double ValueScale = 1.0);
-	void RemoveStatsCounterSeries(uint32 CounterId);
-
-protected:
-	void UpdateFrameSeries(FTimingGraphSeries& Series, const FTimingTrackViewport& Viewport);
-	void UpdateTimerSeries(FTimingGraphSeries& Series, const FTimingTrackViewport& Viewport);
-	void UpdateStatsCounterSeries(FTimingGraphSeries& Series, const FTimingTrackViewport& Viewport);
+	void GenerateSeries(FGraphSeries& Series, const FTimingTrackViewport& Viewport, const int32 EventCount, int32 Seed);
 };
 
 ////////////////////////////////////////////////////////////////////////////////////////////////////
@@ -239,7 +199,7 @@ private:
 	};
 
 public:
-	FGraphTrackBuilder(FGraphTrack& InTrack, FGraphTrackSeries& InSeries, const FTimingTrackViewport& InViewport);
+	FGraphTrackBuilder(FGraphTrack& InTrack, FGraphSeries& InSeries, const FTimingTrackViewport& InViewport);
 	~FGraphTrackBuilder();
 
 	/**
@@ -249,7 +209,7 @@ public:
 	FGraphTrackBuilder& operator=(const FGraphTrackBuilder&) = delete;
 
 	FGraphTrack& GetTrack() const { return Track; }
-	FGraphTrackSeries& GetSeries() const { return Series; }
+	FGraphSeries& GetSeries() const { return Series; }
 	const FTimingTrackViewport& GetViewport() const { return Viewport; }
 
 	void AddEvent(double Time, double Duration, double Value);
@@ -272,7 +232,7 @@ private:
 
 private:
 	FGraphTrack& Track;
-	FGraphTrackSeries& Series;
+	FGraphSeries& Series;
 	const FTimingTrackViewport& Viewport;
 
 	// Used by the point reduction algorithm.
diff --git a/Engine/Source/Developer/TraceInsights/Private/Insights/ViewModels/LoadingTimingTrack.cpp b/Engine/Source/Developer/TraceInsights/Private/Insights/ViewModels/LoadingTimingTrack.cpp
new file mode 100644
--- /dev/null
+++ b/Engine/Source/Developer/TraceInsights/Private/Insights/ViewModels/LoadingTimingTrack.cpp
@@ -0,0 +1,454 @@
+// Copyright 1998-2019 Epic Games, Inc. All Rights Reserved.
+
+#include "LoadingTimingTrack.h"
+
+#include "Fonts/FontMeasure.h"
+#include "Styling/SlateBrush.h"
+#include "TraceServices/AnalysisService.h"
+#include "TraceServices/SessionService.h"
+
+// Insights
+#include "Insights/InsightsManager.h"
+#include "Insights/Common/PaintUtils.h"
+#include "Insights/Common/TimeUtils.h"
+#include "Insights/ViewModels/TimingEvent.h"
+#include "Insights/ViewModels/TimingTrackViewport.h"
+#include "Insights/ViewModels/TimingViewDrawHelper.h"
+
+#define LOCTEXT_NAMESPACE "LoadingTimingTrack"
+
+////////////////////////////////////////////////////////////////////////////////////////////////////
+// FLoadingSharedState
+////////////////////////////////////////////////////////////////////////////////////////////////////
+
+void FLoadingSharedState::Reset()
+{
+	LoadingGetEventNameFn = FLoadingTrackGetEventNameDelegate::CreateRaw(this, &FLoadingSharedState::GetEventName);
+}
+
+////////////////////////////////////////////////////////////////////////////////////////////////////
+
+const TCHAR* GetName(ELoadTimeProfilerPackageEventType Type)
+{
+	switch (Type)
+	{
+	case LoadTimeProfilerPackageEventType_CreateLinker:				return TEXT("CreateLinker");
+	case LoadTimeProfilerPackageEventType_FinishLinker:				return TEXT("FinishLinker");
+	case LoadTimeProfilerPackageEventType_StartImportPackages:		return TEXT("StartImportPackages");
+	case LoadTimeProfilerPackageEventType_SetupImports:				return TEXT("SetupImports");
+	case LoadTimeProfilerPackageEventType_SetupExports:				return TEXT("SetupExports");
+	case LoadTimeProfilerPackageEventType_ProcessImportsAndExports:	return TEXT("ProcessImportsAndExports");
+	case LoadTimeProfilerPackageEventType_ExportsDone:				return TEXT("ExportsDone");
+	case LoadTimeProfilerPackageEventType_PostLoadWait:				return TEXT("PostLoadWait");
+	case LoadTimeProfilerPackageEventType_StartPostLoad:			return TEXT("StartPostLoad");
+	case LoadTimeProfilerPackageEventType_Tick:						return TEXT("Tick");
+	case LoadTimeProfilerPackageEventType_Finish:					return TEXT("Finish");
+	case LoadTimeProfilerPackageEventType_DeferredPostLoad:			return TEXT("DeferredPostLoad");
+	case LoadTimeProfilerPackageEventType_None:						return TEXT("None");
+	default:														return TEXT("");
+	}
+}
+
+////////////////////////////////////////////////////////////////////////////////////////////////////
+
+const TCHAR* GetName(ELoadTimeProfilerObjectEventType Type)
+{
+	switch (Type)
+	{
+	case LoadTimeProfilerObjectEventType_Create:	return TEXT("Create");
+	case LoadTimeProfilerObjectEventType_Serialize:	return TEXT("Serialize");
+	case LoadTimeProfilerObjectEventType_PostLoad:	return TEXT("PostLoad");
+	case LoadTimeProfilerObjectEventType_None:		return TEXT("None");
+	default:										return TEXT("");
+	};
+}
+
+////////////////////////////////////////////////////////////////////////////////////////////////////
+
+const TCHAR* FLoadingSharedState::GetEventNameByPackageEventType(uint32 Depth, const Trace::FLoadTimeProfilerCpuEvent& Event) const
+{
+	return GetName(Event.PackageEventType);
+}
+
+////////////////////////////////////////////////////////////////////////////////////////////////////
+
+const TCHAR* FLoadingSharedState::GetEventNameByExportEventType(uint32 Depth, const Trace::FLoadTimeProfilerCpuEvent& Event) const
+{
+	return GetName(Event.ExportEventType);
+}
+
+////////////////////////////////////////////////////////////////////////////////////////////////////
+
+const TCHAR* FLoadingSharedState::GetEventNameByPackageName(uint32 Depth, const Trace::FLoadTimeProfilerCpuEvent& Event) const
+{
+	return Event.Package ? Event.Package->Name : TEXT("");
+}
+
+////////////////////////////////////////////////////////////////////////////////////////////////////
+
+const TCHAR* FLoadingSharedState::GetEventNameByExportClassName(uint32 Depth, const Trace::FLoadTimeProfilerCpuEvent& Event) const
+{
+	return Event.Export && Event.Export->Class ? Event.Export->Class->Name : TEXT("");
+}
+
+////////////////////////////////////////////////////////////////////////////////////////////////////
+
+const TCHAR* FLoadingSharedState::GetEventName(uint32 Depth, const Trace::FLoadTimeProfilerCpuEvent& Event) const
+{
+	if (Depth == 0)
+	{
+		if (Event.Package)
+		{
+			return Event.Package->Name;
+		}
+	}
+
+	if (Event.Export && Event.Export->Class)
+	{
+		return Event.Export->Class->Name;
+	}
+
+	return TEXT("");
+}
+
+////////////////////////////////////////////////////////////////////////////////////////////////////
+
+const TCHAR* FLoadingSharedState::GetEventNameEx(uint32 Depth, const Trace::FLoadTimeProfilerCpuEvent& Event) const
+{
+	return LoadingGetEventNameFn.Execute(Depth, Event);
+}
+
+////////////////////////////////////////////////////////////////////////////////////////////////////
+
+void FLoadingSharedState::SetColorSchema(int32 Schema)
+{
+	switch (Schema)
+	{
+		case 0: LoadingGetEventNameFn = FLoadingTrackGetEventNameDelegate::CreateRaw(this, &FLoadingSharedState::GetEventNameByPackageEventType); break;
+		case 1: LoadingGetEventNameFn = FLoadingTrackGetEventNameDelegate::CreateRaw(this, &FLoadingSharedState::GetEventNameByExportEventType); break;
+		case 2: LoadingGetEventNameFn = FLoadingTrackGetEventNameDelegate::CreateRaw(this, &FLoadingSharedState::GetEventNameByPackageName); break;
+		case 3: LoadingGetEventNameFn = FLoadingTrackGetEventNameDelegate::CreateRaw(this, &FLoadingSharedState::GetEventNameByExportClassName); break;
+		case 4: LoadingGetEventNameFn = FLoadingTrackGetEventNameDelegate::CreateRaw(this, &FLoadingSharedState::GetEventName); break;
+	};
+}
+
+////////////////////////////////////////////////////////////////////////////////////////////////////
+// FLoadingTimingTrack
+////////////////////////////////////////////////////////////////////////////////////////////////////
+
+void FLoadingTimingTrack::DrawTooltip(FTimingViewTooltip& Tooltip, const FVector2D& MousePosition, const FTimingEvent& HoveredTimingEvent, const FTimingTrackViewport& Viewport, const FDrawContext& DrawContext, const FSlateBrush* WhiteBrush, const FSlateFontInfo& Font) const
+{
+	FString EventName(State->GetEventNameEx(HoveredTimingEvent.Depth, HoveredTimingEvent.LoadingInfo));
+
+	const Trace::FPackageInfo* Package = HoveredTimingEvent.LoadingInfo.Package;
+	const Trace::FPackageExportInfo* Export = HoveredTimingEvent.LoadingInfo.Export;
+
+	FString PackageName = Package ? Package->Name : TEXT("N/A");
+
+	constexpr float ValueOffsetX = 100.0f;
+	constexpr float MinValueTextWidth = 220.0f;
+
+	// Compute desired tooltip width.
+	const TSharedRef<FSlateFontMeasure> FontMeasureService = FSlateApplication::Get().GetRenderer()->GetFontMeasureService();
+	const float NameWidth = FontMeasureService->Measure(EventName, Font).X;
+	const float PackageNameWidth = FontMeasureService->Measure(PackageName, Font).X;
+	const float MinWidth1 = NameWidth + 2 * FTimingViewTooltip::BorderX;
+	const float MinWidth2 = ValueOffsetX + FMath::Max(PackageNameWidth, MinValueTextWidth) + 2 * FTimingViewTooltip::BorderX;
+	const float DesiredTooltipWidth = FMath::Max3(MinWidth1, MinWidth2, FTimingViewTooltip::MinWidth);
+
+	constexpr float LineH = 14.0f;
+
+	// Compute desired tooltip height.
+	int32 LineCount = 5;
+	if (Package)
+	{
+		LineCount += 3;
+	}
+	if (Export)
+	{
+		LineCount += 2;
+	}
+	const float MinHeight = LineCount * LineH + 2 * FTimingViewTooltip::BorderY;
+	const float DesiredTooltipHeight = FMath::Max(MinHeight, FTimingViewTooltip::MinHeight);
+
+	Tooltip.Update(MousePosition, DesiredTooltipWidth, DesiredTooltipHeight, Viewport.Width, Viewport.Height);
+
+	const FLinearColor BackgroundColor(0.05f, 0.05f, 0.05f, Tooltip.Opacity);
+	const FLinearColor NameColor(0.9f, 0.9f, 0.5f, Tooltip.Opacity);
+	const FLinearColor TextColor(0.6f, 0.6f, 0.6f, Tooltip.Opacity);
+	const FLinearColor ValueColor(1.0f, 1.0f, 1.0f, Tooltip.Opacity);
+
+	DrawContext.DrawBox(Tooltip.PosX, Tooltip.PosY, Tooltip.Width, Tooltip.Height, WhiteBrush, BackgroundColor);
+	DrawContext.LayerId++;
+
+	const float X = Tooltip.PosX + FTimingViewTooltip::BorderX;
+	const float ValueX = X + ValueOffsetX;
+	float Y = Tooltip.PosY + FTimingViewTooltip::BorderY;
+
+	DrawContext.DrawText(X, Y, EventName, Font, NameColor);
+	Y += LineH;
+
+	DrawContext.DrawText(X, Y, TEXT("Duration:"), Font, TextColor);
+	FString InclStr = TimeUtils::FormatTimeAuto(HoveredTimingEvent.Duration());
+	DrawContext.DrawText(ValueX, Y, InclStr, Font, ValueColor);
+	Y += LineH;
+
+	DrawContext.DrawText(X, Y, TEXT("Depth:"), Font, TextColor);
+	FString DepthStr = FString::Printf(TEXT("%d"), HoveredTimingEvent.Depth);
+	DrawContext.DrawText(ValueX, Y, DepthStr, Font, ValueColor);
+	Y += LineH;
+
+	DrawContext.DrawText(X, Y, TEXT("Package Event:"), Font, TextColor);
+	DrawContext.DrawText(ValueX, Y, ::GetName(HoveredTimingEvent.LoadingInfo.PackageEventType), Font, ValueColor);
+	Y += LineH;
+
+	if (Package)
+	{
+		DrawContext.DrawText(X, Y, TEXT("Package Name:"), Font, TextColor);
+		DrawContext.DrawText(ValueX, Y, PackageName, Font, ValueColor);
+		Y += LineH;
+
+		DrawContext.DrawText(X, Y, TEXT("Header Size:"), Font, TextColor);
+		FString HeaderSizeStr = FString::Printf(TEXT("%d bytes"), Package->Summary.TotalHeaderSize);
+		DrawContext.DrawText(ValueX, Y, HeaderSizeStr, Font, ValueColor);
+		Y += LineH;
+
+		DrawContext.DrawText(X, Y, TEXT("Package Summary:"), Font, TextColor);
+		FString SummaryStr = FString::Printf(TEXT("%d names, %d imports, %d exports"), Package->Summary.NameCount, Package->Summary.ImportCount, Package->Summary.ExportCount);
+		DrawContext.DrawText(ValueX, Y, SummaryStr, Font, ValueColor);
+		Y += LineH;
+	}
+
+	{
+		DrawContext.DrawText(X, Y, TEXT("Export Event:"), Font, TextColor);
+		FString ExportTypeStr = FString::Printf(TEXT("%s%s"), ::GetName(HoveredTimingEvent.LoadingInfo.ExportEventType), Export && Export->IsAsset ? TEXT(" [asset]") : TEXT(""));
+		DrawContext.DrawText(ValueX, Y, ExportTypeStr, Font, ValueColor);
+		Y += LineH;
+	}
+
+	if (Export)
+	{
+		DrawContext.DrawText(X, Y, TEXT("Export Class:"), Font, TextColor);
+		FString ClassStr = FString::Printf(TEXT("%s"), Export->Class ? Export->Class->Name : TEXT("N/A"));
+		DrawContext.DrawText(ValueX, Y, ClassStr, Font, ValueColor);
+		Y += LineH;
+
+		DrawContext.DrawText(X, Y, TEXT("Serial:"), Font, TextColor);
+		FString SerialStr = FString::Printf(TEXT("Offset: %llu, Size: %llu%s"), Export->SerialOffset, Export->SerialSize);
+		DrawContext.DrawText(ValueX, Y, SerialStr, Font, ValueColor);
+		Y += LineH;
+	}
+
+	DrawContext.LayerId++;
+}
+
+////////////////////////////////////////////////////////////////////////////////////////////////////
+// FSearchLoadTimeProfilerCpuEvent
+////////////////////////////////////////////////////////////////////////////////////////////////////
+
+struct FSearchLoadTimeProfilerCpuEvent
+{
+	const double StartTime;
+	const double EndTime;
+	TFunctionRef<bool(double, double, uint32)> Predicate;
+	FTimingEvent& TimingEvent;
+	const bool bStopAtFirstMatch;
+	const bool bSearchForLargestEvent;
+	mutable bool bFound;
+	mutable bool bContinueSearching;
+	mutable double LargestDuration;
+
+	FSearchLoadTimeProfilerCpuEvent(const double InStartTime, const double InEndTime, TFunctionRef<bool(double, double, uint32)> InPredicate, FTimingEvent& InOutTimingEvent, bool bInStopAtFirstMatch, bool bInSearchForLargestEvent)
+		: StartTime(InStartTime)
+		, EndTime(InEndTime)
+		, Predicate(InPredicate)
+		, TimingEvent(InOutTimingEvent)
+		, bStopAtFirstMatch(bInStopAtFirstMatch)
+		, bSearchForLargestEvent(bInSearchForLargestEvent)
+		, bFound(false)
+		, bContinueSearching(true)
+		, LargestDuration(-1.0)
+	{
+	}
+
+	void CheckEvent(double EventStartTime, double EventEndTime, uint32 EventDepth, const Trace::FLoadTimeProfilerCpuEvent& Event)
+	{
+		if (bContinueSearching && Predicate(EventStartTime, EventEndTime, EventDepth))
+		{
+			if (!bSearchForLargestEvent || EventEndTime - EventStartTime > LargestDuration)
+			{
+				LargestDuration = EventEndTime - EventStartTime;
+
+				TimingEvent.TypeId = 0;
+				TimingEvent.Depth = EventDepth;
+				TimingEvent.StartTime = EventStartTime;
+				TimingEvent.EndTime = EventEndTime;
+
+				TimingEvent.LoadingInfo = Event;
+
+				bFound = true;
+				bContinueSearching = !bStopAtFirstMatch || bSearchForLargestEvent;
+			}
+		}
+	}
+};
+
+////////////////////////////////////////////////////////////////////////////////////////////////////
+// FLoadingMainThreadTimingTrack
+////////////////////////////////////////////////////////////////////////////////////////////////////
+
+void FLoadingMainThreadTimingTrack::Draw(FTimingViewDrawHelper& Helper) const
+{
+	FTimingEventsTrack& Track = *const_cast<FLoadingMainThreadTimingTrack*>(this);
+
+	if (Helper.BeginTimeline(Track))
+	{
+		TSharedPtr<const Trace::IAnalysisSession> Session = FInsightsManager::Get()->GetSession();
+		if (Session.IsValid() && Trace::ReadLoadTimeProfilerProvider(*Session.Get()))
+		{
+			Trace::FAnalysisSessionReadScope SessionReadScope(*Session.Get());
+
+			const Trace::ILoadTimeProfilerProvider& LoadTimeProfilerProvider = *Trace::ReadLoadTimeProfilerProvider(*Session.Get());
+
+			LoadTimeProfilerProvider.ReadMainThreadCpuTimeline([this, &Helper](const Trace::ILoadTimeProfilerProvider::CpuTimeline& Timeline)
+			{
+				if (FTimingEventsTrack::bUseDownSampling)
+				{
+					const double SecondsPerPixel = 1.0 / Helper.GetViewport().ScaleX;
+					Timeline.EnumerateEventsDownSampled(Helper.GetViewport().StartTime, Helper.GetViewport().EndTime, SecondsPerPixel, [this, &Helper](double StartTime, double EndTime, uint32 Depth, const Trace::FLoadTimeProfilerCpuEvent& Event)
+					{
+						const TCHAR* Name = State->GetEventNameEx(Depth, Event);
+						Helper.AddEvent(StartTime, EndTime, Depth, Name);
+					});
+				}
+				else
+				{
+					Timeline.EnumerateEvents(Helper.GetViewport().StartTime, Helper.GetViewport().EndTime, [this, &Helper](double StartTime, double EndTime, uint32 Depth, const Trace::FLoadTimeProfilerCpuEvent& Event)
+					{
+						const TCHAR* Name = State->GetEventNameEx(Depth, Event);
+						Helper.AddEvent(StartTime, EndTime, Depth, Name);
+					});
+				}
+			});
+		}
+
+		Helper.EndTimeline(Track);
+	}
+}
+
+////////////////////////////////////////////////////////////////////////////////////////////////////
+
+bool FLoadingMainThreadTimingTrack::SearchTimingEvent(const double InStartTime,
+	const double InEndTime,
+	TFunctionRef<bool(double, double, uint32)> InPredicate,
+	FTimingEvent& InOutTimingEvent,
+	bool bInStopAtFirstMatch,
+	bool bInSearchForLargestEvent) const
+{
+	FSearchLoadTimeProfilerCpuEvent Ctx(InStartTime, InEndTime, InPredicate, InOutTimingEvent, bInStopAtFirstMatch, bInSearchForLargestEvent);
+
+	TSharedPtr<const Trace::IAnalysisSession> Session = FInsightsManager::Get()->GetSession();
+	if (Session.IsValid())
+	{
+		Trace::FAnalysisSessionReadScope SessionReadScope(*Session.Get());
+
+		const FTimingEventsTrack* Track = Ctx.TimingEvent.Track;
+
+		if (Trace::ReadLoadTimeProfilerProvider(*Session.Get()))
+		{
+			const Trace::ILoadTimeProfilerProvider& LoadTimeProfilerProvider = *Trace::ReadLoadTimeProfilerProvider(*Session.Get());
+
+			LoadTimeProfilerProvider.ReadMainThreadCpuTimeline([&Ctx](const Trace::ILoadTimeProfilerProvider::CpuTimeline& Timeline)
+			{
+				Timeline.EnumerateEvents(Ctx.StartTime, Ctx.EndTime, [&Ctx](double EventStartTime, double EventEndTime, uint32 EventDepth, const Trace::FLoadTimeProfilerCpuEvent& Event)
+				{
+					Ctx.CheckEvent(EventStartTime, EventEndTime, EventDepth, Event);
+				});
+			});
+		}
+	}
+
+	return Ctx.bFound;
+}
+
+////////////////////////////////////////////////////////////////////////////////////////////////////
+// FLoadingAsyncThreadTimingTrack
+////////////////////////////////////////////////////////////////////////////////////////////////////
+
+void FLoadingAsyncThreadTimingTrack::Draw(FTimingViewDrawHelper& Helper) const
+{
+	FTimingEventsTrack& Track = *const_cast<FLoadingAsyncThreadTimingTrack*>(this);
+
+	if (Helper.BeginTimeline(Track))
+	{
+		TSharedPtr<const Trace::IAnalysisSession> Session = FInsightsManager::Get()->GetSession();
+		if (Session.IsValid() && Trace::ReadLoadTimeProfilerProvider(*Session.Get()))
+		{
+			Trace::FAnalysisSessionReadScope SessionReadScope(*Session.Get());
+
+			const Trace::ILoadTimeProfilerProvider& LoadTimeProfilerProvider = *Trace::ReadLoadTimeProfilerProvider(*Session.Get());
+
+			LoadTimeProfilerProvider.ReadAsyncLoadingThreadCpuTimeline([this, &Helper](const Trace::ILoadTimeProfilerProvider::CpuTimeline& Timeline)
+			{
+				if (FTimingEventsTrack::bUseDownSampling)
+				{
+					const double SecondsPerPixel = 1.0 / Helper.GetViewport().ScaleX;
+					Timeline.EnumerateEventsDownSampled(Helper.GetViewport().StartTime, Helper.GetViewport().EndTime, SecondsPerPixel, [this, &Helper](double StartTime, double EndTime, uint32 Depth, const Trace::FLoadTimeProfilerCpuEvent& Event)
+					{
+						const TCHAR* Name = State->GetEventNameEx(Depth, Event);
+						Helper.AddEvent(StartTime, EndTime, Depth, Name);
+					});
+				}
+				else
+				{
+					Timeline.EnumerateEvents(Helper.GetViewport().StartTime, Helper.GetViewport().EndTime, [this, &Helper](double StartTime, double EndTime, uint32 Depth, const Trace::FLoadTimeProfilerCpuEvent& Event)
+					{
+						const TCHAR* Name = State->GetEventNameEx(Depth, Event);
+						Helper.AddEvent(StartTime, EndTime, Depth, Name);
+					});
+				}
+			});
+		}
+
+		Helper.EndTimeline(Track);
+	}
+}
+
+////////////////////////////////////////////////////////////////////////////////////////////////////
+
+bool FLoadingAsyncThreadTimingTrack::SearchTimingEvent(const double InStartTime,
+													   const double InEndTime,
+													   TFunctionRef<bool(double, double, uint32)> InPredicate,
+													   FTimingEvent& InOutTimingEvent,
+													   bool bInStopAtFirstMatch,
+													   bool bInSearchForLargestEvent) const
+{
+	FSearchLoadTimeProfilerCpuEvent Ctx(InStartTime, InEndTime, InPredicate, InOutTimingEvent, bInStopAtFirstMatch, bInSearchForLargestEvent);
+
+	TSharedPtr<const Trace::IAnalysisSession> Session = FInsightsManager::Get()->GetSession();
+	if (Session.IsValid())
+	{
+		Trace::FAnalysisSessionReadScope SessionReadScope(*Session.Get());
+
+		const FTimingEventsTrack* Track = Ctx.TimingEvent.Track;
+
+		if (Trace::ReadLoadTimeProfilerProvider(*Session.Get()))
+		{
+			const Trace::ILoadTimeProfilerProvider& LoadTimeProfilerProvider = *Trace::ReadLoadTimeProfilerProvider(*Session.Get());
+
+			LoadTimeProfilerProvider.ReadAsyncLoadingThreadCpuTimeline([&Ctx](const Trace::ILoadTimeProfilerProvider::CpuTimeline& Timeline)
+			{
+				Timeline.EnumerateEvents(Ctx.StartTime, Ctx.EndTime, [&Ctx](double EventStartTime, double EventEndTime, uint32 EventDepth, const Trace::FLoadTimeProfilerCpuEvent& Event)
+				{
+					Ctx.CheckEvent(EventStartTime, EventEndTime, EventDepth, Event);
+				});
+			});
+		}
+	}
+
+	return Ctx.bFound;
+}
+
+////////////////////////////////////////////////////////////////////////////////////////////////////
+
+#undef LOCTEXT_NAMESPACE
diff --git a/Engine/Source/Developer/TraceInsights/Private/Insights/ViewModels/LoadingTimingTrack.h b/Engine/Source/Developer/TraceInsights/Private/Insights/ViewModels/LoadingTimingTrack.h
new file mode 100644
--- /dev/null
+++ b/Engine/Source/Developer/TraceInsights/Private/Insights/ViewModels/LoadingTimingTrack.h
@@ -0,0 +1,81 @@
+// Copyright 1998-2019 Epic Games, Inc. All Rights Reserved.
+
+#pragma once
+
+#include "CoreMinimal.h"
+#include "TraceServices/AnalysisService.h"
+
+// Insights
+#include "Insights/ViewModels/TimingEventsTrack.h"
+
+/** Defines FLoadingTrackGetEventNameDelegate delegate interface. Returns the name for a timing event in a Loading track. */
+DECLARE_DELEGATE_RetVal_TwoParams(const TCHAR*, FLoadingTrackGetEventNameDelegate, uint32 /*Depth*/, const Trace::FLoadTimeProfilerCpuEvent& /*Event*/);
+
+////////////////////////////////////////////////////////////////////////////////////////////////////
+
+class FLoadingSharedState
+{
+public:
+	void Reset();
+
+	const TCHAR* GetEventNameByPackageEventType(uint32 Depth, const Trace::FLoadTimeProfilerCpuEvent& Event) const;
+	const TCHAR* GetEventNameByExportEventType(uint32 Depth, const Trace::FLoadTimeProfilerCpuEvent& Event) const;
+	const TCHAR* GetEventNameByPackageName(uint32 Depth, const Trace::FLoadTimeProfilerCpuEvent& Event) const;
+	const TCHAR* GetEventNameByExportClassName(uint32 Depth, const Trace::FLoadTimeProfilerCpuEvent& Event) const;
+	const TCHAR* GetEventName(uint32 Depth, const Trace::FLoadTimeProfilerCpuEvent& Event) const;
+
+	const TCHAR* GetEventNameEx(uint32 Depth, const Trace::FLoadTimeProfilerCpuEvent& Event) const;
+
+	void SetColorSchema(int32 Schema);
+
+private:
+	FLoadingTrackGetEventNameDelegate LoadingGetEventNameFn;
+};
+
+////////////////////////////////////////////////////////////////////////////////////////////////////
+
+class FLoadingTimingTrack : public FTimingEventsTrack
+{
+public:
+	FLoadingTimingTrack(uint64 InTrackId, const FName& InSubType, const FString& InName, TSharedPtr<FLoadingSharedState> InState)
+		: FTimingEventsTrack(InTrackId, FName(TEXT("Loading")), InSubType, InName)
+		, State(InState)
+	{
+	}
+	virtual ~FLoadingTimingTrack() {}
+
+	virtual void DrawTooltip(FTimingViewTooltip& Tooltip, const FVector2D& MousePosition, const FTimingEvent& HoveredTimingEvent, const FTimingTrackViewport& Viewport, const FDrawContext& DrawContext, const FSlateBrush* WhiteBrush, const FSlateFontInfo& Font) const override;
+
+protected:
+	TSharedPtr<FLoadingSharedState> State;
+};
+
+////////////////////////////////////////////////////////////////////////////////////////////////////
+
+class FLoadingMainThreadTimingTrack : public FLoadingTimingTrack
+{
+public:
+	FLoadingMainThreadTimingTrack(uint64 InTrackId, TSharedPtr<FLoadingSharedState> InState)
+		: FLoadingTimingTrack(InTrackId, FName(TEXT("MainThread")), TEXT("Loading - Main Thread"), InState)
+	{
+	}
+
+	virtual void Draw(FTimingViewDrawHelper& Helper) const override;
+	virtual bool SearchTimingEvent(const double InStartTime, const double InEndTime, TFunctionRef<bool(double, double, uint32)> InPredicate, FTimingEvent& InOutTimingEvent, bool bInStopAtFirstMatch, bool bInSearchForLargestEvent) const override;
+};
+
+////////////////////////////////////////////////////////////////////////////////////////////////////
+
+class FLoadingAsyncThreadTimingTrack : public FLoadingTimingTrack
+{
+public:
+	FLoadingAsyncThreadTimingTrack(uint64 InTrackId, TSharedPtr<FLoadingSharedState> InState)
+		: FLoadingTimingTrack(InTrackId, FName(TEXT("AsyncThread")), TEXT("Loading - Async Thread"), InState)
+	{
+	}
+
+	virtual void Draw(FTimingViewDrawHelper& Helper) const override;
+	virtual bool SearchTimingEvent(const double InStartTime, const double InEndTime, TFunctionRef<bool(double, double, uint32)> InPredicate, FTimingEvent& InOutTimingEvent, bool bInStopAtFirstMatch, bool bInSearchForLargestEvent) const override;
+};
+
+////////////////////////////////////////////////////////////////////////////////////////////////////
diff --git a/Engine/Source/Developer/TraceInsights/Private/Insights/ViewModels/ThreadTimingTrack.cpp b/Engine/Source/Developer/TraceInsights/Private/Insights/ViewModels/ThreadTimingTrack.cpp
new file mode 100644
--- /dev/null
+++ b/Engine/Source/Developer/TraceInsights/Private/Insights/ViewModels/ThreadTimingTrack.cpp
@@ -0,0 +1,287 @@
+// Copyright 1998-2019 Epic Games, Inc. All Rights Reserved.
+
+#include "ThreadTimingTrack.h"
+
+#include "Fonts/FontMeasure.h"
+#include "Styling/SlateBrush.h"
+
+// Insights
+#include "Insights/Common/PaintUtils.h"
+#include "Insights/Common/TimeUtils.h"
+#include "Insights/InsightsManager.h"
+#include "Insights/TimingProfilerManager.h"
+#include "Insights/ViewModels/TimerNode.h"
+#include "Insights/ViewModels/TimingEvent.h"
+#include "Insights/ViewModels/TimingTrackViewport.h"
+#include "Insights/ViewModels/TimingViewDrawHelper.h"
+
+#define LOCTEXT_NAMESPACE "ThreadTimingTrack"
+
+////////////////////////////////////////////////////////////////////////////////////////////////////
+
+void FThreadTimingTrack::Draw(FTimingViewDrawHelper& Helper) const
+{
+	FTimingEventsTrack& Track = *const_cast<FThreadTimingTrack*>(this);
+
+	if (Helper.BeginTimeline(Track))
+	{
+		TSharedPtr<const Trace::IAnalysisSession> Session = FInsightsManager::Get()->GetSession();
+		if (Session.IsValid() && Trace::ReadTimingProfilerProvider(*Session.Get()))
+		{
+			Trace::FAnalysisSessionReadScope SessionReadScope(*Session.Get());
+
+			const Trace::ITimingProfilerProvider& TimingProfilerProvider = *Trace::ReadTimingProfilerProvider(*Session.Get());
+
+			TimingProfilerProvider.ReadTimers([this, &Helper, &Track, &TimingProfilerProvider](const Trace::FTimingProfilerTimer* Timers, uint64 TimersCount)
+			{
+				TimingProfilerProvider.ReadTimeline(Track.GetId(), [this, &Helper, &Track, Timers](const Trace::ITimingProfilerProvider::Timeline& Timeline)
+				{
+					if (FTimingEventsTrack::bUseDownSampling)
+					{
+						const double SecondsPerPixel = 1.0 / Helper.GetViewport().ScaleX;
+						Timeline.EnumerateEventsDownSampled(Helper.GetViewport().StartTime, Helper.GetViewport().EndTime, SecondsPerPixel, [this, &Helper, Timers](double StartTime, double EndTime, uint32 Depth, const Trace::FTimingProfilerEvent& Event)
+						{
+							Helper.AddEvent(StartTime, EndTime, Depth, Timers[Event.TimerIndex].Name);
+						});
+					}
+					else
+					{
+						Timeline.EnumerateEvents(Helper.GetViewport().StartTime, Helper.GetViewport().EndTime, [this, &Helper, Timers](double StartTime, double EndTime, uint32 Depth, const Trace::FTimingProfilerEvent& Event)
+						{
+							Helper.AddEvent(StartTime, EndTime, Depth, Timers[Event.TimerIndex].Name);
+						});
+					}
+				});
+			});
+
+			Helper.EndTimeline(Track);
+		}
+	}
+}
+
+////////////////////////////////////////////////////////////////////////////////////////////////////
+
+void FThreadTimingTrack::DrawSelectedEventInfo(const FTimingEvent& SelectedTimingEvent, const FTimingTrackViewport& Viewport, const FDrawContext& DrawContext, const FSlateBrush* WhiteBrush, const FSlateFontInfo& Font) const
+{
+	const FTimerNodePtr TimerNodePtr = FTimingProfilerManager::Get()->GetTimerNode(SelectedTimingEvent.TypeId);
+
+	FString Str = FString::Printf(TEXT("%s (Incl.: %s, Excl.: %s)"),
+		TimerNodePtr ? *(TimerNodePtr->GetName().ToString()) : TEXT("N/A"),
+		*TimeUtils::FormatTimeAuto(SelectedTimingEvent.Duration()),
+		*TimeUtils::FormatTimeAuto(SelectedTimingEvent.ExclusiveTime));
+
+	const TSharedRef<FSlateFontMeasure> FontMeasureService = FSlateApplication::Get().GetRenderer()->GetFontMeasureService();
+	const FVector2D Size = FontMeasureService->Measure(Str, Font);
+	const float X = Viewport.Width - Size.X - 23.0f;
+	const float Y = Viewport.Height - Size.Y - 18.0f;
+
+	const FLinearColor BackgroundColor(0.05f, 0.05f, 0.05f, 1.0f);
+	const FLinearColor TextColor(0.7f, 0.7f, 0.7f, 1.0f);
+
+	DrawContext.DrawBox(X - 8.0f, Y - 2.0f, Size.X + 16.0f, Size.Y + 4.0f, WhiteBrush, BackgroundColor);
+	DrawContext.LayerId++;
+
+	DrawContext.DrawText(X, Y, Str, Font, TextColor);
+	DrawContext.LayerId++;
+}
+
+////////////////////////////////////////////////////////////////////////////////////////////////////
+
+void FThreadTimingTrack::DrawTooltip(FTimingViewTooltip& Tooltip,
+									 const FVector2D& MousePosition,
+									 const FTimingEvent& HoveredTimingEvent,
+									 const FTimingTrackViewport& Viewport,
+									 const FDrawContext& DrawContext,
+									 const FSlateBrush* WhiteBrush,
+									 const FSlateFontInfo& Font) const
+{
+	const FTimerNodePtr TimerNodePtr = FTimingProfilerManager::Get()->GetTimerNode(HoveredTimingEvent.TypeId);
+	FString TimerName = TimerNodePtr ? TimerNodePtr->GetName().ToString() : TEXT("N/A");
+
+	const TSharedRef<FSlateFontMeasure> FontMeasureService = FSlateApplication::Get().GetRenderer()->GetFontMeasureService();
+	float DesiredTooltipWidth = FontMeasureService->Measure(TimerName, Font).X + 2 * FTimingViewTooltip::BorderX;
+	if (DesiredTooltipWidth < FTimingViewTooltip::MinWidth)
+	{
+		DesiredTooltipWidth = FTimingViewTooltip::MinWidth;
+	}
+
+	constexpr float LineH = 14.0f;
+	const float DesiredTooltipHeight = 4 * LineH + 2 * FTimingViewTooltip::BorderY;
+
+	Tooltip.Update(MousePosition, DesiredTooltipWidth, DesiredTooltipHeight, Viewport.Width, Viewport.Height);
+
+	const FLinearColor BackgroundColor(0.05f, 0.05f, 0.05f, Tooltip.Opacity);
+	const FLinearColor NameColor(0.9f, 0.9f, 0.5f, Tooltip.Opacity);
+	const FLinearColor TextColor(0.6f, 0.6f, 0.6f, Tooltip.Opacity);
+	const FLinearColor ValueColor(1.0f, 1.0f, 1.0f, Tooltip.Opacity);
+
+	DrawContext.DrawBox(Tooltip.PosX, Tooltip.PosY, Tooltip.Width, Tooltip.Height, WhiteBrush, BackgroundColor);
+	DrawContext.LayerId++;
+
+	float X = Tooltip.PosX + FTimingViewTooltip::BorderX;
+	float Y = Tooltip.PosY + FTimingViewTooltip::BorderY;
+
+	DrawContext.DrawText(X, Y, TimerName, Font, NameColor);
+	Y += LineH;
+
+	const float ValueX = X + 58.0f;
+
+	DrawContext.DrawText(X + 3.0f, Y, TEXT("Incl. Time:"), Font, TextColor);
+	FString InclStr = TimeUtils::FormatTimeAuto(HoveredTimingEvent.Duration());
+	DrawContext.DrawText(ValueX, Y, InclStr, Font, ValueColor);
+	Y += LineH;
+
+	DrawContext.DrawText(X, Y, TEXT("Excl. Time:"), Font, TextColor);
+	FString ExclStr = TimeUtils::FormatTimeAuto(HoveredTimingEvent.ExclusiveTime);
+	DrawContext.DrawText(ValueX, Y, ExclStr, Font, ValueColor);
+	Y += LineH;
+
+	DrawContext.DrawText(X + 24.0f, Y, TEXT("Depth:"), Font, TextColor);
+	FString DepthStr = FString::Printf(TEXT("%d"), HoveredTimingEvent.Depth);
+	DrawContext.DrawText(ValueX, Y, DepthStr, Font, ValueColor);
+	Y += LineH;
+
+	DrawContext.LayerId++;
+}
+
+////////////////////////////////////////////////////////////////////////////////////////////////////
+
+
+bool FThreadTimingTrack::SearchTimingEvent(const double InStartTime,
+										   const double InEndTime,
+										   TFunctionRef<bool(double, double, uint32)> InPredicate,
+										   FTimingEvent& InOutTimingEvent,
+										   bool bInStopAtFirstMatch,
+										   bool bInSearchForLargestEvent) const
+{
+	struct FSearchTimingEventContext
+	{
+		const double StartTime;
+		const double EndTime;
+		TFunctionRef<bool(double, double, uint32)> Predicate;
+		FTimingEvent& TimingEvent;
+		const bool bStopAtFirstMatch;
+		const bool bSearchForLargestEvent;
+		mutable bool bFound;
+		mutable bool bContinueSearching;
+		mutable double LargestDuration;
+
+		FSearchTimingEventContext(const double InStartTime, const double InEndTime, TFunctionRef<bool(double, double, uint32)> InPredicate, FTimingEvent& InOutTimingEvent, bool bInStopAtFirstMatch, bool bInSearchForLargestEvent)
+			: StartTime(InStartTime)
+			, EndTime(InEndTime)
+			, Predicate(InPredicate)
+			, TimingEvent(InOutTimingEvent)
+			, bStopAtFirstMatch(bInStopAtFirstMatch)
+			, bSearchForLargestEvent(bInSearchForLargestEvent)
+			, bFound(false)
+			, bContinueSearching(true)
+			, LargestDuration(-1.0)
+		{
+		}
+
+		void CheckEvent(double EventStartTime, double EventEndTime, uint32 EventDepth, const Trace::FTimingProfilerEvent& Event)
+		{
+			if (bContinueSearching && Predicate(EventStartTime, EventEndTime, EventDepth))
+			{
+				if (!bSearchForLargestEvent || EventEndTime - EventStartTime > LargestDuration)
+				{
+					LargestDuration = EventEndTime - EventStartTime;
+
+					TimingEvent.TypeId = Event.TimerIndex;
+					TimingEvent.Depth = EventDepth;
+					TimingEvent.StartTime = EventStartTime;
+					TimingEvent.EndTime = EventEndTime;
+
+					bFound = true;
+					bContinueSearching = !bStopAtFirstMatch || bSearchForLargestEvent;
+				}
+			}
+		}
+	};
+
+	FSearchTimingEventContext Ctx(InStartTime, InEndTime, InPredicate, InOutTimingEvent, bInStopAtFirstMatch, bInSearchForLargestEvent);
+
+	TSharedPtr<const Trace::IAnalysisSession> Session = FInsightsManager::Get()->GetSession();
+	if (Session.IsValid())
+	{
+		Trace::FAnalysisSessionReadScope SessionReadScope(*Session.Get());
+
+		const FTimingEventsTrack* Track = Ctx.TimingEvent.Track;
+
+		if (Trace::ReadTimingProfilerProvider(*Session.Get()))
+		{
+			const Trace::ITimingProfilerProvider& TimingProfilerProvider = *Trace::ReadTimingProfilerProvider(*Session.Get());
+
+			TimingProfilerProvider.ReadTimeline(Track->GetId(), [&Ctx](const Trace::ITimingProfilerProvider::Timeline& Timeline)
+			{
+				Timeline.EnumerateEvents(Ctx.StartTime, Ctx.EndTime, [&Ctx](double EventStartTime, double EventEndTime, uint32 EventDepth, const Trace::FTimingProfilerEvent& Event)
+				{
+					Ctx.CheckEvent(EventStartTime, EventEndTime, EventDepth, Event);
+				});
+			});
+
+			if (Ctx.bFound)
+			{
+				// Compute Exclusive Time.
+				TimingProfilerProvider.ReadTimeline(Track->GetId(), [&Ctx](const Trace::ITimingProfilerProvider::Timeline& Timeline)
+				{
+					struct FEnumerationState
+					{
+						double EventStartTime;
+						double EventEndTime;
+						uint64 EventDepth;
+
+						uint64 CurrentDepth;
+						double LastTime;
+						double ExclusiveTime;
+						bool IsInEventScope;
+					};
+					FEnumerationState State;
+
+					State.EventStartTime = Ctx.TimingEvent.StartTime;
+					State.EventEndTime = Ctx.TimingEvent.EndTime;
+					State.EventDepth = Ctx.TimingEvent.Depth;
+
+					State.CurrentDepth = 0;
+					State.LastTime = 0.0;
+					State.ExclusiveTime = 0.0;
+					State.IsInEventScope = false;
+
+					Timeline.EnumerateEvents(Ctx.TimingEvent.StartTime, Ctx.TimingEvent.EndTime, [&State](bool IsEnter, double Time, const Trace::FTimingProfilerEvent& Event)
+					{
+						if (IsEnter)
+						{
+							if (State.IsInEventScope && State.CurrentDepth == State.EventDepth + 1)
+							{
+								State.ExclusiveTime += Time - State.LastTime;
+							}
+							if (State.CurrentDepth == State.EventDepth && Time == State.EventStartTime)
+							{
+								State.IsInEventScope = true;
+							}
+							++State.CurrentDepth;
+						}
+						else
+						{
+							--State.CurrentDepth;
+							if (State.CurrentDepth == State.EventDepth && Time == State.EventEndTime)
+							{
+								State.IsInEventScope = false;
+								State.ExclusiveTime += Time - State.LastTime;
+							}
+						}
+						State.LastTime = Time;
+					});
+
+					Ctx.TimingEvent.ExclusiveTime = State.ExclusiveTime;
+				});
+			}
+		}
+	}
+
+	return Ctx.bFound;
+}
+
+////////////////////////////////////////////////////////////////////////////////////////////////////
+
+#undef LOCTEXT_NAMESPACE
diff --git a/Engine/Source/Developer/TraceInsights/Private/Insights/ViewModels/ThreadTimingTrack.h b/Engine/Source/Developer/TraceInsights/Private/Insights/ViewModels/ThreadTimingTrack.h
new file mode 100644
--- /dev/null
+++ b/Engine/Source/Developer/TraceInsights/Private/Insights/ViewModels/ThreadTimingTrack.h
@@ -0,0 +1,62 @@
+// Copyright 1998-2019 Epic Games, Inc. All Rights Reserved.
+
+#pragma once
+
+#include "CoreMinimal.h"
+
+// Insights
+#include "Insights/ViewModels/TimingEventsTrack.h"
+
+////////////////////////////////////////////////////////////////////////////////////////////////////
+
+class FThreadTimingTrack : public FTimingEventsTrack
+{
+public:
+	FThreadTimingTrack(uint64 InTrackId, const FName& InSubType, const FString& InName, const TCHAR* InGroupName)
+		: FTimingEventsTrack(InTrackId, FName(TEXT("Thread")), InSubType, InName)
+		, GroupName(InGroupName)
+		, ThreadId(0)
+	{
+	}
+
+	virtual ~FThreadTimingTrack() {}
+
+	const TCHAR* GetGroupName() const { return GroupName; };
+
+	void SetThreadId(uint32 InThreadId) { ThreadId = InThreadId; }
+	uint32 GetThreadId() const { return ThreadId; }
+
+	virtual void Draw(FTimingViewDrawHelper& Helper) const override;
+	virtual void DrawSelectedEventInfo(const FTimingEvent& SelectedTimingEvent, const FTimingTrackViewport& Viewport, const FDrawContext& DrawContext, const FSlateBrush* WhiteBrush, const FSlateFontInfo& Font) const override;
+	virtual void DrawTooltip(FTimingViewTooltip& Tooltip, const FVector2D& MousePosition, const FTimingEvent& HoveredTimingEvent, const FTimingTrackViewport& Viewport, const FDrawContext& DrawContext, const FSlateBrush* WhiteBrush, const FSlateFontInfo& Font) const override;
+
+	virtual bool SearchTimingEvent(const double InStartTime, const double InEndTime, TFunctionRef<bool(double, double, uint32)> InPredicate, FTimingEvent& InOutTimingEvent, bool bInStopAtFirstMatch, bool bInSearchForLargestEvent) const override;
+
+private:
+	const TCHAR* GroupName;
+	uint32 ThreadId;
+};
+
+////////////////////////////////////////////////////////////////////////////////////////////////////
+
+class FCpuTimingTrack : public FThreadTimingTrack
+{
+public:
+	FCpuTimingTrack(uint64 InTrackId, const FString& InName, const TCHAR* InGroupName)
+		: FThreadTimingTrack(InTrackId, FName(TEXT("CPU")), InName, InGroupName)
+	{
+	}
+};
+
+////////////////////////////////////////////////////////////////////////////////////////////////////
+
+class FGpuTimingTrack : public FThreadTimingTrack
+{
+public:
+	FGpuTimingTrack(uint64 InTrackId, const FString& InName, const TCHAR* InGroupName)
+		: FThreadTimingTrack(InTrackId, FName(TEXT("GPU")), InName, InGroupName)
+	{
+	}
+};
+
+////////////////////////////////////////////////////////////////////////////////////////////////////
diff --git a/Engine/Source/Developer/TraceInsights/Private/Insights/ViewModels/TimingEventsTrack.cpp b/Engine/Source/Developer/TraceInsights/Private/Insights/ViewModels/TimingEventsTrack.cpp
--- a/Engine/Source/Developer/TraceInsights/Private/Insights/ViewModels/TimingEventsTrack.cpp
+++ b/Engine/Source/Developer/TraceInsights/Private/Insights/ViewModels/TimingEventsTrack.cpp
@@ -81,14 +81,17 @@ void FTimingEventsTrackLayout::Update()
 // FTimingEventsTrack
 ////////////////////////////////////////////////////////////////////////////////////////////////////
 
-FTimingEventsTrack::FTimingEventsTrack(uint64 InTrackId, ETimingEventsTrackType InType, const FString& InName, const TCHAR* InGroupName)
+bool FTimingEventsTrack::bUseDownSampling = true;
+
+////////////////////////////////////////////////////////////////////////////////////////////////////
+
+FTimingEventsTrack::FTimingEventsTrack(uint64 InTrackId, const FName& InType, const FName& InSubType, const FString& InName)
 	: FBaseTimingTrack(InTrackId)
 	, Type(InType)
+	, SubType(InSubType)
 	, Name(InName)
-	, GroupName(InGroupName)
-	, ThreadId(0)
 	, Order(0)
-	, Depth(-1)
+	, NumLanes(0)
 	, bIsCollapsed(false)
 {
 }
@@ -105,7 +108,7 @@ void FTimingEventsTrack::Reset()
 {
 	FBaseTimingTrack::Reset();
 
-	Depth = -1;
+	NumLanes = 0;
 	bIsCollapsed = false;
 }
 
diff --git a/Engine/Source/Developer/TraceInsights/Private/Insights/ViewModels/TimingEventsTrack.h b/Engine/Source/Developer/TraceInsights/Private/Insights/ViewModels/TimingEventsTrack.h
--- a/Engine/Source/Developer/TraceInsights/Private/Insights/ViewModels/TimingEventsTrack.h
+++ b/Engine/Source/Developer/TraceInsights/Private/Insights/ViewModels/TimingEventsTrack.h
@@ -7,6 +7,14 @@
 // Insights
 #include "Insights/ViewModels/BaseTimingTrack.h"
 
+struct FDrawContext;
+struct FSlateBrush;
+struct FSlateFontInfo;
+struct FTimingEvent;
+class FTimingTrackViewport;
+struct FTimingViewTooltip;
+class FTimingViewDrawHelper;
+
 ////////////////////////////////////////////////////////////////////////////////////////////////////
 
 struct FTimingEventsTrackLayout
@@ -44,46 +52,39 @@ struct FTimingEventsTrackLayout
 
 ////////////////////////////////////////////////////////////////////////////////////////////////////
 
-enum class ETimingEventsTrackType
-{
-	Cpu,
-	Gpu,
-	Loading,
-	Io,
-};
-
-////////////////////////////////////////////////////////////////////////////////////////////////////
-
 class FTimingEventsTrack : public FBaseTimingTrack
 {
 public:
-	FTimingEventsTrack(uint64 InTrackId, ETimingEventsTrackType InType, const FString& InName, const TCHAR* InGroupName);
+	FTimingEventsTrack(uint64 InTrackId, const FName& InType, const FName& InSubType, const FString& InName);
 	virtual ~FTimingEventsTrack();
 
-	ETimingEventsTrackType GetType() const { return Type; }
-	const FString& GetName() const { return Name; }
-
-	const TCHAR* GetGroupName() const { return GroupName; };
+	const FName& GetType() const { return Type; }
+	const FName& GetSubType() const { return SubType; }
 
-	void SetThreadId(uint32 InThreadId) { ThreadId = InThreadId; }
-	uint32 GetThreadId() const { return ThreadId; }
+	const FString& GetName() const { return Name; }
 
 	void SetOrder(int32 InOrder) { Order = InOrder; }
 	int32 GetOrder() const { return Order; }
 
-	int32 GetDepth() const { return Depth; }
-	int32 GetNumLanes() const { return Depth + 1; }
+	int32 GetDepth() const { return NumLanes - 1; }
+	int32 GetNumLanes() const { return NumLanes; }
+	void SetDepth(int32 InDepth) { NumLanes = InDepth + 1; }
 
 	virtual void Reset() override;
 	virtual void UpdateHoveredState(float MX, float MY, const FTimingTrackViewport& Viewport);
 
-public:
-	ETimingEventsTrackType Type;
+	virtual void Draw(FTimingViewDrawHelper& Helper) const = 0;
+	virtual void DrawSelectedEventInfo(const FTimingEvent& SelectedTimingEvent, const FTimingTrackViewport& Viewport, const FDrawContext& DrawContext, const FSlateBrush* WhiteBrush, const FSlateFontInfo& Font) const {}
+	virtual void DrawTooltip(FTimingViewTooltip& Tooltip, const FVector2D& MousePosition, const FTimingEvent& HoveredTimingEvent, const FTimingTrackViewport& Viewport, const FDrawContext& DrawContext, const FSlateBrush* WhiteBrush, const FSlateFontInfo& Font) const {}
+
+	virtual bool SearchTimingEvent(const double InStartTime, const double InEndTime, TFunctionRef<bool(double, double, uint32)> InPredicate, FTimingEvent& InOutTimingEvent, bool bInStopAtFirstMatch, bool bInSearchForLargestEvent) const { return false; }
+
+private:
+	FName Type; // Thread, Loading, FileActivity
+	FName SubType; // Cpu, Gpu, MainThread, AsyncThread, Overview, Detailed
 	FString Name;
-	const TCHAR* GroupName;
-	uint32 ThreadId;
 	int32 Order;
-	int32 Depth; // number of lanes == Depth + 1
+	int32 NumLanes; // number of lanes (sub-tracks)
 	bool bIsCollapsed;
 
 	// TODO: Cached OnPaint state.
@@ -91,6 +92,9 @@ public:
 	//TArray<FEventBoxInfo> MergedBorderBoxes;
 	//TArray<FEventBoxInfo> Borders;
 	//TArray<FTextInfo> Texts;
+
+public:
+	static bool bUseDownSampling;
 };
 
 ////////////////////////////////////////////////////////////////////////////////////////////////////
diff --git a/Engine/Source/Developer/TraceInsights/Private/Insights/ViewModels/TimingGraphTrack.cpp b/Engine/Source/Developer/TraceInsights/Private/Insights/ViewModels/TimingGraphTrack.cpp
new file mode 100644
--- /dev/null
+++ b/Engine/Source/Developer/TraceInsights/Private/Insights/ViewModels/TimingGraphTrack.cpp
@@ -0,0 +1,253 @@
+// Copyright 1998-2019 Epic Games, Inc. All Rights Reserved.
+
+#include "TimingGraphTrack.h"
+
+#include <limits>
+
+#include "TraceServices/AnalysisService.h"
+#include "TraceServices/SessionService.h"
+
+// Insights
+#include "Insights/InsightsManager.h"
+#include "Insights/TimingProfilerManager.h"
+#include "Insights/ViewModels/TimingTrackViewport.h"
+
+#define LOCTEXT_NAMESPACE "GraphTrack"
+
+////////////////////////////////////////////////////////////////////////////////////////////////////
+// FTimingGraphTrack
+////////////////////////////////////////////////////////////////////////////////////////////////////
+
+FTimingGraphTrack::FTimingGraphTrack(uint64 InTrackId)
+	: FGraphTrack(InTrackId)
+{
+	bDrawPoints = true;
+	bDrawPointsWithBorder = true;
+	bDrawLines = true;
+	bDrawPolygon = true;
+	bUseEventDuration = true;
+	bDrawBoxes = false;
+
+	TSharedPtr<FTimingGraphSeries> GameFramesSeries = MakeShareable(new FTimingGraphSeries());
+	GameFramesSeries->SetName(TEXT("Game Frames"));
+	GameFramesSeries->SetDescription(TEXT("Duration of Game frames"));
+	GameFramesSeries->SetColor(FLinearColor(0.3f, 0.3f, 1.0f, 1.0f), FLinearColor(1.0f, 1.0f, 1.0f, 1.0f));
+	GameFramesSeries->Type = FTimingGraphSeries::ESeriesType::Frame;
+	GameFramesSeries->Id = static_cast<uint32>(TraceFrameType_Game);
+	AllSeries.Add(GameFramesSeries);
+
+	TSharedPtr<FTimingGraphSeries> RenderingFramesSeries = MakeShareable(new FTimingGraphSeries());
+	RenderingFramesSeries->SetName(TEXT("Rendering Frames"));
+	RenderingFramesSeries->SetDescription(TEXT("Duration of Rendering frames"));
+	RenderingFramesSeries->SetColor(FLinearColor(1.0f, 0.3f, 0.3f, 1.0f), FLinearColor(1.0f, 1.0f, 1.0f, 1.0f));
+	RenderingFramesSeries->Type = FTimingGraphSeries::ESeriesType::Frame;
+	RenderingFramesSeries->Id = static_cast<uint32>(TraceFrameType_Rendering);
+	AllSeries.Add(RenderingFramesSeries);
+}
+
+////////////////////////////////////////////////////////////////////////////////////////////////////
+
+TSharedPtr<FTimingGraphSeries> FTimingGraphTrack::GetStatsCounterSeries(uint32 CounterId)
+{
+	TSharedPtr<FGraphSeries>* Ptr = AllSeries.FindByPredicate([CounterId](const TSharedPtr<FGraphSeries>& Series)
+	{
+		const TSharedPtr<FTimingGraphSeries> TimingSeries = StaticCastSharedPtr<FTimingGraphSeries>(Series);
+		return TimingSeries->Type == FTimingGraphSeries::ESeriesType::StatsCounter && TimingSeries->Id == CounterId;
+	});
+	return (Ptr != nullptr) ? StaticCastSharedPtr<FTimingGraphSeries>(*Ptr) : nullptr;
+}
+
+////////////////////////////////////////////////////////////////////////////////////////////////////
+
+void FTimingGraphTrack::AddStatsCounterSeries(uint32 CounterId, FLinearColor Color, double ValueOffset, double ValueScale)
+{
+	TSharedPtr<FTimingGraphSeries> Series = MakeShareable(new FTimingGraphSeries());
+
+	const TCHAR* CounterName = nullptr;
+	bool bIsFloatingPoint = false;
+
+	TSharedPtr<const Trace::IAnalysisSession> Session = FInsightsManager::Get()->GetSession();
+	if (Session.IsValid())
+	{
+		Trace::FAnalysisSessionReadScope SessionReadScope(*Session.Get());
+		const Trace::ICounterProvider& CountersProvider = Trace::ReadCounterProvider(*Session.Get());
+		if (CounterId < CountersProvider.GetCounterCount())
+		{
+			CountersProvider.ReadCounter(CounterId, [&](const Trace::ICounter& Counter)
+			{
+				CounterName = Counter.GetName();
+				bIsFloatingPoint = Counter.IsFloatingPoint();
+			});
+		}
+	}
+
+	Series->SetName(CounterName != nullptr ? CounterName : TEXT("<StatsCounter>"));
+	Series->SetDescription(TEXT("Stats counter series"));
+
+	FLinearColor BorderColor(Color.R + 0.4f, Color.G + 0.4f, Color.B + 0.4f, 1.0f);
+	Series->SetColor(Color, BorderColor);
+
+	Series->Type = FTimingGraphSeries::ESeriesType::StatsCounter;
+	Series->Id = CounterId;
+	Series->bIsFloatingPoint = bIsFloatingPoint;
+	Series->ValueOffset = ValueOffset;
+	Series->ValueScale = ValueScale;
+
+	AllSeries.Add(Series);
+}
+
+////////////////////////////////////////////////////////////////////////////////////////////////////
+
+void FTimingGraphTrack::RemoveStatsCounterSeries(uint32 CounterId)
+{
+	AllSeries.RemoveAll([CounterId](const TSharedPtr<FGraphSeries>& Series)
+	{
+		const TSharedPtr<FTimingGraphSeries> TimingSeries = StaticCastSharedPtr<FTimingGraphSeries>(Series);
+		return TimingSeries->Type == FTimingGraphSeries::ESeriesType::StatsCounter && TimingSeries->Id == CounterId;
+	});
+}
+
+////////////////////////////////////////////////////////////////////////////////////////////////////
+
+FTimingGraphTrack::~FTimingGraphTrack()
+{
+}
+
+////////////////////////////////////////////////////////////////////////////////////////////////////
+
+void FTimingGraphTrack::Update(const FTimingTrackViewport& Viewport)
+{
+	NumAddedEvents = 0;
+
+	// TODO: Vertical panning and zooming needs to be moved out in a Viewport like controller.
+	BaselineY = GetHeight();
+	ScaleY = 200.0 / 0.1; // 200px = 100ms
+
+	for (TSharedPtr<FGraphSeries>& Series : AllSeries)
+	{
+		if (Series->IsVisible())
+		{
+			TSharedPtr<FTimingGraphSeries> TimingSeries = StaticCastSharedPtr<FTimingGraphSeries>(Series);
+			switch (TimingSeries->Type)
+			{
+				case FTimingGraphSeries::ESeriesType::Frame:
+					UpdateFrameSeries(*TimingSeries, Viewport);
+					break;
+
+				case FTimingGraphSeries::ESeriesType::Timer:
+					UpdateTimerSeries(*TimingSeries, Viewport);
+					break;
+
+				case FTimingGraphSeries::ESeriesType::StatsCounter:
+					UpdateStatsCounterSeries(*TimingSeries, Viewport);
+					break;
+			}
+		}
+	}
+
+	UpdateStats();
+}
+
+////////////////////////////////////////////////////////////////////////////////////////////////////
+
+void FTimingGraphTrack::UpdateFrameSeries(FTimingGraphSeries& Series, const FTimingTrackViewport& Viewport)
+{
+	FGraphTrackBuilder Builder(*this, Series, Viewport);
+	TSharedPtr<const Trace::IAnalysisSession> Session = FInsightsManager::Get()->GetSession();
+	if (Session.IsValid())
+	{
+		Trace::FAnalysisSessionReadScope SessionReadScope(*Session.Get());
+		const Trace::IFrameProvider& FramesProvider = ReadFrameProvider(*Session.Get());
+		ETraceFrameType FrameType = static_cast<ETraceFrameType>(Series.Id);
+		uint64 FrameCount = FramesProvider.GetFrameCount(FrameType);
+		FramesProvider.EnumerateFrames(FrameType, 0, FrameCount - 1, [&Builder](const Trace::FFrame& Frame)
+		{
+			//TODO: add a "frame converter" (i.e. to fps, miliseconds or seconds)
+			const double Duration = Frame.EndTime - Frame.StartTime;
+			Builder.AddEvent(Frame.StartTime, Duration, Duration);
+		});
+	}
+}
+
+////////////////////////////////////////////////////////////////////////////////////////////////////
+
+void FTimingGraphTrack::UpdateTimerSeries(FTimingGraphSeries& Series, const FTimingTrackViewport& Viewport)
+{
+	//TODO
+}
+
+////////////////////////////////////////////////////////////////////////////////////////////////////
+
+void FTimingGraphTrack::UpdateStatsCounterSeries(FTimingGraphSeries& Series, const FTimingTrackViewport& Viewport)
+{
+	FGraphTrackBuilder Builder(*this, Series, Viewport);
+	TSharedPtr<const Trace::IAnalysisSession> Session = FInsightsManager::Get()->GetSession();
+	if (Session.IsValid())
+	{
+		Trace::FAnalysisSessionReadScope SessionReadScope(*Session.Get());
+		const Trace::ICounterProvider& CountersProvider = Trace::ReadCounterProvider(*Session.Get());
+		CountersProvider.ReadCounter(Series.Id, [this, &Viewport, &Builder, &Series](const Trace::ICounter& Counter)
+		{
+			//const double ValueOffset = Series.ValueOffset;
+			//const double ValueScale = Series.ValueScale;
+			double MinValue = std::numeric_limits<double>::infinity();
+			double MaxValue = -std::numeric_limits<double>::infinity();
+
+			if (Counter.IsFloatingPoint())
+			{
+				Counter.EnumerateFloatValues(Viewport.StartTime, Viewport.EndTime, [&Builder, &MinValue, &MaxValue](double Time, double Value)
+				{
+					if (Value < MinValue)
+					{
+						MinValue = Value;
+					}
+					if (Value > MaxValue)
+					{
+						MaxValue = Value;
+					}
+				});
+			}
+			else
+			{
+				Counter.EnumerateValues(Viewport.StartTime, Viewport.EndTime, [&Builder, &MinValue, &MaxValue](double Time, int64 IntValue)
+				{
+					const double Value = static_cast<double>(IntValue);
+					if (Value < MinValue)
+					{
+						MinValue = Value;
+					}
+					if (Value > MaxValue)
+					{
+						MaxValue = Value;
+					}
+				});
+			}
+
+			const double ValueOffset = (MinValue != std::numeric_limits<double>::infinity()) ? -MinValue : 0.0;
+
+			const double AdjustedHeight = GetHeight() - 3.0f;
+			const double ValueScale = (MinValue < MaxValue) ? AdjustedHeight / ScaleY / (MaxValue - MinValue) : 1.0;
+
+			if (Counter.IsFloatingPoint())
+			{
+				Counter.EnumerateFloatValues(Viewport.StartTime, Viewport.EndTime, [&Builder, ValueOffset, ValueScale](double Time, double Value)
+				{
+					//TODO: add a "value converter"
+					Builder.AddEvent(Time, 0.0, (Value + ValueOffset) * ValueScale);
+				});
+			}
+			else
+			{
+				Counter.EnumerateValues(Viewport.StartTime, Viewport.EndTime, [&Builder, ValueOffset, ValueScale](double Time, int64 IntValue)
+				{
+					//TODO: add a "value converter"
+					Builder.AddEvent(Time, 0.0, (static_cast<double>(IntValue) + ValueOffset) * ValueScale);
+				});
+			}
+		});
+	}
+}
+
+////////////////////////////////////////////////////////////////////////////////////////////////////
+
+#undef LOCTEXT_NAMESPACE
diff --git a/Engine/Source/Developer/TraceInsights/Private/Insights/ViewModels/TimingGraphTrack.h b/Engine/Source/Developer/TraceInsights/Private/Insights/ViewModels/TimingGraphTrack.h
new file mode 100644
--- /dev/null
+++ b/Engine/Source/Developer/TraceInsights/Private/Insights/ViewModels/TimingGraphTrack.h
@@ -0,0 +1,50 @@
+// Copyright 1998-2019 Epic Games, Inc. All Rights Reserved.
+
+#pragma once
+
+#include "CoreMinimal.h"
+
+// Insights
+#include "Insights/ViewModels/GraphTrack.h"
+
+////////////////////////////////////////////////////////////////////////////////////////////////////
+
+class FTimingGraphSeries : public FGraphSeries
+{
+public:
+	enum class ESeriesType
+	{
+		Frame,
+		Timer,
+		StatsCounter
+	};
+
+public:
+	ESeriesType Type;
+	uint32 Id; // frame type, timer id or stats counter id
+	bool bIsFloatingPoint; // for stats counters
+	double ValueOffset; // offset added to Y values (used to pan graph vertically)
+	double ValueScale; // scale of Y values (used to scale graph vertically)
+};
+
+////////////////////////////////////////////////////////////////////////////////////////////////////
+
+class FTimingGraphTrack : public FGraphTrack
+{
+public:
+	FTimingGraphTrack(uint64 InTrackId);
+	virtual ~FTimingGraphTrack();
+
+	virtual void Update(const FTimingTrackViewport& Viewport) override;
+
+	TSharedPtr<FTimingGraphSeries> GetStatsCounterSeries(uint32 CounterId);
+	void AddStatsCounterSeries(uint32 CounterId, FLinearColor Color, double ValueOffset = 0.0, double ValueScale = 1.0);
+	void RemoveStatsCounterSeries(uint32 CounterId);
+
+protected:
+	void UpdateFrameSeries(FTimingGraphSeries& Series, const FTimingTrackViewport& Viewport);
+	void UpdateTimerSeries(FTimingGraphSeries& Series, const FTimingTrackViewport& Viewport);
+	void UpdateStatsCounterSeries(FTimingGraphSeries& Series, const FTimingTrackViewport& Viewport);
+};
+
+////////////////////////////////////////////////////////////////////////////////////////////////////
diff --git a/Engine/Source/Developer/TraceInsights/Private/Insights/ViewModels/TimingViewDrawHelper.cpp b/Engine/Source/Developer/TraceInsights/Private/Insights/ViewModels/TimingViewDrawHelper.cpp
--- a/Engine/Source/Developer/TraceInsights/Private/Insights/ViewModels/TimingViewDrawHelper.cpp
+++ b/Engine/Source/Developer/TraceInsights/Private/Insights/ViewModels/TimingViewDrawHelper.cpp
@@ -127,7 +127,7 @@ void FTimingViewDrawHelper::EndTimeline(FTimingEventsTrack& Track)
 		}
 	}
 
-	Track.Depth = MaxDepth;
+	Track.SetDepth(MaxDepth);
 
 	float NewH;
 	if (MaxDepth < 0)
diff --git a/Engine/Source/Developer/TraceInsights/Private/Insights/ViewModels/TimingViewDrawHelper.h b/Engine/Source/Developer/TraceInsights/Private/Insights/ViewModels/TimingViewDrawHelper.h
--- a/Engine/Source/Developer/TraceInsights/Private/Insights/ViewModels/TimingViewDrawHelper.h
+++ b/Engine/Source/Developer/TraceInsights/Private/Insights/ViewModels/TimingViewDrawHelper.h
@@ -17,6 +17,61 @@ class FTimingTrackViewport;
 
 ////////////////////////////////////////////////////////////////////////////////////////////////////
 
+struct FTimingViewTooltip
+{
+	static constexpr float BorderX = 6.0f;
+	static constexpr float BorderY = 3.0f;
+	static constexpr float MinWidth = 128.0f;
+	static constexpr float MinHeight = 0.0f;
+
+	FTimingViewTooltip() : PosX(0.0f), PosY(0.0f), Width(FTimingViewTooltip::MinWidth), Height(FTimingViewTooltip::MinHeight), Opacity(0.0f) {}
+
+	void Reset()
+	{
+		PosX = 0.0f;
+		PosY = 0.0f;
+		Width = MinWidth;
+		Height = MinHeight;
+		Opacity = 0.0f;
+	}
+
+	void Update(const FVector2D& MousePosition, const float DesiredWidth, const float DesiredHeight, const float ViewportWidth, const float ViewportHeight)
+	{
+		if (Width != DesiredWidth)
+		{
+			Width = Width * 0.75f + DesiredWidth * 0.25f;
+		}
+		
+		Height = DesiredHeight;
+
+		const float MaxX = FMath::Max(0.0f, ViewportWidth - Width - 12.0f); // -12.0f is to avoid overlapping the vertical scrollbar (one on the right side of the view)
+		const float X = FMath::Clamp<float>(MousePosition.X + 12.0f, 0.0f, MaxX);
+		PosX = X;
+
+		const float MaxY = FMath::Max(0.0f, ViewportHeight - Height - 12.0f); // -12.0f is to avoid overlapping the horizontal scrollbar (one on the bottom of the view)
+		float Y = FMath::Clamp<float>(MousePosition.Y + 15.0f, 0.0f, MaxY);
+		PosY = Y;
+
+		const float DesiredOpacity = 1.0f - FMath::Abs(Width - DesiredWidth) / DesiredWidth;
+		if (Opacity < DesiredOpacity)
+		{
+			Opacity = Opacity * 0.9f + DesiredOpacity * 0.1f;
+		}
+		else
+		{
+			Opacity = DesiredOpacity;
+		}
+	}
+
+	float PosX;
+	float PosY;
+	float Width;
+	float Height;
+	float Opacity;
+};
+
+////////////////////////////////////////////////////////////////////////////////////////////////////
+
 class FTimingViewDrawHelper
 {
 public:
@@ -102,6 +157,11 @@ public:
 	const FTimingTrackViewport& GetViewport() const { return Viewport; }
 	const FTimingEventsTrackLayout& GetLayout() const { return Layout; }
 
+	const FSlateBrush* GetWhiteBrush() const { return WhiteBrush; }
+	const FSlateBrush* GetBorderBrush() const { return BorderBrush; }
+	const FSlateBrush* GetEventsBorderBrush() const { return EventsBorderBrush; }
+	const FSlateFontInfo& GetEventFont() const { return EventFont; }
+
 	int32 GetNumEvents() const              { return Stats.NumEvents; }
 	int32 GetNumDrawBoxes() const           { return Stats.NumDrawBoxes; }
 	int32 GetNumMergedBoxes() const         { return Stats.NumMergedBoxes; }
@@ -143,6 +203,9 @@ private:
 	mutable float ValidAreaX;
 	mutable float ValidAreaW;
 
+	//////////////////////////////////////////////////
+	// Builder state
+
 	float TimelineTopY;
 	float TimelineY;
 	int32 MaxDepth;
@@ -151,6 +214,8 @@ private:
 	TArray<float> LastEventX2; // X2 value for last event on each depth, for current timeline
 	TArray<FBoxData> LastBox;
 
+	//////////////////////////////////////////////////
+
 	/** Debug stats */
 	mutable FStats Stats;
 };
diff --git a/Engine/Source/Developer/TraceInsights/Private/Insights/Widgets/SStatsView.cpp b/Engine/Source/Developer/TraceInsights/Private/Insights/Widgets/SStatsView.cpp
--- a/Engine/Source/Developer/TraceInsights/Private/Insights/Widgets/SStatsView.cpp
+++ b/Engine/Source/Developer/TraceInsights/Private/Insights/Widgets/SStatsView.cpp
@@ -18,7 +18,7 @@
 // Insights
 #include "Insights/TimingProfilerCommon.h"
 #include "Insights/TimingProfilerManager.h"
-#include "Insights/ViewModels/GraphTrack.h"
+#include "Insights/ViewModels/TimingGraphTrack.h"
 #include "Insights/ViewModels/StatsNodeHelper.h"
 #include "Insights/ViewModels/StatsViewColumnFactory.h"
 #include "Insights/Widgets/SStatsViewTooltip.h"
diff --git a/Engine/Source/Developer/TraceInsights/Private/Insights/Widgets/STimingView.cpp b/Engine/Source/Developer/TraceInsights/Private/Insights/Widgets/STimingView.cpp
--- a/Engine/Source/Developer/TraceInsights/Private/Insights/Widgets/STimingView.cpp
+++ b/Engine/Source/Developer/TraceInsights/Private/Insights/Widgets/STimingView.cpp
@@ -19,6 +19,7 @@
 #include "Rendering/DrawElements.h"
 #include "SlateOptMacros.h"
 #include "Styling/CoreStyle.h"
+#include "Styling/SlateBrush.h"
 #include "TraceServices/AnalysisService.h"
 #include "TraceServices/SessionService.h"
 #include "Widgets/Input/SButton.h"
@@ -37,6 +38,10 @@
 #include "Insights/TimingProfilerCommon.h"
 #include "Insights/TimingProfilerManager.h"
 #include "Insights/ViewModels/BaseTimingTrack.h"
+#include "Insights/ViewModels/FileActivityTimingTrack.h"
+#include "Insights/ViewModels/LoadingTimingTrack.h"
+#include "Insights/ViewModels/ThreadTimingTrack.h"
+#include "Insights/ViewModels/TimingGraphTrack.h"
 #include "Insights/ViewModels/DrawHelpers.h"
 #include "Insights/ViewModels/TimingViewDrawHelper.h"
 #include "Insights/Widgets/SStatsView.h"
@@ -62,12 +67,13 @@ uint32 GetFileActivityTypeColor(Trace::EFileActivityType Type);
 ////////////////////////////////////////////////////////////////////////////////////////////////////
 
 STimingView::STimingView()
-	: bAssetLoadingMode(false)
+	: LoadingSharedState(MakeShareable(new FLoadingSharedState()))
+	, bAssetLoadingMode(false)
+	, FileActivitySharedState(MakeShareable(new FFileActivitySharedState()))
 	, TimeRulerTrack(FBaseTimingTrack::GenerateId())
 	, MarkersTrack(FBaseTimingTrack::GenerateId())
 	, GraphTrack(MakeShareable(new FTimingGraphTrack(FBaseTimingTrack::GenerateId())))
-	, TooltipWidth(MinTooltipWidth)
-	, TooltipAlpha(0.0f)
+	, Tooltip()
 	, WhiteBrush(FCoreStyle::Get().GetBrush("WhiteBrush"))
 	, MainFont(FCoreStyle::GetDefaultFontStyle("Regular", 8))
 {
@@ -204,7 +210,7 @@ void STimingView::Reset()
 
 	bAreTimingEventsTracksDirty = true;
 
-	bUseDownSampling = true;
+	FTimingEventsTrack::bUseDownSampling = true;
 
 	//////////////////////////////////////////////////
 
@@ -215,29 +221,25 @@ void STimingView::Reset()
 
 	//////////////////////////////////////////////////
 
+	LoadingSharedState->Reset();
+
 	LoadingMainThreadTrack = nullptr;
 	LoadingAsyncThreadTrack = nullptr;
 
 	LoadingMainThreadId = 0;
 	LoadingAsyncThreadId = 0;
 
-	LoadingGetEventNameFn = FLoadingTrackGetEventNameDelegate::CreateRaw(this, &STimingView::GetLoadTimeProfilerEventName);
-
 	EventAggregationTotalCount = 0;
 	EventAggregation.Reset();
 	ObjectTypeAggregationTotalCount = 0;
 	ObjectTypeAggregation.Reset();
 
+	//////////////////////////////////////////////////
+
+	FileActivitySharedState->Reset();
 	IoOverviewTrack = nullptr;
 	IoActivityTrack = nullptr;
 
-	bForceIoEventsUpdate = false;
-	bMergeIoLanes = true;
-	bShowFileActivityBackgroundEvents = false;
-	FileActivities.Reset();
-	FileActivityMap.Reset();
-	AllIoEvents.Reset();
-
 	//////////////////////////////////////////////////
 
 	TimeRulerTrack.Reset();
@@ -412,7 +414,7 @@ void STimingView::Tick(const FGeometry& AllottedGeometry, const double InCurrent
 		}
 	}
 
-	UpdateIo();
+	FileActivitySharedState->Update();
 
 	TSharedPtr<const Trace::IAnalysisSession> Session = FInsightsManager::Get()->GetSession();
 	if (Session)
@@ -449,16 +451,20 @@ void STimingView::Tick(const FGeometry& AllottedGeometry, const double InCurrent
 			if (LoadingMainThreadTrack == nullptr)
 			{
 				uint64 TrackId = FBaseTimingTrack::GenerateId();
-				LoadingMainThreadTrack = AddTimingEventsTrack(TrackId, ETimingEventsTrackType::Loading, TEXT("Loading - Main Thread"), nullptr, -3);
+				LoadingMainThreadTrack = new FLoadingMainThreadTimingTrack(TrackId, LoadingSharedState);
+				LoadingMainThreadTrack->SetOrder(-3);
 				LoadingMainThreadTrack->SetVisibilityFlag(bShowHideAllLoadingTracks);
+				AddTimingEventsTrack(LoadingMainThreadTrack);
 				bIsTimingEventsTrackDirty = true;
 			}
 
 			if (LoadingAsyncThreadTrack == nullptr)
 			{
 				uint64 TrackId = FBaseTimingTrack::GenerateId();
-				LoadingAsyncThreadTrack = AddTimingEventsTrack(TrackId, ETimingEventsTrackType::Loading, TEXT("Loading - Async Thread"), nullptr, -2);
+				LoadingAsyncThreadTrack = new FLoadingAsyncThreadTimingTrack(TrackId, LoadingSharedState);
+				LoadingMainThreadTrack->SetOrder(-2);
 				LoadingAsyncThreadTrack->SetVisibilityFlag(bShowHideAllLoadingTracks);
+				AddTimingEventsTrack(LoadingAsyncThreadTrack);
 				bIsTimingEventsTrackDirty = true;
 			}
 		}
@@ -469,8 +475,10 @@ void STimingView::Tick(const FGeometry& AllottedGeometry, const double InCurrent
 			{
 				// Note: The I/O timelines are just prototypes for now (will be removed once the functionality is moved in analyzer).
 				const uint64 TrackId = FBaseTimingTrack::GenerateId();
-				IoOverviewTrack = AddTimingEventsTrack(TrackId, ETimingEventsTrackType::Io, TEXT("I/O Overview"), nullptr, -1);
+				IoOverviewTrack = new FOverviewFileActivityTimingTrack(TrackId, FileActivitySharedState);
+				IoOverviewTrack->SetOrder(-1);
 				IoOverviewTrack->SetVisibilityFlag(bShowHideAllIoTracks);
+				AddTimingEventsTrack(IoOverviewTrack);
 				bIsTimingEventsTrackDirty = true;
 			}
 		}
@@ -486,8 +494,10 @@ void STimingView::Tick(const FGeometry& AllottedGeometry, const double InCurrent
 				const uint64 TrackId = static_cast<uint64>(GpuTimelineIndex);
 				if (!CachedTimelines.Contains(TrackId))
 				{
-					GpuTrack = AddTimingEventsTrack(TrackId, ETimingEventsTrackType::Gpu, TEXT("GPU"), nullptr, 0);
+					GpuTrack = new FGpuTimingTrack(TrackId, TEXT("GPU"), nullptr);
+					GpuTrack->SetOrder(0);
 					GpuTrack->SetVisibilityFlag(bShowHideAllGpuTracks);
+					AddTimingEventsTrack(GpuTrack);
 					bIsTimingEventsTrackDirty = true;
 				}
 			}
@@ -530,9 +540,9 @@ void STimingView::Tick(const FGeometry& AllottedGeometry, const double InCurrent
 						FString TrackName(ThreadInfo.Name && *ThreadInfo.Name ? ThreadInfo.Name : FString::Printf(TEXT("Thread %u"), ThreadInfo.Id));
 
 						// Create new Timing Events track for the CPU thread.
-						Track = AddTimingEventsTrack(TrackId, ETimingEventsTrackType::Cpu, TrackName, GroupName, Order);
-
-						Track->SetThreadId(ThreadInfo.Id);
+						Track = new FCpuTimingTrack(TrackId, TrackName, GroupName);
+						Track->SetOrder(Order);
+						((FCpuTimingTrack*)Track)->SetThreadId(ThreadInfo.Id);
 						CpuTracks.Add(ThreadInfo.Id, Track);
 
 						FThreadGroup& ThreadGroup = ThreadGroups[GroupName];
@@ -548,6 +558,7 @@ void STimingView::Tick(const FGeometry& AllottedGeometry, const double InCurrent
 							Track->SetVisibilityFlag(bIsGroupVisible);
 						}
 
+						AddTimingEventsTrack(Track);
 						bIsTimingEventsTrackDirty = true;
 					}
 					else
@@ -572,8 +583,10 @@ void STimingView::Tick(const FGeometry& AllottedGeometry, const double InCurrent
 			{
 				// Note: The I/O timelines are just prototypes for now (will be removed once the functionality is moved in analyzer).
 				uint64 TrackId = FBaseTimingTrack::GenerateId();
-				IoActivityTrack = AddTimingEventsTrack(TrackId, ETimingEventsTrackType::Io, TEXT("I/O Activity"), nullptr, 999999);
+				IoActivityTrack = new FDetailedFileActivityTimingTrack(TrackId, FileActivitySharedState);
+				IoActivityTrack->SetOrder(999999);
 				IoActivityTrack->SetVisibilityFlag(bShowHideAllIoTracks);
+				AddTimingEventsTrack(IoActivityTrack);
 				bIsTimingEventsTrackDirty = true;
 			}
 		}
@@ -750,41 +763,13 @@ int32 STimingView::OnPaint(const FPaintArgs& Args, const FGeometry& AllottedGeom
 			FTimingEventsTrack& Track = *TimingEventsTracks[TrackIndex];
 			if (Track.IsVisible())
 			{
-				//TODO: Track.Draw(Helper);
-				if (Track.GetType() == ETimingEventsTrackType::Cpu)
-				{
-					DrawTimingProfilerTrack(Helper, Track);
-				}
-				else if (Track.GetType() == ETimingEventsTrackType::Gpu)
-				{
-					DrawTimingProfilerTrack(Helper, Track);
-				}
-				else if (Track.GetType() == ETimingEventsTrackType::Loading)
-				{
-					DrawLoadTimeProfilerTrack(Helper, Track);
-				}
-				else if (Track.GetType() == ETimingEventsTrackType::Io)
-				{
-					if (&Track == IoOverviewTrack)
-					{
-						DrawIoOverviewTrack(Helper, Track);
-					}
-					else if (&Track == IoActivityTrack)
-					{
-						DrawIoActivityTrack(Helper, Track);
-					}
-				}
+				Track.Draw(Helper);
 			}
 		}
 
 		Helper.EndTimelines();
 	}
 
-	if (GraphTrack->IsVisible())
-	{
-		GraphTrack->Draw(DrawContext, Viewport);
-	}
-
 	if (!FTimingEvent::AreEquals(SelectedTimingEvent, HoveredTimingEvent))
 	{
 		// Highlight the selected timing event (if any).
@@ -811,6 +796,12 @@ int32 STimingView::OnPaint(const FPaintArgs& Args, const FGeometry& AllottedGeom
 		}
 	}
 
+	if (GraphTrack->IsVisible())
+	{
+		GraphTrack->Draw(DrawContext, Viewport);
+		//GraphTrack->Draw(DrawContext, Viewport, MousePosition);
+	}
+
 	if (MarkersTrack.IsVisible())
 	{
 		// Draw the time markers (vertical lines + category & message texts).
@@ -1024,341 +1015,17 @@ int32 STimingView::OnPaint(const FPaintArgs& Args, const FGeometry& AllottedGeom
 		// Draw info about selected event (bottom-right corner).
 		if (SelectedTimingEvent.IsValid())
 		{
-			if (SelectedTimingEvent.Track->GetType() == ETimingEventsTrackType::Cpu ||
-				SelectedTimingEvent.Track->GetType() == ETimingEventsTrackType::Gpu)
-			{
-				const FTimerNodePtr TimerNodePtr = GetTimerNode(SelectedTimingEvent.TypeId);
-
-				FString Str = FString::Printf(TEXT("%s (Incl.: %s, Excl.: %s)"),
-					TimerNodePtr ? *(TimerNodePtr->GetName().ToString()) : TEXT("N/A"),
-					*TimeUtils::FormatTimeAuto(SelectedTimingEvent.Duration()),
-					*TimeUtils::FormatTimeAuto(SelectedTimingEvent.ExclusiveTime));
-
-				const FVector2D Size = FontMeasureService->Measure(Str, MainFont);
-				const float X = Viewport.Width - Size.X - 23.0f;
-				const float Y = Viewport.Height - Size.Y - 18.0f;
-
-				const FLinearColor BackgroundColor(0.05f, 0.05f, 0.05f, 1.0f);
-				const FLinearColor TextColor(0.7f, 0.7f, 0.7f, 1.0f);
-
-				DrawContext.DrawBox(X - 8.0f, Y - 2.0f, Size.X + 16.0f, Size.Y + 4.0f, WhiteBrush, BackgroundColor);
-				DrawContext.LayerId++;
-
-				DrawContext.DrawText(X, Y, Str, MainFont, TextColor);
-				DrawContext.LayerId++;
-			}
-			else if (SelectedTimingEvent.Track->GetType() == ETimingEventsTrackType::Loading)
-			{
-				//TODO: ...
-			}
-			else if (SelectedTimingEvent.Track->GetType() == ETimingEventsTrackType::Io)
-			{
-				//TODO: ...
-			}
+			SelectedTimingEvent.Track->DrawSelectedEventInfo(SelectedTimingEvent, Viewport, DrawContext, WhiteBrush, MainFont);
 		}
 
 		// Draw info about hovered event (like a tooltip at mouse position).
 		if (HoveredTimingEvent.IsValid() && !MousePosition.IsZero())
 		{
-			if (HoveredTimingEvent.Track->GetType() == ETimingEventsTrackType::Cpu ||
-				HoveredTimingEvent.Track->GetType() == ETimingEventsTrackType::Gpu)
-			{
-				const FTimerNodePtr TimerNodePtr = GetTimerNode(HoveredTimingEvent.TypeId);
-				FString Name = TimerNodePtr ? TimerNodePtr->GetName().ToString() : TEXT("N/A");
-
-				float W = FontMeasureService->Measure(Name, MainFont).X;
-
-				if (W < MinTooltipWidth)
-				{
-					W = MinTooltipWidth;
-				}
-				if (TooltipWidth != W)
-				{
-					TooltipWidth = TooltipWidth * 0.75f + W * 0.25f;
-				}
-
-				const float MaxX = FMath::Max(0.0f, Viewport.Width - TooltipWidth - 21.0f);
-				const float X = FMath::Clamp<float>(MousePosition.X + 18.0f, 0.0f, MaxX);
-
-				constexpr float LineH = 14.0f;
-				constexpr float H = 4 * LineH;
-				const float MaxY = FMath::Max(0.0f, Viewport.Height - H - 18.0f);
-				float Y = FMath::Clamp<float>(MousePosition.Y + 18.0f, 0.0f, MaxY);
-
-				const float Alpha = 1.0f - FMath::Abs(TooltipWidth - W) / W;
-				if (TooltipAlpha < Alpha)
-				{
-					TooltipAlpha = TooltipAlpha * 0.9f + Alpha * 0.1f;
-				}
-				else
-				{
-					TooltipAlpha = Alpha;
-				}
-
-				const FLinearColor BackgroundColor(0.05f, 0.05f, 0.05f, TooltipAlpha);
-				const FLinearColor NameColor(0.9f, 0.9f, 0.5f, TooltipAlpha);
-				const FLinearColor TextColor(0.6f, 0.6f, 0.6f, TooltipAlpha);
-				const FLinearColor ValueColor(1.0f, 1.0f, 1.0f, TooltipAlpha);
-
-				DrawContext.DrawBox(X - 6.0f, Y - 3.0f, TooltipWidth + 12.0f, H + 6.0f, WhiteBrush, BackgroundColor);
-				DrawContext.LayerId++;
-
-				DrawContext.DrawText(X, Y, Name, MainFont, NameColor);
-				Y += LineH;
-
-				const float ValueX = X + 58.0f;
-
-				DrawContext.DrawText(X + 3.0f, Y, TEXT("Incl. Time:"), MainFont, TextColor);
-				FString InclStr = TimeUtils::FormatTimeAuto(HoveredTimingEvent.Duration());
-				DrawContext.DrawText(ValueX, Y, InclStr, MainFont, ValueColor);
-				Y += LineH;
-
-				DrawContext.DrawText(X, Y, TEXT("Excl. Time:"), MainFont, TextColor);
-				FString ExclStr = TimeUtils::FormatTimeAuto(HoveredTimingEvent.ExclusiveTime);
-				DrawContext.DrawText(ValueX, Y, ExclStr, MainFont, ValueColor);
-				Y += LineH;
-
-				DrawContext.DrawText(X + 24.0f, Y, TEXT("Depth:"), MainFont, TextColor);
-				FString DepthStr = FString::Printf(TEXT("%d"), HoveredTimingEvent.Depth);
-				DrawContext.DrawText(ValueX, Y, DepthStr, MainFont, ValueColor);
-				Y += LineH;
-
-				DrawContext.LayerId++;
-			}
-			else if (HoveredTimingEvent.Track->GetType() == ETimingEventsTrackType::Loading)
-			{
-				FString Name(LoadingGetEventNameFn.Execute(HoveredTimingEvent.Depth, HoveredTimingEvent.LoadingInfo));
-
-				const Trace::FPackageInfo* Package = HoveredTimingEvent.LoadingInfo.Package;
-				const Trace::FPackageExportInfo* Export = HoveredTimingEvent.LoadingInfo.Export;
-
-				FString PackageName = Package ? Package->Name : TEXT("N/A");
-
-				constexpr float ValueOffsetX = 100.0f;
-				constexpr float MinValueTextWidth = 220.0f;
-
-				float W = FontMeasureService->Measure(Name, MainFont).X;
-
-				float PackageNameW = ValueOffsetX + FontMeasureService->Measure(PackageName, MainFont).X;
-				if (W < PackageNameW)
-				{
-					W = PackageNameW;
-				}
-				if (W < ValueOffsetX + MinValueTextWidth)
-				{
-					W = ValueOffsetX + MinValueTextWidth;
-				}
-				if (W < MinTooltipWidth)
-				{
-					W = MinTooltipWidth;
-				}
-				if (TooltipWidth != W)
-				{
-					TooltipWidth = TooltipWidth * 0.75f + W * 0.25f;
-				}
-
-				const float MaxX = FMath::Max(0.0f, Viewport.Width - TooltipWidth - 21.0f);
-				const float X = FMath::Clamp<float>(MousePosition.X + 18.0f, 0.0f, MaxX);
-				const float ValueX = X + ValueOffsetX;
-
-				constexpr float LineH = 14.0f;
-				float H = 5 * LineH;
-				if (Package)
-				{
-					H += 3 * LineH;
-				}
-				if (Export)
-				{
-					H += 2 * LineH;
-				}
-				const float MaxY = FMath::Max(0.0f, Viewport.Height - H - 18.0f);
-				float Y = FMath::Clamp<float>(MousePosition.Y + 18.0f, 0.0f, MaxY);
-
-				const float Alpha = 1.0f - FMath::Abs(TooltipWidth - W) / W;
-				if (TooltipAlpha < Alpha)
-				{
-					TooltipAlpha = TooltipAlpha * 0.9f + Alpha * 0.1f;
-				}
-				else
-				{
-					TooltipAlpha = Alpha;
-				}
-
-				const FLinearColor BackgroundColor(0.05f, 0.05f, 0.05f, TooltipAlpha);
-				const FLinearColor NameColor(0.9f, 0.9f, 0.5f, TooltipAlpha);
-				const FLinearColor TextColor(0.6f, 0.6f, 0.6f, TooltipAlpha);
-				const FLinearColor ValueColor(1.0f, 1.0f, 1.0f, TooltipAlpha);
-
-				DrawContext.DrawBox(X - 6.0f, Y - 3.0f, TooltipWidth + 12.0f, H + 6.0f, WhiteBrush, BackgroundColor);
-				DrawContext.LayerId++;
-
-				DrawContext.DrawText(X, Y, Name, MainFont, NameColor);
-				Y += LineH;
-
-				DrawContext.DrawText(X, Y, TEXT("Duration:"), MainFont, TextColor);
-				FString InclStr = TimeUtils::FormatTimeAuto(HoveredTimingEvent.Duration());
-				DrawContext.DrawText(ValueX, Y, InclStr, MainFont, ValueColor);
-				Y += LineH;
-
-				DrawContext.DrawText(X, Y, TEXT("Depth:"), MainFont, TextColor);
-				FString DepthStr = FString::Printf(TEXT("%d"), HoveredTimingEvent.Depth);
-				DrawContext.DrawText(ValueX, Y, DepthStr, MainFont, ValueColor);
-				Y += LineH;
-
-				DrawContext.DrawText(X, Y, TEXT("Package Event:"), MainFont, TextColor);
-				DrawContext.DrawText(ValueX, Y, GetName(HoveredTimingEvent.LoadingInfo.PackageEventType), MainFont, ValueColor);
-				Y += LineH;
-
-				if (Package)
-				{
-					DrawContext.DrawText(X, Y, TEXT("Package Name:"), MainFont, TextColor);
-					DrawContext.DrawText(ValueX, Y, PackageName, MainFont, ValueColor);
-					Y += LineH;
-
-					DrawContext.DrawText(X, Y, TEXT("Header Size:"), MainFont, TextColor);
-					FString HeaderSizeStr = FString::Printf(TEXT("%d bytes"), Package->Summary.TotalHeaderSize);
-					DrawContext.DrawText(ValueX, Y, HeaderSizeStr, MainFont, ValueColor);
-					Y += LineH;
-
-					DrawContext.DrawText(X, Y, TEXT("Package Summary:"), MainFont, TextColor);
-					FString SummaryStr = FString::Printf(TEXT("%d names, %d imports, %d exports"), Package->Summary.NameCount, Package->Summary.ImportCount, Package->Summary.ExportCount);
-					DrawContext.DrawText(ValueX, Y, SummaryStr, MainFont, ValueColor);
-					Y += LineH;
-				}
-
-				{
-					DrawContext.DrawText(X, Y, TEXT("Export Event:"), MainFont, TextColor);
-					FString ExportTypeStr = FString::Printf(TEXT("%s%s"), GetName(HoveredTimingEvent.LoadingInfo.ExportEventType), Export && Export->IsAsset ? TEXT(" [asset]") : TEXT(""));
-					DrawContext.DrawText(ValueX, Y, ExportTypeStr, MainFont, ValueColor);
-					Y += LineH;
-				}
-
-				if (Export)
-				{
-					DrawContext.DrawText(X, Y, TEXT("Export Class:"), MainFont, TextColor);
-					FString ClassStr = FString::Printf(TEXT("%s"), Export->Class ? Export->Class->Name : TEXT("N/A"));
-					DrawContext.DrawText(ValueX, Y, ClassStr, MainFont, ValueColor);
-					Y += LineH;
-
-					DrawContext.DrawText(X, Y, TEXT("Serial:"), MainFont, TextColor);
-					FString SerialStr = FString::Printf(TEXT("Offset: %llu, Size: %llu%s"), Export->SerialOffset, Export->SerialSize);
-					DrawContext.DrawText(ValueX, Y, SerialStr, MainFont, ValueColor);
-					Y += LineH;
-				}
-
-				DrawContext.LayerId++;
-			}
-			else if (HoveredTimingEvent.Track->GetType() == ETimingEventsTrackType::Io)
-			{
-				const Trace::EFileActivityType ActivityType = static_cast<Trace::EFileActivityType>(HoveredTimingEvent.TypeId & 0x0F);
-				const bool bHasFailed = ((HoveredTimingEvent.TypeId & 0xF0) != 0);
-
-				FString Name(HoveredTimingEvent.Path);
-
-				constexpr float ValueOffsetX = 50.0f;
-				constexpr float MinValueTextWidth = 220.0f;
-
-				float W = FontMeasureService->Measure(Name, MainFont).X;
-
-				if (W < ValueOffsetX + MinValueTextWidth)
-				{
-					W = ValueOffsetX + MinValueTextWidth;
-				}
-				if (W < MinTooltipWidth)
-				{
-					W = MinTooltipWidth;
-				}
-				if (TooltipWidth != W)
-				{
-					TooltipWidth = TooltipWidth * 0.75f + W * 0.25f;
-				}
-
-				const float MaxX = FMath::Max(0.0f, Viewport.Width - TooltipWidth - 21.0f);
-				const float X = FMath::Clamp<float>(MousePosition.X + 18.0f, 0.0f, MaxX);
-				const float ValueX = X + ValueOffsetX;
-
-				constexpr float LineH = 14.0f;
-				float H = 4 * LineH;
-				if (ActivityType == Trace::FileActivityType_Read || ActivityType == Trace::FileActivityType_Write)
-				{
-					H += 2 * LineH;
-				}
-
-				const float MaxY = FMath::Max(0.0f, Viewport.Height - H - 18.0f);
-				float Y = FMath::Clamp<float>(MousePosition.Y + 18.0f, 0.0f, MaxY);
-
-				const float Alpha = 1.0f - FMath::Abs(TooltipWidth - W) / W;
-				if (TooltipAlpha < Alpha)
-				{
-					TooltipAlpha = TooltipAlpha * 0.9f + Alpha * 0.1f;
-				}
-				else
-				{
-					TooltipAlpha = Alpha;
-				}
-
-				const FLinearColor BackgroundColor(0.05f, 0.05f, 0.05f, TooltipAlpha);
-				const FLinearColor NameColor(0.9f, 0.9f, 0.5f, TooltipAlpha);
-				const FLinearColor TextColor(0.6f, 0.6f, 0.6f, TooltipAlpha);
-				const FLinearColor ValueColor(1.0f, 1.0f, 1.0f, TooltipAlpha);
-
-				DrawContext.DrawBox(X - 6.0f, Y - 3.0f, TooltipWidth + 12.0f, H + 6.0f, WhiteBrush, BackgroundColor);
-				DrawContext.LayerId++;
-
-				FString TypeStr;
-				uint32 TypeColor;
-				if (bHasFailed)
-				{
-					TypeStr = TEXT("Failed ");
-					TypeStr += GetFileActivityTypeName(ActivityType);
-					TypeColor = 0xFFFF3333;
-				}
-				else
-				{
-					TypeStr = GetFileActivityTypeName(ActivityType);
-					TypeColor = GetFileActivityTypeColor(ActivityType);
-				}
-				FLinearColor TypeLinearColor = FLinearColor(FColor(TypeColor));
-				TypeLinearColor.R *= 2.0f;
-				TypeLinearColor.G *= 2.0f;
-				TypeLinearColor.B *= 2.0f;
-				DrawContext.DrawText(X, Y, TypeStr, MainFont, TypeLinearColor);
-				Y += LineH;
-
-				DrawContext.DrawText(X, Y, Name, MainFont, NameColor);
-				Y += LineH;
-
-				DrawContext.DrawText(X, Y, TEXT("Duration:"), MainFont, TextColor);
-				FString DurationStr = TimeUtils::FormatTimeAuto(HoveredTimingEvent.Duration());
-				DrawContext.DrawText(ValueX, Y, DurationStr, MainFont, ValueColor);
-				Y += LineH;
-
-				DrawContext.DrawText(X, Y, TEXT("Depth:"), MainFont, TextColor);
-				FString DepthStr = FString::Printf(TEXT("%d"), HoveredTimingEvent.Depth);
-				DrawContext.DrawText(ValueX, Y, DepthStr, MainFont, ValueColor);
-				Y += LineH;
-
-				if (ActivityType == Trace::FileActivityType_Read || ActivityType == Trace::FileActivityType_Write)
-				{
-					DrawContext.DrawText(X, Y, TEXT("Offset:"), MainFont, TextColor);
-					FString OffsetStr = FText::AsNumber(HoveredTimingEvent.Offset).ToString();
-					DrawContext.DrawText(ValueX, Y, OffsetStr, MainFont, ValueColor);
-					Y += LineH;
-
-					DrawContext.DrawText(X, Y, TEXT("Size:"), MainFont, TextColor);
-					FString SizeStr = FText::AsNumber(HoveredTimingEvent.Size).ToString();
-					SizeStr += TEXT(" bytes");
-					DrawContext.DrawText(ValueX, Y, SizeStr, MainFont, ValueColor);
-					Y += LineH;
-				}
-
-				DrawContext.LayerId++;
-			}
+			HoveredTimingEvent.Track->DrawTooltip(Tooltip, MousePosition, HoveredTimingEvent, Viewport, DrawContext, WhiteBrush, MainFont);
 		}
 		else
 		{
-			TooltipWidth = MinTooltipWidth;
-			TooltipAlpha = 0.0f;
+			Tooltip.Reset();
 		}
 	}
 
@@ -1425,815 +1092,139 @@ float STimingView::DrawAssetLoadingAggregationTable(FDrawContext& DrawContext, f
 	constexpr float AvgColumnWidth = 60.0f;
 	constexpr float MedColumnWidth = 60.0f;
 
-	const float TableWidth = RowNumberColumnWidth
-						   + NameColumnWidth
-						   + CountColumnWidth
-						   + TotalColumnWidth
-						   + MaxColumnWidth
-						   + AvgColumnWidth
-						   + MedColumnWidth
-						   + MinColumnWidth;
-
-	float W = TableWidth + 2 * BorderX;
-	float H = TableHeight + 2 * BorderY;
-
-	float Y = BottomY - H;
-
-	DrawContext.DrawBox(RightX - W, Y, W, H, WhiteBrush, BackgroundColor);
-	DrawContext.LayerId++;
-	Y += BorderY;
-
-	const float TableX = RightX - W + BorderX;
-
-	// Table name
-	{
-		FString TableNameRow = FString::Printf(TEXT("%s (%d records, sorted by Total time)"), TableName, TotalRowCount);
-		DrawContext.DrawText(TableX, Y, TableNameRow, MainFont, TextColor);
-		Y += LineH;
-	}
-
-	// Column names
-	{
-		float X = TableX;
-
-		X += RowNumberColumnWidth;
-
-		DrawContext.DrawText(X, Y, TEXT("Name"), MainFont, TextColorHeader);
-		X += NameColumnWidth;
-
-		X += CountColumnWidth;
-		DrawContext.DrawTextAligned(HAlign_Right, X, Y, TEXT("Count"), MainFont, TextColorHeader);
-
-		X += TotalColumnWidth;
-		DrawContext.DrawTextAligned(HAlign_Right, X, Y, TEXT("Total"), MainFont, TextColorHeader);
-
-		X += MaxColumnWidth;
-		DrawContext.DrawTextAligned(HAlign_Right, X, Y, TEXT("Max [ms]"), MainFont, TextColorHeader);
-
-		X += AvgColumnWidth;
-		DrawContext.DrawTextAligned(HAlign_Right, X, Y, TEXT("Avg [ms]"), MainFont, TextColorHeader);
-
-		X += MedColumnWidth;
-		DrawContext.DrawTextAligned(HAlign_Right, X, Y, TEXT("Med [ms]"), MainFont, TextColorHeader);
-
-		X += MinColumnWidth;
-		DrawContext.DrawTextAligned(HAlign_Right, X, Y, TEXT("Min [ms]"), MainFont, TextColorHeader);
-
-		Y += LineH;
-	}
-
-	// Records
-	for (int32 Index = 0; Index < Aggregation.Num(); ++Index)
-	{
-		constexpr int32 NumDigits = 2;
-		constexpr bool bAddTimeUnit = false;
-
-		const FAssetLoadingEventAggregationRow& Row = Aggregation[Index];
-
-		float X = TableX;
-
-		X += RowNumberColumnWidth;
-		DrawContext.DrawTextAligned(HAlign_Right, X - 4.0f, Y, FString::Printf(TEXT("%d."), Index + 1), MainFont, TextColorHeader);
-
-		DrawContext.DrawText(X, Y, Row.Name, MainFont, TextColor);
-		X += NameColumnWidth;
-
-		X += CountColumnWidth;
-		DrawContext.DrawTextAligned(HAlign_Right, X, Y, FText::AsNumber(Row.Count).ToString(), MainFont, TextColor);
-
-		X += TotalColumnWidth;
-		DrawContext.DrawTextAligned(HAlign_Right, X, Y, TimeUtils::FormatTimeAuto(Row.Total), MainFont, TextColor);
-
-		X += MaxColumnWidth;
-		DrawContext.DrawTextAligned(HAlign_Right, X, Y, TimeUtils::FormatTimeMs(Row.Max, NumDigits, bAddTimeUnit), MainFont, TextColor);
-
-		X += AvgColumnWidth;
-		DrawContext.DrawTextAligned(HAlign_Right, X, Y, TimeUtils::FormatTimeMs(Row.Avg, NumDigits, bAddTimeUnit), MainFont, TextColor);
-
-		X += MedColumnWidth;
-		DrawContext.DrawTextAligned(HAlign_Right, X, Y, TimeUtils::FormatTimeMs(Row.Med, NumDigits, bAddTimeUnit), MainFont, TextColor);
-
-		X += MinColumnWidth;
-		DrawContext.DrawTextAligned(HAlign_Right, X, Y, TimeUtils::FormatTimeMs(Row.Min, NumDigits, bAddTimeUnit), MainFont, TextColor);
-
-		Y += LineH;
-	}
-
-	if (Aggregation.Num() < TotalRowCount)
-	{
-		DrawContext.DrawText(TableX, Y, TEXT("[...]"), MainFont, TextColorHeader);
-		Y += LineH;
-	}
-
-	DrawContext.LayerId++;
-
-	return H;
-}
-
-////////////////////////////////////////////////////////////////////////////////////////////////////
-
-const TCHAR* GetName(ETimingEventsTrackType Type)
-{
-	switch (Type)
-	{
-		case ETimingEventsTrackType::Gpu:		return TEXT("GPU");
-		case ETimingEventsTrackType::Cpu:		return TEXT("CPU");
-		case ETimingEventsTrackType::Loading:	return TEXT("Loading");
-		case ETimingEventsTrackType::Io:		return TEXT("I/O");
-		default:								return TEXT("unknown");
-	}
-}
-
-////////////////////////////////////////////////////////////////////////////////////////////////////
-
-FTimingEventsTrack* STimingView::AddTimingEventsTrack(uint64 TrackId, ETimingEventsTrackType TrackType, const FString& TrackName, const TCHAR* GroupName, int32 Order)
-{
-	FTimingEventsTrack* Track = new FTimingEventsTrack(TrackId, TrackType, TrackName, GroupName);
-
-	Track->SetOrder(Order);
-
-	UE_LOG(TimingProfiler, Log, TEXT("New Timing Events Track (%d) : %s (\"%s\")"), TimingEventsTracks.Num() + 1, GetName(TrackType), *TrackName);
-
-	CachedTimelines.Add(TrackId, Track);
-	TimingEventsTracks.Add(Track);
-
-	return Track;
-}
-
-////////////////////////////////////////////////////////////////////////////////////////////////////
-
-bool STimingView::IsGpuTrackVisible() const
-{
-	return GpuTrack != nullptr && GpuTrack->IsVisible();
-}
-
-////////////////////////////////////////////////////////////////////////////////////////////////////
-
-bool STimingView::IsCpuTrackVisible(uint32 ThreadId) const
-{
-	return CpuTracks.Contains(ThreadId) && CpuTracks[ThreadId]->IsVisible();
-}
-
-////////////////////////////////////////////////////////////////////////////////////////////////////
-
-void STimingView::DrawTimingProfilerTrack(FTimingViewDrawHelper& Helper, FTimingEventsTrack& Track) const
-{
-	if (Helper.BeginTimeline(Track))
-	{
-		TSharedPtr<const Trace::IAnalysisSession> Session = FInsightsManager::Get()->GetSession();
-		if (Session.IsValid() && Trace::ReadTimingProfilerProvider(*Session.Get()))
-		{
-			Trace::FAnalysisSessionReadScope SessionReadScope(*Session.Get());
-
-			const Trace::ITimingProfilerProvider& TimingProfilerProvider = *Trace::ReadTimingProfilerProvider(*Session.Get());
-
-			TimingProfilerProvider.ReadTimers([this, &Helper, &Track, &TimingProfilerProvider](const Trace::FTimingProfilerTimer* Timers, uint64 TimersCount)
-			{
-				TimingProfilerProvider.ReadTimeline(Track.GetId(), [this, &Helper, &Track, Timers](const Trace::ITimingProfilerProvider::Timeline& Timeline)
-				{
-					if (bUseDownSampling)
-					{
-						const double SecondsPerPixel = 1.0 / Helper.GetViewport().ScaleX;
-						Timeline.EnumerateEventsDownSampled(Helper.GetViewport().StartTime, Helper.GetViewport().EndTime, SecondsPerPixel, [this, &Helper, Timers](double StartTime, double EndTime, uint32 Depth, const Trace::FTimingProfilerEvent& Event)
-						{
-							Helper.AddEvent(StartTime, EndTime, Depth, Timers[Event.TimerIndex].Name);
-						});
-					}
-					else
-					{
-						Timeline.EnumerateEvents(Helper.GetViewport().StartTime, Helper.GetViewport().EndTime, [this, &Helper, Timers](double StartTime, double EndTime, uint32 Depth, const Trace::FTimingProfilerEvent& Event)
-						{
-							Helper.AddEvent(StartTime, EndTime, Depth, Timers[Event.TimerIndex].Name);
-						});
-					}
-				});
-			});
-
-			Helper.EndTimeline(Track);
-		}
-	}
-}
-
-////////////////////////////////////////////////////////////////////////////////////////////////////
-
-const TCHAR* GetName(ELoadTimeProfilerPackageEventType Type)
-{
-	switch (Type)
-	{
-		case LoadTimeProfilerPackageEventType_CreateLinker:				return TEXT("CreateLinker");
-		case LoadTimeProfilerPackageEventType_FinishLinker:				return TEXT("FinishLinker");
-		case LoadTimeProfilerPackageEventType_StartImportPackages:		return TEXT("StartImportPackages");
-		case LoadTimeProfilerPackageEventType_SetupImports:				return TEXT("SetupImports");
-		case LoadTimeProfilerPackageEventType_SetupExports:				return TEXT("SetupExports");
-		case LoadTimeProfilerPackageEventType_ProcessImportsAndExports:	return TEXT("ProcessImportsAndExports");
-		case LoadTimeProfilerPackageEventType_ExportsDone:				return TEXT("ExportsDone");
-		case LoadTimeProfilerPackageEventType_PostLoadWait:				return TEXT("PostLoadWait");
-		case LoadTimeProfilerPackageEventType_StartPostLoad:			return TEXT("StartPostLoad");
-		case LoadTimeProfilerPackageEventType_Tick:						return TEXT("Tick");
-		case LoadTimeProfilerPackageEventType_Finish:					return TEXT("Finish");
-		case LoadTimeProfilerPackageEventType_DeferredPostLoad:			return TEXT("DeferredPostLoad");
-		case LoadTimeProfilerPackageEventType_None:						return TEXT("None");
-		default:														return TEXT("");
-	}
-}
-
-const TCHAR* GetName(ELoadTimeProfilerObjectEventType Type)
-{
-	switch (Type)
-	{
-		case LoadTimeProfilerObjectEventType_Create:	return TEXT("Create");
-		case LoadTimeProfilerObjectEventType_Serialize:	return TEXT("Serialize");
-		case LoadTimeProfilerObjectEventType_PostLoad:	return TEXT("PostLoad");
-		case LoadTimeProfilerObjectEventType_None:		return TEXT("None");
-		default:										return TEXT("");
-	};
-}
-
-const TCHAR* STimingView::GetLoadTimeProfilerEventNameByPackageEventType(uint32 Depth, const Trace::FLoadTimeProfilerCpuEvent& Event) const
-{
-	return GetName(Event.PackageEventType);
-}
-
-const TCHAR* STimingView::GetLoadTimeProfilerEventNameByExportEventType(uint32 Depth, const Trace::FLoadTimeProfilerCpuEvent& Event) const
-{
-	return GetName(Event.ExportEventType);
-}
-
-const TCHAR* STimingView::GetLoadTimeProfilerEventNameByPackageName(uint32 Depth, const Trace::FLoadTimeProfilerCpuEvent& Event) const
-{
-	return Event.Package ?  Event.Package->Name : TEXT("");
-}
-
-const TCHAR* STimingView::GetLoadTimeProfilerEventNameByExportClassName(uint32 Depth, const Trace::FLoadTimeProfilerCpuEvent& Event) const
-{
-	return Event.Export && Event.Export->Class ? Event.Export->Class->Name : TEXT("");
-}
-
-const TCHAR* STimingView::GetLoadTimeProfilerEventName(uint32 Depth, const Trace::FLoadTimeProfilerCpuEvent& Event) const
-{
-	if (Depth == 0)
-	{
-		if (Event.Package)
-		{
-			return Event.Package->Name;
-		}
-	}
-
-	if (Event.Export && Event.Export->Class)
-	{
-		return Event.Export->Class->Name;
-	}
-
-	return TEXT("");
-}
-
-////////////////////////////////////////////////////////////////////////////////////////////////////
-
-void STimingView::DrawLoadTimeProfilerTrack(FTimingViewDrawHelper& Helper, FTimingEventsTrack& Track) const
-{
-	if (Helper.BeginTimeline(Track))
-	{
-		TSharedPtr<const Trace::IAnalysisSession> Session = FInsightsManager::Get()->GetSession();
-		if (Session.IsValid() && Trace::ReadLoadTimeProfilerProvider(*Session.Get()))
-		{
-			Trace::FAnalysisSessionReadScope SessionReadScope(*Session.Get());
-
-			const Trace::ILoadTimeProfilerProvider& LoadTimeProfilerProvider = *Trace::ReadLoadTimeProfilerProvider(*Session.Get());
-
-			if (&Track == LoadingMainThreadTrack)
-			{
-				LoadTimeProfilerProvider.ReadMainThreadCpuTimeline([this, &Helper, &Track](const Trace::ILoadTimeProfilerProvider::CpuTimeline& Timeline)
-				{
-					if (bUseDownSampling)
-					{
-						const double SecondsPerPixel = 1.0 / Helper.GetViewport().ScaleX;
-						Timeline.EnumerateEventsDownSampled(Helper.GetViewport().StartTime, Helper.GetViewport().EndTime, SecondsPerPixel, [this, &Helper](double StartTime, double EndTime, uint32 Depth, const Trace::FLoadTimeProfilerCpuEvent& Event)
-						{
-							const TCHAR* Name = LoadingGetEventNameFn.Execute(Depth, Event);
-							Helper.AddEvent(StartTime, EndTime, Depth, Name);
-						});
-					}
-					else
-					{
-						Timeline.EnumerateEvents(Helper.GetViewport().StartTime, Helper.GetViewport().EndTime, [this, &Helper](double StartTime, double EndTime, uint32 Depth, const Trace::FLoadTimeProfilerCpuEvent& Event)
-						{
-							const TCHAR* Name = LoadingGetEventNameFn.Execute(Depth, Event);
-							Helper.AddEvent(StartTime, EndTime, Depth, Name);
-						});
-					}
-				});
-			}
-			else if (&Track == LoadingAsyncThreadTrack)
-			{
-				LoadTimeProfilerProvider.ReadAsyncLoadingThreadCpuTimeline([this, &Helper, &Track](const Trace::ILoadTimeProfilerProvider::CpuTimeline& Timeline)
-				{
-					if (bUseDownSampling)
-					{
-						const double SecondsPerPixel = 1.0 / Helper.GetViewport().ScaleX;
-						Timeline.EnumerateEventsDownSampled(Helper.GetViewport().StartTime, Helper.GetViewport().EndTime, SecondsPerPixel, [this, &Helper](double StartTime, double EndTime, uint32 Depth, const Trace::FLoadTimeProfilerCpuEvent& Event)
-						{
-							const TCHAR* Name = LoadingGetEventNameFn.Execute(Depth, Event);
-							Helper.AddEvent(StartTime, EndTime, Depth, Name);
-						});
-					}
-					else
-					{
-						Timeline.EnumerateEvents(Helper.GetViewport().StartTime, Helper.GetViewport().EndTime, [this, &Helper](double StartTime, double EndTime, uint32 Depth, const Trace::FLoadTimeProfilerCpuEvent& Event)
-						{
-							const TCHAR* Name = LoadingGetEventNameFn.Execute(Depth, Event);
-							Helper.AddEvent(StartTime, EndTime, Depth, Name);
-						});
-					}
-				});
-			}
-		}
-
-		Helper.EndTimeline(Track);
-	}
-}
-
-////////////////////////////////////////////////////////////////////////////////////////////////////
-// "I/O - File Activity" prototype
-////////////////////////////////////////////////////////////////////////////////////////////////////
-
-// The I/O timelines are just prototypes for now.
-// Below code will be removed once the functionality is moved in analyzer.
-
-////////////////////////////////////////////////////////////////////////////////////////////////////
-
-const TCHAR* GetFileActivityTypeName(Trace::EFileActivityType Type)
-{
-	static_assert(Trace::FileActivityType_Open == 0, "Trace::EFileActivityType enum has changed!?");
-	static_assert(Trace::FileActivityType_Close == 1, "Trace::EFileActivityType enum has changed!?");
-	static_assert(Trace::FileActivityType_Read == 2, "Trace::EFileActivityType enum has changed!?");
-	static_assert(Trace::FileActivityType_Write == 3, "Trace::EFileActivityType enum has changed!?");
-	static_assert(Trace::FileActivityType_Count == 4, "Trace::EFileActivityType enum has changed!?");
-	static const TCHAR* GFileActivityTypeNames[] =
-	{
-		TEXT("Open"),
-		TEXT("Close"),
-		TEXT("Read"),
-		TEXT("Write"),
-		TEXT("Idle"), // virtual events added for cases where Close event is more than 1s away from last Open/Read/Write event.
-		TEXT("NotClosed") // virtual events added when an Open activity never closes
-	};
-	return GFileActivityTypeNames[Type];
-}
-
-////////////////////////////////////////////////////////////////////////////////////////////////////
-
-uint32 GetFileActivityTypeColor(Trace::EFileActivityType Type)
-{
-	static const uint32 GFileActivityTypeColors[] =
-	{
-		0xFFCCAA33, // open
-		0xFF33AACC, // close
-		0xFF33AA33, // read
-		0xFFDD33CC, // write
-		0x55333333, // idle
-		0x55553333, // close
-	};
-	return GFileActivityTypeColors[Type];
-}
-
-////////////////////////////////////////////////////////////////////////////////////////////////////
-
-void STimingView::DrawIoOverviewTrack(FTimingViewDrawHelper& Helper, FTimingEventsTrack& Track) const
-{
-	if (Helper.BeginTimeline(Track))
-	{
-		// Draw the IO Overiew track using cached events (as those are sorted by Start Time).
-		for (const FIoTimingEvent& Event : AllIoEvents)
-		{
-			const Trace::EFileActivityType ActivityType = static_cast<Trace::EFileActivityType>(Event.Type & 0x0F);
-
-			if (ActivityType >= Trace::FileActivityType_Count)
-			{
-				// Ignore "Idle" and "NotClosed" events.
-				continue;
-			}
-
-			//const double EventEndTime = Event.EndTime; // keep duration of events
-			const double EventEndTime = Event.StartTime; // make all 0 duration events
-
-			if (EventEndTime <= Helper.GetViewport().StartTime)
-			{
-				continue;
-			}
-			if (Event.StartTime >= Helper.GetViewport().EndTime)
-			{
-				break;
-			}
-
-			uint32 Color = GetFileActivityTypeColor(ActivityType);
-
-			const bool bHasFailed = ((Event.Type & 0xF0) != 0);
-
-			if (bHasFailed)
-			{
-				FString Name = TEXT("Failed ");
-				Name += GetFileActivityTypeName(ActivityType);
-				Color = 0xFFAA0000;
-				Helper.AddEvent(Event.StartTime, EventEndTime, 0, *Name, Color);
-			}
-			else
-			{
-				Helper.AddEvent(Event.StartTime, EventEndTime, 0, GetFileActivityTypeName(ActivityType), Color);
-			}
-		}
-
-		Helper.EndTimeline(Track);
-	}
-}
-
-////////////////////////////////////////////////////////////////////////////////////////////////////
-
-void STimingView::UpdateIo()
-{
-	if (bForceIoEventsUpdate)
-	{
-		bForceIoEventsUpdate = false;
-
-		FileActivities.Reset();
-		FileActivityMap.Reset();
-		AllIoEvents.Reset();
-
-		FStopwatch Stopwatch;
-		Stopwatch.Start();
-
-		TSharedPtr<const Trace::IAnalysisSession> Session = FInsightsManager::Get()->GetSession();
-		if (Session.IsValid() && Trace::ReadFileActivityProvider(*Session.Get()))
-		{
-			Trace::FAnalysisSessionReadScope SessionReadScope(*Session.Get());
-
-			const Trace::IFileActivityProvider& FileActivityProvider = *Trace::ReadFileActivityProvider(*Session.Get());
-
-			// Enumerate all IO events and cache them.
-			FileActivityProvider.EnumerateFileActivity([this](const Trace::FFileInfo& FileInfo, const Trace::IFileActivityProvider::Timeline& Timeline)
-			{
-				Timeline.EnumerateEvents(-std::numeric_limits<double>::infinity(), +std::numeric_limits<double>::infinity(),
-					[this, &FileInfo, &Timeline](double EventStartTime, double EventEndTime, uint32 EventDepth, const Trace::FFileActivity* FileActivity)
-				{
-					TSharedPtr<FIoFileActivity> Activity;
-
-					//if (EventEndTime == std::numeric_limits<double>::infinity())
-					//{
-					//	EventEndTime = EventStartTime;
-					//}
-
-					if (!FileActivityMap.Contains(FileInfo.Id))
-					{
-						Activity = MakeShareable(new FIoFileActivity());
-
-						Activity->Id = FileInfo.Id;
-						Activity->Path = FileInfo.Path;
-						Activity->StartTime = EventStartTime;
-						Activity->EndTime = EventEndTime;
-						Activity->EventCount = 1;
-						Activity->Depth = -1;
-
-						FileActivities.Add(Activity);
-						FileActivityMap.Add(FileInfo.Id, Activity);
-					}
-					else
-					{
-						Activity = FileActivityMap[FileInfo.Id];
-
-						if (FileActivity->ActivityType != Trace::FileActivityType_Close)
-						{
-							ensure(EventStartTime >= Activity->StartTime);
-							if (EventStartTime < Activity->StartTime)
-							{
-								Activity->StartTime = EventStartTime;
-							}
-
-							if (EventEndTime > Activity->EndTime)
-							{
-								Activity->EndTime = EventEndTime;
-							}
-						}
-
-						Activity->EventCount++;
-					}
-
-					if (bMergeIoLanes)
-					{
-						EventDepth = 0;
-					}
-					else
-					{
-						EventDepth = FileInfo.Id % 32; // simple layout
-					}
-
-					uint32 Type = ((uint32)FileActivity->ActivityType & 0x0F) | (FileActivity->Failed ? 0x80 : 0);
-
-					AllIoEvents.Add(FIoTimingEvent{ EventStartTime, EventEndTime, EventDepth, Type, FileActivity->Offset, FileActivity->Size, Activity });
-				});
-
-				return true;
-			});
-		}
-
-		Stopwatch.Stop();
-		UE_LOG(TimingProfiler, Log, TEXT("[IO] Enumerated %s events (%s file activities) in %s."),
-			*FText::AsNumber(AllIoEvents.Num()).ToString(),
-			*FText::AsNumber(FileActivities.Num()).ToString(),
-			*TimeUtils::FormatTimeAuto(Stopwatch.GetAccumulatedTime()));
-		Stopwatch.Restart();
-
-		// Sort cached IO file activities by Start Time.
-		FileActivities.Sort([](const TSharedPtr<FIoFileActivity>& A, const TSharedPtr<FIoFileActivity>& B) { return A->StartTime < B->StartTime; });
-
-		// Sort cached IO events by Start Time.
-		AllIoEvents.Sort([](const FIoTimingEvent& A, const FIoTimingEvent& B) { return A.StartTime < B.StartTime; });
-
-		Stopwatch.Stop();
-		UE_LOG(TimingProfiler, Log, TEXT("[IO] Sorted file activities and events in %s."), *TimeUtils::FormatTimeAuto(Stopwatch.GetAccumulatedTime()));
-
-		if (bMergeIoLanes)
-		{
-			//////////////////////////////////////////////////
-			// Compute depth for file activities (avoids overlaps).
-
-			Stopwatch.Restart();
-
-			TArray<TSharedPtr<FIoFileActivity>> ActivityLanes;
-
-			for (TSharedPtr<FIoFileActivity> FileActivity : FileActivities)
-			{
-				// Find lane (avoiding overlaps with other file activities).
-				for (int32 LaneIndex = 0; LaneIndex < ActivityLanes.Num(); ++LaneIndex)
-				{
-					TSharedPtr<FIoFileActivity> Lane = ActivityLanes[LaneIndex];
-
-					if (FileActivity->StartTime >= Lane->EndTime)
-					{
-						FileActivity->Depth = LaneIndex;
-						ActivityLanes[LaneIndex] = FileActivity;
-						break;
-					}
-				}
-
-				if (FileActivity->Depth < 0)
-				{
-					const int32 MaxLanes = 10000;
-					if (ActivityLanes.Num() < MaxLanes)
-					{
-						// Add new lane.
-						FileActivity->Depth = ActivityLanes.Num();
-						ActivityLanes.Add(FileActivity);
-					}
-					else
-					{
-						int32 LaneIndex = ActivityLanes.Num() - 1;
-						FileActivity->Depth = LaneIndex;
-						TSharedPtr<FIoFileActivity> Lane = ActivityLanes[LaneIndex];
-						if (FileActivity->EndTime > Lane->EndTime)
-						{
-							ActivityLanes[LaneIndex] = FileActivity;
-						}
-					}
-				}
-			}
-
-			Stopwatch.Stop();
-			UE_LOG(TimingProfiler, Log, TEXT("[IO] Computed layout for file activities in %s."), *TimeUtils::FormatTimeAuto(Stopwatch.GetAccumulatedTime()));
+	const float TableWidth = RowNumberColumnWidth
+						   + NameColumnWidth
+						   + CountColumnWidth
+						   + TotalColumnWidth
+						   + MaxColumnWidth
+						   + AvgColumnWidth
+						   + MedColumnWidth
+						   + MinColumnWidth;
 
-			//////////////////////////////////////////////////
+	float W = TableWidth + 2 * BorderX;
+	float H = TableHeight + 2 * BorderY;
 
-			Stopwatch.Restart();
+	float Y = BottomY - H;
 
-			for (FIoTimingEvent& Event : AllIoEvents)
-			{
-				Event.Depth = static_cast<uint32>(Event.FileActivity->Depth);
-			}
+	DrawContext.DrawBox(RightX - W, Y, W, H, WhiteBrush, BackgroundColor);
+	DrawContext.LayerId++;
+	Y += BorderY;
 
-			Stopwatch.Stop();
-			UE_LOG(TimingProfiler, Log, TEXT("[IO] Updated depth for events in %s."), *TimeUtils::FormatTimeAuto(Stopwatch.GetAccumulatedTime()));
+	const float TableX = RightX - W + BorderX;
 
-			//////////////////////////////////////////////////
-			/*
+	// Table name
+	{
+		FString TableNameRow = FString::Printf(TEXT("%s (%d records, sorted by Total time)"), TableName, TotalRowCount);
+		DrawContext.DrawText(TableX, Y, TableNameRow, MainFont, TextColor);
+		Y += LineH;
+	}
 
-			Stopwatch.Restart();
+	// Column names
+	{
+		float X = TableX;
 
-			struct FIoLane
-			{
-				uint32 FileId;
-				const TCHAR* Path;
-				double LastEndTime;
-				double EndTime;
-			};
+		X += RowNumberColumnWidth;
 
-			TArray<FIoLane> Lanes;
+		DrawContext.DrawText(X, Y, TEXT("Name"), MainFont, TextColorHeader);
+		X += NameColumnWidth;
 
-			TArray<FIoTimingEvent> IoEventsToAdd;
+		X += CountColumnWidth;
+		DrawContext.DrawTextAligned(HAlign_Right, X, Y, TEXT("Count"), MainFont, TextColorHeader);
 
-			for (FIoTimingEvent& Event : AllIoEvents)
-			{
-				uint64 TimelineId = Event.Depth;
+		X += TotalColumnWidth;
+		DrawContext.DrawTextAligned(HAlign_Right, X, Y, TEXT("Total"), MainFont, TextColorHeader);
 
-				int32 Depth = -1;
+		X += MaxColumnWidth;
+		DrawContext.DrawTextAligned(HAlign_Right, X, Y, TEXT("Max [ms]"), MainFont, TextColorHeader);
 
-				bool bIsCloseEvent = false;
+		X += AvgColumnWidth;
+		DrawContext.DrawTextAligned(HAlign_Right, X, Y, TEXT("Avg [ms]"), MainFont, TextColorHeader);
 
-				const Trace::EFileActivityType ActivityType = static_cast<Trace::EFileActivityType>(Event.Type & 0x0F);
+		X += MedColumnWidth;
+		DrawContext.DrawTextAligned(HAlign_Right, X, Y, TEXT("Med [ms]"), MainFont, TextColorHeader);
 
-				if (ActivityType == Trace::FileActivityType_Open)
-				{
-					// Find lane (avoiding overlaps with other opened files).
-					for (int32 LaneIndex = 0; LaneIndex < Lanes.Num(); ++LaneIndex)
-					{
-						FIoLane& Lane = Lanes[LaneIndex];
-						if (Event.StartTime >= Lane.EndTime)
-						{
-							Depth = LaneIndex;
-							break;
-						}
-					}
+		X += MinColumnWidth;
+		DrawContext.DrawTextAligned(HAlign_Right, X, Y, TEXT("Min [ms]"), MainFont, TextColorHeader);
 
-					if (Depth < 0)
-					{
-						// Add new lane.
-						Depth = Lanes.AddDefaulted();
-					}
+		Y += LineH;
+	}
 
-					const bool bHasFailed = ((Event.Type & 0xF0) != 0);
+	// Records
+	for (int32 Index = 0; Index < Aggregation.Num(); ++Index)
+	{
+		constexpr int32 NumDigits = 2;
+		constexpr bool bAddTimeUnit = false;
 
-					FIoLane& Lane = Lanes[Depth];
-					Lane.FileId = TimelineId;
-					Lane.Path = Event.Path;
-					Lane.LastEndTime = Event.EndTime;
-					Lane.EndTime = bHasFailed ? Event.EndTime : std::numeric_limits<double>::infinity();
-				}
-				else if (ActivityType == Trace::FileActivityType_Close)
-				{
-					bIsCloseEvent = true;
+		const FAssetLoadingEventAggregationRow& Row = Aggregation[Index];
 
-					// Find lane with same id.
-					for (int32 LaneIndex = 0; LaneIndex < Lanes.Num(); ++LaneIndex)
-					{
-						FIoLane& Lane = Lanes[LaneIndex];
-						if (Lane.FileId == TimelineId)
-						{
-							// Adds an Idle event, but only if has passed at least 1s since last open/read/write activity.
-							if (Event.StartTime - Lane.LastEndTime > 1.0)
-							{
-								IoEventsToAdd.Add(FIoTimingEvent{ Lane.LastEndTime, Event.StartTime, static_cast<uint32>(LaneIndex), Trace::FileActivityType_Count, Event.Path });
-							}
-							Lane.LastEndTime = Event.EndTime;
-							Lane.EndTime = Event.EndTime;
-							Depth = LaneIndex;
-							break;
-						}
-					}
-					ensure(Depth >= 0);
-				}
-				else
-				{
-					// All other events should be inside the virtual Open-Close event.
-					// Find lane with same id.
-					for (int32 LaneIndex = 0; LaneIndex < Lanes.Num(); ++LaneIndex)
-					{
-						FIoLane& Lane = Lanes[LaneIndex];
-						if (Lane.FileId == TimelineId)
-						{
-							Lane.LastEndTime = Event.EndTime;
-							Depth = LaneIndex;
-							break;
-						}
-					}
-					ensure(Depth >= 0);
-				}
+		float X = TableX;
 
-				if (Depth < 0) // just in case
-				{
-					// Add new lane.
-					Depth = Lanes.AddDefaulted();
-					FIoLane& Lane = Lanes[Depth];
-					Lane.FileId = TimelineId;
-					Lane.Path = Event.Path;
-					Lane.LastEndTime = Event.EndTime;
-					Lane.EndTime = bIsCloseEvent ? Event.EndTime : std::numeric_limits<double>::infinity();
-				}
+		X += RowNumberColumnWidth;
+		DrawContext.DrawTextAligned(HAlign_Right, X - 4.0f, Y, FString::Printf(TEXT("%d."), Index + 1), MainFont, TextColorHeader);
 
-				Event.Depth = Depth;
-			}
+		DrawContext.DrawText(X, Y, Row.Name, MainFont, TextColor);
+		X += NameColumnWidth;
 
-			Stopwatch.Stop();
-			UE_LOG(TimingProfiler, Log, TEXT("[IO] Merge 1/3 in %s."), *TimeUtils::FormatTimeAuto(Stopwatch.GetAccumulatedTime()));
-			Stopwatch.Restart();
+		X += CountColumnWidth;
+		DrawContext.DrawTextAligned(HAlign_Right, X, Y, FText::AsNumber(Row.Count).ToString(), MainFont, TextColor);
 
-			for (int32 LaneIndex = 0; LaneIndex < Lanes.Num(); ++LaneIndex)
-			{
-				FIoLane& Lane = Lanes[LaneIndex];
-				if (Lane.EndTime == std::numeric_limits<double>::infinity())
-				{
-					IoEventsToAdd.Add(FIoTimingEvent{ Lane.LastEndTime, Lane.EndTime, static_cast<uint32>(LaneIndex), Trace::FileActivityType_Count + 1, Lane.Path });
-				}
-			}
+		X += TotalColumnWidth;
+		DrawContext.DrawTextAligned(HAlign_Right, X, Y, TimeUtils::FormatTimeAuto(Row.Total), MainFont, TextColor);
 
-			for (const FIoTimingEvent& Event : IoEventsToAdd)
-			{
-				AllIoEvents.Add(Event);
-			}
+		X += MaxColumnWidth;
+		DrawContext.DrawTextAligned(HAlign_Right, X, Y, TimeUtils::FormatTimeMs(Row.Max, NumDigits, bAddTimeUnit), MainFont, TextColor);
 
-			Stopwatch.Stop();
-			UE_LOG(TimingProfiler, Log, TEXT("[IO] Merge 2/3 in %s."), *TimeUtils::FormatTimeAuto(Stopwatch.GetAccumulatedTime()));
-			Stopwatch.Restart();
+		X += AvgColumnWidth;
+		DrawContext.DrawTextAligned(HAlign_Right, X, Y, TimeUtils::FormatTimeMs(Row.Avg, NumDigits, bAddTimeUnit), MainFont, TextColor);
 
-			// Sort cached IO events one more time, also by Start Time.
-			AllIoEvents.Sort([](const FIoTimingEvent& A, const FIoTimingEvent& B) { return A.StartTime < B.StartTime; });
+		X += MedColumnWidth;
+		DrawContext.DrawTextAligned(HAlign_Right, X, Y, TimeUtils::FormatTimeMs(Row.Med, NumDigits, bAddTimeUnit), MainFont, TextColor);
 
-			//// Sort cached IO events again, by Depth and then by Start Time.
-			//AllIoEvents.Sort([](const FIoTimingEvent& A, const FIoTimingEvent& B)
-			//{
-			//	return A.Depth == B.Depth ? A.StartTime < B.StartTime : A.Depth < B.Depth;
-			//});
+		X += MinColumnWidth;
+		DrawContext.DrawTextAligned(HAlign_Right, X, Y, TimeUtils::FormatTimeMs(Row.Min, NumDigits, bAddTimeUnit), MainFont, TextColor);
 
-			Stopwatch.Stop();
-			UE_LOG(TimingProfiler, Log, TEXT("[IO] Merge 3/3 (sort) in %s."), *TimeUtils::FormatTimeAuto(Stopwatch.GetAccumulatedTime()));
-			*/
-		}
+		Y += LineH;
 	}
-}
-
-////////////////////////////////////////////////////////////////////////////////////////////////////
 
-void STimingView::DrawIoActivityTrack(FTimingViewDrawHelper& Helper, FTimingEventsTrack& Track) const
-{
-	// Draw IO track using cached events.
-	if (Helper.BeginTimeline(Track))
+	if (Aggregation.Num() < TotalRowCount)
 	{
-		// Draw IO file activity background events.
-		if (bShowFileActivityBackgroundEvents)
-		{
-			for (const TSharedPtr<FIoFileActivity> Activity : FileActivities)
-			{
-				if (Activity->EndTime <= Helper.GetViewport().StartTime)
-				{
-					continue;
-				}
-				if (Activity->StartTime >= Helper.GetViewport().EndTime)
-				{
-					break;
-				}
+		DrawContext.DrawText(TableX, Y, TEXT("[...]"), MainFont, TextColorHeader);
+		Y += LineH;
+	}
 
-				ensure(Activity->Depth <= 10000);
+	DrawContext.LayerId++;
 
-				Helper.AddEvent(Activity->StartTime, Activity->EndTime, Activity->Depth, Activity->Path, 0x55333333);
-			}
-		}
+	return H;
+}
 
-		// Draw IO file activity foreground events.
-		for (const FIoTimingEvent& Event : AllIoEvents)
-		{
-			if (Event.EndTime <= Helper.GetViewport().StartTime)
-			{
-				continue;
-			}
-			if (Event.StartTime >= Helper.GetViewport().EndTime)
-			{
-				break;
-			}
+////////////////////////////////////////////////////////////////////////////////////////////////////
 
-			ensure(Event.Depth <= 10000);
+void STimingView::AddTimingEventsTrack(FTimingEventsTrack* Track)
+{
+	UE_LOG(TimingProfiler, Log, TEXT("New Timing Events Track (%d) : %s - %s (\"%s\")"),
+		TimingEventsTracks.Num() + 1,
+		*Track->GetType().ToString(),
+		*Track->GetSubType().ToString(),
+		*Track->GetName());
 
-			const Trace::EFileActivityType ActivityType = static_cast<Trace::EFileActivityType>(Event.Type & 0x0F);
+	CachedTimelines.Add(Track->GetId(), Track);
+	TimingEventsTracks.Add(Track);
+}
 
-			uint32 Color = GetFileActivityTypeColor(ActivityType);
+////////////////////////////////////////////////////////////////////////////////////////////////////
 
-			if (ActivityType < Trace::FileActivityType_Count)
-			{
-				const bool bHasFailed = ((Event.Type & 0xF0) != 0);
+bool STimingView::IsGpuTrackVisible() const
+{
+	return GpuTrack != nullptr && GpuTrack->IsVisible();
+}
 
-				if (bHasFailed)
-				{
-					FString Name = TEXT("Failed ");
-					Name += GetFileActivityTypeName(ActivityType);
-					Color = 0xFFAA0000;
-					Helper.AddEvent(Event.StartTime, Event.EndTime, Event.Depth, *Name, Color);
-				}
-				else
-				{
-					Helper.AddEvent(Event.StartTime, Event.EndTime, Event.Depth, GetFileActivityTypeName(ActivityType), Color);
-				}
-			}
-			else
-			{
-				FString Name = GetFileActivityTypeName(ActivityType);
-				Name += " [";
-				Name += Event.FileActivity->Path;
-				Name += "]";
-				Helper.AddEvent(Event.StartTime, Event.EndTime, Event.Depth, *Name, Color);
-			}
-		}
+////////////////////////////////////////////////////////////////////////////////////////////////////
 
-		Helper.EndTimeline(Track);
-	}
+bool STimingView::IsCpuTrackVisible(uint32 ThreadId) const
+{
+	return CpuTracks.Contains(ThreadId) && CpuTracks[ThreadId]->IsVisible();
 }
 
-////////////////////////////////////////////////////////////////////////////////////////////////////
-// end of "I/O - File Activity" prototype
 ////////////////////////////////////////////////////////////////////////////////////////////////////
 
 FReply STimingView::OnMouseButtonDown(const FGeometry& MyGeometry, const FPointerEvent& MouseEvent)
@@ -2828,7 +1819,7 @@ FReply STimingView::OnKeyDown(const FGeometry& MyGeometry, const FKeyEvent& InKe
 	}
 	else if (InKeyEvent.GetKey() == EKeys::D) // debug: toggles down-sampling on/off
 	{
-		bUseDownSampling = !bUseDownSampling;
+		FTimingEventsTrack::bUseDownSampling = !FTimingEventsTrack::bUseDownSampling;
 		bAreTimingEventsTracksDirty = true;
 		return FReply::Handled();
 	}
@@ -2859,34 +1850,34 @@ FReply STimingView::OnKeyDown(const FGeometry& MyGeometry, const FKeyEvent& InKe
 	}
 	else if (InKeyEvent.GetKey() == EKeys::O)  // debug: toggles IO merge lanes algorithm
 	{
-		//bMergeIoLanes = !bMergeIoLanes;
-		//bForceIoEventsUpdate = true;
-		bShowFileActivityBackgroundEvents = !bShowFileActivityBackgroundEvents;
+		//FileActivitySharedState->ToggleMergeLanes();
+		//FileActivitySharedState->RequestUpdate();
+		FileActivitySharedState->ToggleBackgroundEvents();
 		return FReply::Handled();
 	}
 	else if (InKeyEvent.GetKey() == EKeys::One)
 	{
-		LoadingGetEventNameFn = FLoadingTrackGetEventNameDelegate::CreateRaw(this, &STimingView::GetLoadTimeProfilerEventNameByPackageEventType);
+		LoadingSharedState->SetColorSchema(0);
 		return FReply::Handled();
 	}
 	else if (InKeyEvent.GetKey() == EKeys::Two)
 	{
-		LoadingGetEventNameFn = FLoadingTrackGetEventNameDelegate::CreateRaw(this, &STimingView::GetLoadTimeProfilerEventNameByExportEventType);
+		LoadingSharedState->SetColorSchema(1);
 		return FReply::Handled();
 	}
 	else if (InKeyEvent.GetKey() == EKeys::Three)
 	{
-		LoadingGetEventNameFn = FLoadingTrackGetEventNameDelegate::CreateRaw(this, &STimingView::GetLoadTimeProfilerEventNameByPackageName);
+		LoadingSharedState->SetColorSchema(2);
 		return FReply::Handled();
 	}
 	else if (InKeyEvent.GetKey() == EKeys::Four)
 	{
-		LoadingGetEventNameFn = FLoadingTrackGetEventNameDelegate::CreateRaw(this, &STimingView::GetLoadTimeProfilerEventNameByExportClassName);
+		LoadingSharedState->SetColorSchema(3);
 		return FReply::Handled();
 	}
 	else if (InKeyEvent.GetKey() == EKeys::Five)
 	{
-		LoadingGetEventNameFn = FLoadingTrackGetEventNameDelegate::CreateRaw(this, &STimingView::GetLoadTimeProfilerEventName);
+		LoadingSharedState->SetColorSchema(4);
 		return FReply::Handled();
 	}
 
@@ -2951,9 +1942,10 @@ void STimingView::ShowContextMenu(const FPointerEvent& MouseEvent)
 			FText Title;
 			if (HoveredTimingEvent.Track != nullptr)
 			{
-				if (HoveredTimingEvent.Track->GetGroupName() != nullptr)
+				if (HoveredTimingEvent.Track->GetType() == FName(TEXT("Thread")) &&
+					((FThreadTimingTrack*)HoveredTimingEvent.Track)->GetGroupName() != nullptr)
 				{
-					Title = FText::Format(LOCTEXT("TrackTitleGroupFmt", "{0} (Group: {1})"), FText::FromString(HoveredTimingEvent.Track->GetName()), FText::FromString(HoveredTimingEvent.Track->GetGroupName()));
+					Title = FText::Format(LOCTEXT("TrackTitleGroupFmt", "{0} (Group: {1})"), FText::FromString(HoveredTimingEvent.Track->GetName()), FText::FromString(((FThreadTimingTrack*)HoveredTimingEvent.Track)->GetGroupName()));
 				}
 				else
 				{
@@ -3385,7 +2377,7 @@ void STimingView::UpdateHoveredTimingEvent(float MX, float MY)
 				const double EndTime = StartTime + 2.0 / Viewport.ScaleX; // +2px
 				constexpr bool bStopAtFirstMatch = true; // get first one matching
 				constexpr bool bSearchForLargestEvent = false;
-				SearchTimingEvent(StartTime, EndTime,
+				HoveredTimingEvent.Track->SearchTimingEvent(StartTime, EndTime,
 					[Depth](double, double, uint32 EventDepth)
 					{
 						return EventDepth == Depth;
@@ -3400,265 +2392,14 @@ void STimingView::UpdateHoveredTimingEvent(float MX, float MY)
 
 ////////////////////////////////////////////////////////////////////////////////////////////////////
 
-bool STimingView::SearchTimingEvent(const double InStartTime,
-									const double InEndTime,
-									TFunctionRef<bool(double, double, uint32)> InPredicate,
-									FTimingEvent& InOutTimingEvent,
-									bool bInStopAtFirstMatch,
-									bool bInSearchForLargestEvent) const
-{
-	struct FSearchTimingEventContext
-	{
-		const double StartTime;
-		const double EndTime;
-		TFunctionRef<bool(double, double, uint32)> Predicate;
-		FTimingEvent& TimingEvent;
-		const bool bStopAtFirstMatch;
-		const bool bSearchForLargestEvent;
-		mutable bool bFound;
-		mutable bool bContinueSearching;
-		mutable double LargestDuration;
-
-		FSearchTimingEventContext(const double InStartTime, const double InEndTime, TFunctionRef<bool(double, double, uint32)> InPredicate, FTimingEvent& InOutTimingEvent, bool bInStopAtFirstMatch, bool bInSearchForLargestEvent)
-			: StartTime(InStartTime)
-			, EndTime(InEndTime)
-			, Predicate(InPredicate)
-			, TimingEvent(InOutTimingEvent)
-			, bStopAtFirstMatch(bInStopAtFirstMatch)
-			, bSearchForLargestEvent(bInSearchForLargestEvent)
-			, bFound(false)
-			, bContinueSearching(true)
-			, LargestDuration(-1.0)
-		{
-		}
-
-		void CheckEvent(double EventStartTime, double EventEndTime, uint32 EventDepth, const Trace::FTimingProfilerEvent& Event)
-		{
-			if (bContinueSearching && Predicate(EventStartTime, EventEndTime, EventDepth))
-			{
-				if (!bSearchForLargestEvent || EventEndTime - EventStartTime > LargestDuration)
-				{
-					LargestDuration = EventEndTime - EventStartTime;
-
-					TimingEvent.TypeId = Event.TimerIndex;
-					TimingEvent.Depth = EventDepth;
-					TimingEvent.StartTime = EventStartTime;
-					TimingEvent.EndTime = EventEndTime;
-
-					bFound = true;
-					bContinueSearching = !bStopAtFirstMatch || bSearchForLargestEvent;
-				}
-			}
-		}
-
-		void CheckEvent(double EventStartTime, double EventEndTime, uint32 EventDepth, const Trace::FLoadTimeProfilerCpuEvent& Event)
-		{
-			if (bContinueSearching && Predicate(EventStartTime, EventEndTime, EventDepth))
-			{
-				if (!bSearchForLargestEvent || EventEndTime - EventStartTime > LargestDuration)
-				{
-					LargestDuration = EventEndTime - EventStartTime;
-
-					TimingEvent.TypeId = 0;
-					TimingEvent.Depth = EventDepth;
-					TimingEvent.StartTime = EventStartTime;
-					TimingEvent.EndTime = EventEndTime;
-
-					TimingEvent.LoadingInfo = Event;
-
-					bFound = true;
-					bContinueSearching = !bStopAtFirstMatch || bSearchForLargestEvent;
-				}
-			}
-		}
-
-		void CheckEvent(const FIoTimingEvent& Event, uint32 EventDepth)
-		{
-			if (bContinueSearching && Predicate(Event.StartTime, Event.EndTime, EventDepth))
-			{
-				if (!bSearchForLargestEvent || Event.EndTime - Event.StartTime > LargestDuration)
-				{
-					LargestDuration = Event.EndTime - Event.StartTime;
-
-					TimingEvent.TypeId = Event.Type;
-					TimingEvent.Depth = EventDepth;
-					TimingEvent.StartTime = Event.StartTime;
-					TimingEvent.EndTime = Event.EndTime;
-
-					TimingEvent.Offset = Event.Offset;
-					TimingEvent.Size = Event.Size;
-					TimingEvent.Path = Event.FileActivity->Path;
-
-					bFound = true;
-					bContinueSearching = !bStopAtFirstMatch || bSearchForLargestEvent;
-				}
-			}
-		}
-	};
-
-	FSearchTimingEventContext Ctx(InStartTime, InEndTime, InPredicate, InOutTimingEvent, bInStopAtFirstMatch, bInSearchForLargestEvent);
-
-	TSharedPtr<const Trace::IAnalysisSession> Session = FInsightsManager::Get()->GetSession();
-	if (Session.IsValid())
-	{
-		Trace::FAnalysisSessionReadScope SessionReadScope(*Session.Get());
-
-		const FTimingEventsTrack* Track = Ctx.TimingEvent.Track;
-
-		if (Track->GetType() == ETimingEventsTrackType::Cpu ||
-			Track->GetType() == ETimingEventsTrackType::Gpu)
-		{
-			if (Trace::ReadTimingProfilerProvider(*Session.Get()))
-			{
-				const Trace::ITimingProfilerProvider& TimingProfilerProvider = *Trace::ReadTimingProfilerProvider(*Session.Get());
-
-				TimingProfilerProvider.ReadTimeline(Track->GetId(), [&Ctx](const Trace::ITimingProfilerProvider::Timeline& Timeline)
-				{
-					Timeline.EnumerateEvents(Ctx.StartTime, Ctx.EndTime, [&Ctx](double EventStartTime, double EventEndTime, uint32 EventDepth, const Trace::FTimingProfilerEvent& Event)
-					{
-						Ctx.CheckEvent(EventStartTime, EventEndTime, EventDepth, Event);
-					});
-				});
-
-				if (Ctx.bFound)
-				{
-					// Compute Exclusive Time.
-					TimingProfilerProvider.ReadTimeline(Track->GetId(), [&Ctx](const Trace::ITimingProfilerProvider::Timeline& Timeline)
-					{
-						struct FEnumerationState
-						{
-							double EventStartTime;
-							double EventEndTime;
-							uint64 EventDepth;
-
-							uint64 CurrentDepth;
-							double LastTime;
-							double ExclusiveTime;
-							bool IsInEventScope;
-						};
-						FEnumerationState State;
-
-						State.EventStartTime = Ctx.TimingEvent.StartTime;
-						State.EventEndTime = Ctx.TimingEvent.EndTime;
-						State.EventDepth = Ctx.TimingEvent.Depth;
-
-						State.CurrentDepth = 0;
-						State.LastTime = 0.0;
-						State.ExclusiveTime = 0.0;
-						State.IsInEventScope = false;
-
-						Timeline.EnumerateEvents(Ctx.TimingEvent.StartTime, Ctx.TimingEvent.EndTime, [&State](bool IsEnter, double Time, const Trace::FTimingProfilerEvent& Event)
-						{
-							if (IsEnter)
-							{
-								if (State.IsInEventScope && State.CurrentDepth == State.EventDepth + 1)
-								{
-									State.ExclusiveTime += Time - State.LastTime;
-								}
-								if (State.CurrentDepth == State.EventDepth && Time == State.EventStartTime)
-								{
-									State.IsInEventScope = true;
-								}
-								++State.CurrentDepth;
-							}
-							else
-							{
-								--State.CurrentDepth;
-								if (State.CurrentDepth == State.EventDepth && Time == State.EventEndTime)
-								{
-									State.IsInEventScope = false;
-									State.ExclusiveTime += Time - State.LastTime;
-								}
-							}
-							State.LastTime = Time;
-						});
-
-						Ctx.TimingEvent.ExclusiveTime = State.ExclusiveTime;
-					});
-				}
-			}
-		}
-		else if (Track->GetType() == ETimingEventsTrackType::Loading)
-		{
-			if (Trace::ReadLoadTimeProfilerProvider(*Session.Get()))
-			{
-				const Trace::ILoadTimeProfilerProvider& LoadTimeProfilerProvider = *Trace::ReadLoadTimeProfilerProvider(*Session.Get());
-
-				if (Track == LoadingMainThreadTrack)
-				{
-					LoadTimeProfilerProvider.ReadMainThreadCpuTimeline([&Ctx](const Trace::ILoadTimeProfilerProvider::CpuTimeline& Timeline)
-					{
-						Timeline.EnumerateEvents(Ctx.StartTime, Ctx.EndTime, [&Ctx](double EventStartTime, double EventEndTime, uint32 EventDepth, const Trace::FLoadTimeProfilerCpuEvent& Event)
-						{
-							Ctx.CheckEvent(EventStartTime, EventEndTime, EventDepth, Event);
-						});
-					});
-				}
-				else if (Track == LoadingAsyncThreadTrack)
-				{
-					LoadTimeProfilerProvider.ReadAsyncLoadingThreadCpuTimeline([&Ctx](const Trace::ILoadTimeProfilerProvider::CpuTimeline& Timeline)
-					{
-						Timeline.EnumerateEvents(Ctx.StartTime, Ctx.EndTime, [&Ctx](double EventStartTime, double EventEndTime, uint32 EventDepth, const Trace::FLoadTimeProfilerCpuEvent& Event)
-						{
-							Ctx.CheckEvent(EventStartTime, EventEndTime, EventDepth, Event);
-						});
-					});
-				}
-			}
-		}
-		else if (Track->GetType() == ETimingEventsTrackType::Io)
-		{
-			if (Track == IoOverviewTrack)
-			{
-				for (const FIoTimingEvent& Event : AllIoEvents)
-				{
-					if (Event.EndTime <= Ctx.StartTime)
-					{
-						continue;
-					}
-
-					if (!Ctx.bContinueSearching || Event.StartTime >= Ctx.EndTime)
-					{
-						break;
-					}
-
-					Ctx.CheckEvent(Event, 0);
-				}
-			}
-			else if (Track == IoActivityTrack)
-			{
-				for (const FIoTimingEvent& Event : AllIoEvents)
-				{
-					if (Event.EndTime <= Ctx.StartTime)
-					{
-						continue;
-					}
-
-					if (!Ctx.bContinueSearching || Event.StartTime >= Ctx.EndTime)
-					{
-						break;
-					}
-
-					Ctx.CheckEvent(Event, Event.Depth);
-				}
-			}
-		}
-	}
-
-	return Ctx.bFound;
-}
-
-////////////////////////////////////////////////////////////////////////////////////////////////////
-
 void STimingView::OnSelectedTimingEventChanged()
 {
 	// Select the timer node coresponding to timing event type of selected timing event.
 	if (!bAssetLoadingMode &&
 		SelectedTimingEvent.IsValid() &&
-		(SelectedTimingEvent.Track->GetType() == ETimingEventsTrackType::Cpu ||
-		 SelectedTimingEvent.Track->GetType() == ETimingEventsTrackType::Gpu))
+		(SelectedTimingEvent.Track->GetType() == FName(TEXT("Thread"))))
 	{
-		//TODO: make this a more generic (i.e. no hardocdings on TimingProfilerManager)
+		//TODO: make this more generic
 		TSharedPtr<STimingProfilerWindow> Wnd = FTimingProfilerManager::Get()->GetProfilerWindow();
 		if (Wnd)
 		{
@@ -3694,7 +2435,7 @@ void STimingView::SelectLeftTimingEvent()
 		const double EndTime = SelectedTimingEvent.EndTime;
 		const bool bStopAtFirstMatch = false; // get last one matching
 		const bool bSearchForLargestEvent = false;
-		if (SearchTimingEvent(0.0, StartTime,
+		if (SelectedTimingEvent.Track->SearchTimingEvent(0.0, StartTime,
 			[Depth, StartTime, EndTime](double EventStartTime, double EventEndTime, uint32 EventDepth)
 			{
 				return EventDepth == Depth &&
@@ -3719,7 +2460,7 @@ void STimingView::SelectRightTimingEvent()
 		const double EndTime = SelectedTimingEvent.EndTime;
 		const bool bStopAtFirstMatch = true; // get first one matching
 		const bool bSearchForLargestEvent = false;
-		if (SearchTimingEvent(EndTime, Viewport.MaxValidTime,
+		if (SelectedTimingEvent.Track->SearchTimingEvent(EndTime, Viewport.MaxValidTime,
 			[Depth, StartTime, EndTime](double EventStartTime, double EventEndTime, uint32 EventDepth)
 			{
 				return EventDepth == Depth &&
@@ -3745,7 +2486,7 @@ void STimingView::SelectUpTimingEvent()
 		double EndTime = SelectedTimingEvent.EndTime;
 		const bool bStopAtFirstMatch = true; // get first one matching
 		const bool bSearchForLargestEvent = false;
-		if (SearchTimingEvent(StartTime, EndTime,
+		if (SelectedTimingEvent.Track->SearchTimingEvent(StartTime, EndTime,
 			[Depth, StartTime, EndTime](double EventStartTime, double EventEndTime, uint32 EventDepth)
 			{
 				return EventDepth == Depth
@@ -3771,7 +2512,7 @@ void STimingView::SelectDownTimingEvent()
 		double EndTime = SelectedTimingEvent.EndTime;
 		const bool bStopAtFirstMatch = false; // check all timing events
 		const bool bSearchForLargestEvent = true; // get largest timing event
-		if (SearchTimingEvent(StartTime, EndTime,
+		if (SelectedTimingEvent.Track->SearchTimingEvent(StartTime, EndTime,
 			[Depth, StartTime, EndTime](double EventStartTime, double EventEndTime, uint32 EventDepth)
 			{
 				return EventDepth == Depth
@@ -3861,26 +2602,6 @@ void STimingView::FrameSelection()
 
 ////////////////////////////////////////////////////////////////////////////////////////////////////
 
-const FTimerNodePtr STimingView::GetTimerNode(uint64 TypeId) const
-{
-	const FTimerNodePtr* TimerNodePtrPtr = nullptr;
-	TSharedPtr<STimingProfilerWindow> Wnd = FTimingProfilerManager::Get()->GetProfilerWindow();
-	if (Wnd && Wnd->TimersView)
-	{
-		TimerNodePtrPtr = Wnd->TimersView->GetTimerNode(TypeId);
-		if (TimerNodePtrPtr == nullptr)
-		{
-			// List of timers in TimersView not up to date?
-			// Refresh and try again.
-			Wnd->TimersView->RebuildTree(false);
-			TimerNodePtrPtr = Wnd->TimersView->GetTimerNode(TypeId);
-		}
-	}
-	return TimerNodePtrPtr ? *TimerNodePtrPtr : nullptr;
-}
-
-////////////////////////////////////////////////////////////////////////////////////////////////////
-
 TSharedRef<SWidget> STimingView::MakeTracksFilterMenu()
 {
 	FMenuBuilder MenuBuilder(/*bInShouldCloseWindowAfterMenuSelection=*/true, nullptr);
@@ -4060,7 +2781,7 @@ void STimingView::ShowHideAllCpuTracks_Execute()
 	for (int32 TrackIndex = 0; TrackIndex < TimingEventsTracks.Num(); ++TrackIndex)
 	{
 		FTimingEventsTrack& Track = *TimingEventsTracks[TrackIndex];
-		if (Track.GetType() == ETimingEventsTrackType::Cpu)
+		if (Track.GetType() == FName(TEXT("Thread")) && Track.GetSubType() == FName(TEXT("CPU")))
 		{
 			Track.SetVisibilityFlag(bShowHideAllCpuTracks);
 		}
@@ -4093,7 +2814,7 @@ void STimingView::ShowHideAllGpuTracks_Execute()
 	for (int32 TrackIndex = 0; TrackIndex < TimingEventsTracks.Num(); ++TrackIndex)
 	{
 		FTimingEventsTrack& Track = *TimingEventsTracks[TrackIndex];
-		if (Track.GetType() == ETimingEventsTrackType::Gpu)
+		if (Track.GetType() == FName(TEXT("Thread")) && Track.GetSubType() == FName(TEXT("GPU")))
 		{
 			Track.SetVisibilityFlag(bShowHideAllGpuTracks);
 		}
@@ -4121,7 +2842,7 @@ void STimingView::ShowHideAllLoadingTracks_Execute()
 	for (int32 TrackIndex = 0; TrackIndex < TimingEventsTracks.Num(); ++TrackIndex)
 	{
 		FTimingEventsTrack& Track = *TimingEventsTracks[TrackIndex];
-		if (Track.GetType() == ETimingEventsTrackType::Loading)
+		if (Track.GetType() == FName(TEXT("Loading")))
 		{
 			Track.SetVisibilityFlag(bShowHideAllLoadingTracks);
 		}
@@ -4146,7 +2867,7 @@ void STimingView::ShowHideAllIoTracks_Execute()
 	for (int32 TrackIndex = 0; TrackIndex < TimingEventsTracks.Num(); ++TrackIndex)
 	{
 		FTimingEventsTrack& Track = *TimingEventsTracks[TrackIndex];
-		if (Track.GetType() == ETimingEventsTrackType::Io)
+		if (Track.GetType() == FName(TEXT("FileActivity")))
 		{
 			Track.SetVisibilityFlag(bShowHideAllIoTracks);
 		}
@@ -4156,7 +2877,7 @@ void STimingView::ShowHideAllIoTracks_Execute()
 
 	if (bShowHideAllIoTracks)
 	{
-		bForceIoEventsUpdate = true;
+		FileActivitySharedState->RequestUpdate();
 	}
 }
 
@@ -4240,8 +2961,8 @@ void STimingView::ToggleTrackVisibilityByGroup_Execute(const TCHAR* InGroupName)
 		for (auto& KV : CachedTimelines)
 		{
 			FTimingEventsTrack& Track = *KV.Value;
-			if (Track.GetType() == ETimingEventsTrackType::Cpu &&
-				Track.GetGroupName() == InGroupName)
+			if (Track.GetType() == FName(TEXT("Thread")) && Track.GetSubType() == FName(TEXT("CPU")) &&
+				((FCpuTimingTrack*)&Track)->GetGroupName() == InGroupName)
 			{
 				Track.SetVisibilityFlag(ThreadGroup.bIsVisible);
 			}
diff --git a/Engine/Source/Developer/TraceInsights/Private/Insights/Widgets/STimingView.h b/Engine/Source/Developer/TraceInsights/Private/Insights/Widgets/STimingView.h
--- a/Engine/Source/Developer/TraceInsights/Private/Insights/Widgets/STimingView.h
+++ b/Engine/Source/Developer/TraceInsights/Private/Insights/Widgets/STimingView.h
@@ -15,7 +15,6 @@
 
 // Insights
 #include "Insights/Common/FixedCircularBuffer.h"
-#include "Insights/ViewModels/GraphTrack.h"
 #include "Insights/ViewModels/MarkersTimingTrack.h"
 #include "Insights/ViewModels/TimerNode.h"
 #include "Insights/ViewModels/TimeRulerTrack.h"
@@ -24,6 +23,9 @@
 #include "Insights/ViewModels/TimingTrackViewport.h"
 #include "Insights/ViewModels/TimingViewDrawHelper.h"
 
+class FFileActivitySharedState;
+class FLoadingSharedState;
+class FTimingGraphTrack;
 class FTimingViewDrawHelper;
 class SScrollBar;
 
@@ -36,15 +38,9 @@ DECLARE_DELEGATE_RetVal(int32, FHoveredEventChangedDelegate);
 /** The delegate to be invoked when the selected timing event has changed. Returns id of timing event or -1 (if no event is hovered). */
 DECLARE_DELEGATE_RetVal(int32, FSelectedEventChangedDelegate);
 
-/** Defines FLoadingTrackGetEventNameDelegate delegate interface. Returns the name for a timing event in a Loading track. */
-DECLARE_DELEGATE_RetVal_TwoParams(const TCHAR*, FLoadingTrackGetEventNameDelegate, uint32 /*Depth*/, const Trace::FLoadTimeProfilerCpuEvent& /*Event*/);
-
 /** A custom widget used to display timing events. */
 class STimingView : public SCompoundWidget
 {
-private:
-	static constexpr float MinTooltipWidth = 110.0f;
-
 protected:
 	struct FAssetLoadingEventAggregationRow
 	{
@@ -272,7 +268,7 @@ public:
 	bool IsGpuTrackVisible() const;
 	bool IsCpuTrackVisible(uint32 ThreadId) const;
 
-	TSharedPtr<FTimingGraphTrack> GetMainTimingGraphTrack() { return StaticCastSharedPtr<FTimingGraphTrack>(GraphTrack); }
+	TSharedPtr<FTimingGraphTrack> GetMainTimingGraphTrack() { return GraphTrack; }
 
 protected:
 	virtual FVector2D ComputeDesiredSize(float) const override
@@ -280,20 +276,7 @@ protected:
 		return FVector2D(16.0f, 16.0f);
 	}
 
-	FTimingEventsTrack* AddTimingEventsTrack(uint64 TrackId, ETimingEventsTrackType TrackType, const FString& TrackName, const TCHAR* GroupName, int32 Order);
-
-	void DrawTimingProfilerTrack(FTimingViewDrawHelper& Helper, FTimingEventsTrack& Track) const;
-
-	const TCHAR* GetLoadTimeProfilerEventNameByPackageEventType(uint32 Depth, const Trace::FLoadTimeProfilerCpuEvent& Event) const;
-	const TCHAR* GetLoadTimeProfilerEventNameByExportEventType(uint32 Depth, const Trace::FLoadTimeProfilerCpuEvent& Event) const;
-	const TCHAR* GetLoadTimeProfilerEventNameByPackageName(uint32 Depth, const Trace::FLoadTimeProfilerCpuEvent& Event) const;
-	const TCHAR* GetLoadTimeProfilerEventNameByExportClassName(uint32 Depth, const Trace::FLoadTimeProfilerCpuEvent& Event) const;
-	const TCHAR* GetLoadTimeProfilerEventName(uint32 Depth, const Trace::FLoadTimeProfilerCpuEvent& Event) const;
-	void DrawLoadTimeProfilerTrack(FTimingViewDrawHelper& Helper, FTimingEventsTrack& Track) const;
-
-	void UpdateIo();
-	void DrawIoOverviewTrack(FTimingViewDrawHelper& Helper, FTimingEventsTrack& Track) const;
-	void DrawIoActivityTrack(FTimingViewDrawHelper& Helper, FTimingEventsTrack& Track) const;
+	void AddTimingEventsTrack(FTimingEventsTrack* Track);
 
 	void DrawTimeRangeSelection(FDrawContext& DrawContext) const;
 
@@ -328,13 +311,6 @@ protected:
 
 	void UpdateHoveredTimingEvent(float MX, float MY);
 
-	bool SearchTimingEvent(const double InStartTime,
-						   const double InEndTime,
-						   TFunctionRef<bool(double, double, uint32)> InPredicate,
-						   FTimingEvent& InOutTimingEvent,
-						   bool bInStopAtFirstMatch,
-						   bool bInSearchForLargestEvent) const;
-
 	void OnSelectedTimingEventChanged();
 	void SelectHoveredTimingEvent();
 	void SelectLeftTimingEvent();
@@ -344,8 +320,6 @@ protected:
 
 	void FrameSelection();
 
-	const FTimerNodePtr GetTimerNode(uint64 TypeId) const;
-
 	////////////////////////////////////////////////////////////////////////////////////////////////////
 	// FrameSelectionChanged Event
 
@@ -401,8 +375,6 @@ protected:
 
 	bool bAreTimingEventsTracksDirty;
 
-	bool bUseDownSampling;
-
 	////////////////////////////////////////////////////////////
 	// Cpu/Gpu
 
@@ -424,14 +396,13 @@ protected:
 	////////////////////////////////////////////////////////////
 	// Asset Loading
 
+	TSharedPtr<FLoadingSharedState> LoadingSharedState;
 	FTimingEventsTrack* LoadingMainThreadTrack;
 	FTimingEventsTrack* LoadingAsyncThreadTrack;
 
 	uint32 LoadingMainThreadId;
 	uint32 LoadingAsyncThreadId;
 
-	FLoadingTrackGetEventNameDelegate LoadingGetEventNameFn;
-
 	bool bAssetLoadingMode;
 
 	int32 EventAggregationTotalCount;
@@ -443,41 +414,10 @@ protected:
 	////////////////////////////////////////////////////////////
 	// File activity (I/O)
 
+	TSharedPtr<FFileActivitySharedState> FileActivitySharedState;
 	FTimingEventsTrack* IoOverviewTrack;
 	FTimingEventsTrack* IoActivityTrack;
 
-	bool bForceIoEventsUpdate;
-
-	bool bMergeIoLanes;
-	bool bShowFileActivityBackgroundEvents;
-
-	struct FIoFileActivity
-	{
-		uint64 Id;
-		const TCHAR* Path;
-		double StartTime;
-		double EndTime;
-		int32 EventCount;
-		int32 Depth;
-	};
-
-	TArray<TSharedPtr<FIoFileActivity>> FileActivities;
-	TMap<uint64, TSharedPtr<FIoFileActivity>> FileActivityMap;
-
-	struct FIoTimingEvent
-	{
-		double StartTime;
-		double EndTime;
-		uint32 Depth;
-		uint32 Type; // Trace::EFileActivityType + "Failed" flag
-		uint64 Offset;
-		uint64 Size;
-		TSharedPtr<FIoFileActivity> FileActivity;
-	};
-
-	/** All IO events, cached. */
-	TArray<FIoTimingEvent> AllIoEvents;
-
 	////////////////////////////////////////////////////////////
 
 	/** The time ruler track. */
@@ -487,7 +427,7 @@ protected:
 	FMarkersTimingTrack MarkersTrack;
 
 	/** A graph track for frame times. */
-	TSharedPtr<FGraphTrack> GraphTrack;
+	TSharedPtr<FTimingGraphTrack> GraphTrack;
 
 	////////////////////////////////////////////////////////////
 
@@ -545,8 +485,7 @@ protected:
 	double SelectionStartTime;
 	double SelectionEndTime;
 
-	mutable float TooltipWidth;
-	mutable float TooltipAlpha;
+	mutable FTimingViewTooltip Tooltip;
 
 	FTimingEvent HoveredTimingEvent;
 	FTimingEvent SelectedTimingEvent;
diff --git a/Engine/Source/Developer/TraceInsights/Public/Insights/Version.h b/Engine/Source/Developer/TraceInsights/Public/Insights/Version.h
--- a/Engine/Source/Developer/TraceInsights/Public/Insights/Version.h
+++ b/Engine/Source/Developer/TraceInsights/Public/Insights/Version.h
@@ -5,10 +5,10 @@
 // See https://semver.org/
 
 #define UNREAL_INSIGHTS_VERSION_MAJOR 0
-#define UNREAL_INSIGHTS_VERSION_MINOR 7
+#define UNREAL_INSIGHTS_VERSION_MINOR 8
 #define UNREAL_INSIGHTS_VERSION_PATCH 0
 
-#define UNREAL_INSIGHTS_VERSION_STRING "0.07"
+#define UNREAL_INSIGHTS_VERSION_STRING "0.08"
 
 /** Extra identifier string (like: "alpha", "beta.1", "beta.2", "test", "featureX", etc.). */
 //#define UNREAL_INSIGHTS_VERSION_ID_STRING "alpha"
