[Insights]
  - Graph: Added animated Auto Zoom (vertically) for Stats Counter graph series.
  - Graph: Changed vertical pan/zoom (baseline+scale) to be specific per each graph series.

diff --git a/Engine/Source/Developer/TraceInsights/Private/Insights/ViewModels/GraphTrack.cpp b/Engine/Source/Developer/TraceInsights/Private/Insights/ViewModels/GraphTrack.cpp
--- a/Engine/Source/Developer/TraceInsights/Private/Insights/ViewModels/GraphTrack.cpp
+++ b/Engine/Source/Developer/TraceInsights/Private/Insights/ViewModels/GraphTrack.cpp
@@ -24,6 +24,14 @@ FGraphSeries::FGraphSeries()
 	: Name()
 	, Description()
 	, bIsVisible(true)
+	, bIsDirty(false)
+	, bAutoZoom(false)
+	, TargetAutoZoomLowValue(0.0)
+	, TargetAutoZoomHighValue(1.0)
+	, AutoZoomLowValue(0.0)
+	, AutoZoomHighValue(1.0)
+	, BaselineY(0.0)
+	, ScaleY(1.0)
 	, Color(0.0f, 0.5f, 1.0f, 1.0f)
 	, BorderColor(0.3f, 0.8f, 1.0f, 1.0f)
 	//, Points()
@@ -56,8 +64,6 @@ FGraphTrack::FGraphTrack(uint64 InTrackId)
 	, bDrawPolygon(true)
 	, bUseEventDuration(true)
 	, bDrawBoxes(true)
-	, BaselineY(0.0f)
-	, ScaleY(1.0)
 	, NumAddedEvents(0)
 	, NumDrawPoints(0)
 	, NumDrawLines(0)
@@ -74,6 +80,13 @@ FGraphTrack::~FGraphTrack()
 
 ////////////////////////////////////////////////////////////////////////////////////////////////////
 
+void FGraphTrack::Reset()
+{
+	AllSeries.Reset();
+}
+
+////////////////////////////////////////////////////////////////////////////////////////////////////
+
 void FGraphTrack::UpdateHoveredState(float MouseX, float MouseY, const FTimingTrackViewport& Viewport)
 {
 	constexpr float HeaderWidth = 100.0f;
@@ -135,9 +148,17 @@ void FGraphTrack::Draw(FDrawContext& DrawContext, const FTimingTrackViewport& Vi
 	// Draw top line.
 	//DrawContext.DrawBox(0.0f, GetPosY(), Viewport.Width, 1.0f, WhiteBrush, FLinearColor(0.05f, 0.05f, 0.05f, 1.0f));
 
-	// Draw baseline (Value == 0).
-	DrawContext.DrawBox(0.0f, GetPosY() + BaselineY - 1.0f, Viewport.Width, 1.0f, WhiteBrush, FLinearColor(0.05f, 0.05f, 0.05f, 1.0f));
-	DrawContext.LayerId++;
+	// Draw baseline (Value == 0), for the first series only.
+	if (AllSeries.Num() > 0)
+	{
+		const TSharedPtr<FGraphSeries>& Series = AllSeries[0];
+		const float BaselineY = FMath::RoundToFloat(static_cast<float>(Series->GetBaselineY()));
+		if (BaselineY > 0.0f && BaselineY <= GetHeight())
+		{
+			DrawContext.DrawBox(0.0f, GetPosY() + BaselineY - 1.0f, Viewport.Width, 1.0f, WhiteBrush, FLinearColor(0.05f, 0.05f, 0.05f, 1.0f));
+			DrawContext.LayerId++;
+		}
+	}
 
 	for (const TSharedPtr<FGraphSeries>& Series : AllSeries)
 	{
@@ -155,10 +176,19 @@ void FGraphTrack::DrawSeries(const FGraphSeries& Series, FDrawContext& DrawConte
 	if (bDrawBoxes)
 	{
 		int32 NumBoxes = Series.Boxes.Num();
+		const float TrackPosY = GetPosY();
+		const float BaselineY = static_cast<float>(Series.GetBaselineY());
 		for (int32 Index = 0; Index < NumBoxes; ++Index)
 		{
 			const FGraphBox& Box = Series.Boxes[Index];
-			DrawContext.DrawBox(Box.X, GetPosY() + Box.Y, Box.W, BaselineY - Box.Y, WhiteBrush, Series.Color);
+			if (BaselineY >= Box.Y)
+			{
+				DrawContext.DrawBox(Box.X, TrackPosY + Box.Y, Box.W, BaselineY - Box.Y + 1.0f, WhiteBrush, Series.Color);
+			}
+			else
+			{
+				DrawContext.DrawBox(Box.X, TrackPosY + BaselineY, Box.W, Box.Y - BaselineY, WhiteBrush, Series.Color);
+			}
 		}
 		DrawContext.LayerId++;
 	}
@@ -184,46 +214,50 @@ void FGraphTrack::DrawSeries(const FGraphSeries& Series, FDrawContext& DrawConte
 		Indices.Reserve(Series.LinePoints.Num() * 6);
 		Verts.Reserve(Series.LinePoints.Num() * 2);
 
-		const float TopY = GetPosY();
-		const float BottomY = (TopY + GetHeight()) / Size.Y;
+		const float TopV = GetPosY() / Size.Y;
+		const float BaselineY = static_cast<float>(Series.GetBaselineY());
+		const float BottomV = FMath::Min(GetPosY() + BaselineY, GetPosY() + GetHeight()) / Size.Y;
 
-		for (const FVector2D& LinePoint : Series.LinePoints)
+		if (TopV < BottomV)
 		{
-			// Add verts top->bottom
-			FVector2D UV(LinePoint.X / Size.X, (TopY + LinePoint.Y) / Size.Y);
-			Verts.Add(FSlateVertex::Make<ESlateVertexRounding::Disabled>(RenderTransform, (Pos + UV * Size * Scale), AtlasOffset + UV * AtlasUVSize, FillColor));
-			UV.Y = BottomY;
-			Verts.Add(FSlateVertex::Make<ESlateVertexRounding::Disabled>(RenderTransform, (Pos + UV * Size * Scale), AtlasOffset + FVector2D(UV.X, 0.5f) * AtlasUVSize, FillColor));
-
-			if (Verts.Num() >= 4)
+			for (const FVector2D& LinePoint : Series.LinePoints)
 			{
-				int32 Index0 = Verts.Num() - 4;
-				int32 Index1 = Verts.Num() - 3;
-				int32 Index2 = Verts.Num() - 2;
-				int32 Index3 = Verts.Num() - 1;
-
-				Indices.Add(Index0);
-				Indices.Add(Index1);
-				Indices.Add(Index2);
-
-				Indices.Add(Index1);
-				Indices.Add(Index2);
-				Indices.Add(Index3);
+				// Add verts top->bottom
+				FVector2D UV(LinePoint.X / Size.X, TopV + LinePoint.Y / Size.Y);
+				Verts.Add(FSlateVertex::Make<ESlateVertexRounding::Disabled>(RenderTransform, Pos + UV * Size * Scale, AtlasOffset + UV * AtlasUVSize, FillColor));
+				UV.Y = BottomV;
+				Verts.Add(FSlateVertex::Make<ESlateVertexRounding::Disabled>(RenderTransform, Pos + UV * Size * Scale, AtlasOffset + FVector2D(UV.X, 0.5f) * AtlasUVSize, FillColor));
+
+				if (Verts.Num() >= 4)
+				{
+					int32 Index0 = Verts.Num() - 4;
+					int32 Index1 = Verts.Num() - 3;
+					int32 Index2 = Verts.Num() - 2;
+					int32 Index3 = Verts.Num() - 1;
+
+					Indices.Add(Index0);
+					Indices.Add(Index1);
+					Indices.Add(Index2);
+
+					Indices.Add(Index1);
+					Indices.Add(Index2);
+					Indices.Add(Index3);
+				}
 			}
-		}
 
-		if (Indices.Num())
-		{
-			FSlateDrawElement::MakeCustomVerts(
-				DrawContext.ElementList,
-				DrawContext.LayerId,
-				ResourceHandle,
-				Verts,
-				Indices,
-				nullptr,
-				0,
-				0, ESlateDrawEffect::PreMultipliedAlpha);
-			DrawContext.LayerId++;
+			if (Indices.Num() > 0)
+			{
+				FSlateDrawElement::MakeCustomVerts(
+					DrawContext.ElementList,
+					DrawContext.LayerId,
+					ResourceHandle,
+					Verts,
+					Indices,
+					nullptr,
+					0,
+					0, ESlateDrawEffect::PreMultipliedAlpha);
+				DrawContext.LayerId++;
+			}
 		}
 	}
 
@@ -597,7 +631,12 @@ FRandomGraphTrack::FRandomGraphTrack(uint64 InTrackId)
 	bDrawPolygon = false;
 	bUseEventDuration = false;
 	bDrawBoxes = false;
+}
 
+////////////////////////////////////////////////////////////////////////////////////////////////////
+
+void FRandomGraphTrack::AddDefaultSeries()
+{
 	TSharedPtr<FGraphSeries> Series0 = MakeShareable(new FGraphSeries());
 	Series0->SetName(TEXT("Random Blue"));
 	Series0->SetDescription(TEXT("Random series; for debuging purposes"));
@@ -616,7 +655,7 @@ FRandomGraphTrack::FRandomGraphTrack(uint64 InTrackId)
 	Series2->SetName(TEXT("Random Red"));
 	Series2->SetDescription(TEXT("Random series; for debuging purposes"));
 	Series2->SetColor(FLinearColor(1.0f, 0.1f, 0.2f, 1.0f), FLinearColor(1.0f, 0.4f, 0.5f, 1.0f));
-	Series2->SetVisibility(false);
+	Series2->SetVisibility(true);
 	AllSeries.Add(Series2);
 }
 
@@ -630,18 +669,22 @@ FRandomGraphTrack::~FRandomGraphTrack()
 
 void FRandomGraphTrack::Update(const FTimingTrackViewport& Viewport)
 {
-	NumAddedEvents = 0;
+	if (IsDirty())
+	{
+		ClearDirtyFlag();
 
-	// TODO: vertical panning and zooming
-	BaselineY = GetHeight();
-	ScaleY = 1.0;
+		NumAddedEvents = 0;
 
-	ensure(AllSeries.Num() == 3);
-	GenerateSeries(*AllSeries[0], Viewport, 1000000, 0);
-	GenerateSeries(*AllSeries[1], Viewport, 1000000, 1);
-	GenerateSeries(*AllSeries[2], Viewport, 1000000, 2);
+		int32 Seed = 0;
+		for (TSharedPtr<FGraphSeries>& Series : AllSeries)
+		{
+			Series->SetBaselineY(GetHeight() / 2.0);
+			Series->SetScaleY(GetHeight());
+			GenerateSeries(*Series, Viewport, 1000000, Seed++);
+		}
 
-	UpdateStats();
+		UpdateStats();
+	}
 }
 
 ////////////////////////////////////////////////////////////////////////////////////////////////////
@@ -653,8 +696,9 @@ void FRandomGraphTrack::GenerateSeries(FGraphSeries& Series, const FTimingTrackV
 
 	constexpr double MinDeltaTime = 0.0000001; // 100ns
 	constexpr double MaxDeltaTime = 0.01; // 100ms
-	const float MinValue = 0;
-	const float MaxValue = GetHeight();
+
+	float MinValue = (Seed == 0) ? 0.0f : (Seed == 1) ? -0.25f : -0.5f;
+	float MaxValue = (Seed == 0) ? 0.5f : (Seed == 1) ? +0.25f :  0.0f;
 
 	struct FGraphEvent
 	{
@@ -796,7 +840,7 @@ void FGraphTrackBuilder::AddPoint(double Time, double Value)
 		PointsCurrentX = AlignedX;
 	}
 
-	const float Y = Track.GetYForValue(Value);
+	const float Y = Series.GetYForValue(Value);
 
 	int32 Index = FMath::RoundToInt(Y / FGraphTrack::PointSizeY);
 	if (Index >= 0 && Index < PointsAtCurrentX.Num())
@@ -856,7 +900,7 @@ void FGraphTrackBuilder::AddConnectedLine(double Time, double Value)
 	}
 
 	// Transform Value to Y coordinate (local track space) and ensure Y is not infinite.
-	const float Y = FMath::Clamp<float>(Track.GetYForValue(Value), -FLT_MAX, FLT_MAX);
+	const float Y = FMath::Clamp<float>(Series.GetYForValue(Value), -FLT_MAX, FLT_MAX);
 
 	const float X = Viewport.TimeToSlateUnitsRounded(Viewport.RestrictEndTime(Time));
 
@@ -979,7 +1023,7 @@ void FGraphTrackBuilder::AddBox(double Time, double Duration, double Value)
 	FGraphBox Box;
 	Box.X = X1;
 	Box.W = W;
-	Box.Y = Track.GetYForValue(Value);
+	Box.Y = Series.GetYForValue(Value);
 	Series.Boxes.Add(Box);
 }
 
diff --git a/Engine/Source/Developer/TraceInsights/Private/Insights/ViewModels/GraphTrack.h b/Engine/Source/Developer/TraceInsights/Private/Insights/ViewModels/GraphTrack.h
--- a/Engine/Source/Developer/TraceInsights/Private/Insights/ViewModels/GraphTrack.h
+++ b/Engine/Source/Developer/TraceInsights/Private/Insights/ViewModels/GraphTrack.h
@@ -46,6 +46,10 @@ public:
 	bool IsVisible() const { return bIsVisible; }
 	void SetVisibility(bool bOnOff) { bIsVisible = bOnOff; }
 
+	bool IsDirty() const { return bIsDirty; }
+	void SetDirtyFlag() { bIsDirty = true; }
+	void ClearDirtyFlag() { bIsDirty = false; }
+
 	const FLinearColor& GetColor() const { return Color; }
 	const FLinearColor& GetBorderColor() const { return BorderColor; }
 
@@ -55,11 +59,83 @@ public:
 		BorderColor = InBorderColor;
 	}
 
+	/**
+	 * @return Y position (in viewport local space) of the baseline (with Value == 0); in pixels (Slate units).
+	 * Y == 0 at the top of the graph track, positive values are downward.
+	 */
+	double GetBaselineY() const { return BaselineY; }
+	void SetBaselineY(const double InBaselineY) { BaselineY = InBaselineY; }
+
+	/**
+	 * @return The scale between Value units and viewport units; in pixels (Slate units) / Value unit.
+	 */
+	double GetScaleY() const { return ScaleY; }
+	void SetScaleY(const double InScaleY) { ScaleY = InScaleY; }
+
+	/**
+	 * @param Value a value; in Value units
+	 * @return Y position (in viewport local space) for a Value; in pixels (Slate units).
+	 * Y == 0 at the top of the graph track, positive values are downward.
+	 */
+	float GetYForValue(double Value) const
+	{
+		return static_cast<float>(BaselineY - Value * ScaleY);
+	}
+
+	/**
+	 * @param Y a Y position (in viewport local space); in pixels (Slate units).
+	 * @return Value for specified Y position.
+	 */
+	double GetValueForY(float Y) const
+	{
+		return (BaselineY - static_cast<double>(Y)) / ScaleY;
+	}
+
+	bool IsAutoZoomEnabled() const { return bAutoZoom; }
+	void EnableAutoZoom() { bAutoZoom = true; }
+
+	/** target low value of auto zoom interval (corresponding to bottom of the track) */
+	double GetTargetAutoZoomLowValue() const { return TargetAutoZoomLowValue; }
+	/** target high value of auto zoom interval (corresponding to top of the track) */
+	double GetTargetAutoZoomHighValue() const { return TargetAutoZoomHighValue; }
+	void SetTargetAutoZoomRange(double LowValue, double HighValue) { TargetAutoZoomLowValue = LowValue; TargetAutoZoomHighValue = HighValue; }
+
+	/** current auto zoom low value */
+	double GetAutoZoomLowValue() const { return AutoZoomLowValue; }
+	/** current auto zoom high value */
+	double GetAutoZoomHighValue() const { return AutoZoomHighValue; }
+	void SetAutoZoomRange(double LowValue, double HighValue) { AutoZoomLowValue = LowValue; AutoZoomHighValue = HighValue; }
+	
+	/**
+	 * Compute BaselineY and ScaleY so the [Low, High] Value range will correspond to [Top, Bottom] Y position range.
+	 * GetYForValue(InHighValue) == InTopY
+	 * GetYForValue(InLowValue) == InBottomY
+	 */
+	void ComputeBaselineAndScale(const double InLowValue, const double InHighValue, const float InTopY, const float InBottomY, double& OutBaselineY, double& OutScaleY) const
+	{
+		ensure(InLowValue < InHighValue);
+		ensure(InTopY < InBottomY);
+		const double InvRange = 1.0 / (InHighValue - InLowValue);
+		OutScaleY = static_cast<double>(InBottomY - InTopY) * InvRange;
+		//OutBaselineY = (InHighValue * static_cast<double>(InBottomY) - InLowValue * static_cast<double>(InTopY)) * InvRange;
+		OutBaselineY = static_cast<double>(InTopY) + InHighValue * OutScaleY;
+	}
+
 private:
 	FText Name;
 	FText Description;
 
 	bool bIsVisible;
+	bool bIsDirty;
+
+	bool bAutoZoom;
+	double TargetAutoZoomLowValue; // target low value of auto zoom interval (corresponding to bottom of the track)
+	double TargetAutoZoomHighValue; // target high value of auto zoom interval (corresponding to top of the track)
+	double AutoZoomLowValue; // current auto zoom low value
+	double AutoZoomHighValue; // current auto zoom high value
+
+	double BaselineY; // Y position (in viewport local space) of the baseline (with Value == 0); in pixels (Slate units)
+	double ScaleY; // scale between Value units and viewport units; in pixels (Slate units) / Value unit
 
 	FLinearColor Color;
 	FLinearColor BorderColor;
@@ -88,6 +164,8 @@ public:
 	FGraphTrack(uint64 InTrackId);
 	virtual ~FGraphTrack();
 
+	virtual void Reset() override;
+
 	virtual void UpdateHoveredState(float MouseX, float MouseY, const FTimingTrackViewport& Viewport) override;
 
 	virtual void Update(const FTimingTrackViewport& Viewport) override = 0;
@@ -110,11 +188,6 @@ protected:
 	void DrawBackground(FDrawContext& DrawContext, const FTimingTrackViewport& Viewport) const;
 	void DrawSeries(const FGraphSeries& Series, FDrawContext& DrawContext, const FTimingTrackViewport& Viewport) const;
 
-	float GetYForValue(double Value) const
-	{
-		return BaselineY - static_cast<float>(Value * ScaleY); // TODO: vertical zooming and panning
-	}
-
 private:
 	void ContextMenu_ShowPoints_Execute();
 	bool ContextMenu_ShowPoints_CanExecute();
@@ -160,9 +233,6 @@ protected:
 	bool bUseEventDuration;
 	bool bDrawBoxes;
 
-	float BaselineY; // Y position (in viewport local space) of the baseline (with Value == 0); in pixels (Slate units)
-	double ScaleY; // scale between Value units and viewport units; in pixels (Slate units) / Value unit
-
 	// Stats
 	int32 NumAddedEvents;
 	int32 NumDrawPoints;
@@ -180,6 +250,8 @@ public:
 
 	virtual void Update(const FTimingTrackViewport& Viewport) override;
 
+	void AddDefaultSeries();
+
 protected:
 	void GenerateSeries(FGraphSeries& Series, const FTimingTrackViewport& Viewport, const int32 EventCount, int32 Seed);
 };
diff --git a/Engine/Source/Developer/TraceInsights/Private/Insights/ViewModels/TimingGraphTrack.cpp b/Engine/Source/Developer/TraceInsights/Private/Insights/ViewModels/TimingGraphTrack.cpp
--- a/Engine/Source/Developer/TraceInsights/Private/Insights/ViewModels/TimingGraphTrack.cpp
+++ b/Engine/Source/Developer/TraceInsights/Private/Insights/ViewModels/TimingGraphTrack.cpp
@@ -27,6 +27,14 @@ FTimingGraphTrack::FTimingGraphTrack(uint64 InTrackId)
 	bDrawPolygon = true;
 	bUseEventDuration = true;
 	bDrawBoxes = false;
+}
+
+////////////////////////////////////////////////////////////////////////////////////////////////////
+
+void FTimingGraphTrack::AddDefaultFrameSeries()
+{
+	const float BaselineY = GetHeight() - 1.0f;
+	const double ScaleY = GetHeight() / 0.1; // 100ms
 
 	TSharedPtr<FTimingGraphSeries> GameFramesSeries = MakeShareable(new FTimingGraphSeries());
 	GameFramesSeries->SetName(TEXT("Game Frames"));
@@ -34,6 +42,8 @@ FTimingGraphTrack::FTimingGraphTrack(uint64 InTrackId)
 	GameFramesSeries->SetColor(FLinearColor(0.3f, 0.3f, 1.0f, 1.0f), FLinearColor(1.0f, 1.0f, 1.0f, 1.0f));
 	GameFramesSeries->Type = FTimingGraphSeries::ESeriesType::Frame;
 	GameFramesSeries->Id = static_cast<uint32>(TraceFrameType_Game);
+	GameFramesSeries->SetBaselineY(BaselineY);
+	GameFramesSeries->SetScaleY(ScaleY);
 	AllSeries.Add(GameFramesSeries);
 
 	TSharedPtr<FTimingGraphSeries> RenderingFramesSeries = MakeShareable(new FTimingGraphSeries());
@@ -42,6 +52,8 @@ FTimingGraphTrack::FTimingGraphTrack(uint64 InTrackId)
 	RenderingFramesSeries->SetColor(FLinearColor(1.0f, 0.3f, 0.3f, 1.0f), FLinearColor(1.0f, 1.0f, 1.0f, 1.0f));
 	RenderingFramesSeries->Type = FTimingGraphSeries::ESeriesType::Frame;
 	RenderingFramesSeries->Id = static_cast<uint32>(TraceFrameType_Rendering);
+	RenderingFramesSeries->SetBaselineY(BaselineY);
+	RenderingFramesSeries->SetScaleY(ScaleY);
 	AllSeries.Add(RenderingFramesSeries);
 }
 
@@ -59,7 +71,7 @@ TSharedPtr<FTimingGraphSeries> FTimingGraphTrack::GetStatsCounterSeries(uint32 C
 
 ////////////////////////////////////////////////////////////////////////////////////////////////////
 
-void FTimingGraphTrack::AddStatsCounterSeries(uint32 CounterId, FLinearColor Color, double ValueOffset, double ValueScale)
+TSharedPtr<FTimingGraphSeries> FTimingGraphTrack::AddStatsCounterSeries(uint32 CounterId, FLinearColor Color)
 {
 	TSharedPtr<FTimingGraphSeries> Series = MakeShareable(new FTimingGraphSeries());
 
@@ -90,10 +102,16 @@ void FTimingGraphTrack::AddStatsCounterSeries(uint32 CounterId, FLinearColor Col
 	Series->Type = FTimingGraphSeries::ESeriesType::StatsCounter;
 	Series->Id = CounterId;
 	Series->bIsFloatingPoint = bIsFloatingPoint;
-	Series->ValueOffset = ValueOffset;
-	Series->ValueScale = ValueScale;
+
+	Series->SetBaselineY(GetHeight() - 1.0f);
+	Series->SetScaleY(1.0);
+
+	Series->EnableAutoZoom();
+	Series->SetTargetAutoZoomRange(0.0, 1.0);
+	Series->SetAutoZoomRange(0.0, 1.0);
 
 	AllSeries.Add(Series);
+	return Series;
 }
 
 ////////////////////////////////////////////////////////////////////////////////////////////////////
@@ -117,19 +135,38 @@ FTimingGraphTrack::~FTimingGraphTrack()
 
 void FTimingGraphTrack::Update(const FTimingTrackViewport& Viewport)
 {
-	NumAddedEvents = 0;
+	const bool bIsEntireGraphTrackDirty = IsDirty();
+	bool bNeedsUpdate = bIsEntireGraphTrackDirty;
 
-	// TODO: Vertical panning and zooming needs to be moved out in a Viewport like controller.
-	BaselineY = GetHeight();
-	ScaleY = 200.0 / 0.1; // 200px = 100ms
+	if (!bNeedsUpdate)
+	{
+		for (TSharedPtr<FGraphSeries>& Series : AllSeries)
+		{
+			if (Series->IsVisible() && Series->IsDirty())
+			{
+				// At least one series is dirty.
+				bNeedsUpdate = true;
+				break;
+			}
+		}
+	}
 
-	for (TSharedPtr<FGraphSeries>& Series : AllSeries)
+	if (bNeedsUpdate)
 	{
-		if (Series->IsVisible())
+		ClearDirtyFlag();
+
+		NumAddedEvents = 0;
+
+		for (TSharedPtr<FGraphSeries>& Series : AllSeries)
 		{
-			TSharedPtr<FTimingGraphSeries> TimingSeries = StaticCastSharedPtr<FTimingGraphSeries>(Series);
-			switch (TimingSeries->Type)
+			if (Series->IsVisible() && (bIsEntireGraphTrackDirty || Series->IsDirty()))
 			{
+				// Clear the flag before updating, becasue the update itself may furter need to set the series as dirty.
+				Series->ClearDirtyFlag();
+
+				TSharedPtr<FTimingGraphSeries> TimingSeries = StaticCastSharedPtr<FTimingGraphSeries>(Series);
+				switch (TimingSeries->Type)
+				{
 				case FTimingGraphSeries::ESeriesType::Frame:
 					UpdateFrameSeries(*TimingSeries, Viewport);
 					break;
@@ -141,11 +178,12 @@ void FTimingGraphTrack::Update(const FTimingTrackViewport& Viewport)
 				case FTimingGraphSeries::ESeriesType::StatsCounter:
 					UpdateStatsCounterSeries(*TimingSeries, Viewport);
 					break;
+				}
 			}
 		}
-	}
 
-	UpdateStats();
+		UpdateStats();
+	}
 }
 
 ////////////////////////////////////////////////////////////////////////////////////////////////////
@@ -188,60 +226,93 @@ void FTimingGraphTrack::UpdateStatsCounterSeries(FTimingGraphSeries& Series, con
 		const Trace::ICounterProvider& CountersProvider = Trace::ReadCounterProvider(*Session.Get());
 		CountersProvider.ReadCounter(Series.Id, [this, &Viewport, &Builder, &Series](const Trace::ICounter& Counter)
 		{
-			//const double ValueOffset = Series.ValueOffset;
-			//const double ValueScale = Series.ValueScale;
-			double MinValue = std::numeric_limits<double>::infinity();
-			double MaxValue = -std::numeric_limits<double>::infinity();
-
-			if (Counter.IsFloatingPoint())
+			if (Series.IsAutoZoomEnabled())
 			{
-				Counter.EnumerateFloatValues(Viewport.StartTime, Viewport.EndTime, [&Builder, &MinValue, &MaxValue](double Time, double Value)
+				double MinValue =  std::numeric_limits<double>::infinity();
+				double MaxValue = -std::numeric_limits<double>::infinity();
+
+				if (Counter.IsFloatingPoint())
 				{
-					if (Value < MinValue)
+					Counter.EnumerateFloatValues(Viewport.StartTime, Viewport.EndTime, [&Builder, &MinValue, &MaxValue](double Time, double Value)
 					{
-						MinValue = Value;
-					}
-					if (Value > MaxValue)
+						if (Value < MinValue)
+						{
+							MinValue = Value;
+						}
+						if (Value > MaxValue)
+						{
+							MaxValue = Value;
+						}
+					});
+				}
+				else
+				{
+					Counter.EnumerateValues(Viewport.StartTime, Viewport.EndTime, [&Builder, &MinValue, &MaxValue](double Time, int64 IntValue)
 					{
-						MaxValue = Value;
-					}
-				});
-			}
-			else
-			{
-				Counter.EnumerateValues(Viewport.StartTime, Viewport.EndTime, [&Builder, &MinValue, &MaxValue](double Time, int64 IntValue)
+						const double Value = static_cast<double>(IntValue);
+						if (Value < MinValue)
+						{
+							MinValue = Value;
+						}
+						if (Value > MaxValue)
+						{
+							MaxValue = Value;
+						}
+					});
+				}
+
+				const float TopY = 4.0f;
+				const float BottomY = GetHeight();
+
+				double HighValue = Series.GetValueForY(TopY);
+				double LowValue = Series.GetValueForY(BottomY);
+
+				if (MinValue < MaxValue)
 				{
-					const double Value = static_cast<double>(IntValue);
-					if (Value < MinValue)
+					constexpr bool bIsAutoZoomAnimated = true;
+					if (bIsAutoZoomAnimated)
 					{
-						MinValue = Value;
+						// Interpolate interval (animating the vertical position and scale of the graph series).
+						constexpr double InterpolationSpeed = 0.5;
+						MaxValue = InterpolationSpeed * MaxValue + (1.0 - InterpolationSpeed) * HighValue;
+						MinValue = InterpolationSpeed * MinValue + (1.0 - InterpolationSpeed) * LowValue;
 					}
-					if (Value > MaxValue)
+
+					double BaselineY;
+					double ScaleY;
+					Series.ComputeBaselineAndScale(MinValue, MaxValue, TopY, BottomY, BaselineY, ScaleY);
+
+					constexpr double ErrorTolerance = 1e-15;
+					if (!FMath::IsNearlyEqual(BaselineY, Series.GetBaselineY(), ErrorTolerance) ||
+						!FMath::IsNearlyEqual(ScaleY, Series.GetScaleY(), ErrorTolerance))
 					{
-						MaxValue = Value;
+						// Request a new update so we can further interpolate the coefficients (until reaching targeted ones).
+						Series.SetDirtyFlag();
 					}
-				});
-			}
-
-			const double ValueOffset = (MinValue != std::numeric_limits<double>::infinity()) ? -MinValue : 0.0;
 
-			const double AdjustedHeight = GetHeight() - 3.0f;
-			const double ValueScale = (MinValue < MaxValue) ? AdjustedHeight / ScaleY / (MaxValue - MinValue) : 1.0;
+					Series.SetBaselineY(BaselineY);
+					Series.SetScaleY(ScaleY);
+				}
+				else
+				{
+					// If MinValue == MaxValue, we keep the previous baseline and scale.
+				}
+			}
 
 			if (Counter.IsFloatingPoint())
 			{
-				Counter.EnumerateFloatValues(Viewport.StartTime, Viewport.EndTime, [&Builder, ValueOffset, ValueScale](double Time, double Value)
+				Counter.EnumerateFloatValues(Viewport.StartTime, Viewport.EndTime, [&Builder](double Time, double Value)
 				{
-					//TODO: add a "value converter"
-					Builder.AddEvent(Time, 0.0, (Value + ValueOffset) * ValueScale);
+					//TODO: add a "value unit converter"
+					Builder.AddEvent(Time, 0.0, Value);
 				});
 			}
 			else
 			{
-				Counter.EnumerateValues(Viewport.StartTime, Viewport.EndTime, [&Builder, ValueOffset, ValueScale](double Time, int64 IntValue)
+				Counter.EnumerateValues(Viewport.StartTime, Viewport.EndTime, [&Builder](double Time, int64 IntValue)
 				{
-					//TODO: add a "value converter"
-					Builder.AddEvent(Time, 0.0, (static_cast<double>(IntValue) + ValueOffset) * ValueScale);
+					//TODO: add a "value unit converter"
+					Builder.AddEvent(Time, 0.0, static_cast<double>(IntValue));
 				});
 			}
 		});
diff --git a/Engine/Source/Developer/TraceInsights/Private/Insights/ViewModels/TimingGraphTrack.h b/Engine/Source/Developer/TraceInsights/Private/Insights/ViewModels/TimingGraphTrack.h
--- a/Engine/Source/Developer/TraceInsights/Private/Insights/ViewModels/TimingGraphTrack.h
+++ b/Engine/Source/Developer/TraceInsights/Private/Insights/ViewModels/TimingGraphTrack.h
@@ -23,8 +23,6 @@ public:
 	ESeriesType Type;
 	uint32 Id; // frame type, timer id or stats counter id
 	bool bIsFloatingPoint; // for stats counters
-	double ValueOffset; // offset added to Y values (used to pan graph vertically)
-	double ValueScale; // scale of Y values (used to scale graph vertically)
 };
 
 ////////////////////////////////////////////////////////////////////////////////////////////////////
@@ -37,8 +35,10 @@ public:
 
 	virtual void Update(const FTimingTrackViewport& Viewport) override;
 
+	void AddDefaultFrameSeries();
+
 	TSharedPtr<FTimingGraphSeries> GetStatsCounterSeries(uint32 CounterId);
-	void AddStatsCounterSeries(uint32 CounterId, FLinearColor Color, double ValueOffset = 0.0, double ValueScale = 1.0);
+	TSharedPtr<FTimingGraphSeries> AddStatsCounterSeries(uint32 CounterId, FLinearColor Color);
 	void RemoveStatsCounterSeries(uint32 CounterId);
 
 protected:
diff --git a/Engine/Source/Developer/TraceInsights/Private/Insights/Widgets/SGraphTrack.cpp b/Engine/Source/Developer/TraceInsights/Private/Insights/Widgets/SGraphTrack.cpp
--- a/Engine/Source/Developer/TraceInsights/Private/Insights/Widgets/SGraphTrack.cpp
+++ b/Engine/Source/Developer/TraceInsights/Private/Insights/Widgets/SGraphTrack.cpp
@@ -59,6 +59,7 @@ void SGraphTrack::Reset()
 {
 	TimeRulerTrack->Reset();
 	GraphTrack->Reset();
+	StaticCastSharedPtr<FRandomGraphTrack>(GraphTrack)->AddDefaultSeries();
 
 	Viewport.Reset();
 	Viewport.MaxValidTime = 84.0 * 60.0;
@@ -105,24 +106,17 @@ void SGraphTrack::Tick(const FGeometry& AllottedGeometry, const double InCurrent
 	float TrackWidth = AllottedGeometry.GetAbsoluteSize().X;
 	float TrackHeight = AllottedGeometry.GetAbsoluteSize().Y;
 
-	bool bIsGraphDirty = false;
-
 	if (Viewport.UpdateSize(TrackWidth, TrackHeight) ||
 		bIsViewportDirty ||
 		bIsVerticalViewportDirty ||
 		GraphTrack->GetHeight() != TrackHeight)
 	{
-		bIsGraphDirty = true;
-	}
-
-	if (bIsGraphDirty)
-	{
-		bIsGraphDirty = false;
-
 		GraphTrack->SetPosY(TimeRulerTrack->GetHeight());
 		GraphTrack->SetHeight(Viewport.Height - TimeRulerTrack->GetHeight());
-		GraphTrack->Update(Viewport);
+		GraphTrack->SetDirtyFlag();
 	}
+
+	GraphTrack->Update(Viewport);
 }
 
 ////////////////////////////////////////////////////////////////////////////////////////////////////
diff --git a/Engine/Source/Developer/TraceInsights/Private/Insights/Widgets/SStatsView.cpp b/Engine/Source/Developer/TraceInsights/Private/Insights/Widgets/SStatsView.cpp
--- a/Engine/Source/Developer/TraceInsights/Private/Insights/Widgets/SStatsView.cpp
+++ b/Engine/Source/Developer/TraceInsights/Private/Insights/Widgets/SStatsView.cpp
@@ -865,7 +865,7 @@ void SStatsView::TreeView_OnMouseButtonDoubleClick(FStatsNodePtr StatsNode)
 				}
 				else
 				{
-					GraphTrack->AddStatsCounterSeries(StatsCounterId, StatsNode->GetColor(), 0.0, 1.0);
+					GraphTrack->AddStatsCounterSeries(StatsCounterId, StatsNode->GetColor());
 					GraphTrack->SetDirtyFlag();
 					StatsNode->SetAddedToGraphFlag(true);
 				}
diff --git a/Engine/Source/Developer/TraceInsights/Private/Insights/Widgets/STimingView.cpp b/Engine/Source/Developer/TraceInsights/Private/Insights/Widgets/STimingView.cpp
--- a/Engine/Source/Developer/TraceInsights/Private/Insights/Widgets/STimingView.cpp
+++ b/Engine/Source/Developer/TraceInsights/Private/Insights/Widgets/STimingView.cpp
@@ -246,6 +246,8 @@ void STimingView::Reset()
 	MarkersTrack.Reset();
 
 	GraphTrack->Reset();
+	GraphTrack->SetHeight(200.0f);
+	GraphTrack->AddDefaultFrameSeries();
 	GraphTrack->SetVisibilityFlag(false);
 
 	//////////////////////////////////////////////////
@@ -341,7 +343,6 @@ void STimingView::Tick(const FGeometry& AllottedGeometry, const double InCurrent
 	if (GraphTrack->IsVisible())
 	{
 		GraphTrack->SetPosY(TopOffset);
-		GraphTrack->SetHeight(200.0f);
 		TopOffset += GraphTrack->GetHeight();
 	}
 	Viewport.TopOffset = TopOffset;
@@ -698,9 +699,8 @@ void STimingView::Tick(const FGeometry& AllottedGeometry, const double InCurrent
 		TimelineCacheUpdateDurationHistory.AddValue(Stopwatch.AccumulatedTime);
 	}
 
-	if (GraphTrack->IsVisible() && GraphTrack->IsDirty())
+	if (GraphTrack->IsVisible())
 	{
-		GraphTrack->ClearDirtyFlag();
 		GraphTrack->Update(Viewport);
 	}
 }
