[Insights]
  - Timers view / Stats view: Fixed selection not being restored when list is updated or when aggregated stats are updated.

diff --git a/Engine/Source/Developer/TraceInsights/Private/Insights/Widgets/SStatsView.cpp b/Engine/Source/Developer/TraceInsights/Private/Insights/Widgets/SStatsView.cpp
--- a/Engine/Source/Developer/TraceInsights/Private/Insights/Widgets/SStatsView.cpp
+++ b/Engine/Source/Developer/TraceInsights/Private/Insights/Widgets/SStatsView.cpp
@@ -485,7 +485,6 @@ void SStatsView::TreeViewHeaderRow_CreateColumnArgs(const int32 ColumnIndex)
 	ColumnArgs
 		.ColumnId(Column.Id)
 		.DefaultLabel(Column.ShortName)
-		.SortMode(EColumnSortMode::None)
 		.HAlignHeader(HAlign_Fill)
 		.VAlignHeader(VAlign_Fill)
 		.HeaderContentPadding(TOptional<FMargin>(2.0f))
@@ -1450,6 +1449,7 @@ void SStatsView::ContextMenu_ResetColumns_Execute()
 
 void SStatsView::RebuildTree(bool bResync)
 {
+	TArray<FStatsNodePtr> SelectedItems;
 	bool bListHasChanged = false;
 
 	if (bResync)
@@ -1473,6 +1473,9 @@ void SStatsView::RebuildTree(bool bResync)
 
 		if (bResync)
 		{
+			// Save selection.
+			TreeView->GetSelectedItems(SelectedItems);
+
 			StatsNodes.Empty(StatsNodes.Num());
 			//StatsNodesMap.Empty(StatsNodesMap.Num());
 			StatsNodesIdMap.Empty(StatsNodesIdMap.Num());
@@ -1497,6 +1500,26 @@ void SStatsView::RebuildTree(bool bResync)
 	{
 		UpdateTree();
 		UpdateStats(StatsStartTime, StatsEndTime);
+
+		// Restore selection.
+		if (SelectedItems.Num() > 0)
+		{
+			TreeView->ClearSelection();
+			TArray<FStatsNodePtr> NewSelectedItems;
+			for (const FStatsNodePtr& StatsNode : SelectedItems)
+			{
+				FStatsNodePtr* StatsNodePtrPtr = StatsNodesIdMap.Find(StatsNode->GetId());
+				if (StatsNodePtrPtr != nullptr)
+				{
+					NewSelectedItems.Add(*StatsNodePtrPtr);
+				}
+			}
+			if (NewSelectedItems.Num() > 0)
+			{
+				TreeView->SetItemSelection(NewSelectedItems, true);
+				TreeView->RequestScrollIntoView(NewSelectedItems[0]);
+			}
+		}
 	}
 }
 
@@ -1832,7 +1855,19 @@ void SStatsView::UpdateStats(double StartTime, double EndTime)
 		StatsEndTime = EndTime;
 
 		UpdateTree();
+
+		// Save selection.
+		TArray<FStatsNodePtr> SelectedItems = TreeView->GetSelectedItems();
+
 		TreeView->RebuildList();
+
+		// Restore selection.
+		if (SelectedItems.Num() > 0)
+		{
+			TreeView->ClearSelection();
+			TreeView->SetItemSelection(SelectedItems, true);
+			TreeView->RequestScrollIntoView(SelectedItems[0]);
+		}
 	}
 }
 
diff --git a/Engine/Source/Developer/TraceInsights/Private/Insights/Widgets/STimersView.cpp b/Engine/Source/Developer/TraceInsights/Private/Insights/Widgets/STimersView.cpp
--- a/Engine/Source/Developer/TraceInsights/Private/Insights/Widgets/STimersView.cpp
+++ b/Engine/Source/Developer/TraceInsights/Private/Insights/Widgets/STimersView.cpp
@@ -491,7 +491,6 @@ void STimersView::TreeViewHeaderRow_CreateColumnArgs(const int32 ColumnIndex)
 	ColumnArgs
 		.ColumnId(Column.Id)
 		.DefaultLabel(Column.ShortName)
-		.SortMode(EColumnSortMode::None)
 		.HAlignHeader(HAlign_Fill)
 		.VAlignHeader(VAlign_Fill)
 		.HeaderContentPadding(TOptional<FMargin>(2.0f))
@@ -1430,6 +1429,7 @@ void STimersView::ContextMenu_ResetColumns_Execute()
 
 void STimersView::RebuildTree(bool bResync)
 {
+	TArray<FTimerNodePtr> SelectedItems;
 	bool bListHasChanged = false;
 
 	if (bResync)
@@ -1446,7 +1446,7 @@ void STimersView::RebuildTree(bool bResync)
 
 		const Trace::ITimingProfilerProvider& TimingProfilerProvider = *Trace::ReadTimingProfilerProvider(*Session.Get());
 
-		TimingProfilerProvider.ReadTimers([this, &bResync, &bListHasChanged](const Trace::FTimingProfilerTimer* Timers, uint64 TimersCount)
+		TimingProfilerProvider.ReadTimers([this, &bResync, &SelectedItems, &bListHasChanged](const Trace::FTimingProfilerTimer* Timers, uint64 TimersCount)
 		{
 			if (!bResync)
 			{
@@ -1455,6 +1455,9 @@ void STimersView::RebuildTree(bool bResync)
 
 			if (bResync)
 			{
+				// Save selection.
+				TreeView->GetSelectedItems(SelectedItems);
+
 				TimerNodes.Empty(TimerNodes.Num());
 				//TimerNodesMap.Empty(TimerNodesMap.Num());
 				TimerNodesIdMap.Empty(TimerNodesIdMap.Num());
@@ -1466,8 +1469,7 @@ void STimersView::RebuildTree(bool bResync)
 					FName Name(Timer.Name);// +TEXT(" [GPU]")));
 					FName Group(Timer.IsGpuTimer ? TEXT("GPU") : TEXT("CPU"));
 					ETimerNodeType Type = Timer.IsGpuTimer ? ETimerNodeType::GpuScope : ETimerNodeType::CpuScope;
-					FTimerNode* TimerPtr = new FTimerNode(Timer.Id, Name, Group, Type);
-					FTimerNodePtr TimerNodePtr = MakeShareable(TimerPtr);
+					FTimerNodePtr TimerNodePtr = MakeShareable(new FTimerNode(Timer.Id, Name, Group, Type));
 					TimerNodes.Add(TimerNodePtr);
 					//TimerNodesMap.Add(Name, TimerNodePtr);
 					TimerNodesIdMap.Add(Timer.Id, TimerNodePtr);
@@ -1480,6 +1482,26 @@ void STimersView::RebuildTree(bool bResync)
 	{
 		UpdateTree();
 		UpdateStats(StatsStartTime, StatsEndTime);
+
+		// Restore selection.
+		if (SelectedItems.Num() > 0)
+		{
+			TreeView->ClearSelection();
+			TArray<FTimerNodePtr> NewSelectedItems;
+			for (const FTimerNodePtr& TimerNode : SelectedItems)
+			{
+				FTimerNodePtr* TimerNodePtrPtr = TimerNodesIdMap.Find(TimerNode->GetId());
+				if (TimerNodePtrPtr != nullptr)
+				{
+					NewSelectedItems.Add(*TimerNodePtrPtr);
+				}
+			}
+			if (NewSelectedItems.Num() > 0)
+			{
+				TreeView->SetItemSelection(NewSelectedItems, true);
+				TreeView->RequestScrollIntoView(NewSelectedItems[0]);
+			}
+		}
 	}
 }
 
@@ -1527,7 +1549,19 @@ void STimersView::UpdateStats(double StartTime, double EndTime)
 		StatsEndTime = EndTime;
 
 		UpdateTree();
+
+		// Save selection.
+		TArray<FTimerNodePtr> SelectedItems = TreeView->GetSelectedItems();
+
 		TreeView->RebuildList();
+
+		// Restore selection.
+		if (SelectedItems.Num() > 0)
+		{
+			TreeView->ClearSelection();
+			TreeView->SetItemSelection(SelectedItems, true);
+			TreeView->RequestScrollIntoView(SelectedItems[0]);
+		}
 	}
 }
 
