[Insights]
  - Timing view / Graph track: Added possiblity to manually zoom the values vertically for series that does not have "vertical auto zoom" enabled (i.e. for frame series).
  - Graph track: Fixed "vertical auto zoom" for cases where new range of value [min == max] and is outside the previous range of values. Now correctly zooming out to include new values.

diff --git a/Engine/Source/Developer/TraceInsights/Private/Insights/ViewModels/GraphTrack.h b/Engine/Source/Developer/TraceInsights/Private/Insights/ViewModels/GraphTrack.h
--- a/Engine/Source/Developer/TraceInsights/Private/Insights/ViewModels/GraphTrack.h
+++ b/Engine/Source/Developer/TraceInsights/Private/Insights/ViewModels/GraphTrack.h
@@ -174,6 +174,8 @@ public:
 
 	virtual void BuildContextMenu(FMenuBuilder& MenuBuilder) override;
 
+	TArray<TSharedPtr<FGraphSeries>>& GetSeries() { return AllSeries; }
+
 	//virtual int GetDebugLineCount() override;
 	//virtual void BuildDebugLines(FString& OutStr) override;
 
diff --git a/Engine/Source/Developer/TraceInsights/Private/Insights/ViewModels/TimingGraphTrack.cpp b/Engine/Source/Developer/TraceInsights/Private/Insights/ViewModels/TimingGraphTrack.cpp
--- a/Engine/Source/Developer/TraceInsights/Private/Insights/ViewModels/TimingGraphTrack.cpp
+++ b/Engine/Source/Developer/TraceInsights/Private/Insights/ViewModels/TimingGraphTrack.cpp
@@ -264,8 +264,15 @@ void FTimingGraphTrack::UpdateStatsCounterSeries(FTimingGraphSeries& Series, con
 				const float TopY = 4.0f;
 				const float BottomY = GetHeight();
 
-				double HighValue = Series.GetValueForY(TopY);
-				double LowValue = Series.GetValueForY(BottomY);
+				const double HighValue = Series.GetValueForY(TopY);
+				const double LowValue = Series.GetValueForY(BottomY);
+
+				// If MinValue == MaxValue, we keep the previous baseline and scale, but only if the min/max value is already visible.
+				if (MinValue == MaxValue && (MinValue < LowValue || MaxValue > HighValue))
+				{
+					MinValue = FMath::Min(MinValue, LowValue);
+					MaxValue = FMath::Max(MaxValue, HighValue);
+				}
 
 				if (MinValue < MaxValue)
 				{
diff --git a/Engine/Source/Developer/TraceInsights/Private/Insights/Widgets/STimingView.cpp b/Engine/Source/Developer/TraceInsights/Private/Insights/Widgets/STimingView.cpp
--- a/Engine/Source/Developer/TraceInsights/Private/Insights/Widgets/STimingView.cpp
+++ b/Engine/Source/Developer/TraceInsights/Private/Insights/Widgets/STimingView.cpp
@@ -1533,10 +1533,40 @@ FReply STimingView::OnMouseWheel(const FGeometry& MyGeometry, const FPointerEven
 {
 	if (MouseEvent.GetModifierKeys().IsShiftDown())
 	{
-		// Scroll vertically.
-		constexpr float ScrollSpeedY = 16.0f * 3;
-		const float ScrollPosY = Viewport.ScrollPosY - ScrollSpeedY * MouseEvent.GetWheelDelta();
-		ScrollAtPosY(ScrollPosY);
+		if (GraphTrack->IsVisible() &&
+			MousePosition.Y >= GraphTrack->GetPosY() &&
+			MousePosition.Y < GraphTrack->GetPosY() + GraphTrack->GetHeight())
+		{
+			for (const TSharedPtr<FGraphSeries>& Series : GraphTrack->GetSeries())
+			{
+				if (Series->IsVisible() && !Series->IsAutoZoomEnabled())
+				{
+					// Zoom in/out vertically.
+					const double Delta = MouseEvent.GetWheelDelta();
+					constexpr double ZoomStep = 0.25; // as percent
+					double ScaleY;
+
+					if (Delta > 0)
+					{
+						ScaleY = Series->GetScaleY() * FMath::Pow(1.0 + ZoomStep, Delta);
+					}
+					else
+					{
+						ScaleY = Series->GetScaleY() * FMath::Pow(1.0 / (1.0 + ZoomStep), -Delta);
+					}
+
+					Series->SetScaleY(ScaleY);
+					Series->SetDirtyFlag();
+				}
+			}
+		}
+		else
+		{
+			// Scroll vertically.
+			constexpr float ScrollSpeedY = 16.0f * 3;
+			const float ScrollPosY = Viewport.ScrollPosY - ScrollSpeedY * MouseEvent.GetWheelDelta();
+			ScrollAtPosY(ScrollPosY);
+		}
 	}
 	else if (MouseEvent.GetModifierKeys().IsControlDown())
 	{
