[Insights]
  - Graph track: Added event detection and highlight (point / horizontal line / box) for event under mouse.
  - Graph track: Fixed the drawing of the filled polygon for cases where values are crossing the baseline.

diff --git a/Engine/Source/Developer/TraceInsights/Private/Insights/ViewModels/GraphTrack.cpp b/Engine/Source/Developer/TraceInsights/Private/Insights/ViewModels/GraphTrack.cpp
--- a/Engine/Source/Developer/TraceInsights/Private/Insights/ViewModels/GraphTrack.cpp
+++ b/Engine/Source/Developer/TraceInsights/Private/Insights/ViewModels/GraphTrack.cpp
@@ -34,6 +34,7 @@ FGraphSeries::FGraphSeries()
 	, ScaleY(1.0)
 	, Color(0.0f, 0.5f, 1.0f, 1.0f)
 	, BorderColor(0.3f, 0.8f, 1.0f, 1.0f)
+	//, Events()
 	//, Points()
 	//, LinePoints()
 	//, Boxes()
@@ -46,6 +47,65 @@ FGraphSeries::~FGraphSeries()
 {
 }
 
+////////////////////////////////////////////////////////////////////////////////////////////////////
+
+const FGraphSeries::FEvent* FGraphSeries::GetEvent(const float X, const float Y, const FTimingTrackViewport& Viewport, bool bCheckLine, bool bCheckBox) const
+{
+	const float LocalBaselineY = static_cast<float>(GetBaselineY());
+
+	for (const FGraphSeries::FEvent& Event : Events)
+	{
+		const float EventX1 = Viewport.TimeToSlateUnitsRounded(Event.Time);
+		const float EventX2 = Viewport.TimeToSlateUnitsRounded(Event.Time + Event.Duration);
+
+		const float EventY = GetRoundedYForValue(Event.Value);
+
+		// Check bounding box of the visual point.
+		constexpr float PointTolerance = 5.0f;
+		if (X >= EventX1 - PointTolerance && X <= EventX1 + PointTolerance &&
+			Y >= EventY - PointTolerance && Y <= EventY + PointTolerance)
+		{
+			return &Event;
+		}
+
+		if (bCheckLine)
+		{
+			// Check bounding box of the horizontal line.
+			constexpr float LineTolerance = 2.0f;
+			if (X >= EventX1 - LineTolerance && X <= EventX2 + LineTolerance &&
+				Y >= EventY - LineTolerance && Y <= EventY + LineTolerance)
+			{
+				return &Event;
+			}
+		}
+
+		if (bCheckBox)
+		{
+			// Check bounding box of the visual box.
+			constexpr float BoxTolerance = 1.0f;
+			if (X >= EventX1 - BoxTolerance && X <= EventX2 + BoxTolerance)
+			{
+				if (EventY < LocalBaselineY)
+				{
+					if (Y >= EventY - BoxTolerance && Y <= LocalBaselineY + BoxTolerance)
+					{
+						return &Event;
+					}
+				}
+				else
+				{
+					if (Y >= LocalBaselineY - BoxTolerance && Y <= EventY + BoxTolerance)
+					{
+						return &Event;
+					}
+				}
+			}
+		}
+	}
+
+	return nullptr;
+}
+
 ////////////////////////////////////////////////////////////////////////////////////////////////////
 // FGraphTrack
 ////////////////////////////////////////////////////////////////////////////////////////////////////
@@ -139,11 +199,19 @@ void FGraphTrack::DrawBackground(FDrawContext& DrawContext, const FTimingTrackVi
 
 ////////////////////////////////////////////////////////////////////////////////////////////////////
 
-void FGraphTrack::Draw(FDrawContext& DrawContext, const FTimingTrackViewport& Viewport) const
+void FGraphTrack::Draw(FDrawContext& DrawContext, const FTimingTrackViewport& Viewport, const FVector2D& MousePosition) const
 {
 	DrawBackground(DrawContext, Viewport);
 
-	//TODO: Set clipping to (0, GetPosY(), Viewport.Width, GetHeight())!
+	{
+		FVector2D AbsPos = DrawContext.Geometry.GetAbsolutePosition();
+		const float L = AbsPos.X;
+		const float R = AbsPos.X + Viewport.Width;
+		const float T = AbsPos.Y + GetPosY();
+		const float B = AbsPos.Y + GetPosY() + GetHeight();
+		FSlateClippingZone ClipZone(FVector2D(L, T), FVector2D(R, T), FVector2D(L, B), FVector2D(R, B));
+		DrawContext.ElementList.PushClip(ClipZone);
+	}
 
 	// Draw top line.
 	//DrawContext.DrawBox(0.0f, GetPosY(), Viewport.Width, 1.0f, WhiteBrush, FLinearColor(0.05f, 0.05f, 0.05f, 1.0f));
@@ -153,9 +221,9 @@ void FGraphTrack::Draw(FDrawContext& DrawContext, const FTimingTrackViewport& Vi
 	{
 		const TSharedPtr<FGraphSeries>& Series = AllSeries[0];
 		const float BaselineY = FMath::RoundToFloat(static_cast<float>(Series->GetBaselineY()));
-		if (BaselineY > 0.0f && BaselineY <= GetHeight())
+		if (BaselineY >= 0.0f && BaselineY < GetHeight())
 		{
-			DrawContext.DrawBox(0.0f, GetPosY() + BaselineY - 1.0f, Viewport.Width, 1.0f, WhiteBrush, FLinearColor(0.05f, 0.05f, 0.05f, 1.0f));
+			DrawContext.DrawBox(0.0f, GetPosY() + BaselineY, Viewport.Width, 1.0f, WhiteBrush, FLinearColor(0.05f, 0.05f, 0.05f, 1.0f));
 			DrawContext.LayerId++;
 		}
 	}
@@ -167,6 +235,18 @@ void FGraphTrack::Draw(FDrawContext& DrawContext, const FTimingTrackViewport& Vi
 			DrawSeries(*Series, DrawContext, Viewport);
 		}
 	}
+
+	if (MousePosition.Y >= GetPosY() && MousePosition.Y < GetPosY() + GetHeight())
+	{
+		FGraphTrack::FEvent GraphEvent;
+		GraphEvent.Series = nullptr;
+		if (GetEvent(MousePosition.X, MousePosition.Y, Viewport, GraphEvent))
+		{
+			DrawHighlightedEvent(DrawContext, Viewport, GraphEvent);
+		}
+	}
+
+	DrawContext.ElementList.PopClip();
 }
 
 ////////////////////////////////////////////////////////////////////////////////////////////////////
@@ -177,17 +257,17 @@ void FGraphTrack::DrawSeries(const FGraphSeries& Series, FDrawContext& DrawConte
 	{
 		int32 NumBoxes = Series.Boxes.Num();
 		const float TrackPosY = GetPosY();
-		const float BaselineY = static_cast<float>(Series.GetBaselineY());
+		const float BaselineY = FMath::RoundToFloat(static_cast<float>(Series.GetBaselineY()));
 		for (int32 Index = 0; Index < NumBoxes; ++Index)
 		{
-			const FGraphBox& Box = Series.Boxes[Index];
+			const FGraphSeries::FBox& Box = Series.Boxes[Index];
 			if (BaselineY >= Box.Y)
 			{
 				DrawContext.DrawBox(Box.X, TrackPosY + Box.Y, Box.W, BaselineY - Box.Y + 1.0f, WhiteBrush, Series.Color);
 			}
 			else
 			{
-				DrawContext.DrawBox(Box.X, TrackPosY + BaselineY, Box.W, Box.Y - BaselineY, WhiteBrush, Series.Color);
+				DrawContext.DrawBox(Box.X, TrackPosY + BaselineY, Box.W, Box.Y - BaselineY + 1.0f, WhiteBrush, Series.Color);
 			}
 		}
 		DrawContext.LayerId++;
@@ -215,17 +295,59 @@ void FGraphTrack::DrawSeries(const FGraphSeries& Series, FDrawContext& DrawConte
 		Verts.Reserve(Series.LinePoints.Num() * 2);
 
 		const float TopV = GetPosY() / Size.Y;
+		const float BottomV = (GetPosY() + GetHeight()) / Size.Y;
 		const float BaselineY = static_cast<float>(Series.GetBaselineY());
-		const float BottomV = FMath::Min(GetPosY() + BaselineY, GetPosY() + GetHeight()) / Size.Y;
+		const float BaselineV = FMath::Clamp<float>((GetPosY() + BaselineY) / Size.Y, TopV, BottomV);
 
-		if (TopV < BottomV)
+		int32 PrevSide = 0;
+
+		for (int32 PointIndex = 0; PointIndex < Series.LinePoints.Num(); ++PointIndex)
 		{
-			for (const FVector2D& LinePoint : Series.LinePoints)
+			const FVector2D& LinePoint = Series.LinePoints[PointIndex];
+
+			// When crossing baseline the polygon needs to be intersected.
+			int32 CrtSide = (LinePoint.Y > BaselineY) ? 1 : (LinePoint.Y < BaselineY) ? -1 : 0;
+			if (PrevSide != 0 && CrtSide + PrevSide == 0) // alternating sides?
+			{
+				// Compute intersection point.
+				const FVector2D& PrevLinePoint = Series.LinePoints[PointIndex - 1];
+				float X = PrevLinePoint.X + (LinePoint.X - PrevLinePoint.X) / ((BaselineY - LinePoint.Y) / (PrevLinePoint.Y - BaselineY) + 1.0f);
+
+				// Add an interesection point vertex.
+				FVector2D UV(X / Size.X, BaselineV);
+				Verts.Add(FSlateVertex::Make<ESlateVertexRounding::Disabled>(RenderTransform, Pos + UV * Size * Scale, AtlasOffset + UV * AtlasUVSize, FillColor));
+
+				// Add a value point vertex.
+				UV.X = LinePoint.X / Size.X;
+				UV.Y = TopV + LinePoint.Y / Size.Y;
+				Verts.Add(FSlateVertex::Make<ESlateVertexRounding::Disabled>(RenderTransform, Pos + UV * Size * Scale, AtlasOffset + UV * AtlasUVSize, FillColor));
+
+				// Add a baseline vertex.
+				UV.Y = BaselineV;
+				Verts.Add(FSlateVertex::Make<ESlateVertexRounding::Disabled>(RenderTransform, Pos + UV * Size * Scale, AtlasOffset + FVector2D(UV.X, 0.5f) * AtlasUVSize, FillColor));
+
+				int32 Index0 = Verts.Num() - 5;
+				int32 Index1 = Verts.Num() - 4;
+				int32 Index2 = Verts.Num() - 3;
+				int32 Index3 = Verts.Num() - 2;
+				int32 Index4 = Verts.Num() - 1;
+
+				Indices.Add(Index0);
+				Indices.Add(Index1);
+				Indices.Add(Index2);
+
+				Indices.Add(Index2);
+				Indices.Add(Index3);
+				Indices.Add(Index4);
+			}
+			else
 			{
-				// Add verts top->bottom
+				// Add a value point vertex.
 				FVector2D UV(LinePoint.X / Size.X, TopV + LinePoint.Y / Size.Y);
 				Verts.Add(FSlateVertex::Make<ESlateVertexRounding::Disabled>(RenderTransform, Pos + UV * Size * Scale, AtlasOffset + UV * AtlasUVSize, FillColor));
-				UV.Y = BottomV;
+
+				// Add a baseline vertex.
+				UV.Y = BaselineV;
 				Verts.Add(FSlateVertex::Make<ESlateVertexRounding::Disabled>(RenderTransform, Pos + UV * Size * Scale, AtlasOffset + FVector2D(UV.X, 0.5f) * AtlasUVSize, FillColor));
 
 				if (Verts.Num() >= 4)
@@ -239,32 +361,35 @@ void FGraphTrack::DrawSeries(const FGraphSeries& Series, FDrawContext& DrawConte
 					Indices.Add(Index1);
 					Indices.Add(Index2);
 
-					Indices.Add(Index1);
 					Indices.Add(Index2);
+					Indices.Add(Index1);
 					Indices.Add(Index3);
 				}
 			}
+			PrevSide = CrtSide;
+		}
 
-			if (Indices.Num() > 0)
-			{
-				FSlateDrawElement::MakeCustomVerts(
-					DrawContext.ElementList,
-					DrawContext.LayerId,
-					ResourceHandle,
-					Verts,
-					Indices,
-					nullptr,
-					0,
-					0, ESlateDrawEffect::PreMultipliedAlpha);
-				DrawContext.LayerId++;
-			}
+		if (Indices.Num() > 0)
+		{
+			FSlateDrawElement::MakeCustomVerts(
+				DrawContext.ElementList,
+				DrawContext.LayerId,
+				ResourceHandle,
+				Verts,
+				Indices,
+				nullptr,
+				0,
+				0, ESlateDrawEffect::PreMultipliedAlpha);
+			DrawContext.LayerId++;
 		}
 	}
 
+	const float LocalPosY = FMath::RoundToFloat(GetPosY());
+
 	if (bDrawLines)
 	{
 		FPaintGeometry Geo = DrawContext.Geometry.ToPaintGeometry();
-		Geo.AppendTransform(FSlateLayoutTransform(FVector2D(0.5f, 0.5f + GetPosY())));
+		Geo.AppendTransform(FSlateLayoutTransform(FVector2D(1.0f, LocalPosY + 1.0f)));
 		FSlateDrawElement::MakeLines(DrawContext.ElementList, DrawContext.LayerId, Geo, Series.LinePoints, DrawContext.DrawEffects, Series.Color, false, 1.0f);
 		DrawContext.LayerId++;
 	}
@@ -282,8 +407,8 @@ void FGraphTrack::DrawSeries(const FGraphSeries& Series, FDrawContext& DrawConte
 			for (int32 Index = 0; Index < NumPoints; ++Index)
 			{
 				const FVector2D& Pt = Series.Points[Index];
-				const float PtX = Pt.X - PointVisualSize / 2.0f - 1.0f;
-				const float PtY = GetPosY() + Pt.Y - PointVisualSize / 2.0f - 1.0f;
+				const float PtX = Pt.X - PointVisualSize / 2.0f - 0.5f;
+				const float PtY = LocalPosY + Pt.Y - PointVisualSize / 2.0f - 0.5f;
 				DrawContext.DrawBox(PtX, PtY, PointVisualSize + 2.0f, PointVisualSize + 2.0f, PointBrush, Series.BorderColor);
 			}
 			DrawContext.LayerId++;
@@ -293,8 +418,8 @@ void FGraphTrack::DrawSeries(const FGraphSeries& Series, FDrawContext& DrawConte
 		for (int32 Index = 0; Index < NumPoints; ++Index)
 		{
 			const FVector2D& Pt = Series.Points[Index];
-			const float PtX = Pt.X - PointVisualSize / 2.0f;
-			const float PtY = GetPosY() + Pt.Y - PointVisualSize / 2.0f;
+			const float PtX = Pt.X - PointVisualSize / 2.0f + 0.5f;
+			const float PtY = LocalPosY + Pt.Y - PointVisualSize / 2.0f + 0.5f;
 			DrawContext.DrawBox(PtX, PtY, PointVisualSize, PointVisualSize, PointBrush, Series.Color);
 		}
 		DrawContext.LayerId++;
@@ -311,8 +436,8 @@ void FGraphTrack::DrawSeries(const FGraphSeries& Series, FDrawContext& DrawConte
 			for (int32 Index = 0; Index < NumPoints; ++Index)
 			{
 				const FVector2D& Pt = Series.Points[Index];
-				const float PtX = Pt.X - BorderPtSize / 2.0f;
-				const float PtY = GetPosY() + Pt.Y - BorderPtSize / 2.0f;
+				const float PtX = Pt.X - BorderPtSize / 2.0f + 0.5f;
+				const float PtY = LocalPosY + Pt.Y - BorderPtSize / 2.0f + 0.5f;
 				DrawContext.DrawBox(PtX, PtY, BorderPtSize, BorderPtSize, BorderBrush, Series.BorderColor);
 				//DrawContext.DrawRotatedBox(PtX, PtY, BorderPtSize, BorderPtSize, BorderBrush, Series.BorderColor, Angle, BorderRotationPoint);
 			}
@@ -325,15 +450,118 @@ void FGraphTrack::DrawSeries(const FGraphSeries& Series, FDrawContext& DrawConte
 		for (int32 Index = 0; Index < NumPoints; ++Index)
 		{
 			const FVector2D& Pt = Series.Points[Index];
-			const float PtX = Pt.X - PtSize / 2.0f;
-			const float PtY = GetPosY() + Pt.Y - PtSize / 2.0f;
+			const float PtX = Pt.X - PtSize / 2.0f + 0.5f;
+			const float PtY = LocalPosY + Pt.Y - PtSize / 2.0f + 0.5f;
 			DrawContext.DrawBox(PtX, PtY, PtSize, PtSize, WhiteBrush, Series.Color);
 			//DrawContext.DrawRotatedBox(PtX, PtY, PtSize, PtSize, WhiteBrush, Series.Color, Angle, RotationPoint);
 		}
 		DrawContext.LayerId++;
 
 #endif
+
+		if (false) // for debugging only
+		{
+			// Draw points (cross).
+			for (int32 Index = 0; Index < NumPoints; ++Index)
+			{
+				const FVector2D& Pt = Series.Points[Index];
+				const float PtX = Pt.X;
+				const float PtY = LocalPosY + Pt.Y;
+				DrawContext.DrawBox(PtX - 5.0f, PtY, 11.0f, 1.0f, WhiteBrush, Series.Color);
+				DrawContext.DrawBox(PtX, PtY - 5.0f, 1.0f, 11.0f, WhiteBrush, Series.Color);
+			}
+			DrawContext.LayerId++;
+		}
+	}
+}
+
+////////////////////////////////////////////////////////////////////////////////////////////////////
+
+void FGraphTrack::DrawHighlightedEvent(FDrawContext& DrawContext, const FTimingTrackViewport& Viewport, const FGraphTrack::FEvent& GraphEvent) const
+{
+	ensure(GraphEvent.IsValid());
+
+	const float EventX1 = Viewport.TimeToSlateUnitsRounded(GraphEvent.SeriesEvent.Time);
+	const float EventX2 = Viewport.TimeToSlateUnitsRounded(GraphEvent.SeriesEvent.Time + GraphEvent.SeriesEvent.Duration);
+
+	const float EventY1 = GetPosY() + GraphEvent.Series->GetRoundedYForValue(GraphEvent.SeriesEvent.Value);
+
+	const FLinearColor HighlightColor(1.0f, 1.0f, 0.0f, 1.0f);
+
+	// Draw highlighted box.
+	if (bDrawBoxes)
+	{
+		float W = EventX2 - EventX1;
+		ensure(W >= 0); // we expect events to be sorted
+
+		// Timing events are displayed with minimum 1px (including empty ones).
+		if (W == 0)
+		{
+			W = 1.0f;
+		}
+
+		const float EventY2 = GetPosY() + static_cast<float>(GraphEvent.Series->GetBaselineY());
+
+		const float Y1 = FMath::Min(EventY1, EventY2);
+		const float DY = FMath::Abs(EventY2 - EventY1);
+
+		DrawContext.DrawBox(EventX1 - 1.0f, Y1 - 1.0f, W + 2.0f, DY + 3.0f, WhiteBrush, HighlightColor);
+		DrawContext.LayerId++;
+		DrawContext.DrawBox(EventX1, Y1, W, DY + 1.0f, WhiteBrush, GraphEvent.Series->Color);
+		DrawContext.LayerId++;
+	}
+
+	// Draw highlighted line.
+	if ((EventX2 > EventX1) && ((bUseEventDuration && (bDrawLines || bDrawPolygon)) || bDrawBoxes))
+	{
+		float W = EventX2 - EventX1;
+		ensure(W >= 0); // we expect events to be sorted
+
+		// Timing events are displayed with minimum 1px (including empty ones).
+		if (W == 0)
+		{
+			W = 1.0f;
+		}
+
+		DrawContext.DrawBox(EventX1 - 1.0f, EventY1 - 1.0f, W + 2.0f, 3.0f, WhiteBrush, HighlightColor);
+		DrawContext.LayerId++;
+		DrawContext.DrawBox(EventX1, EventY1, W, 1.0f, WhiteBrush, GraphEvent.Series->Color);
+		DrawContext.LayerId++;
+	}
+
+	// Draw highlighted point.
+	DrawContext.DrawBox(EventX1 - PointVisualSize / 2.0f - 1.5f, EventY1 - PointVisualSize / 2.0f - 1.5f, PointVisualSize + 4.0f, PointVisualSize + 4.0f, PointBrush, HighlightColor);
+	DrawContext.LayerId++;
+	DrawContext.DrawBox(EventX1 - PointVisualSize / 2.0f + 0.5f, EventY1 - PointVisualSize / 2.0f + 0.5f, PointVisualSize, PointVisualSize, PointBrush, GraphEvent.Series->Color);
+	DrawContext.LayerId++;
+}
+
+////////////////////////////////////////////////////////////////////////////////////////////////////
+
+const bool FGraphTrack::GetEvent(const float MouseX, const float MouseY, const FTimingTrackViewport& Viewport, FGraphTrack::FEvent& OutEvent) const
+{
+	const float LocalPosX = MouseX;
+	const float LocalPosY = MouseY - GetPosY();
+
+	const bool bCheckLine = bUseEventDuration && (bDrawLines || bDrawPolygon);
+
+	// Search series in reverse order.
+	for (int32 SeriesIndex = AllSeries.Num() - 1; SeriesIndex >= 0; --SeriesIndex)
+	{
+		const TSharedPtr<FGraphSeries>& Series = AllSeries[SeriesIndex];
+		if (Series->IsVisible())
+		{
+			const FGraphSeries::FEvent* Event = Series->GetEvent(LocalPosX, LocalPosY, Viewport, bCheckLine, bDrawBoxes);
+			if (Event != nullptr)
+			{
+				OutEvent.Series = Series;
+				OutEvent.SeriesEvent = *Event;
+				return true;
+			}
+		}
 	}
+
+	return false;
 }
 
 ////////////////////////////////////////////////////////////////////////////////////////////////////
@@ -762,6 +990,7 @@ FGraphTrackBuilder::FGraphTrackBuilder(FGraphTrack& InTrack, FGraphSeries& InSer
 	, Series(InSeries)
 	, Viewport(InViewport)
 {
+	Series.Events.Reset();
 	BeginPoints();
 	BeginConnectedLines();
 	BeginBoxes();
@@ -782,9 +1011,11 @@ void FGraphTrackBuilder::AddEvent(double Time, double Duration, double Value)
 {
 	Track.NumAddedEvents++;
 
-	if (Track.bDrawPoints)
+	bool bIsEventVisible = false;
+
+	//if (Track.bDrawPoints)
 	{
-		AddPoint(Time, Value);
+		bIsEventVisible = AddPoint(Time, Value);
 	}
 
 	if (Track.bDrawLines || Track.bDrawPolygon)
@@ -801,6 +1032,11 @@ void FGraphTrackBuilder::AddEvent(double Time, double Duration, double Value)
 	{
 		AddBox(Time, Duration, Value);
 	}
+
+	if (bIsEventVisible)
+	{
+		Series.Events.Add({ Time, Duration, Value });
+	}
 }
 
 ////////////////////////////////////////////////////////////////////////////////////////////////////
@@ -823,12 +1059,12 @@ void FGraphTrackBuilder::BeginPoints()
 
 ////////////////////////////////////////////////////////////////////////////////////////////////////
 
-void FGraphTrackBuilder::AddPoint(double Time, double Value)
+bool FGraphTrackBuilder::AddPoint(double Time, double Value)
 {
 	const float X = Viewport.TimeToSlateUnitsRounded(Time);
 	if (X < -FGraphTrack::PointVisualSize / 2.0f || X >= Viewport.Width + FGraphTrack::PointVisualSize / 2.0f)
 	{
-		return;
+		return false;
 	}
 
 	// Align the X with a grid of GraphTrackPointDX pixels in size, in the global space (i.e. scroll independent).
@@ -840,16 +1076,21 @@ void FGraphTrackBuilder::AddPoint(double Time, double Value)
 		PointsCurrentX = AlignedX;
 	}
 
-	const float Y = Series.GetYForValue(Value);
+	const float Y = Series.GetRoundedYForValue(Value);
+
+	bool bIsNewVisiblePoint = false;
 
 	int32 Index = FMath::RoundToInt(Y / FGraphTrack::PointSizeY);
 	if (Index >= 0 && Index < PointsAtCurrentX.Num())
 	{
 		FPointInfo& Pt = PointsAtCurrentX[Index];
+		bIsNewVisiblePoint = !Pt.bValid;
 		Pt.bValid = true;
 		Pt.X = X;
 		Pt.Y = Y;
 	}
+
+	return bIsNewVisiblePoint;
 }
 
 ////////////////////////////////////////////////////////////////////////////////////////////////////
@@ -899,10 +1140,8 @@ void FGraphTrackBuilder::AddConnectedLine(double Time, double Value)
 		return;
 	}
 
-	// Transform Value to Y coordinate (local track space) and ensure Y is not infinite.
-	const float Y = FMath::Clamp<float>(Series.GetYForValue(Value), -FLT_MAX, FLT_MAX);
-
 	const float X = Viewport.TimeToSlateUnitsRounded(Viewport.RestrictEndTime(Time));
+	const float Y = Series.GetRoundedYForValue(Value);
 
 	ensure(X >= LinesCurrentX); // we are assuming events are already sorted by Time
 
@@ -1020,10 +1259,10 @@ void FGraphTrackBuilder::AddBox(double Time, double Duration, double Value)
 	}
 
 	// TODO: reduction algorithm
-	FGraphBox Box;
+	FGraphSeries::FBox Box;
 	Box.X = X1;
 	Box.W = W;
-	Box.Y = Series.GetYForValue(Value);
+	Box.Y = Series.GetRoundedYForValue(Value);
 	Series.Boxes.Add(Box);
 }
 
diff --git a/Engine/Source/Developer/TraceInsights/Private/Insights/ViewModels/GraphTrack.h b/Engine/Source/Developer/TraceInsights/Private/Insights/ViewModels/GraphTrack.h
--- a/Engine/Source/Developer/TraceInsights/Private/Insights/ViewModels/GraphTrack.h
+++ b/Engine/Source/Developer/TraceInsights/Private/Insights/ViewModels/GraphTrack.h
@@ -15,20 +15,26 @@ class FTimingTrackViewport;
 
 ////////////////////////////////////////////////////////////////////////////////////////////////////
 
-struct FGraphBox
-{
-	float X;
-	float W;
-	float Y;
-};
-
-////////////////////////////////////////////////////////////////////////////////////////////////////
-
 class FGraphSeries
 {
 	friend class FGraphTrack;
 	friend class FGraphTrackBuilder;
 
+public:
+	struct FEvent
+	{
+		double Time;
+		double Duration;
+		double Value;
+	};
+
+	struct FBox
+	{
+		float X;
+		float W;
+		float Y;
+	};
+
 public:
 	FGraphSeries();
 	~FGraphSeries();
@@ -82,6 +88,11 @@ public:
 		return static_cast<float>(BaselineY - Value * ScaleY);
 	}
 
+	float GetRoundedYForValue(double Value) const
+	{
+		return FMath::RoundToFloat(FMath::Clamp<float>(GetYForValue(Value), -FLT_MAX, FLT_MAX));
+	}
+
 	/**
 	 * @param Y a Y position (in viewport local space); in pixels (Slate units).
 	 * @return Value for specified Y position.
@@ -94,18 +105,18 @@ public:
 	bool IsAutoZoomEnabled() const { return bAutoZoom; }
 	void EnableAutoZoom() { bAutoZoom = true; }
 
-	/** target low value of auto zoom interval (corresponding to bottom of the track) */
+	/** Target low value of auto zoom interval (corresponding to bottom of the track) */
 	double GetTargetAutoZoomLowValue() const { return TargetAutoZoomLowValue; }
-	/** target high value of auto zoom interval (corresponding to top of the track) */
+	/** Target high value of auto zoom interval (corresponding to top of the track) */
 	double GetTargetAutoZoomHighValue() const { return TargetAutoZoomHighValue; }
 	void SetTargetAutoZoomRange(double LowValue, double HighValue) { TargetAutoZoomLowValue = LowValue; TargetAutoZoomHighValue = HighValue; }
 
-	/** current auto zoom low value */
+	/** Current auto zoom low value */
 	double GetAutoZoomLowValue() const { return AutoZoomLowValue; }
-	/** current auto zoom high value */
+	/** Current auto zoom high value */
 	double GetAutoZoomHighValue() const { return AutoZoomHighValue; }
 	void SetAutoZoomRange(double LowValue, double HighValue) { AutoZoomLowValue = LowValue; AutoZoomHighValue = HighValue; }
-	
+
 	/**
 	 * Compute BaselineY and ScaleY so the [Low, High] Value range will correspond to [Top, Bottom] Y position range.
 	 * GetYForValue(InHighValue) == InTopY
@@ -121,6 +132,16 @@ public:
 		OutBaselineY = static_cast<double>(InTopY) + InHighValue * OutScaleY;
 	}
 
+	/**
+	 * @param X The horizontal coordinate of the point tested; in Slate pixels (local graph coordinates)
+	 * @param Y The vertical coordinate of the point tested; in Slate pixels (local graph coordinates)
+	 * @param Viewport The timing viewport used to transform time in local graph coordinates
+	 * @param bCheckLine If needs to check the bounding box of the horizontal line (determined by duration of event and value) or only the bounding box of the visual point
+	 * @param bCheckBox If needs to check the bounding box of the entire visual box (determined by duration of event, value and baseline)
+	 * @return A pointer to an Event located at (X, Y) coordinates, if any; nullptr if no event is located at respective coordinates
+	 */
+	const FGraphSeries::FEvent* GetEvent(const float PosX, const float PosY, const FTimingTrackViewport& Viewport, bool bCheckLine, bool bCheckBox) const;
+
 private:
 	FText Name;
 	FText Description;
@@ -141,9 +162,10 @@ private:
 	FLinearColor BorderColor;
 
 protected:
-	TArray<FVector2D> Points;
-	TArray<FVector2D> LinePoints;
-	TArray<FGraphBox> Boxes;
+	TArray<FEvent> Events; // reduced list of events; used to identify an event at a certain screen position (ex.: the event hovered by mouse)
+	TArray<FVector2D> Points; // reduced list of points; for drawing points
+	TArray<FVector2D> LinePoints; // reduced list of points; for drawing the connected line and filled polygon
+	TArray<FBox> Boxes; // reduced list of boxes; for drawing boxes
 };
 
 ////////////////////////////////////////////////////////////////////////////////////////////////////
@@ -152,6 +174,15 @@ class FGraphTrack : public FBaseTimingTrack
 {
 	friend class FGraphTrackBuilder;
 
+public:
+	struct FEvent
+	{
+		TSharedPtr<FGraphSeries> Series;
+		FGraphSeries::FEvent SeriesEvent;
+
+		bool IsValid() const { return Series.IsValid(); }
+	};
+
 private:
 	// Visual size of points (in pixels).
 	static constexpr float PointVisualSize = 5.0f;
@@ -170,7 +201,18 @@ public:
 
 	virtual void Update(const FTimingTrackViewport& Viewport) override = 0;
 
-	void Draw(FDrawContext& DrawContext, const FTimingTrackViewport& Viewport) const;
+	void Draw(FDrawContext& DrawContext, const FTimingTrackViewport& Viewport, const FVector2D& MousePosition) const;
+
+	/**
+	 * @param X The horizontal coordinate of the point tested; in Slate pixels (local graph coordinates)
+	 * @param Y The vertical coordinate of the point tested; in Slate pixels (local graph coordinates)
+	 * @param Viewport The timing viewport used to transform time in local graph coordinates
+	 * @param OutEvent The returned event located at (PosX, PosY) coordinates, if any.
+	 * @return True if an event is found at (PosX, PosY) coordinates.
+	 */
+	const bool GetEvent(const float MouseX, const float MouseY, const FTimingTrackViewport& Viewport, FGraphTrack::FEvent& OutEvent) const;
+
+	void DrawHighlightedEvent(FDrawContext& DrawContext, const FTimingTrackViewport& Viewport, const FGraphTrack::FEvent& GraphEvent) const;
 
 	virtual void BuildContextMenu(FMenuBuilder& MenuBuilder) override;
 
@@ -236,7 +278,7 @@ protected:
 	bool bDrawBoxes;
 
 	// Stats
-	int32 NumAddedEvents;
+	int32 NumAddedEvents; // total event count
 	int32 NumDrawPoints;
 	int32 NumDrawLines;
 	int32 NumDrawBoxes;
@@ -290,7 +332,7 @@ public:
 
 private:
 	void BeginPoints();
-	void AddPoint(double Time, double Value);
+	bool AddPoint(double Time, double Value);
 	void FlushPoints();
 	void EndPoints();
 
diff --git a/Engine/Source/Developer/TraceInsights/Private/Insights/ViewModels/TimingGraphTrack.cpp b/Engine/Source/Developer/TraceInsights/Private/Insights/ViewModels/TimingGraphTrack.cpp
--- a/Engine/Source/Developer/TraceInsights/Private/Insights/ViewModels/TimingGraphTrack.cpp
+++ b/Engine/Source/Developer/TraceInsights/Private/Insights/ViewModels/TimingGraphTrack.cpp
@@ -262,7 +262,7 @@ void FTimingGraphTrack::UpdateStatsCounterSeries(FTimingGraphSeries& Series, con
 				}
 
 				const float TopY = 4.0f;
-				const float BottomY = GetHeight();
+				const float BottomY = GetHeight() - 4.0f;
 
 				const double HighValue = Series.GetValueForY(TopY);
 				const double LowValue = Series.GetValueForY(BottomY);
diff --git a/Engine/Source/Developer/TraceInsights/Private/Insights/Widgets/SGraphTrack.cpp b/Engine/Source/Developer/TraceInsights/Private/Insights/Widgets/SGraphTrack.cpp
--- a/Engine/Source/Developer/TraceInsights/Private/Insights/Widgets/SGraphTrack.cpp
+++ b/Engine/Source/Developer/TraceInsights/Private/Insights/Widgets/SGraphTrack.cpp
@@ -127,7 +127,7 @@ int32 SGraphTrack::OnPaint(const FPaintArgs& Args, const FGeometry& AllottedGeom
 	const ESlateDrawEffect DrawEffects = bEnabled ? ESlateDrawEffect::None : ESlateDrawEffect::DisabledEffect;
 	FDrawContext DrawContext(AllottedGeometry, MyCullingRect, InWidgetStyle, DrawEffects, OutDrawElements, LayerId);
 
-	GraphTrack->Draw(DrawContext, Viewport);
+	GraphTrack->Draw(DrawContext, Viewport, MousePosition);
 
 	TimeRulerTrack->Draw(DrawContext, Viewport, MousePosition, bIsSelecting, SelectionStartTime, SelectionEndTime);
 	DrawContext.DrawBox(0.0f, TimeRulerTrack->GetHeight(), Viewport.Width, 1.0f, WhiteBrush, FLinearColor(0.05f, 0.05f, 0.05f, 1.0f));
diff --git a/Engine/Source/Developer/TraceInsights/Private/Insights/Widgets/STimingView.cpp b/Engine/Source/Developer/TraceInsights/Private/Insights/Widgets/STimingView.cpp
--- a/Engine/Source/Developer/TraceInsights/Private/Insights/Widgets/STimingView.cpp
+++ b/Engine/Source/Developer/TraceInsights/Private/Insights/Widgets/STimingView.cpp
@@ -798,8 +798,7 @@ int32 STimingView::OnPaint(const FPaintArgs& Args, const FGeometry& AllottedGeom
 
 	if (GraphTrack->IsVisible())
 	{
-		GraphTrack->Draw(DrawContext, Viewport);
-		//GraphTrack->Draw(DrawContext, Viewport, MousePosition);
+		GraphTrack->Draw(DrawContext, Viewport, MousePosition);
 	}
 
 	if (MarkersTrack.IsVisible())
