[Insights] Graph track: Added tooltip for event hovered by mouse.

diff --git a/Engine/Source/Developer/TraceInsights/Private/Insights/ViewModels/GraphTrack.cpp b/Engine/Source/Developer/TraceInsights/Private/Insights/ViewModels/GraphTrack.cpp
--- a/Engine/Source/Developer/TraceInsights/Private/Insights/ViewModels/GraphTrack.cpp
+++ b/Engine/Source/Developer/TraceInsights/Private/Insights/ViewModels/GraphTrack.cpp
@@ -11,6 +11,7 @@
 
 // Insights
 #include "Insights/Common/PaintUtils.h"
+#include "Insights/Common/TimeUtils.h"
 #include "Insights/ViewModels/DrawHelpers.h"
 #include "Insights/ViewModels/TimingTrackViewport.h"
 
@@ -152,10 +153,14 @@ void FGraphTrack::UpdateHoveredState(float MouseX, float MouseY, const FTimingTr
 	constexpr float HeaderWidth = 100.0f;
 	constexpr float HeaderHeight = 14.0f;
 
+	HoveredGraphEvent.Series = nullptr;
+
 	if (MouseY >= GetPosY() && MouseY < GetPosY() + GetHeight())
 	{
 		SetHoveredState(true);
 		SetHeaderHoveredState(MouseX < HeaderWidth && MouseY < GetPosY() + HeaderHeight);
+
+		GetEvent(MouseX, MouseY, Viewport, HoveredGraphEvent);
 	}
 	else
 	{
@@ -236,14 +241,9 @@ void FGraphTrack::Draw(FDrawContext& DrawContext, const FTimingTrackViewport& Vi
 		}
 	}
 
-	if (MousePosition.Y >= GetPosY() && MousePosition.Y < GetPosY() + GetHeight())
+	if (HoveredGraphEvent.IsValid())
 	{
-		FGraphTrack::FEvent GraphEvent;
-		GraphEvent.Series = nullptr;
-		if (GetEvent(MousePosition.X, MousePosition.Y, Viewport, GraphEvent))
-		{
-			DrawHighlightedEvent(DrawContext, Viewport, GraphEvent);
-		}
+		DrawHighlightedEvent(DrawContext, Viewport, HoveredGraphEvent);
 	}
 
 	DrawContext.ElementList.PopClip();
@@ -251,6 +251,20 @@ void FGraphTrack::Draw(FDrawContext& DrawContext, const FTimingTrackViewport& Vi
 
 ////////////////////////////////////////////////////////////////////////////////////////////////////
 
+void FGraphTrack::PostDraw(FDrawContext& DrawContext, const FTimingTrackViewport& Viewport, const FVector2D& MousePosition) const
+{
+	if (HoveredGraphEvent.IsValid() && !MousePosition.IsZero())
+	{
+		DrawTooltip(DrawContext, Viewport, MousePosition, HoveredGraphEvent);
+	}
+	else
+	{
+		Tooltip.Reset();
+	}
+}
+
+////////////////////////////////////////////////////////////////////////////////////////////////////
+
 void FGraphTrack::DrawSeries(const FGraphSeries& Series, FDrawContext& DrawContext, const FTimingTrackViewport& Viewport) const
 {
 	if (bDrawBoxes)
@@ -538,6 +552,55 @@ void FGraphTrack::DrawHighlightedEvent(FDrawContext& DrawContext, const FTimingT
 
 ////////////////////////////////////////////////////////////////////////////////////////////////////
 
+void FGraphTrack::DrawTooltip(FDrawContext& DrawContext, const FTimingTrackViewport& Viewport, const FVector2D& MousePosition, const FGraphTrack::FEvent& GraphEvent) const
+{
+	const TSharedRef<FSlateFontMeasure> FontMeasureService = FSlateApplication::Get().GetRenderer()->GetFontMeasureService();
+	float DesiredTooltipWidth = FontMeasureService->Measure(GraphEvent.Series->GetName(), Font).X + 2 * FTimingViewTooltip::BorderX;
+	if (DesiredTooltipWidth < FTimingViewTooltip::MinWidth)
+	{
+		DesiredTooltipWidth = FTimingViewTooltip::MinWidth;
+	}
+
+	constexpr float LineH = 14.0f;
+	const float DesiredTooltipHeight = 4 * LineH + 2 * FTimingViewTooltip::BorderY;
+
+	Tooltip.Update(MousePosition, DesiredTooltipWidth, DesiredTooltipHeight, Viewport.Width, Viewport.Height);
+
+	const FLinearColor BackgroundColor(0.05f, 0.05f, 0.05f, Tooltip.Opacity);
+	const FLinearColor TextColor(0.6f, 0.6f, 0.6f, Tooltip.Opacity);
+	const FLinearColor ValueColor(1.0f, 1.0f, 1.0f, Tooltip.Opacity);
+
+	DrawContext.DrawBox(Tooltip.PosX, Tooltip.PosY, Tooltip.Width, Tooltip.Height, WhiteBrush, BackgroundColor);
+	DrawContext.LayerId++;
+
+	float X = Tooltip.PosX + FTimingViewTooltip::BorderX;
+	float Y = Tooltip.PosY + FTimingViewTooltip::BorderY;
+
+	DrawContext.DrawText(X, Y, GraphEvent.Series->GetName().ToString(), Font, GraphEvent.Series->GetColor());
+	Y += LineH;
+
+	const float ValueX = X + 48.0f;
+
+	DrawContext.DrawTextAligned(HAlign_Right, ValueX - 4.0f, Y, TEXT("Time:"), Font, TextColor);
+	FString InclStr = TimeUtils::FormatTimeAuto(GraphEvent.SeriesEvent.Time);
+	DrawContext.DrawText(ValueX, Y, InclStr, Font, ValueColor);
+	Y += LineH;
+
+	DrawContext.DrawTextAligned(HAlign_Right, ValueX - 4.0f, Y, TEXT("Duration:"), Font, TextColor);
+	FString ExclStr = TimeUtils::FormatTimeAuto(GraphEvent.SeriesEvent.Duration);
+	DrawContext.DrawText(ValueX, Y, ExclStr, Font, ValueColor);
+	Y += LineH;
+
+	DrawContext.DrawTextAligned(HAlign_Right, ValueX - 4.0f, Y, TEXT("Value:"), Font, TextColor);
+	FString DepthStr = FString::Printf(TEXT("%g"), GraphEvent.SeriesEvent.Value);
+	DrawContext.DrawText(ValueX, Y, DepthStr, Font, ValueColor);
+	Y += LineH;
+
+	DrawContext.LayerId++;
+}
+
+////////////////////////////////////////////////////////////////////////////////////////////////////
+
 const bool FGraphTrack::GetEvent(const float MouseX, const float MouseY, const FTimingTrackViewport& Viewport, FGraphTrack::FEvent& OutEvent) const
 {
 	const float LocalPosX = MouseX;
@@ -1011,11 +1074,11 @@ void FGraphTrackBuilder::AddEvent(double Time, double Duration, double Value)
 {
 	Track.NumAddedEvents++;
 
-	bool bIsEventVisible = false;
+	bool bKeepEvent = (Viewport.GetViewportDXForDuration(Duration) > 1.0f); // always keep the events wider than 1px
 
 	//if (Track.bDrawPoints)
 	{
-		bIsEventVisible = AddPoint(Time, Value);
+		bKeepEvent |= AddPoint(Time, Value);
 	}
 
 	if (Track.bDrawLines || Track.bDrawPolygon)
@@ -1033,7 +1096,7 @@ void FGraphTrackBuilder::AddEvent(double Time, double Duration, double Value)
 		AddBox(Time, Duration, Value);
 	}
 
-	if (bIsEventVisible)
+	if (bKeepEvent)
 	{
 		Series.Events.Add({ Time, Duration, Value });
 	}
diff --git a/Engine/Source/Developer/TraceInsights/Private/Insights/ViewModels/GraphTrack.h b/Engine/Source/Developer/TraceInsights/Private/Insights/ViewModels/GraphTrack.h
--- a/Engine/Source/Developer/TraceInsights/Private/Insights/ViewModels/GraphTrack.h
+++ b/Engine/Source/Developer/TraceInsights/Private/Insights/ViewModels/GraphTrack.h
@@ -7,6 +7,7 @@
 
 // Insights
 #include "Insights/ViewModels/BaseTimingTrack.h"
+#include "Insights/ViewModels/TimingViewDrawHelper.h"
 
 struct FDrawContext;
 struct FSlateBrush;
@@ -202,6 +203,7 @@ public:
 	virtual void Update(const FTimingTrackViewport& Viewport) override = 0;
 
 	void Draw(FDrawContext& DrawContext, const FTimingTrackViewport& Viewport, const FVector2D& MousePosition) const;
+	void PostDraw(FDrawContext& DrawContext, const FTimingTrackViewport& Viewport, const FVector2D& MousePosition) const;
 
 	/**
 	 * @param X The horizontal coordinate of the point tested; in Slate pixels (local graph coordinates)
@@ -212,8 +214,6 @@ public:
 	 */
 	const bool GetEvent(const float MouseX, const float MouseY, const FTimingTrackViewport& Viewport, FGraphTrack::FEvent& OutEvent) const;
 
-	void DrawHighlightedEvent(FDrawContext& DrawContext, const FTimingTrackViewport& Viewport, const FGraphTrack::FEvent& GraphEvent) const;
-
 	virtual void BuildContextMenu(FMenuBuilder& MenuBuilder) override;
 
 	TArray<TSharedPtr<FGraphSeries>>& GetSeries() { return AllSeries; }
@@ -231,6 +231,8 @@ protected:
 
 	void DrawBackground(FDrawContext& DrawContext, const FTimingTrackViewport& Viewport) const;
 	void DrawSeries(const FGraphSeries& Series, FDrawContext& DrawContext, const FTimingTrackViewport& Viewport) const;
+	void DrawHighlightedEvent(FDrawContext& DrawContext, const FTimingTrackViewport& Viewport, const FGraphTrack::FEvent& GraphEvent) const;
+	void DrawTooltip(FDrawContext& DrawContext, const FTimingTrackViewport& Viewport, const FVector2D& MousePosition, const FGraphTrack::FEvent& GraphEvent) const;
 
 private:
 	void ContextMenu_ShowPoints_Execute();
@@ -277,6 +279,9 @@ protected:
 	bool bUseEventDuration;
 	bool bDrawBoxes;
 
+	mutable FGraphTrack::FEvent HoveredGraphEvent;
+	mutable FTimingViewTooltip Tooltip;
+
 	// Stats
 	int32 NumAddedEvents; // total event count
 	int32 NumDrawPoints;
diff --git a/Engine/Source/Developer/TraceInsights/Private/Insights/Widgets/SGraphTrack.cpp b/Engine/Source/Developer/TraceInsights/Private/Insights/Widgets/SGraphTrack.cpp
--- a/Engine/Source/Developer/TraceInsights/Private/Insights/Widgets/SGraphTrack.cpp
+++ b/Engine/Source/Developer/TraceInsights/Private/Insights/Widgets/SGraphTrack.cpp
@@ -134,6 +134,8 @@ int32 SGraphTrack::OnPaint(const FPaintArgs& Args, const FGeometry& AllottedGeom
 
 	FDrawHelpers::DrawTimeRangeSelection(DrawContext, Viewport, SelectionStartTime, SelectionEndTime, WhiteBrush, MainFont);
 
+	GraphTrack->PostDraw(DrawContext, Viewport, MousePosition);
+
 	const bool bShouldDisplayDebugInfo = FInsightsManager::Get()->IsDebugInfoEnabled();
 	if (bShouldDisplayDebugInfo)
 	{
@@ -402,6 +404,10 @@ FReply SGraphTrack::OnMouseMove(const FGeometry& MyGeometry, const FPointerEvent
 
 		Reply = FReply::Handled();
 	}
+	else
+	{
+		GraphTrack->UpdateHoveredState(MousePosition.X, MousePosition.Y, Viewport);
+	}
 
 	return Reply;
 }
diff --git a/Engine/Source/Developer/TraceInsights/Private/Insights/Widgets/STimingView.cpp b/Engine/Source/Developer/TraceInsights/Private/Insights/Widgets/STimingView.cpp
--- a/Engine/Source/Developer/TraceInsights/Private/Insights/Widgets/STimingView.cpp
+++ b/Engine/Source/Developer/TraceInsights/Private/Insights/Widgets/STimingView.cpp
@@ -1026,6 +1026,11 @@ int32 STimingView::OnPaint(const FPaintArgs& Args, const FGeometry& AllottedGeom
 		{
 			Tooltip.Reset();
 		}
+
+		if (GraphTrack->IsVisible())
+		{
+			GraphTrack->PostDraw(DrawContext, Viewport, MousePosition);
+		}
 	}
 
 	if (bAssetLoadingMode)
@@ -1496,6 +1501,11 @@ FReply STimingView::OnMouseMove(const FGeometry& MyGeometry, const FPointerEvent
 			MarkersTrack.UpdateHoveredState(MousePosition.X, MousePosition.Y, Viewport);
 		}
 
+		if (GraphTrack->IsVisible())
+		{
+			GraphTrack->UpdateHoveredState(MousePosition.X, MousePosition.Y, Viewport);
+		}
+
 		UpdateHoveredTimingEvent(MousePosition.X, MousePosition.Y);
 	}
 
