[Insights]
  - Graph track: Added custom formatting for value displayed in tooltips. For Stats Counters, value can be "memory" (bytes), an integer number or a floating point number. For "Frame" series the value is displayed as time and as "fps".
  - Graph track: Improved drawing of boxes (implemented simple reduction algorithm).
  - Graph track: Improved performance when auto zooming.

diff --git a/Engine/Source/Developer/TraceInsights/Private/Insights/ViewModels/GraphTrack.cpp b/Engine/Source/Developer/TraceInsights/Private/Insights/ViewModels/GraphTrack.cpp
--- a/Engine/Source/Developer/TraceInsights/Private/Insights/ViewModels/GraphTrack.cpp
+++ b/Engine/Source/Developer/TraceInsights/Private/Insights/ViewModels/GraphTrack.cpp
@@ -107,6 +107,13 @@ const FGraphSeries::FEvent* FGraphSeries::GetEvent(const float X, const float Y,
 	return nullptr;
 }
 
+////////////////////////////////////////////////////////////////////////////////////////////////////
+
+FString FGraphSeries::FormatValue(double Value) const
+{
+	return FString::Printf(TEXT("%g"), Value);
+}
+
 ////////////////////////////////////////////////////////////////////////////////////////////////////
 // FGraphTrack
 ////////////////////////////////////////////////////////////////////////////////////////////////////
@@ -178,9 +185,12 @@ void FGraphTrack::UpdateStats()
 
 	for (const TSharedPtr<FGraphSeries>& Series : AllSeries)
 	{
-		NumDrawPoints += Series->Points.Num();
-		NumDrawLines += Series->LinePoints.Num() / 2;
-		NumDrawBoxes += Series->Boxes.Num();
+		if (Series->IsVisible())
+		{
+			NumDrawPoints += Series->Points.Num();
+			NumDrawLines += Series->LinePoints.Num() / 2;
+			NumDrawBoxes += Series->Boxes.Num();
+		}
 	}
 }
 
@@ -555,10 +565,22 @@ void FGraphTrack::DrawHighlightedEvent(FDrawContext& DrawContext, const FTimingT
 void FGraphTrack::DrawTooltip(FDrawContext& DrawContext, const FTimingTrackViewport& Viewport, const FVector2D& MousePosition, const FGraphTrack::FEvent& GraphEvent) const
 {
 	const TSharedRef<FSlateFontMeasure> FontMeasureService = FSlateApplication::Get().GetRenderer()->GetFontMeasureService();
-	float DesiredTooltipWidth = FontMeasureService->Measure(GraphEvent.Series->GetName(), Font).X + 2 * FTimingViewTooltip::BorderX;
-	if (DesiredTooltipWidth < FTimingViewTooltip::MinWidth)
+
+	constexpr float ValueOffsetX = 48.0f;
+
+	float DesiredTooltipWidth = FTimingViewTooltip::MinWidth;
+
+	const float DesiredTooltipWidthByName = FontMeasureService->Measure(GraphEvent.Series->GetName(), Font).X + 2 * FTimingViewTooltip::BorderX;
+	if (DesiredTooltipWidth < DesiredTooltipWidthByName)
+	{
+		DesiredTooltipWidth = DesiredTooltipWidthByName;
+	}
+
+	FString ValueStr = GraphEvent.Series->FormatValue(GraphEvent.SeriesEvent.Value);
+	const float DesiredTooltipWidthByValue = ValueOffsetX + FontMeasureService->Measure(ValueStr, Font).X + 2 * FTimingViewTooltip::BorderX;
+	if (DesiredTooltipWidth < DesiredTooltipWidthByValue)
 	{
-		DesiredTooltipWidth = FTimingViewTooltip::MinWidth;
+		DesiredTooltipWidth = DesiredTooltipWidthByValue;
 	}
 
 	constexpr float LineH = 14.0f;
@@ -579,21 +601,20 @@ void FGraphTrack::DrawTooltip(FDrawContext& DrawContext, const FTimingTrackViewp
 	DrawContext.DrawText(X, Y, GraphEvent.Series->GetName().ToString(), Font, GraphEvent.Series->GetColor());
 	Y += LineH;
 
-	const float ValueX = X + 48.0f;
+	const float ValueX = X + ValueOffsetX;
 
 	DrawContext.DrawTextAligned(HAlign_Right, ValueX - 4.0f, Y, TEXT("Time:"), Font, TextColor);
-	FString InclStr = TimeUtils::FormatTimeAuto(GraphEvent.SeriesEvent.Time);
-	DrawContext.DrawText(ValueX, Y, InclStr, Font, ValueColor);
+	FString TimeStr = TimeUtils::FormatTimeAuto(GraphEvent.SeriesEvent.Time);
+	DrawContext.DrawText(ValueX, Y, TimeStr, Font, ValueColor);
 	Y += LineH;
 
 	DrawContext.DrawTextAligned(HAlign_Right, ValueX - 4.0f, Y, TEXT("Duration:"), Font, TextColor);
-	FString ExclStr = TimeUtils::FormatTimeAuto(GraphEvent.SeriesEvent.Duration);
-	DrawContext.DrawText(ValueX, Y, ExclStr, Font, ValueColor);
+	FString DurationStr = TimeUtils::FormatTimeAuto(GraphEvent.SeriesEvent.Duration);
+	DrawContext.DrawText(ValueX, Y, DurationStr, Font, ValueColor);
 	Y += LineH;
 
 	DrawContext.DrawTextAligned(HAlign_Right, ValueX - 4.0f, Y, TEXT("Value:"), Font, TextColor);
-	FString DepthStr = FString::Printf(TEXT("%g"), GraphEvent.SeriesEvent.Value);
-	DrawContext.DrawText(ValueX, Y, DepthStr, Font, ValueColor);
+	DrawContext.DrawText(ValueX, Y, ValueStr, Font, ValueColor);
 	Y += LineH;
 
 	DrawContext.LayerId++;
@@ -1321,11 +1342,47 @@ void FGraphTrackBuilder::AddBox(double Time, double Duration, double Value)
 		W = 1.0f;
 	}
 
-	// TODO: reduction algorithm
+	const float Y = Series.GetRoundedYForValue(Value);
+
+	// simple reduction
+	if (W == 1.0f && Series.Boxes.Num() > 0)
+	{
+		FGraphSeries::FBox& LastBox = Series.Boxes.Last();
+		if (LastBox.W == 1.0f && LastBox.X == X1)
+		{
+			const float RoundedBaselineY = FMath::RoundToFloat(Series.GetBaselineY());
+
+			if (LastBox.Y < RoundedBaselineY)
+			{
+				if (Y < RoundedBaselineY)
+				{
+					// Merge current box with previous one.
+					if (Y < LastBox.Y)
+					{
+						LastBox.Y = Y;
+					}
+					return;
+				}
+			}
+			else
+			{
+				if (Y >= RoundedBaselineY)
+				{
+					// Merge current box with previous one.
+					if (Y > LastBox.Y)
+					{
+						LastBox.Y = Y;
+					}
+					return;
+				}
+			}
+		}
+	}
+
 	FGraphSeries::FBox Box;
 	Box.X = X1;
 	Box.W = W;
-	Box.Y = Series.GetRoundedYForValue(Value);
+	Box.Y = Y;
 	Series.Boxes.Add(Box);
 }
 
diff --git a/Engine/Source/Developer/TraceInsights/Private/Insights/ViewModels/GraphTrack.h b/Engine/Source/Developer/TraceInsights/Private/Insights/ViewModels/GraphTrack.h
--- a/Engine/Source/Developer/TraceInsights/Private/Insights/ViewModels/GraphTrack.h
+++ b/Engine/Source/Developer/TraceInsights/Private/Insights/ViewModels/GraphTrack.h
@@ -38,7 +38,7 @@ public:
 
 public:
 	FGraphSeries();
-	~FGraphSeries();
+	virtual ~FGraphSeries();
 
 	const FText& GetName() const { return Name; }
 	void SetName(const TCHAR* InName) { Name = FText::FromString(InName); }
@@ -143,6 +143,8 @@ public:
 	 */
 	const FGraphSeries::FEvent* GetEvent(const float PosX, const float PosY, const FTimingTrackViewport& Viewport, bool bCheckLine, bool bCheckBox) const;
 
+	virtual FString FormatValue(double Value) const;
+
 private:
 	FText Name;
 	FText Description;
diff --git a/Engine/Source/Developer/TraceInsights/Private/Insights/ViewModels/TimingGraphTrack.cpp b/Engine/Source/Developer/TraceInsights/Private/Insights/ViewModels/TimingGraphTrack.cpp
--- a/Engine/Source/Developer/TraceInsights/Private/Insights/ViewModels/TimingGraphTrack.cpp
+++ b/Engine/Source/Developer/TraceInsights/Private/Insights/ViewModels/TimingGraphTrack.cpp
@@ -8,12 +8,48 @@
 #include "TraceServices/SessionService.h"
 
 // Insights
+#include "Insights/Common/TimeUtils.h"
 #include "Insights/InsightsManager.h"
 #include "Insights/TimingProfilerManager.h"
 #include "Insights/ViewModels/TimingTrackViewport.h"
 
 #define LOCTEXT_NAMESPACE "GraphTrack"
 
+////////////////////////////////////////////////////////////////////////////////////////////////////
+// FTimingGraphSeries
+////////////////////////////////////////////////////////////////////////////////////////////////////
+
+FString FTimingGraphSeries::FormatValue(double Value) const
+{
+	switch (Type)
+	{
+	case FTimingGraphSeries::ESeriesType::Frame:
+		return FString::Printf(TEXT("%s (%g fps)"), *TimeUtils::FormatTimeAuto(Value), 1.0 / Value);
+
+	case FTimingGraphSeries::ESeriesType::Timer:
+		return TimeUtils::FormatTimeAuto(Value);
+
+	case FTimingGraphSeries::ESeriesType::StatsCounter:
+		if (bIsMemory)
+		{
+			const int64 MemValue = static_cast<int64>(Value);
+			return FString::Printf(TEXT("%s (%s bytes)"), *FText::AsMemory(MemValue).ToString(), *FText::AsNumber(MemValue).ToString());
+		}
+		if (bIsFloatingPoint)
+		{
+			return FString::Printf(TEXT("%g"), Value);
+		}
+		else
+		{
+			const int64 Int64Value = static_cast<int64>(Value);
+			return FText::AsNumber(Int64Value).ToString();
+		}
+		break;
+	}
+
+	return FGraphSeries::FormatValue(Value);
+}
+
 ////////////////////////////////////////////////////////////////////////////////////////////////////
 // FTimingGraphTrack
 ////////////////////////////////////////////////////////////////////////////////////////////////////
@@ -77,6 +113,7 @@ TSharedPtr<FTimingGraphSeries> FTimingGraphTrack::AddStatsCounterSeries(uint32 C
 
 	const TCHAR* CounterName = nullptr;
 	bool bIsFloatingPoint = false;
+	bool bIsMemory = false;
 
 	TSharedPtr<const Trace::IAnalysisSession> Session = FInsightsManager::Get()->GetSession();
 	if (Session.IsValid())
@@ -89,6 +126,7 @@ TSharedPtr<FTimingGraphSeries> FTimingGraphTrack::AddStatsCounterSeries(uint32 C
 			{
 				CounterName = Counter.GetName();
 				bIsFloatingPoint = Counter.IsFloatingPoint();
+				bIsMemory = (Counter.GetDisplayHint() == Trace::CounterDisplayHint_Memory);
 			});
 		}
 	}
@@ -102,6 +140,7 @@ TSharedPtr<FTimingGraphSeries> FTimingGraphTrack::AddStatsCounterSeries(uint32 C
 	Series->Type = FTimingGraphSeries::ESeriesType::StatsCounter;
 	Series->Id = CounterId;
 	Series->bIsFloatingPoint = bIsFloatingPoint;
+	Series->bIsMemory = bIsMemory;
 
 	Series->SetBaselineY(GetHeight() - 1.0f);
 	Series->SetScaleY(1.0);
@@ -264,8 +303,8 @@ void FTimingGraphTrack::UpdateStatsCounterSeries(FTimingGraphSeries& Series, con
 				const float TopY = 4.0f;
 				const float BottomY = GetHeight() - 4.0f;
 
-				const double HighValue = Series.GetValueForY(TopY);
 				const double LowValue = Series.GetValueForY(BottomY);
+				const double HighValue = Series.GetValueForY(TopY);
 
 				// If MinValue == MaxValue, we keep the previous baseline and scale, but only if the min/max value is already visible.
 				if (MinValue == MaxValue && (MinValue < LowValue || MaxValue > HighValue))
@@ -279,24 +318,27 @@ void FTimingGraphTrack::UpdateStatsCounterSeries(FTimingGraphSeries& Series, con
 					constexpr bool bIsAutoZoomAnimated = true;
 					if (bIsAutoZoomAnimated)
 					{
-						// Interpolate interval (animating the vertical position and scale of the graph series).
+						// Interpolate the min-max interval (animating the vertical position and scale of the graph series).
 						constexpr double InterpolationSpeed = 0.5;
-						MaxValue = InterpolationSpeed * MaxValue + (1.0 - InterpolationSpeed) * HighValue;
-						MinValue = InterpolationSpeed * MinValue + (1.0 - InterpolationSpeed) * LowValue;
+						const double NewMinValue = InterpolationSpeed * MinValue + (1.0 - InterpolationSpeed) * LowValue;
+						const double NewMaxValue = InterpolationSpeed * MaxValue + (1.0 - InterpolationSpeed) * HighValue;
+
+						// Check if we reach the target min-max interval.
+						const double ErrorTolerance = 0.5 / Series.GetScaleY(); // delta value for dy ~= 0.5 pixels
+						if (!FMath::IsNearlyEqual(NewMinValue, MinValue, ErrorTolerance) ||
+							!FMath::IsNearlyEqual(NewMaxValue, MaxValue, ErrorTolerance))
+						{
+							MinValue = NewMinValue;
+							MaxValue = NewMaxValue;
+
+							// Request a new update so we can further interpolate the min-max interval.
+							Series.SetDirtyFlag();
+						}
 					}
 
 					double BaselineY;
 					double ScaleY;
 					Series.ComputeBaselineAndScale(MinValue, MaxValue, TopY, BottomY, BaselineY, ScaleY);
-
-					constexpr double ErrorTolerance = 1e-15;
-					if (!FMath::IsNearlyEqual(BaselineY, Series.GetBaselineY(), ErrorTolerance) ||
-						!FMath::IsNearlyEqual(ScaleY, Series.GetScaleY(), ErrorTolerance))
-					{
-						// Request a new update so we can further interpolate the coefficients (until reaching targeted ones).
-						Series.SetDirtyFlag();
-					}
-
 					Series.SetBaselineY(BaselineY);
 					Series.SetScaleY(ScaleY);
 				}
diff --git a/Engine/Source/Developer/TraceInsights/Private/Insights/ViewModels/TimingGraphTrack.h b/Engine/Source/Developer/TraceInsights/Private/Insights/ViewModels/TimingGraphTrack.h
--- a/Engine/Source/Developer/TraceInsights/Private/Insights/ViewModels/TimingGraphTrack.h
+++ b/Engine/Source/Developer/TraceInsights/Private/Insights/ViewModels/TimingGraphTrack.h
@@ -19,10 +19,14 @@ public:
 		StatsCounter
 	};
 
+	virtual FString FormatValue(double Value) const override;
+
 public:
 	ESeriesType Type;
 	uint32 Id; // frame type, timer id or stats counter id
+
 	bool bIsFloatingPoint; // for stats counters
+	bool bIsMemory; // for stats counters
 };
 
 ////////////////////////////////////////////////////////////////////////////////////////////////////
