[Insights]
  - Added "Generic Table Tree View" implementation (supports filtering, grouping and sorting).
  - Moved/renamed IoProfiler* files to LoadingProfiler subfolder.
  - LoadingProfiler: Changed layout to use docking windows.
  - LoadingProfiler: Added two tree table views: Event Aggregation and Object Type Aggregation.
  - LoadingProfiler: Added toolbar with on/off toggles for Timing view, Event Aggregation view and Object Type Aggregation view.
  - Increased version to 0.09

diff --git a/Engine/Source/Developer/TraceInsights/Private/Insights/InsightsManager.cpp b/Engine/Source/Developer/TraceInsights/Private/Insights/InsightsManager.cpp
--- a/Engine/Source/Developer/TraceInsights/Private/Insights/InsightsManager.cpp
+++ b/Engine/Source/Developer/TraceInsights/Private/Insights/InsightsManager.cpp
@@ -7,7 +7,7 @@
 #include "Templates/UniquePtr.h"
 
 // Insights
-#include "Insights/IoProfilerManager.h"
+#include "Insights/LoadingProfiler/LoadingProfilerManager.h"
 #include "Insights/TimingProfilerManager.h"
 #include "Insights/Widgets/SStartPageWindow.h"
 #include "Insights/Widgets/STimingProfilerWindow.h"
@@ -18,7 +18,7 @@
 
 const FName FInsightsManagerTabs::StartPageTabId(TEXT("StartPage"));
 const FName FInsightsManagerTabs::TimingProfilerTabId(TEXT("TimingProfiler"));
-const FName FInsightsManagerTabs::IoProfilerTabId(TEXT("IoProfiler"));
+const FName FInsightsManagerTabs::LoadingProfilerTabId(TEXT("LoadingProfiler"));
 
 TSharedPtr<FInsightsManager> FInsightsManager::Instance = nullptr;
 
@@ -148,10 +148,10 @@ void FInsightsManager::OnSessionChanged()
 		TimingProfilerManager->OnSessionChanged();
 	}
 
-	if (TSharedPtr<FIoProfilerManager> IoProfilerManager = FIoProfilerManager::Get())
+	if (TSharedPtr<FLoadingProfilerManager> LoadingProfilerManager = FLoadingProfilerManager::Get())
 	{
-		// FIXME: make IoProfilerManager to register to SessionChangedEvent instead
-		IoProfilerManager->OnSessionChanged();
+		// FIXME: make LoadingProfilerManager to register to SessionChangedEvent instead
+		LoadingProfilerManager->OnSessionChanged();
 	}
 
 	SessionChangedEvent.Broadcast();
@@ -166,9 +166,9 @@ void FInsightsManager::SpawnAndActivateTabs()
 	{
 		FGlobalTabmanager::Get()->InvokeTab(FInsightsManagerTabs::TimingProfilerTabId);
 	}
-	if (FGlobalTabmanager::Get()->HasTabSpawner(FInsightsManagerTabs::IoProfilerTabId))
+	if (FGlobalTabmanager::Get()->HasTabSpawner(FInsightsManagerTabs::LoadingProfilerTabId))
 	{
-		FGlobalTabmanager::Get()->InvokeTab(FInsightsManagerTabs::IoProfilerTabId);
+		FGlobalTabmanager::Get()->InvokeTab(FInsightsManagerTabs::LoadingProfilerTabId);
 	}
 
 	// Ensure Timing Insights / Timing View is the active tab / view.
@@ -176,7 +176,7 @@ void FInsightsManager::SpawnAndActivateTabs()
 	{
 		TimingInsightsTab->ActivateInParent(ETabActivationCause::SetDirectly);
 
-		//TOOD: FTimingProfilerManager::Get()->SpawnAndActivateTabs();
+		//TODO: FTimingProfilerManager::Get()->ActivateWindow();
 		TSharedPtr<class STimingProfilerWindow> Wnd = FTimingProfilerManager::Get()->GetProfilerWindow();
 		if (Wnd)
 		{
diff --git a/Engine/Source/Developer/TraceInsights/Private/Insights/InsightsManager.h b/Engine/Source/Developer/TraceInsights/Private/Insights/InsightsManager.h
--- a/Engine/Source/Developer/TraceInsights/Private/Insights/InsightsManager.h
+++ b/Engine/Source/Developer/TraceInsights/Private/Insights/InsightsManager.h
@@ -23,7 +23,7 @@ struct FInsightsManagerTabs
 {
 	static const FName StartPageTabId;
 	static const FName TimingProfilerTabId;
-	static const FName IoProfilerTabId;
+	static const FName LoadingProfilerTabId;
 };
 
 ////////////////////////////////////////////////////////////////////////////////////////////////////
diff --git a/Engine/Source/Developer/TraceInsights/Private/Insights/InsightsStyle.cpp b/Engine/Source/Developer/TraceInsights/Private/Insights/InsightsStyle.cpp
--- a/Engine/Source/Developer/TraceInsights/Private/Insights/InsightsStyle.cpp
+++ b/Engine/Source/Developer/TraceInsights/Private/Insights/InsightsStyle.cpp
@@ -45,15 +45,24 @@ TSharedRef<FSlateStyleSet> FInsightsStyle::Create()
 	Style.Set("StartPage.Icon.Large", new IMAGE_BRUSH("/Icons/icon_tab_Tools_16x", FVector2D(32.0f, 32.0f)));
 	Style.Set("StartPage.Icon.Small", new IMAGE_BRUSH("/Icons/icon_tab_Tools_16x", FVector2D(16.0f, 16.0f)));
 
-	Style.Set("TimingInsights.Icon.Large", new IMAGE_BRUSH("Icons/Profiler/profiler_stats_40x", FVector2D(32.0f, 32.0f)));
-	Style.Set("TimingInsights.Icon.Small", new IMAGE_BRUSH("Icons/Profiler/profiler_stats_40x", FVector2D(16.0f, 16.0f)));
+	Style.Set("TimingProfiler.Icon.Large", new IMAGE_BRUSH("Icons/Profiler/profiler_stats_40x", FVector2D(32.0f, 32.0f)));
+	Style.Set("TimingProfiler.Icon.Small", new IMAGE_BRUSH("Icons/Profiler/profiler_stats_40x", FVector2D(16.0f, 16.0f)));
 
-	Style.Set("AssetLoadingInsights.Icon.Large", new IMAGE_BRUSH("Icons/Profiler/profiler_stats_40x", FVector2D(32.0f, 32.0f)));
-	Style.Set("AssetLoadingInsights.Icon.Small", new IMAGE_BRUSH("Icons/Profiler/profiler_stats_40x", FVector2D(16.0f, 16.0f)));
+	Style.Set("LoadingProfiler.Icon.Large", new IMAGE_BRUSH("Icons/Profiler/profiler_stats_40x", FVector2D(32.0f, 32.0f)));
+	Style.Set("LoadingProfiler.Icon.Small", new IMAGE_BRUSH("Icons/Profiler/profiler_stats_40x", FVector2D(16.0f, 16.0f)));
 
 	Style.Set("Toolbar.Icon.Large", new IMAGE_BRUSH("/Icons/icon_tab_Tools_16x", FVector2D(32.0f, 32.0f)));
 	Style.Set("Toolbar.Icon.Small", new IMAGE_BRUSH("/Icons/icon_tab_Tools_16x", FVector2D(16.0f, 16.0f)));
 
+	//////////////////////////////////////////////////
+	// Start Page buttons
+
+	Style.Set("Open.Icon.Large", new IMAGE_BRUSH("Icons/Profiler/Profiler_LoadMultiple_Profiler_40x", FVector2D(32.0f, 32.0f)));
+	Style.Set("Open.Icon.Small", new IMAGE_BRUSH("Icons/Profiler/Profiler_Load_Profiler_40x", FVector2D(16.0f, 16.0f)));
+
+	Style.Set("OpenFile.Icon.Large", new IMAGE_BRUSH("Icons/LV_Load", FVector2D(32.0f, 32.0f)));
+	Style.Set("OpenFile.Icon.Small", new IMAGE_BRUSH("Icons/LV_Load", FVector2D(16.0f, 16.0f)));
+
 	//////////////////////////////////////////////////
 	// Timing Insights
 
@@ -75,14 +84,8 @@ TSharedRef<FSlateStyleSet> FInsightsStyle::Create()
 	Style.Set("LogView.Icon.Large", new IMAGE_BRUSH("Icons/Profiler/profiler_CopyToClipboard_32x", FVector2D(32.0f, 32.0f)));
 	Style.Set("LogView.Icon.Small", new IMAGE_BRUSH("Icons/Profiler/profiler_CopyToClipboard_32x", FVector2D(16.0f, 16.0f)));
 
-	//////////////////////////////////////////////////
-	// Start Page buttons
-
-	Style.Set("Open.Icon.Large", new IMAGE_BRUSH("Icons/Profiler/Profiler_LoadMultiple_Profiler_40x", FVector2D(32.0f, 32.0f)));
-	Style.Set("Open.Icon.Small", new IMAGE_BRUSH("Icons/Profiler/Profiler_Load_Profiler_40x", FVector2D(16.0f, 16.0f)));
-
-	Style.Set("OpenFile.Icon.Large", new IMAGE_BRUSH("Icons/LV_Load", FVector2D(32.0f, 32.0f)));
-	Style.Set("OpenFile.Icon.Small", new IMAGE_BRUSH("Icons/LV_Load", FVector2D(16.0f, 16.0f)));
+	Style.Set("TableTreeView.Icon.Large", new IMAGE_BRUSH("Icons/Profiler/Profiler_Data_Capture_40x", FVector2D(32.0f, 32.0f)));
+	Style.Set("TableTreeView.Icon.Small", new IMAGE_BRUSH("Icons/Profiler/Profiler_Data_Capture_40x", FVector2D(16.0f, 16.0f)));
 
 	//////////////////////////////////////////////////
 
diff --git a/Engine/Source/Developer/TraceInsights/Private/Insights/LoadingProfiler/LoadingProfilerCommands.cpp b/Engine/Source/Developer/TraceInsights/Private/Insights/LoadingProfiler/LoadingProfilerCommands.cpp
new file mode 100644
--- /dev/null
+++ b/Engine/Source/Developer/TraceInsights/Private/Insights/LoadingProfiler/LoadingProfilerCommands.cpp
@@ -0,0 +1,106 @@
+// Copyright 1998-2019 Epic Games, Inc. All Rights Reserved.
+
+#include "LoadingProfilerCommands.h"
+
+#include "DesktopPlatformModule.h"
+#include "EditorStyleSet.h"
+#include "Framework/Application/SlateApplication.h"
+#include "Framework/MultiBox/MultiBoxBuilder.h"
+
+// Insights
+#include "Insights/InsightsManager.h"
+#include "Insights/LoadingProfiler/LoadingProfilerManager.h"
+
+////////////////////////////////////////////////////////////////////////////////////////////////////
+
+#define LOCTEXT_NAMESPACE "FLoadingProfilerCommands"
+
+////////////////////////////////////////////////////////////////////////////////////////////////////
+// FLoadingProfilerMenuBuilder
+////////////////////////////////////////////////////////////////////////////////////////////////////
+
+void FLoadingProfilerMenuBuilder::AddMenuEntry(FMenuBuilder& MenuBuilder, const TSharedPtr< FUICommandInfo >& UICommandInfo, const FUIAction& UIAction)
+{
+	MenuBuilder.AddMenuEntry
+	(
+		UICommandInfo->GetLabel(),
+		UICommandInfo->GetDescription(),
+		UICommandInfo->GetIcon(),
+		UIAction,
+		NAME_None,
+		UICommandInfo->GetUserInterfaceType()
+	);
+}
+
+////////////////////////////////////////////////////////////////////////////////////////////////////
+// FLoadingProfilerCommands
+////////////////////////////////////////////////////////////////////////////////////////////////////
+
+FLoadingProfilerCommands::FLoadingProfilerCommands()
+	: TCommands<FLoadingProfilerCommands>(
+		TEXT("LoadingProfilerCommand"), // Context name for fast lookup
+		NSLOCTEXT("Contexts", "LoadingProfilerCommand", "Asset Loading Insights"), // Localized context name for displaying
+		NAME_None, // Parent
+		FEditorStyle::GetStyleSetName() // Icon Style Set
+	)
+{
+}
+
+////////////////////////////////////////////////////////////////////////////////////////////////////
+
+// UI_COMMAND takes long for the compiler to optimize
+PRAGMA_DISABLE_OPTIMIZATION
+void FLoadingProfilerCommands::RegisterCommands()
+{
+	UI_COMMAND(ToggleTimingViewVisibility, "Timing", "Toggles the visibility of the main Timing view", EUserInterfaceActionType::ToggleButton, FInputChord(EModifierKey::Control, EKeys::T));
+	UI_COMMAND(ToggleEventAggregationTreeViewVisibility, "Event Aggregation", "Toggles the visibility of the Event Aggregation table/tree view", EUserInterfaceActionType::ToggleButton, FInputChord());
+	UI_COMMAND(ToggleObjectTypeAggregationTreeViewVisibility, "Object Type Aggregation", "Toggles the visibility of the Object Type Aggregation table/tree view", EUserInterfaceActionType::ToggleButton, FInputChord());
+}
+PRAGMA_ENABLE_OPTIMIZATION
+
+////////////////////////////////////////////////////////////////////////////////////////////////////
+// Toggle Commands
+////////////////////////////////////////////////////////////////////////////////////////////////////
+
+#define IMPLEMENT_TOGGLE_COMMAND(CmdName, IsEnabled, SetIsEnabled) \
+	\
+	void FLoadingProfilerActionManager::Map_##CmdName##_Global()\
+	{\
+		This->CommandList->MapAction(This->GetCommands().CmdName, CmdName##_Custom());\
+	}\
+	\
+	const FUIAction FLoadingProfilerActionManager::CmdName##_Custom() \
+	{\
+		FUIAction UIAction;\
+		UIAction.ExecuteAction = FExecuteAction::CreateRaw(this, &FLoadingProfilerActionManager::CmdName##_Execute);\
+		UIAction.CanExecuteAction = FCanExecuteAction::CreateRaw(this, &FLoadingProfilerActionManager::CmdName##_CanExecute);\
+		UIAction.GetActionCheckState = FGetActionCheckState::CreateRaw(this, &FLoadingProfilerActionManager::CmdName##_GetCheckState);\
+		return UIAction;\
+	}\
+	\
+	void FLoadingProfilerActionManager::CmdName##_Execute()\
+	{\
+		const bool b##IsEnabled = !This->IsEnabled();\
+		This->SetIsEnabled(b##IsEnabled);\
+	}\
+	\
+	bool FLoadingProfilerActionManager::CmdName##_CanExecute() const\
+	{\
+		return FInsightsManager::Get()->GetSession().IsValid();\
+	}\
+	\
+	ECheckBoxState FLoadingProfilerActionManager::CmdName##_GetCheckState() const\
+	{\
+		const bool b##IsEnabled = This->IsEnabled();\
+		return b##IsEnabled ? ECheckBoxState::Checked : ECheckBoxState::Unchecked;\
+	}
+
+IMPLEMENT_TOGGLE_COMMAND(ToggleTimingViewVisibility, IsTimingViewVisible, ShowHideTimingView)
+IMPLEMENT_TOGGLE_COMMAND(ToggleEventAggregationTreeViewVisibility, IsEventAggregationTreeViewVisible, ShowHideEventAggregationTreeView)
+IMPLEMENT_TOGGLE_COMMAND(ToggleObjectTypeAggregationTreeViewVisibility, IsObjectTypeAggregationTreeViewVisible, ShowHideObjectTypeAggregationTreeView)
+
+#undef IMPLEMENT_TOGGLE_COMMAND
+
+////////////////////////////////////////////////////////////////////////////////////////////////////
+
+#undef LOCTEXT_NAMESPACE
diff --git a/Engine/Source/Developer/TraceInsights/Private/Insights/LoadingProfiler/LoadingProfilerCommands.h b/Engine/Source/Developer/TraceInsights/Private/Insights/LoadingProfiler/LoadingProfilerCommands.h
new file mode 100644
--- /dev/null
+++ b/Engine/Source/Developer/TraceInsights/Private/Insights/LoadingProfiler/LoadingProfilerCommands.h
@@ -0,0 +1,101 @@
+// Copyright 1998-2019 Epic Games, Inc. All Rights Reserved.
+
+#pragma once
+
+#include "CoreMinimal.h"
+#include "Framework/Commands/UIAction.h"
+#include "Framework/Commands/Commands.h"
+
+class FMenuBuilder;
+
+////////////////////////////////////////////////////////////////////////////////////////////////////
+/**
+ * Class that holds all profiler commands.
+ */
+class FLoadingProfilerCommands : public TCommands<FLoadingProfilerCommands>
+{
+public:
+	/** Default constructor. */
+	FLoadingProfilerCommands();
+
+	/** Initialize commands. */
+	virtual void RegisterCommands() override;
+
+public:
+	//////////////////////////////////////////////////
+	// Global commands need to implement following method:
+	//     void Map_<CommandName>_Global();
+	// Custom commands needs to implement also the following method:
+	//     const FUIAction <CommandName>_Custom(...) const;
+	//////////////////////////////////////////////////
+
+	/** Toggles visibility for the Timing view. Global and custom command. */
+	TSharedPtr<FUICommandInfo> ToggleTimingViewVisibility;
+
+	/** Toggles visibility for the Event Aggregation tree view. Global and custom command. */
+	TSharedPtr<FUICommandInfo> ToggleEventAggregationTreeViewVisibility;
+
+	/** Toggles visibility for the Object Type Aggregation tree view. Global and custom command. */
+	TSharedPtr<FUICommandInfo> ToggleObjectTypeAggregationTreeViewVisibility;
+};
+
+////////////////////////////////////////////////////////////////////////////////////////////////////
+/**
+ * Menu builder. Helper class for adding a customized menu entry using the global UI command info.
+ */
+class FLoadingProfilerMenuBuilder
+{
+public:
+	/**
+	 * Helper method for adding a customized menu entry using the global UI command info.
+	 * FUICommandInfo cannot be executed with custom parameters, so we need to create a custom FUIAction,
+	 * but sometime we have global and local version for the UI command, so reuse data from the global UI command info.
+	 * Ex:
+	 *     SessionInstance_ToggleCapture          - Global version will toggle capture process for all active session instances
+	 *     SessionInstance_ToggleCapture_OneParam - Local version will toggle capture process only for the specified session instance
+	 *
+	 * @param MenuBuilder The menu to add items to
+	 * @param FUICommandInfo A shared pointer to the UI command info
+	 * @param UIAction Customized version of the UI command info stored in an UI action
+	 */
+	static void AddMenuEntry(FMenuBuilder& MenuBuilder, const TSharedPtr<FUICommandInfo>& UICommandInfo, const FUIAction& UIAction);
+};
+
+////////////////////////////////////////////////////////////////////////////////////////////////////
+/**
+ * Class that provides helper functions for the commands to avoid cluttering profiler manager with many small functions.
+ * Can't contain any variables. Directly operates on the profiler manager instance.
+ */
+class FLoadingProfilerActionManager
+{
+	friend class FLoadingProfilerManager;
+
+private:
+	/** Private constructor. */
+	FLoadingProfilerActionManager(class FLoadingProfilerManager* Instance)
+		: This(Instance)
+	{}
+
+	//////////////////////////////////////////////////
+	// Toggle Commands
+
+#define DECLARE_TOGGLE_COMMAND(CmdName)\
+public:\
+	void Map_##CmdName##_Global(); /**< Maps UI command info CmdName with the specified UI command list. */\
+	const FUIAction CmdName##_Custom(); /**< UI action for CmdName command. */\
+protected:\
+	void CmdName##_Execute(); /**< Handles FExecuteAction for CmdName. */\
+	bool CmdName##_CanExecute() const; /**< Handles FCanExecuteAction for CmdName. */\
+	ECheckBoxState CmdName##_GetCheckState() const; /**< Handles FGetActionCheckState for CmdName. */
+
+	DECLARE_TOGGLE_COMMAND(ToggleTimingViewVisibility)
+	DECLARE_TOGGLE_COMMAND(ToggleEventAggregationTreeViewVisibility)
+	DECLARE_TOGGLE_COMMAND(ToggleObjectTypeAggregationTreeViewVisibility)
+#undef DECLARE_TOGGLE_COMMAND
+
+	//////////////////////////////////////////////////
+
+protected:
+	/** Reference to the global instance of the profiler manager. */
+	class FLoadingProfilerManager* This;
+};
diff --git a/Engine/Source/Developer/TraceInsights/Private/Insights/LoadingProfiler/LoadingProfilerManager.cpp b/Engine/Source/Developer/TraceInsights/Private/Insights/LoadingProfiler/LoadingProfilerManager.cpp
new file mode 100644
--- /dev/null
+++ b/Engine/Source/Developer/TraceInsights/Private/Insights/LoadingProfiler/LoadingProfilerManager.cpp
@@ -0,0 +1,166 @@
+// Copyright 1998-2019 Epic Games, Inc. All Rights Reserved.
+
+#include "LoadingProfilerManager.h"
+
+#include "Modules/ModuleManager.h"
+#include "TraceServices/AnalysisService.h"
+#include "TraceServices/SessionService.h"
+
+// Insights
+#include "Insights/Table/Widgets/STableTreeView.h"
+#include "Insights/InsightsManager.h"
+#include "Insights/LoadingProfiler/Widgets/SLoadingProfilerWindow.h"
+#include "Insights/Widgets/STimingView.h"
+
+////////////////////////////////////////////////////////////////////////////////////////////////////
+
+#define LOCTEXT_NAMESPACE "LoadingProfilerManager"
+
+//DEFINE_LOG_CATEGORY(LoadingProfiler);
+//
+//DEFINE_STAT(STAT_FT_OnPaint);
+//DEFINE_STAT(STAT_GT_OnPaint);
+//DEFINE_STAT(STAT_TT_OnPaint);
+//DEFINE_STAT(STAT_IOPM_Tick);
+
+TSharedPtr<FLoadingProfilerManager> FLoadingProfilerManager::Instance = nullptr;
+
+////////////////////////////////////////////////////////////////////////////////////////////////////
+
+FLoadingProfilerManager::FLoadingProfilerManager(TSharedRef<FUICommandList> InCommandList)
+	: CommandList(InCommandList)
+	, ActionManager(this)
+	, ProfilerWindow(nullptr)
+	, bIsTimingViewVisible(true)
+	, bIsEventAggregationTreeViewVisible(true)
+	, bIsObjectTypeAggregationTreeViewVisible(true)
+{
+}
+
+////////////////////////////////////////////////////////////////////////////////////////////////////
+
+void FLoadingProfilerManager::PostConstructor()
+{
+	// Register tick functions.
+	//OnTick = FTickerDelegate::CreateSP(this, &FLoadingProfilerManager::Tick);
+	//OnTickHandle = FTicker::GetCoreTicker().AddTicker(OnTick, 1.0f);
+
+	FLoadingProfilerCommands::Register();
+	BindCommands();
+}
+
+////////////////////////////////////////////////////////////////////////////////////////////////////
+
+void FLoadingProfilerManager::BindCommands()
+{
+	ActionManager.Map_ToggleTimingViewVisibility_Global();
+	ActionManager.Map_ToggleEventAggregationTreeViewVisibility_Global();
+	ActionManager.Map_ToggleObjectTypeAggregationTreeViewVisibility_Global();
+}
+
+////////////////////////////////////////////////////////////////////////////////////////////////////
+
+FLoadingProfilerManager::~FLoadingProfilerManager()
+{
+	FLoadingProfilerCommands::Unregister();
+
+	// Unregister tick function.
+	//FTicker::GetCoreTicker().RemoveTicker(OnTickHandle);
+}
+
+////////////////////////////////////////////////////////////////////////////////////////////////////
+
+TSharedPtr<FLoadingProfilerManager> FLoadingProfilerManager::Get()
+{
+	return FLoadingProfilerManager::Instance;
+}
+
+////////////////////////////////////////////////////////////////////////////////////////////////////
+
+const TSharedRef<FUICommandList> FLoadingProfilerManager::GetCommandList() const
+{
+	return CommandList;
+}
+
+////////////////////////////////////////////////////////////////////////////////////////////////////
+
+const FLoadingProfilerCommands& FLoadingProfilerManager::GetCommands()
+{
+	return FLoadingProfilerCommands::Get();
+}
+
+////////////////////////////////////////////////////////////////////////////////////////////////////
+
+FLoadingProfilerActionManager& FLoadingProfilerManager::GetActionManager()
+{
+	return FLoadingProfilerManager::Instance->ActionManager;
+}
+
+////////////////////////////////////////////////////////////////////////////////////////////////////
+
+bool FLoadingProfilerManager::Tick(float DeltaTime)
+{
+	//SCOPE_CYCLE_COUNTER(STAT_IOPM_Tick);
+
+	return true;
+}
+
+////////////////////////////////////////////////////////////////////////////////////////////////////
+
+void FLoadingProfilerManager::OnSessionChanged()
+{
+	TSharedPtr<SLoadingProfilerWindow> Wnd = GetProfilerWindow();
+	if (Wnd.IsValid())
+	{
+		Wnd->Reset();
+	}
+}
+
+////////////////////////////////////////////////////////////////////////////////////////////////////
+
+void FLoadingProfilerManager::ShowHideTimingView(const bool bTimingViewVisibleState)
+{
+	//bool bWasTimingViewVisible = bIsTimingViewVisible;
+
+	bIsTimingViewVisible = bTimingViewVisibleState;
+
+	TSharedPtr<SLoadingProfilerWindow> Wnd = GetProfilerWindow();
+	if (Wnd.IsValid())
+	{
+		Wnd->ShowHideTab(FLoadingProfilerTabs::TimingViewID, bIsTimingViewVisible);
+	}
+}
+
+////////////////////////////////////////////////////////////////////////////////////////////////////
+
+void FLoadingProfilerManager::ShowHideEventAggregationTreeView(const bool bEventAggregationTreeViewVisibleState)
+{
+	//bool bWasEventAggregationTreeViewVisible = bIsEventAggregationTreeViewVisible;
+
+	bIsEventAggregationTreeViewVisible = bEventAggregationTreeViewVisibleState;
+
+	TSharedPtr<SLoadingProfilerWindow> Wnd = GetProfilerWindow();
+	if (Wnd.IsValid())
+	{
+		Wnd->ShowHideTab(FLoadingProfilerTabs::EventAggregationTreeViewID, bIsEventAggregationTreeViewVisible);
+	}
+}
+
+////////////////////////////////////////////////////////////////////////////////////////////////////
+
+void FLoadingProfilerManager::ShowHideObjectTypeAggregationTreeView(const bool bObjectTypeAggregationTreeViewVisibleState)
+{
+	//bool bWasObjectTypeAggregationTreeViewVisible = bIsObjectTypeAggregationTreeViewVisible;
+
+	bIsObjectTypeAggregationTreeViewVisible = bObjectTypeAggregationTreeViewVisibleState;
+
+	TSharedPtr<SLoadingProfilerWindow> Wnd = GetProfilerWindow();
+	if (Wnd.IsValid())
+	{
+		Wnd->ShowHideTab(FLoadingProfilerTabs::ObjectTypeAggregationTreeViewID, bIsObjectTypeAggregationTreeViewVisible);
+	}
+}
+
+////////////////////////////////////////////////////////////////////////////////////////////////////
+
+#undef LOCTEXT_NAMESPACE
diff --git a/Engine/Source/Developer/TraceInsights/Private/Insights/LoadingProfiler/LoadingProfilerManager.h b/Engine/Source/Developer/TraceInsights/Private/Insights/LoadingProfiler/LoadingProfilerManager.h
new file mode 100644
--- /dev/null
+++ b/Engine/Source/Developer/TraceInsights/Private/Insights/LoadingProfiler/LoadingProfilerManager.h
@@ -0,0 +1,150 @@
+// Copyright 1998-2019 Epic Games, Inc. All Rights Reserved.
+
+#pragma once
+
+#include "Containers/Ticker.h"
+#include "CoreMinimal.h"
+#include "Framework/Commands/UICommandList.h"
+
+// Insights
+#include "Insights/InsightsManager.h"
+#include "Insights/LoadingProfiler/LoadingProfilerCommands.h"
+
+class SLoadingProfilerWindow;
+
+namespace Trace
+{
+	class ISessionService;
+	class IAnalysisService;
+	class IAnalysisSession;
+}
+
+////////////////////////////////////////////////////////////////////////////////////////////////////
+/**
+ * This class manages the Timing Profiler state and settings.
+ */
+class FLoadingProfilerManager : public TSharedFromThis<FLoadingProfilerManager>
+{
+	friend class FLoadingProfilerActionManager;
+
+public:
+	/** Creates the Timing Profiler manager, only one instance can exist. */
+	FLoadingProfilerManager(TSharedRef<FUICommandList> InCommandList);
+
+	/** Virtual destructor. */
+	virtual ~FLoadingProfilerManager();
+
+	/** Creates an instance of the profiler manager. */
+	static TSharedPtr<FLoadingProfilerManager> Initialize()
+	{
+		if (FLoadingProfilerManager::Instance.IsValid())
+		{
+			FLoadingProfilerManager::Instance.Reset();
+		}
+
+		FLoadingProfilerManager::Instance = MakeShareable(new FLoadingProfilerManager(FInsightsManager::Get()->GetCommandList()));
+		FLoadingProfilerManager::Instance->PostConstructor();
+
+		return FLoadingProfilerManager::Instance;
+	}
+
+	/** Shutdowns the Timing Profiler manager. */
+	void Shutdown()
+	{
+		FLoadingProfilerManager::Instance.Reset();
+	}
+
+protected:
+	/** Finishes initialization of the profiler manager. */
+	void PostConstructor();
+
+	/** Binds our UI commands to delegates. */
+	void BindCommands();
+
+public:
+	/**
+	 * @return the global instance of the Loading Profiler (Asset Loading Insights) manager.
+	 * This is an internal singleton and cannot be used outside ProfilerModule.
+	 * For external use:
+	 *     IProfilerModule& ProfilerModule = FModuleManager::Get().LoadModuleChecked<IProfilerModule>("Profiler");
+	 *     ProfilerModule.GetProfilerManager();
+	 */
+	static TSharedPtr<FLoadingProfilerManager> Get();
+
+	/** @returns UI command list for the Timing Profiler manager. */
+	const TSharedRef<FUICommandList> GetCommandList() const;
+
+	/** @return an instance of the Timing Profiler commands. */
+	static const FLoadingProfilerCommands& GetCommands();
+
+	/** @return an instance of the Timing Profiler action manager. */
+	static FLoadingProfilerActionManager& GetActionManager();
+
+	void AssignProfilerWindow(const TSharedRef<SLoadingProfilerWindow>& InProfilerWindow)
+	{
+		ProfilerWindow = InProfilerWindow;
+	}
+
+	/**
+	 * Converts profiler window weak pointer to a shared pointer and returns it.
+	 * Make sure the returned pointer is valid before trying to dereference it.
+	 */
+	TSharedPtr<class SLoadingProfilerWindow> GetProfilerWindow() const
+	{
+		return ProfilerWindow.Pin();
+	}
+
+	////////////////////////////////////////////////////////////////////////////////////////////////////
+	// Getters and setters used by Toggle Commands.
+
+	/** @return true, if Timing view is visible */
+	const bool IsTimingViewVisible() const { return bIsTimingViewVisible; }
+	void SetTimingViewVisible(const bool bTimingViewVisibleState) { bIsTimingViewVisible = bTimingViewVisibleState; }
+	void ShowHideTimingView(const bool bTimingViewVisibleState);
+
+	/** @return true, if Event Aggregation tree view is visible */
+	const bool IsEventAggregationTreeViewVisible() const { return bIsEventAggregationTreeViewVisible; }
+	void SetEventAggregationTreeViewVisible(const bool bVisibleState) { bIsEventAggregationTreeViewVisible = bVisibleState; }
+	void ShowHideEventAggregationTreeView(const bool bVisibleState);
+
+	/** @return true, if Object Type Aggregation tree view is visible */
+	const bool IsObjectTypeAggregationTreeViewVisible() const { return bIsObjectTypeAggregationTreeViewVisible; }
+	void SetObjectTypeAggregationTreeViewVisible(const bool bVisibleState) { bIsObjectTypeAggregationTreeViewVisible = bVisibleState; }
+	void ShowHideObjectTypeAggregationTreeView(const bool bVisibleState);
+
+	////////////////////////////////////////////////////////////////////////////////////////////////////
+
+	void OnSessionChanged();
+
+protected:
+	/** Updates this manager, done through FCoreTicker. */
+	bool Tick(float DeltaTime);
+
+protected:
+	/** The delegate to be invoked when this manager ticks. */
+	FTickerDelegate OnTick;
+
+	/** Handle to the registered OnTick. */
+	FDelegateHandle OnTickHandle;
+
+	/** List of UI commands for this manager. This will be filled by this and corresponding classes. */
+	TSharedRef<FUICommandList> CommandList;
+
+	/** An instance of the Loading Profiler action manager. */
+	FLoadingProfilerActionManager ActionManager;
+
+	/** A weak pointer to the Loading Profiler window. */
+	TWeakPtr<class SLoadingProfilerWindow> ProfilerWindow;
+
+	/** If the Timing view is visible or hidden. */
+	bool bIsTimingViewVisible;
+
+	/** If the Event Aggregation tree view is visible or hidden. */
+	bool bIsEventAggregationTreeViewVisible;
+
+	/** If the Object Type Aggregation tree view is visible or hidden. */
+	bool bIsObjectTypeAggregationTreeViewVisible;
+
+	/** A shared pointer to the global instance of the profiler manager. */
+	static TSharedPtr<FLoadingProfilerManager> Instance;
+};
diff --git a/Engine/Source/Developer/TraceInsights/Private/Insights/LoadingProfiler/Widgets/SLoadingProfilerToolbar.cpp b/Engine/Source/Developer/TraceInsights/Private/Insights/LoadingProfiler/Widgets/SLoadingProfilerToolbar.cpp
new file mode 100644
--- /dev/null
+++ b/Engine/Source/Developer/TraceInsights/Private/Insights/LoadingProfiler/Widgets/SLoadingProfilerToolbar.cpp
@@ -0,0 +1,112 @@
+// Copyright 1998-2019 Epic Games, Inc. All Rights Reserved.
+
+#include "SLoadingProfilerToolbar.h"
+
+#include "EditorStyleSet.h"
+#include "Framework/Application/SlateApplication.h"
+#include "Framework/Docking/TabManager.h"
+#include "Framework/MultiBox/MultiBoxDefs.h"
+#include "Framework/MultiBox/MultiBoxBuilder.h"
+#include "Widgets/Docking/SDockTab.h"
+#include "Widgets/Input/SCheckBox.h"
+#include "Widgets/Layout/SBorder.h"
+#include "Widgets/SBoxPanel.h"
+#include "Widgets/SToolTip.h"
+
+// Insights
+#include "Insights/InsightsCommands.h"
+#include "Insights/InsightsManager.h"
+#include "Insights/LoadingProfiler/LoadingProfilerCommands.h"
+#include "Insights/LoadingProfiler/LoadingProfilerManager.h"
+
+////////////////////////////////////////////////////////////////////////////////////////////////////
+
+#define LOCTEXT_NAMESPACE "SLoadingProfilerToolbar"
+
+////////////////////////////////////////////////////////////////////////////////////////////////////
+
+SLoadingProfilerToolbar::SLoadingProfilerToolbar()
+{
+}
+
+////////////////////////////////////////////////////////////////////////////////////////////////////
+
+SLoadingProfilerToolbar::~SLoadingProfilerToolbar()
+{
+}
+
+////////////////////////////////////////////////////////////////////////////////////////////////////
+
+void SLoadingProfilerToolbar::Construct(const FArguments& InArgs)
+{
+	struct Local
+	{
+		static void FillViewToolbar(FToolBarBuilder& ToolbarBuilder)
+		{
+			ToolbarBuilder.BeginSection("View");
+			{
+				ToolbarBuilder.AddToolBarButton(FLoadingProfilerCommands::Get().ToggleTimingViewVisibility);
+				ToolbarBuilder.AddToolBarButton(FLoadingProfilerCommands::Get().ToggleEventAggregationTreeViewVisibility);
+				ToolbarBuilder.AddToolBarButton(FLoadingProfilerCommands::Get().ToggleObjectTypeAggregationTreeViewVisibility);
+			}
+			ToolbarBuilder.EndSection();
+		}
+
+		static void FillRightSideToolbar(FToolBarBuilder& ToolbarBuilder)
+		{
+			ToolbarBuilder.BeginSection("Debug");
+			{
+				ToolbarBuilder.AddToolBarButton(FInsightsCommands::Get().ToggleDebugInfo);
+			}
+			ToolbarBuilder.EndSection();
+		}
+	};
+
+	TSharedPtr<FUICommandList> CommandList = FInsightsManager::Get()->GetCommandList();
+
+	FToolBarBuilder ToolbarBuilder(CommandList.ToSharedRef(), FMultiBoxCustomization::None);
+	Local::FillViewToolbar(ToolbarBuilder);
+
+	FToolBarBuilder RightSideToolbarBuilder(CommandList.ToSharedRef(), FMultiBoxCustomization::None);
+	Local::FillRightSideToolbar(RightSideToolbarBuilder);
+
+	// Create the tool bar!
+	ChildSlot
+	[
+		SNew(SHorizontalBox)
+
+		+SHorizontalBox::Slot()
+		.HAlign(HAlign_Left)
+		.VAlign(VAlign_Center)
+		.FillWidth(1.0)
+		.Padding(0.0f)
+		[
+			SNew(SBorder)
+			.Padding(0)
+			.BorderImage(FEditorStyle::GetBrush("NoBorder"))
+			.IsEnabled(FSlateApplication::Get().GetNormalExecutionAttribute())
+			[
+				ToolbarBuilder.MakeWidget()
+			]
+		]
+
+		+SHorizontalBox::Slot()
+		.HAlign(HAlign_Right)
+		.VAlign(VAlign_Center)
+		.FillWidth(1.0)
+		.Padding(0.0f)
+		[
+			SNew(SBorder)
+			.Padding(0)
+			.BorderImage(FEditorStyle::GetBrush("NoBorder"))
+			.IsEnabled(FSlateApplication::Get().GetNormalExecutionAttribute())
+			[
+				RightSideToolbarBuilder.MakeWidget()
+			]
+		]
+	];
+}
+
+////////////////////////////////////////////////////////////////////////////////////////////////////
+
+#undef LOCTEXT_NAMESPACE
diff --git a/Engine/Source/Developer/TraceInsights/Private/Insights/LoadingProfiler/Widgets/SLoadingProfilerToolbar.h b/Engine/Source/Developer/TraceInsights/Private/Insights/LoadingProfiler/Widgets/SLoadingProfilerToolbar.h
new file mode 100644
--- /dev/null
+++ b/Engine/Source/Developer/TraceInsights/Private/Insights/LoadingProfiler/Widgets/SLoadingProfilerToolbar.h
@@ -0,0 +1,38 @@
+// Copyright 1998-2019 Epic Games, Inc. All Rights Reserved.
+
+
+#pragma once
+
+#include "CoreMinimal.h"
+#include "Styling/SlateTypes.h"
+#include "Widgets/DeclarativeSyntaxSupport.h"
+#include "Widgets/SCompoundWidget.h"
+
+class SBorder;
+class FToolBarBuilder;
+
+/** Ribbon based toolbar used as a main menu in the Profiler window. */
+class SLoadingProfilerToolbar : public SCompoundWidget
+{
+public:
+	/** Default constructor. */
+	SLoadingProfilerToolbar();
+
+	/** Virtual destructor. */
+	virtual ~SLoadingProfilerToolbar();
+
+	SLATE_BEGIN_ARGS( SLoadingProfilerToolbar )
+		{}
+	SLATE_END_ARGS()
+
+	/**
+	 * Construct this widget
+	 *
+	 * @param	InArgs	The declaration data for this widget
+	 */
+	void Construct( const FArguments& InArgs );
+
+private:
+	/** Create the UI commands for the toolbar */
+	void CreateCommands();
+};
diff --git a/Engine/Source/Developer/TraceInsights/Private/Insights/LoadingProfiler/Widgets/SLoadingProfilerWindow.cpp b/Engine/Source/Developer/TraceInsights/Private/Insights/LoadingProfiler/Widgets/SLoadingProfilerWindow.cpp
new file mode 100644
--- /dev/null
+++ b/Engine/Source/Developer/TraceInsights/Private/Insights/LoadingProfiler/Widgets/SLoadingProfilerWindow.cpp
@@ -0,0 +1,502 @@
+// Copyright 1998-2019 Epic Games, Inc. All Rights Reserved.
+
+#include "SLoadingProfilerWindow.h"
+
+#include "EditorStyleSet.h"
+#include "Framework/Docking/WorkspaceItem.h"
+#include "SlateOptMacros.h"
+#include "Widgets/Docking/SDockTab.h"
+#include "Widgets/Images/SImage.h"
+#include "Widgets/Layout/SBorder.h"
+#include "Widgets/Layout/SBox.h"
+#include "Widgets/Layout/SSpacer.h"
+#include "Widgets/Notifications/SNotificationList.h"
+#include "Widgets/SBoxPanel.h"
+#include "Widgets/Text/STextBlock.h"
+
+#if WITH_EDITOR
+	#include "EngineAnalytics.h"
+	#include "Runtime/Analytics/Analytics/Public/AnalyticsEventAttribute.h"
+	#include "Runtime/Analytics/Analytics/Public/Interfaces/IAnalyticsProvider.h"
+#endif // WITH_EDITOR
+
+// Insights
+#include "Insights/InsightsManager.h"
+#include "Insights/InsightsStyle.h"
+#include "Insights/LoadingProfiler/LoadingProfilerManager.h"
+#include "Insights/LoadingProfiler/Widgets/SLoadingProfilerToolbar.h"
+#include "Insights/Table/Widgets/STableTreeView.h"
+#include "Insights/Version.h"
+#include "Insights/Widgets/SInsightsSettings.h"
+#include "Insights/Widgets/STimingView.h"
+
+////////////////////////////////////////////////////////////////////////////////////////////////////
+
+#define LOCTEXT_NAMESPACE "SLoadingProfilerWindow"
+
+////////////////////////////////////////////////////////////////////////////////////////////////////
+
+const FName FLoadingProfilerTabs::ToolbarID(TEXT("Toolbar"));
+const FName FLoadingProfilerTabs::TimingViewID(TEXT("TimingView"));
+const FName FLoadingProfilerTabs::EventAggregationTreeViewID(TEXT("EventAggregation"));
+const FName FLoadingProfilerTabs::ObjectTypeAggregationTreeViewID(TEXT("ObjectTypeAggregation"));
+
+////////////////////////////////////////////////////////////////////////////////////////////////////
+
+SLoadingProfilerWindow::SLoadingProfilerWindow()
+	: DurationActive(0.0f)
+{
+}
+
+////////////////////////////////////////////////////////////////////////////////////////////////////
+
+SLoadingProfilerWindow::~SLoadingProfilerWindow()
+{
+#if WITH_EDITOR
+	if (DurationActive > 0.0f && FEngineAnalytics::IsAvailable())
+	{
+		FEngineAnalytics::GetProvider().RecordEvent(TEXT("Insights.Usage.LoadingProfiler"), FAnalyticsEventAttribute(TEXT("Duration"), DurationActive));
+	}
+#endif // WITH_EDITOR
+}
+
+////////////////////////////////////////////////////////////////////////////////////////////////////
+
+void SLoadingProfilerWindow::Reset()
+{
+	if (TimingView)
+	{
+		TimingView->Reset();
+	}
+
+	if (EventAggregationTreeView)
+	{
+		EventAggregationTreeView->RebuildTree(true);
+	}
+
+	if (ObjectTypeAggregationTreeView)
+	{
+		ObjectTypeAggregationTreeView->RebuildTree(true);
+	}
+}
+
+////////////////////////////////////////////////////////////////////////////////////////////////////
+
+BEGIN_SLATE_FUNCTION_BUILD_OPTIMIZATION
+
+TSharedRef<SDockTab> SLoadingProfilerWindow::SpawnTab_Toolbar(const FSpawnTabArgs& Args)
+{
+	const TSharedRef<SDockTab> DockTab = SNew(SDockTab)
+		.ShouldAutosize(true)
+		.TabRole(ETabRole::PanelTab)
+		[
+			SNew(SLoadingProfilerToolbar)
+		];
+
+	DockTab->SetOnTabClosed(SDockTab::FOnTabClosedCallback::CreateRaw(this, &SLoadingProfilerWindow::OnToolbarTabClosed));
+
+	return DockTab;
+}
+
+////////////////////////////////////////////////////////////////////////////////////////////////////
+
+void SLoadingProfilerWindow::OnToolbarTabClosed(TSharedRef<SDockTab> TabBeingClosed)
+{
+}
+
+////////////////////////////////////////////////////////////////////////////////////////////////////
+
+TSharedRef<SDockTab> SLoadingProfilerWindow::SpawnTab_TimingView(const FSpawnTabArgs& Args)
+{
+	FLoadingProfilerManager::Get()->SetTimingViewVisible(true);
+
+	const TSharedRef<SDockTab> DockTab = SNew(SDockTab)
+		.ShouldAutosize(false)
+		.TabRole(ETabRole::PanelTab)
+		[
+			SAssignNew(TimingView, STimingView)
+		];
+
+	TimingView->EnableAssetLoadingMode();
+
+	DockTab->SetOnTabClosed(SDockTab::FOnTabClosedCallback::CreateRaw(this, &SLoadingProfilerWindow::OnTimingViewTabClosed));
+
+	return DockTab;
+}
+
+////////////////////////////////////////////////////////////////////////////////////////////////////
+
+void SLoadingProfilerWindow::OnTimingViewTabClosed(TSharedRef<SDockTab> TabBeingClosed)
+{
+	FLoadingProfilerManager::Get()->SetTimingViewVisible(false);
+	TimingView = nullptr;
+}
+
+////////////////////////////////////////////////////////////////////////////////////////////////////
+
+TSharedRef<SDockTab> SLoadingProfilerWindow::SpawnTab_EventAggregationTreeView(const FSpawnTabArgs& Args)
+{
+	FLoadingProfilerManager::Get()->SetEventAggregationTreeViewVisible(true);
+
+	TSharedPtr<Insights::FTable> TablePtr = MakeShareable(new Insights::FTable());
+	TSharedPtr<const Trace::IAnalysisSession> Session = FInsightsManager::Get()->GetSession();
+	if (Session.IsValid() && Trace::ReadLoadTimeProfilerProvider(*Session.Get()))
+	{
+		Trace::FAnalysisSessionReadScope SessionReadScope(*Session.Get());
+		const Trace::ILoadTimeProfilerProvider& LoadTimeProfilerProvider = *Trace::ReadLoadTimeProfilerProvider(*Session.Get());
+		const double SelectionStartTime = TimingView ? TimingView->GetSelectionStartTime() : 0.0;
+		const double SelectionEndTime = TimingView ? TimingView->GetSelectionEndTime() : 0.0;
+		TSharedPtr<Trace::ITable<Trace::FLoadTimeProfilerAggregatedStats>> SourceTable = MakeShareable(LoadTimeProfilerProvider.CreateEventAggregation(SelectionStartTime, SelectionEndTime));
+		TablePtr->Init(SourceTable);
+	}
+
+	const TSharedRef<SDockTab> DockTab = SNew(SDockTab)
+		.ShouldAutosize(false)
+		.TabRole(ETabRole::PanelTab)
+		[
+			SAssignNew(EventAggregationTreeView, Insights::STableTreeView, TablePtr)
+		];
+
+	EventAggregationTreeView->RebuildTree(true);
+
+	DockTab->SetOnTabClosed(SDockTab::FOnTabClosedCallback::CreateRaw(this, &SLoadingProfilerWindow::OnEventAggregationTreeViewTabClosed));
+
+	return DockTab;
+}
+
+////////////////////////////////////////////////////////////////////////////////////////////////////
+
+void SLoadingProfilerWindow::OnEventAggregationTreeViewTabClosed(TSharedRef<SDockTab> TabBeingClosed)
+{
+	FLoadingProfilerManager::Get()->SetEventAggregationTreeViewVisible(false);
+	EventAggregationTreeView = nullptr;
+}
+
+////////////////////////////////////////////////////////////////////////////////////////////////////
+
+TSharedRef<SDockTab> SLoadingProfilerWindow::SpawnTab_ObjectTypeAggregationTreeView(const FSpawnTabArgs& Args)
+{
+	FLoadingProfilerManager::Get()->SetObjectTypeAggregationTreeViewVisible(true);
+
+	TSharedPtr<Insights::FTable> TablePtr = MakeShareable(new Insights::FTable());
+	TSharedPtr<const Trace::IAnalysisSession> Session = FInsightsManager::Get()->GetSession();
+	if (Session.IsValid() && Trace::ReadLoadTimeProfilerProvider(*Session.Get()))
+	{
+		Trace::FAnalysisSessionReadScope SessionReadScope(*Session.Get());
+		const Trace::ILoadTimeProfilerProvider& LoadTimeProfilerProvider = *Trace::ReadLoadTimeProfilerProvider(*Session.Get());
+		const double SelectionStartTime = TimingView ? TimingView->GetSelectionStartTime() : 0.0;
+		const double SelectionEndTime = TimingView ? TimingView->GetSelectionEndTime() : 0.0;
+		TSharedPtr<Trace::ITable<Trace::FLoadTimeProfilerAggregatedStats>> SourceTable = MakeShareable(LoadTimeProfilerProvider.CreateObjectTypeAggregation(SelectionStartTime, SelectionEndTime));
+		TablePtr->Init(SourceTable);
+	}
+
+	const TSharedRef<SDockTab> DockTab = SNew(SDockTab)
+		.ShouldAutosize(false)
+		.TabRole(ETabRole::PanelTab)
+		[
+			SAssignNew(ObjectTypeAggregationTreeView, Insights::STableTreeView, TablePtr)
+		];
+
+	ObjectTypeAggregationTreeView->RebuildTree(true);
+
+	DockTab->SetOnTabClosed(SDockTab::FOnTabClosedCallback::CreateRaw(this, &SLoadingProfilerWindow::OnObjectTypeAggregationTreeViewTabClosed));
+
+	return DockTab;
+}
+
+////////////////////////////////////////////////////////////////////////////////////////////////////
+
+void SLoadingProfilerWindow::OnObjectTypeAggregationTreeViewTabClosed(TSharedRef<SDockTab> TabBeingClosed)
+{
+	FLoadingProfilerManager::Get()->SetObjectTypeAggregationTreeViewVisible(false);
+	ObjectTypeAggregationTreeView = nullptr;
+}
+
+////////////////////////////////////////////////////////////////////////////////////////////////////
+
+void SLoadingProfilerWindow::Construct(const FArguments& InArgs, const TSharedRef<SDockTab>& ConstructUnderMajorTab, const TSharedPtr<SWindow>& ConstructUnderWindow)
+{
+	// Create & initialize tab manager.
+	TabManager = FGlobalTabmanager::Get()->NewTabManager(ConstructUnderMajorTab);
+	TSharedRef<FWorkspaceItem> AppMenuGroup = TabManager->AddLocalWorkspaceMenuCategory(LOCTEXT("LoadingProfilerMenuGroupName", "Timing Insights"));
+
+	TabManager->RegisterTabSpawner(FLoadingProfilerTabs::ToolbarID, FOnSpawnTab::CreateRaw(this, &SLoadingProfilerWindow::SpawnTab_Toolbar))
+		.SetDisplayName(LOCTEXT("DeviceToolbarTabTitle", "Toolbar"))
+		.SetIcon(FSlateIcon(FInsightsStyle::GetStyleSetName(), "Toolbar.Icon.Small"))
+		.SetGroup(AppMenuGroup);
+
+	TabManager->RegisterTabSpawner(FLoadingProfilerTabs::TimingViewID, FOnSpawnTab::CreateRaw(this, &SLoadingProfilerWindow::SpawnTab_TimingView))
+		.SetDisplayName(LOCTEXT("LoadingProfiler.TimingViewTabTitle", "Timing View"))
+		.SetIcon(FSlateIcon(FInsightsStyle::GetStyleSetName(), "TimingView.Icon.Small"))
+		.SetGroup(AppMenuGroup);
+
+	TabManager->RegisterTabSpawner(FLoadingProfilerTabs::EventAggregationTreeViewID, FOnSpawnTab::CreateRaw(this, &SLoadingProfilerWindow::SpawnTab_EventAggregationTreeView))
+		.SetDisplayName(LOCTEXT("LoadingProfiler.EventAggregationTreeViewTabTitle", "Event Aggregation"))
+		.SetIcon(FSlateIcon(FInsightsStyle::GetStyleSetName(), "TableTreeView.Icon.Small"))
+		.SetGroup(AppMenuGroup);
+
+	TabManager->RegisterTabSpawner(FLoadingProfilerTabs::ObjectTypeAggregationTreeViewID, FOnSpawnTab::CreateRaw(this, &SLoadingProfilerWindow::SpawnTab_ObjectTypeAggregationTreeView))
+		.SetDisplayName(LOCTEXT("LoadingProfiler.ObjectTypeAggregationTreeViewTabTitle", "Object Type Aggregation"))
+		.SetIcon(FSlateIcon(FInsightsStyle::GetStyleSetName(), "TableTreeView.Icon.Small"))
+		.SetGroup(AppMenuGroup);
+
+	TSharedPtr<FLoadingProfilerManager> LoadingProfilerManager = FLoadingProfilerManager::Get();
+	ensure(LoadingProfilerManager.IsValid());
+
+	// Create tab layout.
+	const TSharedRef<FTabManager::FLayout> Layout = FTabManager::NewLayout("InsightsLoadingProfilerLayout_v1.0")
+		->AddArea
+		(
+			FTabManager::NewPrimaryArea()
+			->SetOrientation(Orient_Vertical)
+			->Split
+			(
+				FTabManager::NewStack()
+				->AddTab(FLoadingProfilerTabs::ToolbarID, ETabState::OpenedTab)
+				->SetHideTabWell(true)
+			)
+			->Split
+			(
+				FTabManager::NewSplitter()
+				->SetOrientation(Orient_Horizontal)
+				->SetSizeCoefficient(1.0f)
+				->Split
+				(
+					FTabManager::NewStack()
+					->SetSizeCoefficient(0.5f)
+					->SetHideTabWell(true)
+					->AddTab(FLoadingProfilerTabs::TimingViewID, LoadingProfilerManager->IsTimingViewVisible() ? ETabState::OpenedTab : ETabState::ClosedTab)
+				)
+				->Split
+				(
+					FTabManager::NewStack()
+					->SetSizeCoefficient(0.35f)
+					->AddTab(FLoadingProfilerTabs::EventAggregationTreeViewID, LoadingProfilerManager->IsEventAggregationTreeViewVisible() ? ETabState::OpenedTab : ETabState::ClosedTab)
+					->AddTab(FLoadingProfilerTabs::ObjectTypeAggregationTreeViewID, LoadingProfilerManager->IsObjectTypeAggregationTreeViewVisible() ? ETabState::OpenedTab : ETabState::ClosedTab)
+					->SetForegroundTab(FLoadingProfilerTabs::EventAggregationTreeViewID)
+				)
+			)
+		);
+
+	// Create & initialize main menu.
+	FMenuBarBuilder MenuBarBuilder = FMenuBarBuilder(TSharedPtr<FUICommandList>());
+
+	MenuBarBuilder.AddPullDownMenu(
+		LOCTEXT("MenuLabel", "MENU"),
+		FText::GetEmpty(),
+		FNewMenuDelegate::CreateStatic(&SLoadingProfilerWindow::FillMenu, TabManager),
+		FName(TEXT("MENU"))
+	);
+
+	ChildSlot
+		[
+			SNew(SOverlay)
+
+			// Version
+			+ SOverlay::Slot()
+				.HAlign(HAlign_Right)
+				.VAlign(VAlign_Top)
+				.Padding(0.0f, -16.0f, 0.0f, 0.0f)
+				[
+					SNew(STextBlock)
+						.Clipping(EWidgetClipping::ClipToBoundsWithoutIntersecting)
+						.Text(LOCTEXT("UnrealInsightsVersion", UNREAL_INSIGHTS_VERSION_STRING_EX))
+						.ColorAndOpacity(FLinearColor(0.15f, 0.15f, 0.15f, 1.0f))
+				]
+
+			// Overlay slot for the main window area
+			+ SOverlay::Slot()
+				.HAlign(HAlign_Fill)
+				.VAlign(VAlign_Fill)
+				[
+					SNew(SVerticalBox)
+
+					+ SVerticalBox::Slot()
+						.AutoHeight()
+						[
+							MenuBarBuilder.MakeWidget()
+						]
+
+					+ SVerticalBox::Slot()
+						.FillHeight(1.0f)
+						[
+							TabManager->RestoreFrom(Layout, ConstructUnderWindow).ToSharedRef()
+						]
+				]
+
+			// Session hint overlay
+			+ SOverlay::Slot()
+				.HAlign(HAlign_Center)
+				.VAlign(VAlign_Center)
+				[
+					SNew(SBorder)
+						.Visibility(this, &SLoadingProfilerWindow::IsSessionOverlayVisible)
+						.BorderImage(FEditorStyle::GetBrush("NotificationList.ItemBackground"))
+						.Padding(8.0f)
+						[
+							SNew(STextBlock)
+								.Text(LOCTEXT("SelectTraceOverlayText", "Please select a trace."))
+						]
+				]
+		];
+
+	// Tell tab-manager about the global menu bar.
+	TabManager->SetMenuMultiBox(MenuBarBuilder.GetMultiBox());
+}
+
+END_SLATE_FUNCTION_BUILD_OPTIMIZATION
+
+////////////////////////////////////////////////////////////////////////////////////////////////////
+
+void SLoadingProfilerWindow::FillMenu(FMenuBuilder& MenuBuilder, const TSharedPtr<FTabManager> TabManager)
+{
+	if (!TabManager.IsValid())
+	{
+		return;
+	}
+
+#if !WITH_EDITOR
+	//TODO: FGlobalTabmanager::Get()->PopulateTabSpawnerMenu(MenuBuilder, WorkspaceMenu::GetMenuStructure().GetStructureRoot());
+#endif //!WITH_EDITOR
+
+	TabManager->PopulateLocalTabSpawnerMenu(MenuBuilder);
+}
+
+////////////////////////////////////////////////////////////////////////////////////////////////////
+
+void SLoadingProfilerWindow::ShowTab(const FName& TabID)
+{
+	if (TabManager->HasTabSpawner(TabID))
+	{
+		TabManager->InvokeTab(TabID);
+	}
+}
+
+////////////////////////////////////////////////////////////////////////////////////////////////////
+
+void SLoadingProfilerWindow::HideTab(const FName& TabID)
+{
+	TSharedPtr<SDockTab> Tab = TabManager->FindExistingLiveTab(TabID);
+	if (Tab.IsValid())
+	{
+		Tab->RequestCloseTab();
+	}
+}
+
+////////////////////////////////////////////////////////////////////////////////////////////////////
+
+EVisibility SLoadingProfilerWindow::IsSessionOverlayVisible() const
+{
+	if (FInsightsManager::Get()->GetSession().IsValid())
+	{
+		return EVisibility::Hidden;
+	}
+	else
+	{
+		return EVisibility::Visible;
+	}
+}
+
+////////////////////////////////////////////////////////////////////////////////////////////////////
+
+bool SLoadingProfilerWindow::IsProfilerEnabled() const
+{
+	return FInsightsManager::Get()->GetSession().IsValid();
+}
+
+////////////////////////////////////////////////////////////////////////////////////////////////////
+
+EActiveTimerReturnType SLoadingProfilerWindow::UpdateActiveDuration(double InCurrentTime, float InDeltaTime)
+{
+	DurationActive += InDeltaTime;
+
+	// The profiler window will explicitly unregister this active timer when the mouse leaves.
+	return EActiveTimerReturnType::Continue;
+}
+
+////////////////////////////////////////////////////////////////////////////////////////////////////
+
+void SLoadingProfilerWindow::OnMouseEnter(const FGeometry& MyGeometry, const FPointerEvent& MouseEvent)
+{
+	SCompoundWidget::OnMouseEnter(MyGeometry, MouseEvent);
+
+	if (!ActiveTimerHandle.IsValid())
+	{
+		ActiveTimerHandle = RegisterActiveTimer(0.f, FWidgetActiveTimerDelegate::CreateSP(this, &SLoadingProfilerWindow::UpdateActiveDuration));
+	}
+}
+
+////////////////////////////////////////////////////////////////////////////////////////////////////
+
+void SLoadingProfilerWindow::OnMouseLeave(const FPointerEvent& MouseEvent)
+{
+	SCompoundWidget::OnMouseLeave(MouseEvent);
+
+	auto PinnedActiveTimerHandle = ActiveTimerHandle.Pin();
+	if (PinnedActiveTimerHandle.IsValid())
+	{
+		UnRegisterActiveTimer(PinnedActiveTimerHandle.ToSharedRef());
+	}
+}
+
+////////////////////////////////////////////////////////////////////////////////////////////////////
+
+FReply SLoadingProfilerWindow::OnKeyDown(const FGeometry& MyGeometry, const FKeyEvent& InKeyEvent)
+{
+	return FLoadingProfilerManager::Get()->GetCommandList()->ProcessCommandBindings(InKeyEvent) ? FReply::Handled() : FReply::Unhandled();
+}
+
+////////////////////////////////////////////////////////////////////////////////////////////////////
+
+FReply SLoadingProfilerWindow::OnDragOver(const FGeometry& MyGeometry, const FDragDropEvent& DragDropEvent)
+{
+	TSharedPtr<FExternalDragOperation> DragDropOp = DragDropEvent.GetOperationAs<FExternalDragOperation>();
+	if (DragDropOp.IsValid())
+	{
+		if (DragDropOp->HasFiles())
+		{
+			const TArray<FString>& Files = DragDropOp->GetFiles();
+			if (Files.Num() == 1)
+			{
+				const FString DraggedFileExtension = FPaths::GetExtension(Files[0], true);
+				if (DraggedFileExtension == TEXT(".utrace"))
+				{
+					return FReply::Handled();
+				}
+			}
+		}
+	}
+
+	return SCompoundWidget::OnDragOver(MyGeometry,DragDropEvent);
+}
+
+////////////////////////////////////////////////////////////////////////////////////////////////////
+
+FReply SLoadingProfilerWindow::OnDrop(const FGeometry& MyGeometry, const FDragDropEvent& DragDropEvent)
+{
+	TSharedPtr<FExternalDragOperation> DragDropOp = DragDropEvent.GetOperationAs<FExternalDragOperation>();
+	if (DragDropOp.IsValid())
+	{
+		if (DragDropOp->HasFiles())
+		{
+			// For now, only allow a single file.
+			const TArray<FString>& Files = DragDropOp->GetFiles();
+			if (Files.Num() == 1)
+			{
+				const FString DraggedFileExtension = FPaths::GetExtension(Files[0], true);
+				if (DraggedFileExtension == TEXT(".utrace"))
+				{
+					// Enqueue load operation.
+					FInsightsManager::Get()->LoadTraceFile(Files[0]);
+					return FReply::Handled();
+				}
+			}
+		}
+	}
+
+	return SCompoundWidget::OnDrop(MyGeometry,DragDropEvent);
+}
+
+////////////////////////////////////////////////////////////////////////////////////////////////////
+
+#undef LOCTEXT_NAMESPACE
diff --git a/Engine/Source/Developer/TraceInsights/Private/Insights/LoadingProfiler/Widgets/SLoadingProfilerWindow.h b/Engine/Source/Developer/TraceInsights/Private/Insights/LoadingProfiler/Widgets/SLoadingProfilerWindow.h
new file mode 100644
--- /dev/null
+++ b/Engine/Source/Developer/TraceInsights/Private/Insights/LoadingProfiler/Widgets/SLoadingProfilerWindow.h
@@ -0,0 +1,171 @@
+// Copyright 1998-2019 Epic Games, Inc. All Rights Reserved.
+
+#pragma once
+
+#include "CoreMinimal.h"
+#include "Framework/Docking/TabManager.h"
+#include "Input/Reply.h"
+#include "Layout/Visibility.h"
+#include "Misc/Guid.h"
+#include "SlateFwd.h"
+#include "Widgets/DeclarativeSyntaxSupport.h"
+#include "Widgets/Docking/SDockTab.h"
+#include "Widgets/SCompoundWidget.h"
+
+// Insights
+#include "Insights/InsightsManager.h"
+
+////////////////////////////////////////////////////////////////////////////////////////////////////
+
+class FActiveTimerHandle;
+class STimingView;
+
+namespace Insights
+{
+	class STableTreeView;
+}
+
+////////////////////////////////////////////////////////////////////////////////////////////////////
+
+struct FLoadingProfilerTabs
+{
+	// Tab identifiers
+	static const FName ToolbarID;
+	static const FName TimingViewID;
+	static const FName EventAggregationTreeViewID;
+	static const FName ObjectTypeAggregationTreeViewID;
+};
+
+////////////////////////////////////////////////////////////////////////////////////////////////////
+
+/** Implements the timing profiler window. */
+class SLoadingProfilerWindow : public SCompoundWidget
+{
+public:
+	/** Default constructor. */
+	SLoadingProfilerWindow();
+
+	/** Virtual destructor. */
+	virtual ~SLoadingProfilerWindow();
+
+	SLATE_BEGIN_ARGS(SLoadingProfilerWindow) {}
+	SLATE_END_ARGS()
+
+	void Reset();
+
+	/** Constructs this widget. */
+	void Construct(const FArguments& InArgs, const TSharedRef<SDockTab>& ConstructUnderMajorTab, const TSharedPtr<SWindow>& ConstructUnderWindow);
+
+	void ShowTab(const FName& TabID);
+	void HideTab(const FName& TabID);
+	void ShowHideTab(const FName& TabID, bool bShow) { bShow ? ShowTab(TabID) : HideTab(TabID); }
+
+	TSharedPtr<FTabManager> GetTabManager() const { return TabManager; }
+
+	TSharedPtr<STimingView> GetTimingView() const { return TimingView; }
+	TSharedPtr<Insights::STableTreeView> GetEventAggregationTreeView() const { return EventAggregationTreeView; }
+	TSharedPtr<Insights::STableTreeView> GetObjectTypeAggregationTreeView() const { return ObjectTypeAggregationTreeView; }
+
+private:
+	TSharedRef<SDockTab> SpawnTab_Toolbar(const FSpawnTabArgs& Args);
+	void OnToolbarTabClosed(TSharedRef<SDockTab> TabBeingClosed);
+
+	TSharedRef<SDockTab> SpawnTab_TimingView(const FSpawnTabArgs& Args);
+	void OnTimingViewTabClosed(TSharedRef<SDockTab> TabBeingClosed);
+
+	TSharedRef<SDockTab> SpawnTab_EventAggregationTreeView(const FSpawnTabArgs& Args);
+	void OnEventAggregationTreeViewTabClosed(TSharedRef<SDockTab> TabBeingClosed);
+
+	TSharedRef<SDockTab> SpawnTab_ObjectTypeAggregationTreeView(const FSpawnTabArgs& Args);
+	void OnObjectTypeAggregationTreeViewTabClosed(TSharedRef<SDockTab> TabBeingClosed);
+
+	/**
+	 * Fill the main menu with menu items.
+	 *
+	 * @param MenuBuilder The multi-box builder that should be filled with content for this pull-down menu.
+	 * @param TabManager A Tab Manager from which to populate tab spawner menu items.
+	 */
+	static void FillMenu(FMenuBuilder& MenuBuilder, const TSharedPtr<FTabManager> TabManager);
+
+	/** Callback for determining the visibility of the 'Select a session' overlay. */
+	EVisibility IsSessionOverlayVisible() const;
+
+	/** Callback for getting the enabled state of the profiler window. */
+	bool IsProfilerEnabled() const;
+
+	/** Updates the amount of time the profiler has been active. */
+	EActiveTimerReturnType UpdateActiveDuration(double InCurrentTime, float InDeltaTime);
+
+	/**
+	 * Ticks this widget. Override in derived classes, but always call the parent implementation.
+	 *
+	 * @param  AllottedGeometry The space allotted for this widget
+	 * @param  InCurrentTime  Current absolute real time
+	 * @param  InDeltaTime  Real time passed since last tick
+	 */
+	//virtual void Tick(const FGeometry& AllottedGeometry, const double InCurrentTime, const float InDeltaTime) override;
+
+	/**
+	 * The system will use this event to notify a widget that the cursor has entered it. This event is NOT bubbled.
+	 *
+	 * @param MyGeometry The Geometry of the widget receiving the event
+	 * @param MouseEvent Information about the input event
+	 */
+	virtual void OnMouseEnter(const FGeometry& MyGeometry, const FPointerEvent& MouseEvent) override;
+
+	/**
+	 * The system will use this event to notify a widget that the cursor has left it. This event is NOT bubbled.
+	 *
+	 * @param MouseEvent Information about the input event
+	 */
+	virtual void OnMouseLeave(const FPointerEvent& MouseEvent) override;
+
+	/**
+	 * Called after a key is pressed when this widget has focus
+	 *
+	 * @param MyGeometry The Geometry of the widget receiving the event
+	 * @param  InKeyEvent  Key event
+	 *
+	 * @return  Returns whether the event was handled, along with other possible actions
+	 */
+	virtual FReply OnKeyDown(const FGeometry& MyGeometry, const FKeyEvent& InKeyEvent) override;
+
+	/**
+	 * Called when the user is dropping something onto a widget; terminates drag and drop.
+	 *
+	 * @param MyGeometry      The geometry of the widget receiving the event.
+	 * @param DragDropEvent   The drag and drop event.
+	 *
+	 * @return A reply that indicated whether this event was handled.
+	 */
+	virtual FReply OnDrop(const FGeometry& MyGeometry, const FDragDropEvent& DragDropEvent) override;
+
+	/**
+	 * Called during drag and drop when the the mouse is being dragged over a widget.
+	 *
+	 * @param MyGeometry      The geometry of the widget receiving the event.
+	 * @param DragDropEvent   The drag and drop event.
+	 *
+	 * @return A reply that indicated whether this event was handled.
+	 */
+	virtual FReply OnDragOver(const FGeometry& MyGeometry, const FDragDropEvent& DragDropEvent)  override;
+
+private:
+	/** The Timing view (multi-track) widget */
+	TSharedPtr<STimingView> TimingView;
+
+	/** The Event Aggregation tree view widget */
+	TSharedPtr<Insights::STableTreeView> EventAggregationTreeView;
+
+	/** The Object Type Aggregation tree view widget */
+	TSharedPtr<Insights::STableTreeView> ObjectTypeAggregationTreeView;
+
+	/** Holds the tab manager that manages the front-end's tabs. */
+	TSharedPtr<FTabManager> TabManager;
+
+	/** The handle to the active update duration tick */
+	TWeakPtr<FActiveTimerHandle> ActiveTimerHandle;
+
+	/** The number of seconds the profiler has been active */
+	float DurationActive;
+};
diff --git a/Engine/Source/Developer/TraceInsights/Private/Insights/Table/ViewModels/BaseTreeNode.cpp b/Engine/Source/Developer/TraceInsights/Private/Insights/Table/ViewModels/BaseTreeNode.cpp
new file mode 100644
--- /dev/null
+++ b/Engine/Source/Developer/TraceInsights/Private/Insights/Table/ViewModels/BaseTreeNode.cpp
@@ -0,0 +1,59 @@
+// Copyright 1998-2019 Epic Games, Inc. All Rights Reserved.
+
+#include "BaseTreeNode.h"
+
+// Insights
+#include "Insights/Table/ViewModels/TreeNodeSorting.h"
+
+#define LOCTEXT_NAMESPACE "Insights_TreeNode"
+
+namespace Insights
+{
+
+////////////////////////////////////////////////////////////////////////////////////////////////////
+
+const FText FBaseTreeNode::GetDisplayName() const
+{
+	FText Text = FText::GetEmpty();
+
+	if (IsGroup())
+	{
+		const int32 NumChildren = Children.Num();
+		const int32 NumFilteredChildren = FilteredChildren.Num();
+
+		if (NumFilteredChildren == NumChildren)
+		{
+			Text = FText::Format(LOCTEXT("TreeNodeGroupTextFmt1", "{0} ({1})"), FText::FromName(GetName()), FText::AsNumber(NumChildren));
+		}
+		else
+		{
+			Text = FText::Format(LOCTEXT("TreeNodeGroupTextFmt2", "{0} ({1} / {2})"), FText::FromName(GetName()), FText::AsNumber(NumFilteredChildren), FText::AsNumber(NumChildren));
+		}
+	}
+	else
+	{
+		Text = FText::FromName(GetName());
+	}
+
+	return Text;
+}
+
+////////////////////////////////////////////////////////////////////////////////////////////////////
+
+void FBaseTreeNode::SortChildrenAscending(const ITreeNodeSorting* Sorter)
+{
+	Sorter->SortAscending(Children);
+}
+
+////////////////////////////////////////////////////////////////////////////////////////////////////
+
+void FBaseTreeNode::SortChildrenDescending(const ITreeNodeSorting* Sorter)
+{
+	Sorter->SortDescending(Children);
+}
+
+////////////////////////////////////////////////////////////////////////////////////////////////////
+
+} // namespace Insights
+
+#undef LOCTEXT_NAMESPACE
diff --git a/Engine/Source/Developer/TraceInsights/Private/Insights/Table/ViewModels/BaseTreeNode.h b/Engine/Source/Developer/TraceInsights/Private/Insights/Table/ViewModels/BaseTreeNode.h
new file mode 100644
--- /dev/null
+++ b/Engine/Source/Developer/TraceInsights/Private/Insights/Table/ViewModels/BaseTreeNode.h
@@ -0,0 +1,185 @@
+// Copyright 1998-2019 Epic Games, Inc. All Rights Reserved.
+
+#pragma once
+
+#include "CoreMinimal.h"
+
+namespace Insights
+{
+
+class ITreeNodeSorting;
+
+////////////////////////////////////////////////////////////////////////////////////////////////////
+
+class FBaseTreeNode;
+
+/** Type definition for shared pointers to instances of FBaseTreeNode. */
+typedef TSharedPtr<class FBaseTreeNode> FBaseTreeNodePtr;
+
+/** Type definition for shared references to instances of FBaseTreeNode. */
+typedef TSharedRef<class FBaseTreeNode> FBaseTreeNodeRef;
+
+/** Type definition for shared references to const instances of FBaseTreeNode. */
+typedef TSharedRef<const class FBaseTreeNode> FBaseTreeNodeRefConst;
+
+/** Type definition for weak references to instances of FBaseTreeNode. */
+typedef TWeakPtr<class FBaseTreeNode> FBaseTreeNodeWeak;
+
+////////////////////////////////////////////////////////////////////////////////////////////////////
+
+class FBaseTreeNode : public TSharedFromThis<FBaseTreeNode>
+{
+public:
+	static constexpr int32 InvalidId = -1;
+
+public:
+	/** Initialization constructor for the node. */
+	FBaseTreeNode(uint64 InId, const FName InName, bool bInIsGroup)
+		: Id(InId)
+		, Name(InName)
+		, bIsGroup(bInIsGroup)
+		, bIsExpanded(false)
+	{
+	}
+
+	virtual ~FBaseTreeNode()
+	{
+	}
+
+	/**
+	 * @return an id of this node.
+	 */
+	const int32 GetId() const
+	{
+		return Id;
+	}
+
+	/**
+	 * @return a name of this node.
+	 */
+	const FName& GetName() const
+	{
+		return Name;
+	}
+
+	/**
+	 * @return a type id of this node.
+	 */
+	const FName& GetTypeId() const
+	{
+		return TypeId;
+	}
+
+	/**
+	 * @return a name of this node to display in tree view.  Ex.: group nodes may include additional info, like the number of visible/total children.
+	 */
+	virtual const FText GetDisplayName() const;
+
+	/**
+	 * @return true, if this node is a group node.
+	 */
+	bool IsGroup() const
+	{
+		return bIsGroup;
+	}
+
+	/**
+	 * @return a const reference to the child nodes of this group.
+	 */
+	FORCEINLINE_DEBUGGABLE const TArray<FBaseTreeNodePtr>& GetChildren() const
+	{
+		return Children;
+	}
+
+	/**
+	 * @return a const reference to the child nodes that should be visible to the UI based on filtering.
+	 */
+	FORCEINLINE_DEBUGGABLE const TArray<FBaseTreeNodePtr>& GetFilteredChildren() const
+	{
+		return FilteredChildren;
+	}
+
+	/**
+	 * @return a weak reference to the group of this node, may be invalid.
+	 */
+	FBaseTreeNodeWeak GetGroupPtr() const
+	{
+		return GroupPtr;
+	}
+
+	/**
+	 * @return a name of the group node that this node belongs to.
+	 */
+	const FName& GetGroupName() const
+	{
+		return GroupPtr.Pin()->GetName();
+	}
+
+	virtual bool IsFiltered() const
+	{
+		return false;
+	}
+
+	void SortChildrenAscending(const ITreeNodeSorting* Sorter);
+	void SortChildrenDescending(const ITreeNodeSorting* Sorter);
+
+	/** Adds specified child to the children and sets group for it. */
+	FORCEINLINE_DEBUGGABLE void AddChildAndSetGroupPtr(const FBaseTreeNodePtr& ChildPtr)
+	{
+		ChildPtr->GroupPtr = AsShared();
+		Children.Add(ChildPtr);
+	}
+
+	/** Adds specified child to the filtered children. */
+	FORCEINLINE_DEBUGGABLE void AddFilteredChild(const FBaseTreeNodePtr& ChildPtr)
+	{
+		FilteredChildren.Add(ChildPtr);
+	}
+
+	/** Clears filtered children. */
+	void ClearFilteredChildren()
+	{
+		FilteredChildren.Reset();
+	}
+
+	bool IsExpanded() const { return bIsExpanded; }
+	void SetExpansion(bool bOnOff) { bIsExpanded = bOnOff; }
+
+protected:
+	/**
+	 * @return a reference to the child nodes of this group.
+	 */
+	FORCEINLINE_DEBUGGABLE TArray<FBaseTreeNodePtr>& GetChildrenMutable()
+	{
+		return Children;
+	}
+
+private:
+	/** The id of this node. */
+	const int32 Id;
+
+	/** The name of this node. */
+	const FName Name;
+
+	/** The type id of this node (like a category name). */
+	const FName TypeId;
+
+	/** Children of this node. */
+	TArray<FBaseTreeNodePtr> Children;
+
+	/** Filtered children of this node. */
+	TArray<FBaseTreeNodePtr> FilteredChildren;
+
+	/** A weak pointer to the group/parent of this node. */
+	FBaseTreeNodeWeak GroupPtr;
+
+	/** If this node is a group node or not. */
+	bool bIsGroup;
+
+	/** Whether this group node should be expanded or not. */
+	bool bIsExpanded;
+};
+
+////////////////////////////////////////////////////////////////////////////////////////////////////
+
+} // namespace Insights
diff --git a/Engine/Source/Developer/TraceInsights/Private/Insights/Table/ViewModels/Table.cpp b/Engine/Source/Developer/TraceInsights/Private/Insights/Table/ViewModels/Table.cpp
new file mode 100644
--- /dev/null
+++ b/Engine/Source/Developer/TraceInsights/Private/Insights/Table/ViewModels/Table.cpp
@@ -0,0 +1,306 @@
+// Copyright 1998-2019 Epic Games, Inc. All Rights Reserved.
+
+#include "Table.h"
+#include "TraceServices/Containers/Tables.h"
+
+// Insights
+#include "Insights/Table/ViewModels/TableColumn.h"
+
+#define LOCTEXT_NAMESPACE "Insights_Table"
+
+namespace Insights
+{
+
+////////////////////////////////////////////////////////////////////////////////////////////////////
+
+FTable::FTable()
+	: Name()
+	, Description()
+	, SourceTable()
+	, TableReader()
+	, Columns()
+	, VisibleColumns()
+	, ColumnIdToPtrMapping()
+{
+}
+
+////////////////////////////////////////////////////////////////////////////////////////////////////
+
+FTable::~FTable()
+{
+	Reset();
+}
+
+////////////////////////////////////////////////////////////////////////////////////////////////////
+
+void FTable::Reset()
+{
+	SourceTable.Reset();
+	TableReader.Reset();
+
+	Columns.Reset();
+	VisibleColumns.Reset();
+	ColumnIdToPtrMapping.Reset();
+}
+
+////////////////////////////////////////////////////////////////////////////////////////////////////
+
+void FTable::Init(TSharedPtr<Trace::IUntypedTable> InSourceTable)
+{
+	Reset();
+
+	SourceTable = InSourceTable;
+
+	if (SourceTable)
+	{
+		TableReader = MakeShareable(SourceTable->CreateReader());
+		CreateColumns();
+	}
+}
+
+////////////////////////////////////////////////////////////////////////////////////////////////////
+
+bool AreTableLayoutsEqual(const Trace::ITableLayout& TableLayoutA, const Trace::ITableLayout& TableLayoutB)
+{
+	if (TableLayoutA.GetColumnCount() != TableLayoutB.GetColumnCount())
+	{
+		return false;
+	}
+
+	int32 ColumnCount = TableLayoutA.GetColumnCount();
+	for (int32 ColumnIndex = 0; ColumnIndex < ColumnCount; ++ColumnIndex)
+	{
+		if (TableLayoutA.GetColumnType(ColumnIndex) != TableLayoutA.GetColumnType(ColumnIndex))
+		{
+			return false;
+		}
+		if (FCString::Strcmp(TableLayoutA.GetColumnName(ColumnIndex), TableLayoutA.GetColumnName(ColumnIndex)) != 0)
+		{
+			return false;
+		}
+	}
+
+	return true;
+}
+
+////////////////////////////////////////////////////////////////////////////////////////////////////
+
+void FTable::UpdateSourceTable(TSharedPtr<Trace::IUntypedTable> InSourceTable)
+{
+	check(InSourceTable.IsValid() && SourceTable.IsValid());
+	check(AreTableLayoutsEqual(InSourceTable->GetLayout(), SourceTable->GetLayout()));
+	SourceTable = InSourceTable;
+	TableReader = MakeShareable(SourceTable->CreateReader());
+}
+
+////////////////////////////////////////////////////////////////////////////////////////////////////
+
+void FTable::CreateHierarchyColumn(const TCHAR* ColumnName)
+{
+	const FString ColumnNameStr(ColumnName);
+	const FText ColumnNameText = FText::FromString(ColumnNameStr);
+
+	constexpr EHorizontalAlignment HorizontalAlignment = HAlign_Left;
+	constexpr ETableColumnFlags ColumnFlags = ETableColumnFlags::ShouldBeVisible | ETableColumnFlags::CanBeSorted | ETableColumnFlags::CanBeFiltered | ETableColumnFlags::IsHierarchy;
+
+	constexpr float InitialColumnWidth = 200.0f;
+	constexpr float MinColumnWidth = 0.0f;
+	constexpr float MaxColumnWidth = FLT_MAX;
+
+	FTableColumn::FGetValueAsTextFunction GetValueAsTextFn = [](const FTable& Table, const FTableColumn& Column, const FTableRowId& RowId)->FText
+	{
+		return FText();
+	};
+
+	TSharedPtr<FTableColumn> ColumnPtr = MakeShareable(new FTableColumn
+	(
+		0, // Order
+		0,
+		FName(ColumnName), // Id
+		ColumnNameText, // Short Name
+		ColumnNameText, // Title Name
+		FText::GetEmpty(), // Description
+		ColumnFlags,
+		HorizontalAlignment,
+		InitialColumnWidth,
+		MinColumnWidth,
+		MaxColumnWidth,
+		GetValueAsTextFn
+	));
+
+	AddColumn(ColumnPtr);
+}
+
+////////////////////////////////////////////////////////////////////////////////////////////////////
+
+void FTable::AddColumn(TSharedPtr<FTableColumn> ColumnPtr)
+{
+	ColumnPtr->SetParentTable(SharedThis(this));
+
+	Columns.Add(ColumnPtr);
+
+	if (ColumnPtr->IsVisible())
+	{
+		VisibleColumns.Add(ColumnPtr);
+	}
+
+	ColumnIdToPtrMapping.Add(ColumnPtr->GetId(), ColumnPtr);
+}
+
+////////////////////////////////////////////////////////////////////////////////////////////////////
+
+void FTable::CreateColumns()
+{
+	ensure(TableReader.IsValid());
+	ensure(Columns.Num() == 0);
+
+	const TCHAR* HierarchyColumnName = TEXT("Name");
+
+	CreateHierarchyColumn(HierarchyColumnName);
+
+	const Trace::ITableLayout& TableLayout = SourceTable->GetLayout();
+	const int32 ColumnCount = TableLayout.GetColumnCount();
+
+	for (int32 ColumnIndex = 0; ColumnIndex < ColumnCount; ++ColumnIndex)
+	{
+		Trace::ETableColumnType ColumnType = TableLayout.GetColumnType(ColumnIndex);
+		const TCHAR* ColumnName = TableLayout.GetColumnName(ColumnIndex);
+
+		if (ColumnType == Trace::TableColumnType_CString && FCString::Stricmp(ColumnName, HierarchyColumnName) == 0)
+		{
+			// Skip the "Name" column. It is already created as the name / hierarchy column.
+			continue;
+		}
+
+		const FString ColumnNameStr(ColumnName);
+		const FText ColumnNameText = FText::FromString(ColumnNameStr);
+
+		EHorizontalAlignment HorizontalAlignment = HAlign_Left;
+
+		const ETableColumnFlags ColumnFlags = ETableColumnFlags::ShouldBeVisible | ETableColumnFlags::CanBeSorted | ETableColumnFlags::CanBeFiltered | ETableColumnFlags::CanBeHidden;
+
+		float InitialColumnWidth = 60.0f;
+		float MinColumnWidth = 0.0f;
+		float MaxColumnWidth = FLT_MAX;
+
+		FTableColumn::FGetValueAsTextFunction GetUnknownValueAsTextFn = [](const FTable& Table, const FTableColumn& Column, const FTableRowId& RowId)->FText
+		{
+			return LOCTEXT("UnknownValue", "!?");
+		};
+
+		FTableColumn::FGetValueAsTextFunction GetBoolValueAsTextFn = [](const FTable& Table, const FTableColumn& Column, const FTableRowId& RowId) -> FText
+		{
+			TSharedPtr<Trace::IUntypedTableReader> Reader = Table.GetTableReader();
+			if (Reader.IsValid() && RowId.HasValidIndex())
+			{
+				Reader->SetRowIndex(RowId.RowIndex);
+				bool Value = Reader->GetValueBool(Column.GetIndex());
+				return FText::FromString(Value ? TEXT("True") : TEXT("False"));
+			}
+			return FText::GetEmpty();
+		};
+
+		FTableColumn::FGetValueAsTextFunction GetIntValueAsTextFn = [](const FTable& Table, const FTableColumn& Column, const FTableRowId& RowId) -> FText
+		{
+			TSharedPtr<Trace::IUntypedTableReader> Reader = Table.GetTableReader();
+			if (Reader.IsValid() && RowId.HasValidIndex())
+			{
+				Reader->SetRowIndex(RowId.RowIndex);
+				int64 Value = Reader->GetValueInt(Column.GetIndex());
+				return FText::AsNumber(Value);
+			}
+			return FText::GetEmpty();
+		};
+
+		FTableColumn::FGetValueAsTextFunction GetFloatValueAsTextFn = [](const FTable& Table, const FTableColumn& Column, const FTableRowId& RowId) -> FText
+		{
+			TSharedPtr<Trace::IUntypedTableReader> Reader = Table.GetTableReader();
+			if (Reader.IsValid() && RowId.HasValidIndex())
+			{
+				Reader->SetRowIndex(RowId.RowIndex);
+				float Value = Reader->GetValueFloat(Column.GetIndex());
+				return FText::AsNumber(Value);
+			}
+			return FText::GetEmpty();
+		};
+
+		FTableColumn::FGetValueAsTextFunction GetDoubleValueAsTextFn = [](const FTable& Table, const FTableColumn& Column, const FTableRowId& RowId) -> FText
+		{
+			TSharedPtr<Trace::IUntypedTableReader> Reader = Table.GetTableReader();
+			if (Reader.IsValid() && RowId.HasValidIndex())
+			{
+				Reader->SetRowIndex(RowId.RowIndex);
+				double Value = Reader->GetValueDouble(Column.GetIndex());
+				return FText::AsNumber(Value);
+			}
+			return FText::GetEmpty();
+		};
+
+		FTableColumn::FGetValueAsTextFunction GetStringValueAsTextFn = [](const FTable& Table, const FTableColumn& Column, const FTableRowId& RowId) -> FText
+		{
+			TSharedPtr<Trace::IUntypedTableReader> Reader = Table.GetTableReader();
+			if (Reader.IsValid() && RowId.HasValidIndex())
+			{
+				Reader->SetRowIndex(RowId.RowIndex);
+				const TCHAR* Value = Reader->GetValueCString(Column.GetIndex());
+				return FText::FromString(Value);
+			}
+			return FText::GetEmpty();
+		};
+
+		FTableColumn::FGetValueAsTextFunction GetValueAsTextFn = GetUnknownValueAsTextFn;
+
+		switch (ColumnType)
+		{
+		case Trace::TableColumnType_Bool:
+			HorizontalAlignment = HAlign_Right;
+			InitialColumnWidth = 40.0f;
+			GetValueAsTextFn = GetBoolValueAsTextFn;
+			break;
+		case Trace::TableColumnType_Int:
+			HorizontalAlignment = HAlign_Right;
+			InitialColumnWidth = 60.0f;
+			GetValueAsTextFn = GetIntValueAsTextFn;
+			break;
+		case Trace::TableColumnType_Float:
+			HorizontalAlignment = HAlign_Right;
+			InitialColumnWidth = 60.0f;
+			GetValueAsTextFn = GetFloatValueAsTextFn;
+			break;
+		case Trace::TableColumnType_Double:
+			HorizontalAlignment = HAlign_Right;
+			InitialColumnWidth = 80.0f;
+			GetValueAsTextFn = GetDoubleValueAsTextFn;
+			break;
+		case Trace::TableColumnType_CString:
+			HorizontalAlignment = HAlign_Left;
+			InitialColumnWidth = FMath::Max(120.0f, 6.0f * ColumnNameStr.Len());
+			GetValueAsTextFn = GetStringValueAsTextFn;
+			break;
+		}
+
+		TSharedPtr<FTableColumn> ColumnPtr = MakeShareable(new FTableColumn
+		(
+			ColumnIndex + 1, // Order
+			ColumnIndex,
+			FName(ColumnName), // Id
+			ColumnNameText, // Short Name
+			ColumnNameText, // Title Name
+			FText::GetEmpty(), // Description
+			ColumnFlags,
+			HorizontalAlignment,
+			InitialColumnWidth,
+			MinColumnWidth,
+			MaxColumnWidth,
+			GetValueAsTextFn
+		));
+
+		AddColumn(ColumnPtr);
+	}
+}
+
+////////////////////////////////////////////////////////////////////////////////////////////////////
+
+} // namespace Insights
+
+#undef LOCTEXT_NAMESPACE
diff --git a/Engine/Source/Developer/TraceInsights/Private/Insights/Table/ViewModels/Table.h b/Engine/Source/Developer/TraceInsights/Private/Insights/Table/ViewModels/Table.h
new file mode 100644
--- /dev/null
+++ b/Engine/Source/Developer/TraceInsights/Private/Insights/Table/ViewModels/Table.h
@@ -0,0 +1,101 @@
+// Copyright 1998-2019 Epic Games, Inc. All Rights Reserved.
+
+#pragma once
+
+#include "CoreMinimal.h"
+
+namespace Trace
+{
+	class IUntypedTable;
+	class IUntypedTableReader;
+}
+
+namespace Insights
+{
+
+class FTableColumn;
+
+////////////////////////////////////////////////////////////////////////////////////////////////////
+
+struct FTableRowId
+{
+	static constexpr int32 InvalidRowIndex = -1;
+
+	FTableRowId(int32 InRowIndex) : RowIndex(InRowIndex), Flags(0) {}
+
+	bool HasValidIndex() const { return RowIndex >= 0; }
+
+	union
+	{
+		struct
+		{
+			int32 RowIndex;
+			uint32 Flags;
+		};
+
+		void* Data;
+	};
+};
+
+////////////////////////////////////////////////////////////////////////////////////////////////////
+/**
+ * Table View Model.
+ * View model class for the STableListView and STableTreeView widgets.
+ */
+class FTable : public TSharedFromThis<FTable>
+{
+public:
+	FTable();
+	virtual ~FTable();
+
+	const FName& GetName() const { return Name; }
+	FText GetDisplayName() const { return FText::FromName(Name); }
+	const FText& GetDescription() const { return Description; }
+
+	TSharedPtr<Trace::IUntypedTable> GetSourceTable() const { return SourceTable; }
+	TSharedPtr<Trace::IUntypedTableReader> GetTableReader() const { return TableReader; }
+
+	void Reset();
+	void Init(TSharedPtr<Trace::IUntypedTable> InSourceTable);
+	void UpdateSourceTable(TSharedPtr<Trace::IUntypedTable> InSourceTable);
+
+	bool IsValid() const { return Columns.Num() > 0; }
+
+	const TArray<TSharedPtr<FTableColumn>>& GetColumns() const { return Columns; }
+
+	TSharedPtr<FTableColumn> FindColumnChecked(const FName& ColumnId) const
+	{
+		return ColumnIdToPtrMapping.FindChecked(ColumnId);
+	}
+
+	TSharedPtr<FTableColumn> FindColumn(const FName& ColumnId) const
+	{
+		const TSharedPtr<FTableColumn>* const ColumnPtrPtr = ColumnIdToPtrMapping.Find(ColumnId);
+		return (ColumnPtrPtr != nullptr) ? *ColumnPtrPtr : nullptr;
+	}
+
+private:
+	void CreateHierarchyColumn(const TCHAR* ColumnName);
+	void AddColumn(TSharedPtr<FTableColumn> ColumnPtr);
+	void CreateColumns();
+
+private:
+	FName Name;
+	FText Description;
+
+	TSharedPtr<Trace::IUntypedTable> SourceTable;
+	TSharedPtr<Trace::IUntypedTableReader> TableReader;
+
+	/** All available columns. */
+	TArray<TSharedPtr<FTableColumn>> Columns;
+
+	/** The visible columns. */
+	TArray<TSharedPtr<FTableColumn>> VisibleColumns;
+
+	/** Mapping between column Ids and FTableColumn shared pointers. */
+	TMap<FName, TSharedPtr<FTableColumn>> ColumnIdToPtrMapping;
+};
+
+////////////////////////////////////////////////////////////////////////////////////////////////////
+
+} // namespace Insights
diff --git a/Engine/Source/Developer/TraceInsights/Private/Insights/Table/ViewModels/TableColumn.h b/Engine/Source/Developer/TraceInsights/Private/Insights/Table/ViewModels/TableColumn.h
new file mode 100644
--- /dev/null
+++ b/Engine/Source/Developer/TraceInsights/Private/Insights/Table/ViewModels/TableColumn.h
@@ -0,0 +1,166 @@
+// Copyright 1998-2019 Epic Games, Inc. All Rights Reserved.
+
+#pragma once
+
+#include "CoreMinimal.h"
+#include "Types/SlateEnums.h"
+#include "Misc/EnumClassFlags.h"
+
+namespace Insights
+{
+
+class FTable;
+struct FTableRowId;
+
+////////////////////////////////////////////////////////////////////////////////////////////////////
+
+enum class ETableColumnFlags : uint32
+{
+	None = 0,
+
+	ShouldBeVisible = (1 << 0),
+	CanBeHidden     = (1 << 1),
+	CanBeSorted     = (1 << 2),
+	CanBeFiltered   = (1 << 3),
+	IsHierarchy     = (1 << 4),
+};
+ENUM_CLASS_FLAGS(ETableColumnFlags);
+
+////////////////////////////////////////////////////////////////////////////////////////////////////
+/**
+ * Table Column View Model.
+ * Holds information about a column in STableListView or STableTreeView widgets.
+ */
+class FTableColumn
+{
+public:
+	typedef TFunction<FText(const FTable& Table, const FTableColumn& Column, const FTableRowId& RowId)> FGetValueAsTextFunction;
+
+public:
+	/** No default constructor. */
+	FTableColumn() = delete;
+
+	/** Initialization constructor. */
+	FTableColumn
+	(
+		int32 InOrder,
+		int32 InIndex,
+		const FName InId,
+		FText InShortName,
+		FText InTitleName,
+		FText InDescription,
+		const ETableColumnFlags InFlags,
+		const EHorizontalAlignment InHorizontalAlignment,
+		const float InInitialWidth,
+		const float InMinWidth,
+		const float InMaxWidth,
+		const FGetValueAsTextFunction InGetValueAsTextFn
+	)
+		: Order(InOrder)
+		, Index(InIndex)
+		, Id(InId)
+		, ShortName(MoveTemp(InShortName))
+		, TitleName(MoveTemp(InTitleName))
+		, Description(MoveTemp(InDescription))
+		, Flags(InFlags)
+		, HorizontalAlignment(InHorizontalAlignment)
+		, InitialWidth(InInitialWidth)
+		, MinWidth(InMinWidth)
+		, MaxWidth(InMaxWidth)
+		, GetValueAsTextFn(InGetValueAsTextFn)
+		, ParentTable(nullptr)
+	{
+	}
+
+	int32 GetOrder() const { return Order; }
+	int32 GetIndex() const { return Index; }
+
+	const FName& GetId() const { return Id; }
+
+	const FText& GetShortName() const { return ShortName; }
+	const FText& GetTitleName() const { return TitleName; }
+	const FText& GetDescription() const { return Description; }
+
+	bool IsVisible() const { return bIsVisible; }
+	void Show() { bIsVisible = true; OnVisibilityChanged(); }
+	void Hide() { bIsVisible = false; OnVisibilityChanged(); }
+	void ToggleVisibility() { bIsVisible = !bIsVisible; OnVisibilityChanged(); }
+	void SetVisibilityFlag(bool bOnOff) { bIsVisible = bOnOff; OnVisibilityChanged(); }
+	/*virtual*/ void OnVisibilityChanged() {}
+
+	/** Whether this column should be initially visible. */
+	bool ShouldBeVisible() const { return EnumHasAnyFlags(Flags, ETableColumnFlags::ShouldBeVisible); }
+
+	/** Whether this column can be hidden. */
+	bool CanBeHidden() const { return EnumHasAnyFlags(Flags, ETableColumnFlags::CanBeHidden); }
+
+	/** Whether this column can be used for sorting. */
+	bool CanBeSorted() const { return EnumHasAnyFlags(Flags, ETableColumnFlags::CanBeSorted); }
+
+	/** Whether this column can be used for filtering displayed results. */
+	bool CanBeFiltered() const { return EnumHasAnyFlags(Flags, ETableColumnFlags::CanBeFiltered); }
+
+	/** Whether this column is the hierarcy (name) column, in a tree view. */
+	bool IsHierarchy() const { return EnumHasAnyFlags(Flags, ETableColumnFlags::IsHierarchy); }
+
+	EHorizontalAlignment GetHorizontalAlignment() const { return HorizontalAlignment; }
+
+	float GetInitialWidth() const { return InitialWidth; }
+	float GetMinWidth() const { return MinWidth; }
+	float GetMaxWidth() const { return MaxWidth; }
+	//float GetWidth() const { return Width; }
+
+	/** If MinWidth == MaxWidth, this column has fixed width and cannot be resized. */
+	bool IsFixedWidth() const { return MinWidth == MaxWidth; }
+
+	FText GetValueAsText(const FTableRowId& InRowId) const
+	{
+		return GetValueAsTextFn(*ParentTable.Pin(), *this, InRowId);
+	}
+
+	TWeakPtr<FTable> GetParentTable() const { return ParentTable; }
+	void SetParentTable(TWeakPtr<FTable> InParentTable) { ParentTable = InParentTable; }
+
+private:
+	/** Order value, to sort columns in the list/tree view. */
+	int32 Order;
+
+	/** Column index in source table. */
+	int32 Index;
+
+	/** Id of the column. */
+	FName Id;
+
+	/** Short name of the column, displayed in the column header. */
+	FText ShortName;
+
+	/** Title name of the column, displayed as title in the column tooltip. */
+	FText TitleName;
+
+	/** Long name of the column, displayed in the column tooltip. */
+	FText Description;
+
+	/** Is this column visible? */
+	bool bIsVisible;
+
+	/** Other on/off switches. */
+	ETableColumnFlags Flags;
+
+	/** Horizontal alignment of the content in this column. */
+	EHorizontalAlignment HorizontalAlignment;
+
+	float InitialWidth; /**< Initial column width. */
+	float MinWidth; /**< Minimum column width. */
+	float MaxWidth; /**< Maximum column width. */
+	//float Width; /**< Current column width. */
+
+	/** Custom function used to format (as an FText) the value displayed by this column. */
+	FGetValueAsTextFunction GetValueAsTextFn;
+
+	/* Parent table. Only one table instance can own this column. */
+	TWeakPtr<FTable> ParentTable;
+};
+
+////////////////////////////////////////////////////////////////////////////////////////////////////
+
+} // namespace Insights
diff --git a/Engine/Source/Developer/TraceInsights/Private/Insights/Table/ViewModels/TableTreeNode.cpp b/Engine/Source/Developer/TraceInsights/Private/Insights/Table/ViewModels/TableTreeNode.cpp
new file mode 100644
--- /dev/null
+++ b/Engine/Source/Developer/TraceInsights/Private/Insights/Table/ViewModels/TableTreeNode.cpp
@@ -0,0 +1,114 @@
+// Copyright 1998-2019 Epic Games, Inc. All Rights Reserved.
+
+#include "TableTreeNode.h"
+
+#define LOCTEXT_NAMESPACE "Insights_TableTreeNode"
+
+namespace Insights
+{
+
+////////////////////////////////////////////////////////////////////////////////////////////////////
+
+bool FTableTreeNode::GetValueBool(int32 InColumnIndex) const
+{
+	bool Value = false;
+
+	TSharedPtr<FTable> Table = ParentTable.Pin();
+	if (Table.IsValid())
+	{
+		TSharedPtr<Trace::IUntypedTableReader> Reader = Table->GetTableReader();
+		if (Reader.IsValid() && RowId.HasValidIndex())
+		{
+			Reader->SetRowIndex(RowId.RowIndex);
+			Value = Reader->GetValueBool(InColumnIndex);
+		}
+	}
+
+	return Value;
+}
+
+////////////////////////////////////////////////////////////////////////////////////////////////////
+
+int64 FTableTreeNode::GetValueInt(int32 InColumnIndex) const
+{
+	int64 Value = 0;
+
+	TSharedPtr<FTable> Table = ParentTable.Pin();
+	if (Table.IsValid())
+	{
+		TSharedPtr<Trace::IUntypedTableReader> Reader = Table->GetTableReader();
+		if (Reader.IsValid() && RowId.HasValidIndex())
+		{
+			Reader->SetRowIndex(RowId.RowIndex);
+			Value = Reader->GetValueInt(InColumnIndex);
+		}
+	}
+
+	return Value;
+}
+
+////////////////////////////////////////////////////////////////////////////////////////////////////
+
+float FTableTreeNode::GetValueFloat(int32 InColumnIndex) const
+{
+	float Value = 0.0f;
+
+	TSharedPtr<FTable> Table = ParentTable.Pin();
+	if (Table.IsValid())
+	{
+		TSharedPtr<Trace::IUntypedTableReader> Reader = Table->GetTableReader();
+		if (Reader.IsValid() && RowId.HasValidIndex())
+		{
+			Reader->SetRowIndex(RowId.RowIndex);
+			Value = Reader->GetValueFloat(InColumnIndex);
+		}
+	}
+
+	return Value;
+}
+
+////////////////////////////////////////////////////////////////////////////////////////////////////
+
+double FTableTreeNode::GetValueDouble(int32 InColumnIndex) const
+{
+	double Value = 0.0;
+
+	TSharedPtr<FTable> Table = ParentTable.Pin();
+	if (Table.IsValid())
+	{
+		TSharedPtr<Trace::IUntypedTableReader> Reader = Table->GetTableReader();
+		if (Reader.IsValid() && RowId.HasValidIndex())
+		{
+			Reader->SetRowIndex(RowId.RowIndex);
+			Value = Reader->GetValueDouble(InColumnIndex);
+		}
+	}
+
+	return Value;
+}
+
+////////////////////////////////////////////////////////////////////////////////////////////////////
+
+const TCHAR* FTableTreeNode::GetValueCString(int32 InColumnIndex) const
+{
+	const TCHAR* Value = nullptr;
+
+	TSharedPtr<FTable> Table = ParentTable.Pin();
+	if (Table.IsValid())
+	{
+		TSharedPtr<Trace::IUntypedTableReader> Reader = Table->GetTableReader();
+		if (Reader.IsValid() && RowId.HasValidIndex())
+		{
+			Reader->SetRowIndex(RowId.RowIndex);
+			Value = Reader->GetValueCString(InColumnIndex);
+		}
+	}
+
+	return Value;
+}
+
+////////////////////////////////////////////////////////////////////////////////////////////////////
+
+} // namespace Insights
+
+#undef LOCTEXT_NAMESPACE
diff --git a/Engine/Source/Developer/TraceInsights/Private/Insights/Table/ViewModels/TableTreeNode.h b/Engine/Source/Developer/TraceInsights/Private/Insights/Table/ViewModels/TableTreeNode.h
new file mode 100644
--- /dev/null
+++ b/Engine/Source/Developer/TraceInsights/Private/Insights/Table/ViewModels/TableTreeNode.h
@@ -0,0 +1,72 @@
+// Copyright 1998-2019 Epic Games, Inc. All Rights Reserved.
+
+#pragma once
+
+#include "CoreMinimal.h"
+#include "TraceServices/AnalysisService.h"
+
+// Insights
+#include "Insights/Table/ViewModels/BaseTreeNode.h"
+#include "Insights/Table/ViewModels/Table.h"
+
+namespace Insights
+{
+
+////////////////////////////////////////////////////////////////////////////////////////////////////
+
+class FTableTreeNode;
+
+/** Type definition for shared pointers to instances of FTableTreeNode. */
+typedef TSharedPtr<class FTableTreeNode> FTableTreeNodePtr;
+
+/** Type definition for shared references to instances of FTableTreeNode. */
+typedef TSharedRef<class FTableTreeNode> FTableTreeNodeRef;
+
+/** Type definition for shared references to const instances of FTableTreeNode. */
+typedef TSharedRef<const class FTableTreeNode> FTableTreeNodeRefConst;
+
+/** Type definition for weak references to instances of FTableTreeNode. */
+typedef TWeakPtr<class FTableTreeNode> FTableTreeNodeWeak;
+
+////////////////////////////////////////////////////////////////////////////////////////////////////
+/**
+ * Table Tree Node View Model.
+ * Class used to store information about a generic table tree node (used in STableTreeView).
+ */
+class FTableTreeNode : public FBaseTreeNode
+{
+public:
+	/** Initialization constructor for a table record node. */
+	FTableTreeNode(uint64 InId, const FName InName, TWeakPtr<FTable> InParentTable, int32 InRowIndex)
+		: FBaseTreeNode(InId, InName, false)
+		, ParentTable(InParentTable)
+		, RowId(InRowIndex)
+	{
+	}
+
+	/** Initialization constructor for a group node. */
+	FTableTreeNode(const FName InGroupName, TWeakPtr<FTable> InParentTable)
+		: FBaseTreeNode(0, InGroupName, true)
+		, ParentTable(InParentTable)
+		, RowId(FTableRowId::InvalidRowIndex)
+	{
+	}
+
+	TWeakPtr<FTable> GetParentTable() { return ParentTable; }
+	FTableRowId GetRowId() const { return RowId; }
+	int32 GetRowIndex() const { return RowId.RowIndex; }
+
+	bool GetValueBool(int32 InColumnIndex) const;
+	int64 GetValueInt(int32 InColumnIndex) const;
+	float GetValueFloat(int32 InColumnIndex) const;
+	double GetValueDouble(int32 InColumnIndex) const;
+	const TCHAR* GetValueCString(int32 InColumnIndex) const;
+
+protected:
+	TWeakPtr<FTable> ParentTable;
+	FTableRowId RowId;
+};
+
+////////////////////////////////////////////////////////////////////////////////////////////////////
+
+} // namespace Insights
diff --git a/Engine/Source/Developer/TraceInsights/Private/Insights/Table/ViewModels/TreeNodeGrouping.cpp b/Engine/Source/Developer/TraceInsights/Private/Insights/Table/ViewModels/TreeNodeGrouping.cpp
new file mode 100644
--- /dev/null
+++ b/Engine/Source/Developer/TraceInsights/Private/Insights/Table/ViewModels/TreeNodeGrouping.cpp
@@ -0,0 +1,80 @@
+// Copyright 1998-2019 Epic Games, Inc. All Rights Reserved.
+
+#include "TreeNodeGrouping.h"
+
+#include "EditorStyleSet.h"
+
+#define LOCTEXT_NAMESPACE "Insights_TreeNode"
+
+namespace Insights
+{
+
+////////////////////////////////////////////////////////////////////////////////////////////////////
+
+FTreeNodeGrouping::FTreeNodeGrouping(FText InName, FText InDescription, FName InBrushName,  FSlateBrush* InIcon)
+	: Name(InName)
+	, Description(InDescription)
+	, BrushName(InBrushName)
+	, Icon(InIcon)
+{
+}
+
+////////////////////////////////////////////////////////////////////////////////////////////////////
+
+FTreeNodeGroupingFlat::FTreeNodeGroupingFlat()
+	: FTreeNodeGrouping(
+		LOCTEXT("Grouping_Name_Flat", "Flat"),
+		LOCTEXT("Grouping_Desc_Flat", "Creates a single group. Includes all items."),
+		TEXT("Profiler.FiltersAndPresets.GroupNameIcon"), //TODO: "Icons.Grouping.Flat"
+		nullptr)
+{
+}
+
+////////////////////////////////////////////////////////////////////////////////////////////////////
+
+FTreeNodeGroupInfo FTreeNodeGroupingFlat::GetGroupForNode(const FBaseTreeNodePtr InNode) const
+{
+	return { FName(TEXT("All")), true };
+}
+
+////////////////////////////////////////////////////////////////////////////////////////////////////
+
+FTreeNodeGroupingByNameFirstLetter::FTreeNodeGroupingByNameFirstLetter()
+	: FTreeNodeGrouping(
+		LOCTEXT("Grouping_ByName_Name", "Name"),
+		LOCTEXT("Grouping_ByName_Desc", "Creates a group for each first letter of node names."),
+		TEXT("Profiler.FiltersAndPresets.Group1NameIcon"), //TODO: "Icons.Grouping.ByName"
+		nullptr)
+{
+}
+
+////////////////////////////////////////////////////////////////////////////////////////////////////
+
+FTreeNodeGroupInfo FTreeNodeGroupingByNameFirstLetter::GetGroupForNode(const FBaseTreeNodePtr InNode) const
+{
+	return { *InNode->GetName().GetPlainNameString().Left(1), false };
+}
+
+////////////////////////////////////////////////////////////////////////////////////////////////////
+
+FTreeNodeGroupingByType::FTreeNodeGroupingByType()
+	: FTreeNodeGrouping(
+		LOCTEXT("Grouping_ByType_Name", "Type"),
+		LOCTEXT("Grouping_ByType_Desc", "Creates a group for each type."),
+		TEXT("Profiler.FiltersAndPresets.StatTypeIcon"), //TODO
+		nullptr)
+{
+}
+
+////////////////////////////////////////////////////////////////////////////////////////////////////
+
+FTreeNodeGroupInfo FTreeNodeGroupingByType::GetGroupForNode(const FBaseTreeNodePtr InNode) const
+{
+	return { InNode->GetTypeId(), true };
+}
+
+////////////////////////////////////////////////////////////////////////////////////////////////////
+
+} // namespace Insights
+
+#undef LOCTEXT_NAMESPACE
diff --git a/Engine/Source/Developer/TraceInsights/Private/Insights/Table/ViewModels/TreeNodeGrouping.h b/Engine/Source/Developer/TraceInsights/Private/Insights/Table/ViewModels/TreeNodeGrouping.h
new file mode 100644
--- /dev/null
+++ b/Engine/Source/Developer/TraceInsights/Private/Insights/Table/ViewModels/TreeNodeGrouping.h
@@ -0,0 +1,96 @@
+// Copyright 1998-2019 Epic Games, Inc. All Rights Reserved.
+
+#pragma once
+
+#include "CoreMinimal.h"
+#include "Styling/SlateBrush.h"
+#include "Delegates/DelegateCombinations.h"
+
+// Insights
+#include "Insights/Table/ViewModels/BaseTreeNode.h"
+
+namespace Insights
+{
+
+////////////////////////////////////////////////////////////////////////////////////////////////////
+
+class ITreeNodeGrouping
+{
+public:
+	virtual FText GetName() const = 0;
+	virtual FText GetDescription() const = 0;
+	virtual FName GetBrushName() const = 0;
+	virtual const FSlateBrush* GetIcon() const = 0;
+
+	//virtual void CreateGrouping() = 0;
+};
+
+////////////////////////////////////////////////////////////////////////////////////////////////////
+
+struct FTreeNodeGroupInfo
+{
+	FName Name;
+	bool IsExpanded;
+};
+
+////////////////////////////////////////////////////////////////////////////////////////////////////
+
+class FTreeNodeGrouping : public ITreeNodeGrouping
+{
+public:
+	FTreeNodeGrouping(FText InName, FText InDescription, FName InBrushName,  FSlateBrush* InIcon);
+	virtual ~FTreeNodeGrouping() {}
+
+	virtual FText GetName() const override { return Name; }
+	virtual FText GetDescription() const override { return Description; }
+	virtual FName GetBrushName() const override { return BrushName; }
+	virtual const FSlateBrush* GetIcon() const override { return Icon; }
+
+	virtual FTreeNodeGroupInfo GetGroupForNode(const FBaseTreeNodePtr InNode) const { return { FName(), false }; }
+
+protected:
+	FText Name;
+	FText Description;
+	FName BrushName;
+	const FSlateBrush* Icon;
+};
+
+////////////////////////////////////////////////////////////////////////////////////////////////////
+
+/** Creates a single group for all nodes. */
+class FTreeNodeGroupingFlat : public FTreeNodeGrouping
+{
+public:
+	FTreeNodeGroupingFlat();
+	virtual ~FTreeNodeGroupingFlat() {}
+
+	virtual FTreeNodeGroupInfo GetGroupForNode(const FBaseTreeNodePtr InNode) const override;
+};
+
+////////////////////////////////////////////////////////////////////////////////////////////////////
+
+/** Creates a group for each first letter of node names. */
+class FTreeNodeGroupingByNameFirstLetter : public FTreeNodeGrouping
+{
+public:
+	FTreeNodeGroupingByNameFirstLetter();
+	virtual ~FTreeNodeGroupingByNameFirstLetter() {}
+
+	virtual FTreeNodeGroupInfo GetGroupForNode(const FBaseTreeNodePtr InNode) const override;
+};
+
+////////////////////////////////////////////////////////////////////////////////////////////////////
+
+/** Creates a group for each type. */
+class FTreeNodeGroupingByType : public FTreeNodeGrouping
+{
+public:
+	FTreeNodeGroupingByType();
+	virtual ~FTreeNodeGroupingByType() {}
+
+	virtual FTreeNodeGroupInfo GetGroupForNode(const FBaseTreeNodePtr InNode) const override;
+};
+
+////////////////////////////////////////////////////////////////////////////////////////////////////
+
+} // namespace Insights
