[Insights]

diff --git a/Engine/Source/Developer/TraceInsights/Private/Insights/Table/ViewModels/TreeNodeSorting.cpp b/Engine/Source/Developer/TraceInsights/Private/Insights/Table/ViewModels/TreeNodeSorting.cpp
new file mode 100644
--- /dev/null
+++ b/Engine/Source/Developer/TraceInsights/Private/Insights/Table/ViewModels/TreeNodeSorting.cpp
@@ -0,0 +1,362 @@
+// Copyright 1998-2019 Epic Games, Inc. All Rights Reserved.
+
+#include "TreeNodeSorting.h"
+
+// Insights
+#include "Insights/Table/ViewModels/TableColumn.h"
+#include "Insights/Table/ViewModels/TableTreeNode.h"
+
+#define LOCTEXT_NAMESPACE "Insights_TreeNode"
+
+namespace Insights
+{
+
+////////////////////////////////////////////////////////////////////////////////////////////////////
+// FTreeNodeSorting
+////////////////////////////////////////////////////////////////////////////////////////////////////
+
+FTreeNodeSorting::FTreeNodeSorting(FName InName, FText InShortName, FText InTitleName, FText InDescription, FName InColumnId)
+	: Name(InName)
+	, ShortName(InShortName)
+	, TitleName(InTitleName)
+	, Description(InDescription)
+	, ColumnId(InColumnId)
+	, AscendingIcon(nullptr)
+	, DescendingIcon(nullptr)
+	, AscendingCompareDelegate(nullptr)
+	, DescendingCompareDelegate(nullptr)
+{
+}
+
+////////////////////////////////////////////////////////////////////////////////////////////////////
+
+void FTreeNodeSorting::SortAscending(TArray<FBaseTreeNodePtr>& NodesToSort) const
+{
+	if (AscendingCompareDelegate)
+	{
+		NodesToSort.Sort(AscendingCompareDelegate);
+	}
+}
+
+////////////////////////////////////////////////////////////////////////////////////////////////////
+
+void FTreeNodeSorting::SortDescending(TArray<FBaseTreeNodePtr>& NodesToSort) const
+{
+	if (DescendingCompareDelegate)
+	{
+		NodesToSort.Sort(DescendingCompareDelegate);
+	}
+}
+
+////////////////////////////////////////////////////////////////////////////////////////////////////
+// FTreeNodeSortingByName
+////////////////////////////////////////////////////////////////////////////////////////////////////
+
+FTreeNodeSortingByName::FTreeNodeSortingByName()
+	: FTreeNodeSorting(
+		FName(TEXT("ByName")),
+		LOCTEXT("Sorting_ByName_Name", "By Name"),
+		LOCTEXT("Sorting_ByName_Title", "Sort By Name"),
+		LOCTEXT("Sorting_ByName_Desc", "Sort by name."),
+		FName(TEXT("Name")))
+{
+	AscendingCompareDelegate = [](const FBaseTreeNodePtr& A, const FBaseTreeNodePtr& B) -> bool
+	{
+		// Sort by name (ascending).
+		return A->GetName().LexicalLess(B->GetName());
+	};
+
+	DescendingCompareDelegate = [](const FBaseTreeNodePtr& A, const FBaseTreeNodePtr& B) -> bool
+	{
+		// Sort by name (descending).
+		return B->GetName().LexicalLess(A->GetName());
+	};
+}
+
+////////////////////////////////////////////////////////////////////////////////////////////////////
+// FTreeNodeSortingByType
+////////////////////////////////////////////////////////////////////////////////////////////////////
+
+FTreeNodeSortingByType::FTreeNodeSortingByType()
+	: FTreeNodeSorting(
+		FName(TEXT("ByType")),
+		LOCTEXT("Sorting_ByName_Name", "By Type"),
+		LOCTEXT("Sorting_ByName_Title", "Sort By Type"),
+		LOCTEXT("Sorting_ByName_Desc", "Sort by type (ascending or descending), then by name (ascending)."),
+		FName(TEXT("Type")))
+{
+	AscendingCompareDelegate = [](const FBaseTreeNodePtr& A, const FBaseTreeNodePtr& B) -> bool
+	{
+		const FName& TypeA = A->GetTypeId();
+		const FName& TypeB = B->GetTypeId();
+
+		if (TypeA == TypeB)
+		{
+			// Sort by name (ascending).
+			return A->GetName().LexicalLess(B->GetName());
+		}
+		else
+		{
+			// Sort by type (ascending).
+			return TypeA.FastLess(TypeB); // TypeA < TypeB;
+		}
+	};
+
+	DescendingCompareDelegate = [](const FBaseTreeNodePtr& A, const FBaseTreeNodePtr& B) -> bool
+	{
+		const FName& TypeA = A->GetTypeId();
+		const FName& TypeB = B->GetTypeId();
+
+		if (TypeA == TypeB)
+		{
+			// Sort by name (ascending).
+			return A->GetName().LexicalLess(B->GetName());
+		}
+		else
+		{
+			// Sort by type (descending).
+			return TypeB.FastLess(TypeA); // TypeA > TypeB;
+		}
+	};
+}
+
+////////////////////////////////////////////////////////////////////////////////////////////////////
+// FTableTreeNodeSorting
+////////////////////////////////////////////////////////////////////////////////////////////////////
+
+FTableTreeNodeSorting::FTableTreeNodeSorting(TSharedPtr<FTableColumn> InColumnPtr)
+	: FTreeNodeSorting(
+		InColumnPtr->GetId(),
+		FText::Format(LOCTEXT("Sorting_TableValue_Name", "By {0}"), InColumnPtr->GetShortName()),
+		FText::Format(LOCTEXT("Sorting_TableValue_Title", "Sort By {0}"), InColumnPtr->GetTitleName()),
+		FText::Format(LOCTEXT("Sorting_TableValue_Desc", "Sort by {0} (ascending or descending), then by name (ascending)."), InColumnPtr->GetShortName()),
+		InColumnPtr->GetId())
+	, ColumnPtr(InColumnPtr)
+{
+}
+
+////////////////////////////////////////////////////////////////////////////////////////////////////
+// FTreeNodeSortingByBoolValue
+////////////////////////////////////////////////////////////////////////////////////////////////////
+
+FTreeNodeSortingByBoolValue::FTreeNodeSortingByBoolValue(TSharedPtr<FTableColumn> InColumnPtr)
+	: FTableTreeNodeSorting(InColumnPtr)
+{
+	const int32 ColumnIndex = ColumnPtr->GetIndex();
+
+	AscendingCompareDelegate = [ColumnIndex](const FBaseTreeNodePtr& A, const FBaseTreeNodePtr& B) -> bool
+	{
+		const bool ValueA = StaticCastSharedPtr<FTableTreeNode>(A)->GetValueBool(ColumnIndex);
+		const bool ValueB = StaticCastSharedPtr<FTableTreeNode>(B)->GetValueBool(ColumnIndex);
+
+		if (ValueA == ValueB)
+		{
+			// Sort by name (ascending).
+			return A->GetName().LexicalLess(B->GetName());
+		}
+		else
+		{
+			// Sort by value (ascending).
+			return ValueA < ValueB;
+		}
+	};
+
+	DescendingCompareDelegate = [ColumnIndex](const FBaseTreeNodePtr& A, const FBaseTreeNodePtr& B) -> bool
+	{
+		const bool ValueA = StaticCastSharedPtr<FTableTreeNode>(A)->GetValueBool(ColumnIndex);
+		const bool ValueB = StaticCastSharedPtr<FTableTreeNode>(B)->GetValueBool(ColumnIndex);
+
+		if (ValueA == ValueB)
+		{
+			// Sort by name (ascending).
+			return A->GetName().LexicalLess(B->GetName());
+		}
+		else
+		{
+			// Sort by value (descending).
+			return ValueB < ValueA;
+		}
+	};
+}
+
+////////////////////////////////////////////////////////////////////////////////////////////////////
+// FTreeNodeSortingByIntValue
+////////////////////////////////////////////////////////////////////////////////////////////////////
+
+FTreeNodeSortingByIntValue::FTreeNodeSortingByIntValue(TSharedPtr<FTableColumn> InColumnPtr)
+	: FTableTreeNodeSorting(InColumnPtr)
+{
+	const int32 ColumnIndex = ColumnPtr->GetIndex();
+
+	AscendingCompareDelegate = [ColumnIndex](const FBaseTreeNodePtr& A, const FBaseTreeNodePtr& B) -> bool
+	{
+		const int64 ValueA = StaticCastSharedPtr<FTableTreeNode>(A)->GetValueInt(ColumnIndex);
+		const int64 ValueB = StaticCastSharedPtr<FTableTreeNode>(B)->GetValueInt(ColumnIndex);
+
+		if (ValueA == ValueB)
+		{
+			// Sort by name (ascending).
+			return A->GetName().LexicalLess(B->GetName());
+		}
+		else
+		{
+			// Sort by value (ascending).
+			return ValueA < ValueB;
+		}
+	};
+
+	DescendingCompareDelegate = [ColumnIndex](const FBaseTreeNodePtr& A, const FBaseTreeNodePtr& B) -> bool
+	{
+		const int64 ValueA = StaticCastSharedPtr<FTableTreeNode>(A)->GetValueInt(ColumnIndex);
+		const int64 ValueB = StaticCastSharedPtr<FTableTreeNode>(B)->GetValueInt(ColumnIndex);
+
+		if (ValueA == ValueB)
+		{
+			// Sort by name (ascending).
+			return A->GetName().LexicalLess(B->GetName());
+		}
+		else
+		{
+			// Sort by value (descending).
+			return ValueB < ValueA;
+		}
+	};
+}
+
+////////////////////////////////////////////////////////////////////////////////////////////////////
+// FTreeNodeSortingByFloatValue
+////////////////////////////////////////////////////////////////////////////////////////////////////
+
+FTreeNodeSortingByFloatValue::FTreeNodeSortingByFloatValue(TSharedPtr<FTableColumn> InColumnPtr)
+	: FTableTreeNodeSorting(InColumnPtr)
+{
+	const int32 ColumnIndex = ColumnPtr->GetIndex();
+
+	AscendingCompareDelegate = [ColumnIndex](const FBaseTreeNodePtr& A, const FBaseTreeNodePtr& B) -> bool
+	{
+		const float ValueA = StaticCastSharedPtr<FTableTreeNode>(A)->GetValueFloat(ColumnIndex);
+		const float ValueB = StaticCastSharedPtr<FTableTreeNode>(B)->GetValueFloat(ColumnIndex);
+
+		if (ValueA == ValueB)
+		{
+			// Sort by name (ascending).
+			return A->GetName().LexicalLess(B->GetName());
+		}
+		else
+		{
+			// Sort by value (ascending).
+			return ValueA < ValueB;
+		}
+	};
+
+	DescendingCompareDelegate = [ColumnIndex](const FBaseTreeNodePtr& A, const FBaseTreeNodePtr& B) -> bool
+	{
+		const float ValueA = StaticCastSharedPtr<FTableTreeNode>(A)->GetValueFloat(ColumnIndex);
+		const float ValueB = StaticCastSharedPtr<FTableTreeNode>(B)->GetValueFloat(ColumnIndex);
+
+		if (ValueA == ValueB)
+		{
+			// Sort by name (ascending).
+			return A->GetName().LexicalLess(B->GetName());
+		}
+		else
+		{
+			// Sort by value (descending).
+			return ValueB < ValueA;
+		}
+	};
+}
+
+////////////////////////////////////////////////////////////////////////////////////////////////////
+// FTreeNodeSortingByDoubleValue
+////////////////////////////////////////////////////////////////////////////////////////////////////
+
+FTreeNodeSortingByDoubleValue::FTreeNodeSortingByDoubleValue(TSharedPtr<FTableColumn> InColumnPtr)
+	: FTableTreeNodeSorting(InColumnPtr)
+{
+	const int32 ColumnIndex = ColumnPtr->GetIndex();
+
+	AscendingCompareDelegate = [ColumnIndex](const FBaseTreeNodePtr& A, const FBaseTreeNodePtr& B) -> bool
+	{
+		const double ValueA = StaticCastSharedPtr<FTableTreeNode>(A)->GetValueDouble(ColumnIndex);
+		const double ValueB = StaticCastSharedPtr<FTableTreeNode>(B)->GetValueDouble(ColumnIndex);
+
+		if (ValueA == ValueB)
+		{
+			// Sort by name (ascending).
+			return A->GetName().LexicalLess(B->GetName());
+		}
+		else
+		{
+			// Sort by value (ascending).
+			return ValueA < ValueB;
+		}
+	};
+
+	DescendingCompareDelegate = [ColumnIndex](const FBaseTreeNodePtr& A, const FBaseTreeNodePtr& B) -> bool
+	{
+		const double ValueA = StaticCastSharedPtr<FTableTreeNode>(A)->GetValueDouble(ColumnIndex);
+		const double ValueB = StaticCastSharedPtr<FTableTreeNode>(B)->GetValueDouble(ColumnIndex);
+
+		if (ValueA == ValueB)
+		{
+			// Sort by name (ascending).
+			return A->GetName().LexicalLess(B->GetName());
+		}
+		else
+		{
+			// Sort by value (descending).
+			return ValueB < ValueA;
+		}
+	};
+}
+
+////////////////////////////////////////////////////////////////////////////////////////////////////
+// FTreeNodeSortingByCStringValue
+////////////////////////////////////////////////////////////////////////////////////////////////////
+
+FTreeNodeSortingByCStringValue::FTreeNodeSortingByCStringValue(TSharedPtr<FTableColumn> InColumnPtr)
+	: FTableTreeNodeSorting(InColumnPtr)
+{
+	const int32 ColumnIndex = ColumnPtr->GetIndex();
+
+	AscendingCompareDelegate = [ColumnIndex](const FBaseTreeNodePtr& A, const FBaseTreeNodePtr& B) -> bool
+	{
+		const TCHAR* ValueA = StaticCastSharedPtr<FTableTreeNode>(A)->GetValueCString(ColumnIndex);
+		const TCHAR* ValueB = StaticCastSharedPtr<FTableTreeNode>(B)->GetValueCString(ColumnIndex);
+
+		if (ValueA == ValueB)
+		{
+			// Sort by name (ascending).
+			return A->GetName().LexicalLess(B->GetName());
+		}
+		else
+		{
+			// Sort by value (ascending).
+			return FCString::Strcmp(ValueA, ValueB);
+		}
+	};
+
+	DescendingCompareDelegate = [ColumnIndex](const FBaseTreeNodePtr& A, const FBaseTreeNodePtr& B) -> bool
+	{
+		const TCHAR* ValueA = StaticCastSharedPtr<FTableTreeNode>(A)->GetValueCString(ColumnIndex);
+		const TCHAR* ValueB = StaticCastSharedPtr<FTableTreeNode>(B)->GetValueCString(ColumnIndex);
+
+		if (ValueA == ValueB)
+		{
+			// Sort by name (ascending).
+			return A->GetName().LexicalLess(B->GetName());
+		}
+		else
+		{
+			// Sort by value (descending).
+			return FCString::Strcmp(ValueB, ValueA);
+		}
+	};
+}
+
+////////////////////////////////////////////////////////////////////////////////////////////////////
+
+} // namespace Insights
+
+#undef LOCTEXT_NAMESPACE
diff --git a/Engine/Source/Developer/TraceInsights/Private/Insights/Table/ViewModels/TreeNodeSorting.h b/Engine/Source/Developer/TraceInsights/Private/Insights/Table/ViewModels/TreeNodeSorting.h
new file mode 100644
--- /dev/null
+++ b/Engine/Source/Developer/TraceInsights/Private/Insights/Table/ViewModels/TreeNodeSorting.h
@@ -0,0 +1,158 @@
+// Copyright 1998-2019 Epic Games, Inc. All Rights Reserved.
+
+#pragma once
+
+#include "CoreMinimal.h"
+#include "Templates/Function.h"
+#include "Styling/SlateBrush.h"
+
+// Insights
+#include "Insights/Table/ViewModels/BaseTreeNode.h"
+
+namespace Insights
+{
+
+class FTableColumn;
+class FTableTreeNode;
+
+////////////////////////////////////////////////////////////////////////////////////////////////////
+
+typedef TFunction<bool(const FBaseTreeNodePtr& A, const FBaseTreeNodePtr& B)> FTreeNodeSortingCompareFn;
+
+class ITreeNodeSorting
+{
+public:
+	virtual FName GetName() const = 0;
+	virtual FText GetShortName() const = 0;
+	virtual FText GetTitleName() const = 0;
+	virtual FText GetDescription() const = 0;
+	virtual FName GetColumnId() const = 0;
+
+	virtual FSlateBrush* GetAscendingIcon() const = 0;
+	virtual FSlateBrush* GetDescendingIcon() const = 0;
+
+	virtual FTreeNodeSortingCompareFn GetAscendingCompareDelegate() const = 0;
+	virtual FTreeNodeSortingCompareFn GetDescendingCompareDelegate() const = 0;
+
+	virtual void SortAscending(TArray<FBaseTreeNodePtr>& NodesToSort) const = 0;
+	virtual void SortDescending(TArray<FBaseTreeNodePtr>& NodesToSort) const = 0;
+};
+
+////////////////////////////////////////////////////////////////////////////////////////////////////
+
+class FTreeNodeSorting : public ITreeNodeSorting
+{
+public:
+	typedef bool FCompareFn(const FBaseTreeNodePtr& A, const FBaseTreeNodePtr& B) const;
+
+public:
+	FTreeNodeSorting(FName InName, FText InShortName, FText InTitleName, FText InDescription, FName InColumnId);
+	virtual ~FTreeNodeSorting() {}
+
+	virtual FName GetName() const override { return Name; }
+	virtual FText GetShortName() const override { return ShortName; }
+	virtual FText GetTitleName() const override { return TitleName; }
+	virtual FText GetDescription() const override { return Description; }
+	virtual FName GetColumnId() const override { return ColumnId; }
+
+	virtual FSlateBrush* GetAscendingIcon() const override { return AscendingIcon; }
+	virtual FSlateBrush* GetDescendingIcon() const override { return DescendingIcon; }
+
+	virtual FTreeNodeSortingCompareFn GetAscendingCompareDelegate() const override { return AscendingCompareDelegate; }
+	virtual FTreeNodeSortingCompareFn GetDescendingCompareDelegate() const override { return DescendingCompareDelegate; }
+
+	virtual void SortAscending(TArray<FBaseTreeNodePtr>& NodesToSort) const override;
+	virtual void SortDescending(TArray<FBaseTreeNodePtr>& NodesToSort) const override;
+
+protected:
+	FName Name;
+	FText ShortName;
+	FText TitleName;
+	FText Description;
+	FName ColumnId;
+	
+	FSlateBrush* AscendingIcon;
+	FSlateBrush* DescendingIcon;
+
+	FTreeNodeSortingCompareFn AscendingCompareDelegate;
+	FTreeNodeSortingCompareFn DescendingCompareDelegate;
+};
+
+////////////////////////////////////////////////////////////////////////////////////////////////////
+
+class FTreeNodeSortingByName : public FTreeNodeSorting
+{
+public:
+	FTreeNodeSortingByName();
+	virtual ~FTreeNodeSortingByName() {}
+};
+
+////////////////////////////////////////////////////////////////////////////////////////////////////
+
+class FTreeNodeSortingByType : public FTreeNodeSorting
+{
+public:
+	FTreeNodeSortingByType();
+	virtual ~FTreeNodeSortingByType() {}
+};
+
+////////////////////////////////////////////////////////////////////////////////////////////////////
+
+class FTableTreeNodeSorting : public FTreeNodeSorting
+{
+public:
+	FTableTreeNodeSorting(TSharedPtr<FTableColumn> InColumnPtr);
+	virtual ~FTableTreeNodeSorting() {}
+
+protected:
+	TSharedPtr<FTableColumn> ColumnPtr;
+};
+
+////////////////////////////////////////////////////////////////////////////////////////////////////
+
+class FTreeNodeSortingByBoolValue : public FTableTreeNodeSorting
+{
+public:
+	FTreeNodeSortingByBoolValue(TSharedPtr<FTableColumn> InColumnPtr);
+	virtual ~FTreeNodeSortingByBoolValue() {}
+};
+
+////////////////////////////////////////////////////////////////////////////////////////////////////
+
+class FTreeNodeSortingByIntValue : public FTableTreeNodeSorting
+{
+public:
+	FTreeNodeSortingByIntValue(TSharedPtr<FTableColumn> InColumnPtr);
+	virtual ~FTreeNodeSortingByIntValue() {}
+};
+
+////////////////////////////////////////////////////////////////////////////////////////////////////
+
+class FTreeNodeSortingByFloatValue : public FTableTreeNodeSorting
+{
+public:
+	FTreeNodeSortingByFloatValue(TSharedPtr<FTableColumn> InColumnPtr);
+	virtual ~FTreeNodeSortingByFloatValue() {}
+};
+
+////////////////////////////////////////////////////////////////////////////////////////////////////
+
+class FTreeNodeSortingByDoubleValue : public FTableTreeNodeSorting
+{
+public:
+	FTreeNodeSortingByDoubleValue(TSharedPtr<FTableColumn> InColumnPtr);
+	virtual ~FTreeNodeSortingByDoubleValue() {}
+};
+
+////////////////////////////////////////////////////////////////////////////////////////////////////
+
+class FTreeNodeSortingByCStringValue : public FTableTreeNodeSorting
+{
+public:
+	FTreeNodeSortingByCStringValue(TSharedPtr<FTableColumn> InColumnPtr);
+	virtual ~FTreeNodeSortingByCStringValue() {}
+};
+
+////////////////////////////////////////////////////////////////////////////////////////////////////
+
+} // namespace Insights
diff --git a/Engine/Source/Developer/TraceInsights/Private/Insights/Table/Widgets/STableTreeView.cpp b/Engine/Source/Developer/TraceInsights/Private/Insights/Table/Widgets/STableTreeView.cpp
new file mode 100644
--- /dev/null
+++ b/Engine/Source/Developer/TraceInsights/Private/Insights/Table/Widgets/STableTreeView.cpp
@@ -0,0 +1,1412 @@
+// Copyright 1998-2019 Epic Games, Inc. All Rights Reserved.
+
+#include "STableTreeView.h"
+
+#include "EditorStyleSet.h"
+#include "Framework/MultiBox/MultiBoxBuilder.h"
+#include "SlateOptMacros.h"
+#include "TraceServices/AnalysisService.h"
+#include "Widgets/Input/SCheckBox.h"
+#include "Widgets/Layout/SScrollBox.h"
+#include "Widgets/Input/SSearchBox.h"
+#include "Widgets/Images/SImage.h"
+#include "Widgets/Layout/SGridPanel.h"
+#include "Widgets/Layout/SSeparator.h"
+#include "Widgets/SToolTip.h"
+#include "Widgets/Views/STableViewBase.h"
+
+// Insights
+#include "Insights/InsightsManager.h"
+#include "Insights/Table/ViewModels/Table.h"
+#include "Insights/Table/ViewModels/TableColumn.h"
+#include "Insights/Table/ViewModels/TableTreeNode.h"
+#include "Insights/Table/ViewModels/TreeNodeGrouping.h"
+#include "Insights/Table/ViewModels/TreeNodeSorting.h"
+#include "Insights/Table/Widgets/STableTreeViewTooltip.h"
+#include "Insights/Table/Widgets/STableTreeViewRow.h"
+
+#define LOCTEXT_NAMESPACE "STableTreeView"
+
+namespace Insights
+{
+
+////////////////////////////////////////////////////////////////////////////////////////////////////
+
+const FName STableTreeView::DefaultColumnBeingSorted;
+const EColumnSortMode::Type STableTreeView::DefaultColumnSortMode(EColumnSortMode::Descending);
+
+////////////////////////////////////////////////////////////////////////////////////////////////////
+
+STableTreeView::STableTreeView()
+	: Table()
+	, Session(FInsightsManager::Get()->GetSession())
+	, TreeView(nullptr)
+	, TreeViewHeaderColumnArgs()
+	, TreeViewHeaderRow(nullptr)
+	, ExternalScrollbar(nullptr)
+	, HoveredColumnId()
+	, HoveredNodePtr(nullptr)
+	, HighlightedNodeName()
+	, TableTreeNodes()
+	, GroupNodes()
+	, FilteredGroupNodes()
+	, TableTreeNodesIdMap()
+	, ExpandedNodes()
+	, bExpansionSaved(false)
+	, SearchBox(nullptr)
+	, TextFilter(nullptr)
+	, Filters(nullptr)
+	, AvailableGroupings()
+	, GroupByComboBox(nullptr)
+	, CurrentGrouping(nullptr)
+	, AvailableSortings()
+	, ColumnToSortingMap()
+	, CurrentSorting(nullptr)
+	, ColumnBeingSorted(DefaultColumnBeingSorted)
+	, ColumnSortMode(DefaultColumnSortMode)
+	, StatsStartTime(0.0)
+	, StatsEndTime(0.0)
+{
+}
+
+////////////////////////////////////////////////////////////////////////////////////////////////////
+
+STableTreeView::~STableTreeView()
+{
+	// Remove ourselves from the Insights manager.
+	if (FInsightsManager::Get().IsValid())
+	{
+		FInsightsManager::Get()->GetSessionChangedEvent().RemoveAll(this);
+	}
+}
+
+////////////////////////////////////////////////////////////////////////////////////////////////////
+
+BEGIN_SLATE_FUNCTION_BUILD_OPTIMIZATION
+void STableTreeView::Construct(const FArguments& InArgs, TSharedPtr<FTable> InTablePtr)
+{
+	check(InTablePtr.IsValid());
+	Table = InTablePtr;
+
+	SAssignNew(ExternalScrollbar, SScrollBar)
+	.AlwaysShowScrollbar(true);
+
+	ChildSlot
+	[
+		SNew(SVerticalBox)
+
+		+SVerticalBox::Slot()
+		.VAlign(VAlign_Center)
+		.AutoHeight()
+		[
+			SNew(SBorder)
+			.BorderImage(FEditorStyle::GetBrush("ToolPanel.GroupBorder"))
+			.Padding(2.0f)
+			[
+				SNew(SVerticalBox)
+
+				// Search box
+				+SVerticalBox::Slot()
+				.VAlign(VAlign_Center)
+				.Padding(2.0f)
+				.AutoHeight()
+				[
+					SAssignNew(SearchBox, SSearchBox)
+					.HintText(LOCTEXT("SearchBoxHint", "Search timers or groups"))
+					.OnTextChanged(this, &STableTreeView::SearchBox_OnTextChanged)
+					.IsEnabled(this, &STableTreeView::SearchBox_IsEnabled)
+					.ToolTipText(LOCTEXT("FilterSearchHint", "Type here to search timer or group"))
+				]
+
+				// Group by
+				+SVerticalBox::Slot()
+				.VAlign(VAlign_Center)
+				.Padding(2.0f)
+				.AutoHeight()
+				[
+					SNew(SHorizontalBox)
+
+					+SHorizontalBox::Slot()
+					.FillWidth(1.0f)
+					.VAlign(VAlign_Center)
+					[
+						SNew(STextBlock)
+						.Text(LOCTEXT("GroupByText", "Group by"))
+					]
+
+					+SHorizontalBox::Slot()
+					.FillWidth(2.0f)
+					.VAlign(VAlign_Center)
+					[
+						SAssignNew(GroupByComboBox, SComboBox<TSharedPtr<FTreeNodeGrouping>>)
+						.ToolTipText(this, &STableTreeView::GroupBy_GetSelectedTooltipText)
+						.OptionsSource(&AvailableGroupings)
+						.OnSelectionChanged(this, &STableTreeView::GroupBy_OnSelectionChanged)
+						.OnGenerateWidget(this, &STableTreeView::GroupBy_OnGenerateWidget)
+						[
+							SNew(STextBlock)
+							.Text(this, &STableTreeView::GroupBy_GetSelectedText)
+						]
+					]
+				]
+			]
+		]
+
+		// Tree view
+		+SVerticalBox::Slot()
+		.FillHeight(1.0f)
+		.Padding(0.0f, 6.0f, 0.0f, 0.0f)
+		[
+			SNew(SHorizontalBox)
+
+			+ SHorizontalBox::Slot()
+			.FillWidth(1.0f)
+			.Padding(0.0f)
+			[
+				SNew(SScrollBox)
+				.Orientation(Orient_Horizontal)
+
+				+ SScrollBox::Slot()
+				[
+					SNew(SBorder)
+					.BorderImage(FEditorStyle::GetBrush("ToolPanel.GroupBorder"))
+					.Padding(0.0f)
+					[
+						SAssignNew(TreeView, STreeView<FTableTreeNodePtr>)
+						.ExternalScrollbar(ExternalScrollbar)
+						.SelectionMode(ESelectionMode::Multi)
+						.TreeItemsSource(&FilteredGroupNodes)
+						.OnGetChildren(this, &STableTreeView::TreeView_OnGetChildren)
+						.OnGenerateRow(this, &STableTreeView::TreeView_OnGenerateRow)
+						.OnSelectionChanged(this, &STableTreeView::TreeView_OnSelectionChanged)
+						.OnMouseButtonDoubleClick(this, &STableTreeView::TreeView_OnMouseButtonDoubleClick)
+						.OnContextMenuOpening(FOnContextMenuOpening::CreateSP(this, &STableTreeView::TreeView_GetMenuContent))
+						.ItemHeight(12.0f)
+						.HeaderRow
+						(
+							SAssignNew(TreeViewHeaderRow, SHeaderRow)
+							.Visibility(EVisibility::Visible)
+						)
+					]
+				]
+			]
+
+			+ SHorizontalBox::Slot()
+			.AutoWidth()
+			.Padding(0.0f)
+			[
+				SNew(SBox)
+				.WidthOverride(FOptionalSize(13.0f))
+				[
+					ExternalScrollbar.ToSharedRef()
+				]
+			]
+		]
+	];
+
+	InitializeAndShowHeaderColumns();
+	//BindCommands();
+
+	// Create the search filters: text based, type based etc.
+	TextFilter = MakeShareable(new FTableTreeNodeTextFilter(FTableTreeNodeTextFilter::FItemToStringArray::CreateSP(this, &STableTreeView::HandleItemToStringArray)));
+	Filters = MakeShareable(new FTableTreeNodeFilterCollection());
+	Filters->Add(TextFilter);
+
+	CreateGroupings();
+	CreateSortings();
+
+	// Register ourselves with the Insights manager.
+	FInsightsManager::Get()->GetSessionChangedEvent().AddSP(this, &STableTreeView::InsightsManager_OnSessionChanged);
+}
+END_SLATE_FUNCTION_BUILD_OPTIMIZATION
+
+////////////////////////////////////////////////////////////////////////////////////////////////////
+
+TSharedPtr<SWidget> STableTreeView::TreeView_GetMenuContent()
+{
+	const TArray<FTableTreeNodePtr> SelectedNodes = TreeView->GetSelectedItems();
+	const int32 NumSelectedNodes = SelectedNodes.Num();
+	FTableTreeNodePtr SelectedNode = NumSelectedNodes ? SelectedNodes[0] : nullptr;
+
+	const TSharedPtr<FTableColumn> HoveredColumnPtr = Table->FindColumn(HoveredColumnId);
+
+	FText SelectionStr;
+	FText PropertyName;
+	FText PropertyValue;
+
+	if (NumSelectedNodes == 0)
+	{
+		SelectionStr = LOCTEXT("NothingSelected", "Nothing selected");
+	}
+	else if (NumSelectedNodes == 1)
+	{
+		if (HoveredColumnPtr != nullptr)
+		{
+			PropertyName = HoveredColumnPtr->GetShortName();
+			PropertyValue = HoveredColumnPtr->GetValueAsText(SelectedNode->GetRowId());
+		}
+		SelectionStr = FText::FromName(SelectedNode->GetName());
+	}
+	else
+	{
+		SelectionStr = LOCTEXT("MultipleSelection", "Multiple selection");
+	}
+
+	const bool bShouldCloseWindowAfterMenuSelection = true;
+	FMenuBuilder MenuBuilder(bShouldCloseWindowAfterMenuSelection, NULL);
+
+	// Selection menu
+	MenuBuilder.BeginSection("Selection", LOCTEXT("ContextMenu_Header_Selection", "Selection"));
+	{
+		struct FLocal
+		{
+			static bool ReturnFalse()
+			{
+				return false;
+			}
+		};
+
+		FUIAction DummyUIAction;
+		DummyUIAction.CanExecuteAction = FCanExecuteAction::CreateStatic(&FLocal::ReturnFalse);
+		MenuBuilder.AddMenuEntry
+		(
+			SelectionStr,
+			LOCTEXT("ContextMenu_Selection", "Currently selected items"),
+			FSlateIcon(FEditorStyle::GetStyleSetName(), "@missing.icon"), DummyUIAction, NAME_None, EUserInterfaceActionType::Button
+		);
+	}
+	MenuBuilder.EndSection();
+
+	MenuBuilder.BeginSection("Misc", LOCTEXT("ContextMenu_Header_Misc", "Miscellaneous"));
+	{
+		/*TODO
+		FUIAction Action_CopySelectedToClipboard
+		(
+			FExecuteAction::CreateSP(this, &STableTreeView::ContextMenu_CopySelectedToClipboard_Execute),
+			FCanExecuteAction::CreateSP(this, &STableTreeView::ContextMenu_CopySelectedToClipboard_CanExecute)
+		);
+		MenuBuilder.AddMenuEntry
+		(
+			LOCTEXT("ContextMenu_Header_Misc_CopySelectedToClipboard", "Copy To Clipboard"),
+			LOCTEXT("ContextMenu_Header_Misc_CopySelectedToClipboard_Desc", "Copies selection to clipboard"),
+			FSlateIcon(FEditorStyle::GetStyleSetName(), "Profiler.Misc.CopyToClipboard"), Action_CopySelectedToClipboard, NAME_None, EUserInterfaceActionType::Button
+		);
+		*/
+
+		MenuBuilder.AddSubMenu
+		(
+			LOCTEXT("ContextMenu_Header_Misc_Sort", "Sort By"),
+			LOCTEXT("ContextMenu_Header_Misc_Sort_Desc", "Sort by column"),
+			FNewMenuDelegate::CreateSP(this, &STableTreeView::TreeView_BuildSortByMenu),
+			false,
+			FSlateIcon(FEditorStyle::GetStyleSetName(), "Profiler.Misc.SortBy")
+		);
+	}
+	MenuBuilder.EndSection();
+
+	MenuBuilder.BeginSection("Columns", LOCTEXT("ContextMenu_Header_Columns", "Columns"));
+	{
+		MenuBuilder.AddSubMenu
+		(
+			LOCTEXT("ContextMenu_Header_Columns_View", "View Column"),
+			LOCTEXT("ContextMenu_Header_Columns_View_Desc", "Hides or shows columns"),
+			FNewMenuDelegate::CreateSP(this, &STableTreeView::TreeView_BuildViewColumnMenu),
+			false,
+			FSlateIcon(FEditorStyle::GetStyleSetName(), "Profiler.EventGraph.ViewColumn")
+		);
+
+		FUIAction Action_ShowAllColumns
+		(
+			FExecuteAction::CreateSP(this, &STableTreeView::ContextMenu_ShowAllColumns_Execute),
+			FCanExecuteAction::CreateSP(this, &STableTreeView::ContextMenu_ShowAllColumns_CanExecute)
+		);
+		MenuBuilder.AddMenuEntry
+		(
+			LOCTEXT("ContextMenu_Header_Columns_ShowAllColumns", "Show All Columns"),
+			LOCTEXT("ContextMenu_Header_Columns_ShowAllColumns_Desc", "Resets tree view to show all columns"),
+			FSlateIcon(FEditorStyle::GetStyleSetName(), "Profiler.EventGraph.ResetColumn"), Action_ShowAllColumns, NAME_None, EUserInterfaceActionType::Button
+		);
+
+		FUIAction Action_ResetColumns
+		(
+			FExecuteAction::CreateSP(this, &STableTreeView::ContextMenu_ResetColumns_Execute),
+			FCanExecuteAction::CreateSP(this, &STableTreeView::ContextMenu_ResetColumns_CanExecute)
+		);
+		MenuBuilder.AddMenuEntry
+		(
+			LOCTEXT("ContextMenu_Header_Columns_ResetColumns", "Reset Columns to Default"),
+			LOCTEXT("ContextMenu_Header_Columns_ResetColumns_Desc", "Resets columns to default"),
+			FSlateIcon(FEditorStyle::GetStyleSetName(), "Profiler.EventGraph.ResetColumn"), Action_ResetColumns, NAME_None, EUserInterfaceActionType::Button
+		);
+	}
+	MenuBuilder.EndSection();
+
+	return MenuBuilder.MakeWidget();
+}
+
+////////////////////////////////////////////////////////////////////////////////////////////////////
+
+void STableTreeView::TreeView_BuildSortByMenu(FMenuBuilder& MenuBuilder)
+{
+	// TODO: Refactor later @see TSharedPtr<SWidget> SCascadePreviewViewportToolBar::GenerateViewMenu() const
+
+	MenuBuilder.BeginSection("ColumnName", LOCTEXT("ContextMenu_Header_Misc_ColumnName", "Column Name"));
+
+	//TODO: for (Sorting : AvailableSortings)
+	for (const TSharedPtr<FTableColumn>& ColumnPtr : Table->GetColumns())
+	{
+		const FTableColumn& Column = *ColumnPtr;
+
+		if (Column.IsVisible() && Column.CanBeSorted())
+		{
+			FUIAction Action_SortByColumn
+			(
+				FExecuteAction::CreateSP(this, &STableTreeView::ContextMenu_SortByColumn_Execute, Column.GetId()),
+				FCanExecuteAction::CreateSP(this, &STableTreeView::ContextMenu_SortByColumn_CanExecute, Column.GetId()),
+				FIsActionChecked::CreateSP(this, &STableTreeView::ContextMenu_SortByColumn_IsChecked, Column.GetId())
+			);
+			MenuBuilder.AddMenuEntry
+			(
+				Column.GetTitleName(),
+				Column.GetDescription(),
+				FSlateIcon(), Action_SortByColumn, NAME_None, EUserInterfaceActionType::RadioButton
+			);
+		}
+	}
+
+	MenuBuilder.EndSection();
+
+	//-----------------------------------------------------------------------------
+
+	MenuBuilder.BeginSection("SortMode", LOCTEXT("ContextMenu_Header_Misc_Sort_SortMode", "Sort Mode"));
+	{
+		FUIAction Action_SortAscending
+		(
+			FExecuteAction::CreateSP(this, &STableTreeView::ContextMenu_SortMode_Execute, EColumnSortMode::Ascending),
+			FCanExecuteAction::CreateSP(this, &STableTreeView::ContextMenu_SortMode_CanExecute, EColumnSortMode::Ascending),
+			FIsActionChecked::CreateSP(this, &STableTreeView::ContextMenu_SortMode_IsChecked, EColumnSortMode::Ascending)
+		);
+		MenuBuilder.AddMenuEntry
+		(
+			LOCTEXT("ContextMenu_Header_Misc_Sort_SortAscending", "Sort Ascending"),
+			LOCTEXT("ContextMenu_Header_Misc_Sort_SortAscending_Desc", "Sorts ascending"),
+			FSlateIcon(FEditorStyle::GetStyleSetName(), "Profiler.Misc.SortAscending"), Action_SortAscending, NAME_None, EUserInterfaceActionType::RadioButton
+		);
+
+		FUIAction Action_SortDescending
+		(
+			FExecuteAction::CreateSP(this, &STableTreeView::ContextMenu_SortMode_Execute, EColumnSortMode::Descending),
+			FCanExecuteAction::CreateSP(this, &STableTreeView::ContextMenu_SortMode_CanExecute, EColumnSortMode::Descending),
+			FIsActionChecked::CreateSP(this, &STableTreeView::ContextMenu_SortMode_IsChecked, EColumnSortMode::Descending)
+		);
+		MenuBuilder.AddMenuEntry
+		(
+			LOCTEXT("ContextMenu_Header_Misc_Sort_SortDescending", "Sort Descending"),
+			LOCTEXT("ContextMenu_Header_Misc_Sort_SortDescending_Desc", "Sorts descending"),
+			FSlateIcon(FEditorStyle::GetStyleSetName(), "Profiler.Misc.SortDescending"), Action_SortDescending, NAME_None, EUserInterfaceActionType::RadioButton
+		);
+	}
+	MenuBuilder.EndSection();
+}
+
+////////////////////////////////////////////////////////////////////////////////////////////////////
+
+void STableTreeView::TreeView_BuildViewColumnMenu(FMenuBuilder& MenuBuilder)
+{
+	MenuBuilder.BeginSection("ViewColumn", LOCTEXT("ContextMenu_Header_Columns_View", "View Column"));
+
+	for (const TSharedPtr<FTableColumn>& ColumnPtr : Table->GetColumns())
+	{
+		const FTableColumn& Column = *ColumnPtr;
+
+		FUIAction Action_ToggleColumn
+		(
+			FExecuteAction::CreateSP(this, &STableTreeView::ToggleColumnVisibility, Column.GetId()),
+			FCanExecuteAction::CreateSP(this, &STableTreeView::CanToggleColumnVisibility, Column.GetId()),
+			FIsActionChecked::CreateSP(this, &STableTreeView::IsColumnVisible, Column.GetId())
+		);
+		MenuBuilder.AddMenuEntry
+		(
+			Column.GetTitleName(),
+			Column.GetDescription(),
+			FSlateIcon(), Action_ToggleColumn, NAME_None, EUserInterfaceActionType::ToggleButton
+		);
+	}
+
+	MenuBuilder.EndSection();
+}
+
+////////////////////////////////////////////////////////////////////////////////////////////////////
+
+void STableTreeView::InitializeAndShowHeaderColumns()
+{
+	for (const TSharedPtr<FTableColumn>& ColumnPtr : Table->GetColumns())
+	{
+		TreeViewHeaderRow_CreateColumnArgs(*ColumnPtr);
+
+		if (ColumnPtr->ShouldBeVisible())
+		{
+			ShowColumn(ColumnPtr->GetId());
+		}
+	}
+}
+
+////////////////////////////////////////////////////////////////////////////////////////////////////
+
+void STableTreeView::TreeViewHeaderRow_CreateColumnArgs(const FTableColumn& Column)
+{
+	SHeaderRow::FColumn::FArguments ColumnArgs;
+
+	ColumnArgs
+		.ColumnId(Column.GetId())
+		.DefaultLabel(Column.GetShortName())
+		.SortMode(EColumnSortMode::None)
+		.HAlignHeader(HAlign_Fill)
+		.VAlignHeader(VAlign_Fill)
+		.HeaderContentPadding(TOptional<FMargin>(2.0f))
+		.HAlignCell(HAlign_Fill)
+		.VAlignCell(VAlign_Fill)
+		.SortMode(this, &STableTreeView::GetSortModeForColumn, Column.GetId())
+		.OnSort(this, &STableTreeView::OnSortModeChanged)
+		.ManualWidth(Column.GetInitialWidth())
+		.FixedWidth(Column.IsFixedWidth() ? Column.GetInitialWidth() : TOptional<float>())
+		.HeaderContent()
+		[
+			SNew(SHorizontalBox)
+			.ToolTip(STableTreeViewTooltip::GetColumnTooltip(Column))
+
+			+ SHorizontalBox::Slot()
+			.HAlign(Column.GetHorizontalAlignment())
+			.VAlign(VAlign_Center)
+			[
+				SNew(STextBlock)
+				.Text(Column.GetShortName())
+			]
+		]
+		.MenuContent()
+		[
+			TreeViewHeaderRow_GenerateColumnMenu(Column)
+		];
+
+	TreeViewHeaderColumnArgs.Add(Column.GetId(), ColumnArgs);
+}
+
+////////////////////////////////////////////////////////////////////////////////////////////////////
+
+TSharedRef<SWidget> STableTreeView::TreeViewHeaderRow_GenerateColumnMenu(const FTableColumn& Column)
+{
+	bool bIsMenuVisible = false;
+
+	const bool bShouldCloseWindowAfterMenuSelection = true;
+	FMenuBuilder MenuBuilder(bShouldCloseWindowAfterMenuSelection, NULL);
+	{
+		if (Column.CanBeHidden())
+		{
+			MenuBuilder.BeginSection("Column", LOCTEXT("TreeViewHeaderRow_Header_Column", "Column"));
+
+			FUIAction Action_HideColumn
+			(
+				FExecuteAction::CreateSP(this, &STableTreeView::HideColumn, Column.GetId()),
+				FCanExecuteAction::CreateSP(this, &STableTreeView::CanHideColumn, Column.GetId())
+			);
+
+			MenuBuilder.AddMenuEntry
+			(
+				LOCTEXT("TreeViewHeaderRow_HideColumn", "Hide"),
+				LOCTEXT("TreeViewHeaderRow_HideColumn_Desc", "Hides the selected column"),
+				FSlateIcon(), Action_HideColumn, NAME_None, EUserInterfaceActionType::Button
+			);
+			bIsMenuVisible = true;
+
+			MenuBuilder.EndSection();
+		}
+
+		if (Column.CanBeSorted())
+		{
+			MenuBuilder.BeginSection("SortMode", LOCTEXT("ContextMenu_Header_Misc_Sort_SortMode", "Sort Mode"));
+
+			FUIAction Action_SortAscending
+			(
+				FExecuteAction::CreateSP(this, &STableTreeView::HeaderMenu_SortMode_Execute, Column.GetId(), EColumnSortMode::Ascending),
+				FCanExecuteAction::CreateSP(this, &STableTreeView::HeaderMenu_SortMode_CanExecute, Column.GetId(), EColumnSortMode::Ascending),
+				FIsActionChecked::CreateSP(this, &STableTreeView::HeaderMenu_SortMode_IsChecked, Column.GetId(), EColumnSortMode::Ascending)
+			);
+			MenuBuilder.AddMenuEntry
+			(
+				LOCTEXT("ContextMenu_Header_Misc_Sort_SortAscending", "Sort Ascending"),
+				LOCTEXT("ContextMenu_Header_Misc_Sort_SortAscending_Desc", "Sorts ascending"),
+				FSlateIcon(FEditorStyle::GetStyleSetName(), "Profiler.Misc.SortAscending"), Action_SortAscending, NAME_None, EUserInterfaceActionType::RadioButton
+			);
+
+			FUIAction Action_SortDescending
+			(
+				FExecuteAction::CreateSP(this, &STableTreeView::HeaderMenu_SortMode_Execute, Column.GetId(), EColumnSortMode::Descending),
+				FCanExecuteAction::CreateSP(this, &STableTreeView::HeaderMenu_SortMode_CanExecute, Column.GetId(), EColumnSortMode::Descending),
+				FIsActionChecked::CreateSP(this, &STableTreeView::HeaderMenu_SortMode_IsChecked, Column.GetId(), EColumnSortMode::Descending)
+			);
+			MenuBuilder.AddMenuEntry
+			(
+				LOCTEXT("ContextMenu_Header_Misc_Sort_SortDescending", "Sort Descending"),
+				LOCTEXT("ContextMenu_Header_Misc_Sort_SortDescending_Desc", "Sorts descending"),
+				FSlateIcon(FEditorStyle::GetStyleSetName(), "Profiler.Misc.SortDescending"), Action_SortDescending, NAME_None, EUserInterfaceActionType::RadioButton
+			);
+			bIsMenuVisible = true;
+
+			MenuBuilder.EndSection();
+		}
+
+		if (Column.CanBeFiltered())
+		{
+			MenuBuilder.BeginSection("FilterMode", LOCTEXT("ContextMenu_Header_Misc_Filter_FilterMode", "Filter Mode"));
+			bIsMenuVisible = true;
+			MenuBuilder.EndSection();
+		}
+	}
+
+	return bIsMenuVisible ? MenuBuilder.MakeWidget() : (TSharedRef<SWidget>)SNullWidget::NullWidget;
+}
+
+////////////////////////////////////////////////////////////////////////////////////////////////////
+
+void STableTreeView::InsightsManager_OnSessionChanged()
+{
+	TSharedPtr<const Trace::IAnalysisSession> NewSession = FInsightsManager::Get()->GetSession();
+
+	if (NewSession != Session)
+	{
+		Session = NewSession;
+		RebuildTree();
+	}
+	else
+	{
+		UpdateTree();
+	}
+}
+
+////////////////////////////////////////////////////////////////////////////////////////////////////
+
+void STableTreeView::UpdateTree()
+{
+	CreateGroups();
+	SortTreeNodes();
+	ApplyFiltering();
+}
+
+////////////////////////////////////////////////////////////////////////////////////////////////////
+
+void STableTreeView::ApplyFiltering()
+{
+	FilteredGroupNodes.Reset();
+
+	// Apply filter to all groups and its children.
+	const int32 NumGroups = GroupNodes.Num();
+	for (int32 GroupIndex = 0; GroupIndex < NumGroups; ++GroupIndex)
+	{
+		FTableTreeNodePtr& GroupPtr = GroupNodes[GroupIndex];
+		GroupPtr->ClearFilteredChildren();
+		const bool bIsGroupVisible = Filters->PassesAllFilters(GroupPtr);
+
+		const TArray<FBaseTreeNodePtr>& GroupChildren = GroupPtr->GetChildren();
+		const int32 NumChildren = GroupChildren.Num();
+		int32 NumVisibleChildren = 0;
+		for (int32 Cx = 0; Cx < NumChildren; ++Cx)
+		{
+			// Add a child.
+			const FTableTreeNodePtr& NodePtr = StaticCastSharedPtr<FTableTreeNode>(GroupChildren[Cx]);
+			const bool bIsChildVisible = Filters->PassesAllFilters(NodePtr);
+			if (bIsChildVisible)
+			{
+				GroupPtr->AddFilteredChild(NodePtr);
+				NumVisibleChildren++;
+			}
+		}
+
+		if (bIsGroupVisible || NumVisibleChildren > 0)
+		{
+			// Add a group.
+			FilteredGroupNodes.Add(GroupPtr);
+			GroupPtr->SetExpansion(true);
+		}
+		else
+		{
+			GroupPtr->SetExpansion(false);
+		}
+	}
+
+	// Only expand timer nodes if we have a text filter.
+	const bool bNonEmptyTextFilter = !TextFilter->GetRawFilterText().IsEmpty();
+	if (bNonEmptyTextFilter)
+	{
+		if (!bExpansionSaved)
+		{
+			ExpandedNodes.Empty();
+			TreeView->GetExpandedItems(ExpandedNodes);
+			bExpansionSaved = true;
+		}
+
+		for (int32 Fx = 0; Fx < FilteredGroupNodes.Num(); Fx++)
+		{
+			const FTableTreeNodePtr& GroupPtr = FilteredGroupNodes[Fx];
+			TreeView->SetItemExpansion(GroupPtr, GroupPtr->IsExpanded());
+		}
+	}
+	else
+	{
+		if (bExpansionSaved)
+		{
+			// Restore previously expanded nodes when the text filter is disabled.
+			TreeView->ClearExpandedItems();
+			for (auto It = ExpandedNodes.CreateConstIterator(); It; ++It)
+			{
+				TreeView->SetItemExpansion(*It, true);
+			}
+			bExpansionSaved = false;
+		}
+	}
+
+	// Request tree refresh
+	TreeView->RequestTreeRefresh();
+}
+
+////////////////////////////////////////////////////////////////////////////////////////////////////
+
+void STableTreeView::HandleItemToStringArray(const FTableTreeNodePtr& FTableTreeNodePtr, TArray<FString>& OutSearchStrings) const
+{
+	OutSearchStrings.Add(FTableTreeNodePtr->GetName().GetPlainNameString());
+}
+
+////////////////////////////////////////////////////////////////////////////////////////////////////
+// TreeView
+////////////////////////////////////////////////////////////////////////////////////////////////////
+
+void STableTreeView::TreeView_Refresh()
+{
+	if (TreeView.IsValid())
+	{
+		TreeView->RequestTreeRefresh();
+	}
+}
+
+////////////////////////////////////////////////////////////////////////////////////////////////////
+
+void STableTreeView::TreeView_OnSelectionChanged(FTableTreeNodePtr SelectedItem, ESelectInfo::Type SelectInfo)
+{
+	if (SelectInfo != ESelectInfo::Direct)
+	{
+		TArray<FTableTreeNodePtr> SelectedItems = TreeView->GetSelectedItems();
+		if (SelectedItems.Num() == 1)
+		{
+			//HighlightedNodeName = SelectedItems[0]->GetName();
+		}
+		else
+		{
+			//HighlightedNodeName = NAME_None;
+		}
+	}
+}
+
+////////////////////////////////////////////////////////////////////////////////////////////////////
+
+void STableTreeView::TreeView_OnGetChildren(FTableTreeNodePtr InParent, TArray<FTableTreeNodePtr>& OutChildren)
+{
+	const TArray<FBaseTreeNodePtr>& FilteredChildren = InParent->GetFilteredChildren();
+	for (const FBaseTreeNodePtr& NodePtr : FilteredChildren)
+	{
+		OutChildren.Add(StaticCastSharedPtr<FTableTreeNode>(NodePtr));
+	}
+}
+
+////////////////////////////////////////////////////////////////////////////////////////////////////
+
+void STableTreeView::TreeView_OnMouseButtonDoubleClick(FTableTreeNodePtr TimerNode)
+{
+	if (!TimerNode->IsGroup())
+	{
+		//im:TODO: const bool bIsStatTracked = FTimingProfilerManager::Get()->IsStatTracked(TimerNode->GetId());
+	//	if (!bIsStatTracked)
+	//	{
+	//		// Add a new graph.
+	//		FTimingProfilerManager::Get()->TrackStat(TimerNode->GetId());
+	//	}
+	//	else
+	//	{
+	//		// Remove a graph
+	//		FTimingProfilerManager::Get()->UntrackStat(TimerNode->GetId());
+	//	}
+	}
+	else
+	{
+		const bool bIsGroupExpanded = TreeView->IsItemExpanded(TimerNode);
+		TreeView->SetItemExpansion(TimerNode, !bIsGroupExpanded);
+	}
+}
+
+////////////////////////////////////////////////////////////////////////////////////////////////////
+// Tree View's Table Row
+////////////////////////////////////////////////////////////////////////////////////////////////////
+
+TSharedRef<ITableRow> STableTreeView::TreeView_OnGenerateRow(FTableTreeNodePtr NodePtr, const TSharedRef<STableViewBase>& OwnerTable)
+{
+	TSharedRef<ITableRow> TableRow =
+		SNew(STableTreeViewRow, OwnerTable)
+		.OnShouldBeEnabled(this, &STableTreeView::TableRow_ShouldBeEnabled)
+		.OnIsColumnVisible(this, &STableTreeView::IsColumnVisible)
+		.OnSetHoveredCell(this, &STableTreeView::TableRow_SetHoveredCell)
+		.OnGetColumnOutlineHAlignmentDelegate(this, &STableTreeView::TableRow_GetColumnOutlineHAlignment)
+		.HighlightText(this, &STableTreeView::TableRow_GetHighlightText)
+		.HighlightedNodeName(this, &STableTreeView::TableRow_GetHighlightedNodeName)
+		.TablePtr(Table)
+		.TableTreeNodePtr(NodePtr);
+
+	return TableRow;
+}
+
+////////////////////////////////////////////////////////////////////////////////////////////////////
+
+void STableTreeView::TableRow_SetHoveredCell(TSharedPtr<FTable> TablePtr, TSharedPtr<FTableColumn> ColumnPtr, const FTableTreeNodePtr NodePtr)
+{
+	HoveredColumnId = ColumnPtr ? ColumnPtr->GetId() : FName();
+
+	const bool bIsAnyMenusVisible = FSlateApplication::Get().AnyMenusVisible();
+	if (!HasMouseCapture() && !bIsAnyMenusVisible)
+	{
+		HoveredNodePtr = NodePtr;
+	}
+
+	//UE_LOG(TimingProfiler, Log, TEXT("%s -> %s"), *HoveredColumnId.GetPlainNameString(), NodePtr.IsValid() ? *NodePtr->GetName().GetPlainNameString() : TEXT("nullptr"));
+}
+
+////////////////////////////////////////////////////////////////////////////////////////////////////
+
+EHorizontalAlignment STableTreeView::TableRow_GetColumnOutlineHAlignment(const FName ColumnId) const
+{
+	const TIndirectArray<SHeaderRow::FColumn>& Columns = TreeViewHeaderRow->GetColumns();
+	const int32 LastColumnIdx = Columns.Num() - 1;
+
+	// First column
+	if (Columns[0].ColumnId == ColumnId)
+	{
+		return HAlign_Left;
+	}
+	// Last column
+	else if (Columns[LastColumnIdx].ColumnId == ColumnId)
+	{
+		return HAlign_Right;
+	}
+	// Middle columns
+	{
+		return HAlign_Center;
+	}
+}
+
+////////////////////////////////////////////////////////////////////////////////////////////////////
+
+FText STableTreeView::TableRow_GetHighlightText() const
+{
+	return SearchBox->GetText();
+}
+
+////////////////////////////////////////////////////////////////////////////////////////////////////
+
+FName STableTreeView::TableRow_GetHighlightedNodeName() const
+{
+	return HighlightedNodeName;
+}
+
+////////////////////////////////////////////////////////////////////////////////////////////////////
+
+bool STableTreeView::TableRow_ShouldBeEnabled(const uint32 TimerId) const
+{
+	return true;//im:TODO: Session->GetAggregatedStat(TimerId) != nullptr;
+}
+
+////////////////////////////////////////////////////////////////////////////////////////////////////
+// SearchBox
+////////////////////////////////////////////////////////////////////////////////////////////////////
+
+void STableTreeView::SearchBox_OnTextChanged(const FText& InFilterText)
+{
+	TextFilter->SetRawFilterText(InFilterText);
+	SearchBox->SetError(TextFilter->GetFilterErrorText());
+	ApplyFiltering();
+}
+
+////////////////////////////////////////////////////////////////////////////////////////////////////
+
+bool STableTreeView::SearchBox_IsEnabled() const
+{
+	return TableTreeNodes.Num() > 0;
+}
+
+////////////////////////////////////////////////////////////////////////////////////////////////////
+// Grouping
+////////////////////////////////////////////////////////////////////////////////////////////////////
+
+void STableTreeView::CreateGroups()
+{
+	GroupNodes.Reset();
+	TMap<FName, FTableTreeNodePtr> GroupMap;
+
+	for (const FTableTreeNodePtr& NodePtr : TableTreeNodes)
+	{
+		FTreeNodeGroupInfo GroupInfo = CurrentGrouping->GetGroupForNode(NodePtr);
+
+		FTableTreeNodePtr GroupPtr = nullptr;
+		FTableTreeNodePtr* GroupPtrPtr = GroupMap.Find(GroupInfo.Name);
+		if (!GroupPtrPtr)
+		{
+			GroupPtr = MakeShareable(new FTableTreeNode(GroupInfo.Name, Table));
+
+			GroupMap.Add(GroupInfo.Name, GroupPtr);
+			GroupNodes.Add(GroupPtr);
+
+			GroupPtr->SetExpansion(GroupInfo.IsExpanded);
+			TreeView->SetItemExpansion(GroupPtr, GroupInfo.IsExpanded);
+		}
+		else
+		{
+			GroupPtr = *GroupPtrPtr;
+		}
+
+		GroupPtr->AddChildAndSetGroupPtr(NodePtr);
+	}
+
+	// Sort groups by name.
+	GroupNodes.Sort([](const FBaseTreeNodePtr& A, const FBaseTreeNodePtr& B) { return A->GetName().LexicalLess(B->GetName()); });
+}
+
+////////////////////////////////////////////////////////////////////////////////////////////////////
+
+void STableTreeView::CreateGroupings()
+{
+	AvailableGroupings.Reset(3);
+
+	AvailableGroupings.Add(MakeShareable(new FTreeNodeGroupingFlat()));
+	AvailableGroupings.Add(MakeShareable(new FTreeNodeGroupingByNameFirstLetter()));
+	AvailableGroupings.Add(MakeShareable(new FTreeNodeGroupingByType()));
+
+	CurrentGrouping = AvailableGroupings[0];
+
+	GroupByComboBox->SetSelectedItem(CurrentGrouping);
+	GroupByComboBox->RefreshOptions();
+}
+
+////////////////////////////////////////////////////////////////////////////////////////////////////
+
+void STableTreeView::GroupBy_OnSelectionChanged(TSharedPtr<FTreeNodeGrouping> NewGrouping, ESelectInfo::Type SelectInfo)
+{
+	if (SelectInfo != ESelectInfo::Direct)
+	{
+		CurrentGrouping = NewGrouping;
+
+		CreateGroups();
+		SortTreeNodes();
+		ApplyFiltering();
+	}
+}
+
+////////////////////////////////////////////////////////////////////////////////////////////////////
+
+TSharedRef<SWidget> STableTreeView::GroupBy_OnGenerateWidget(TSharedPtr<FTreeNodeGrouping> InGrouping) const
+{
+	return SNew(STextBlock)
+		.Text(InGrouping->GetName())
+		.ToolTipText(InGrouping->GetDescription());
+}
+
+////////////////////////////////////////////////////////////////////////////////////////////////////
+
+FText STableTreeView::GroupBy_GetSelectedText() const
+{
+	return CurrentGrouping.IsValid() ? CurrentGrouping->GetName() : FText::GetEmpty();
+}
+
+////////////////////////////////////////////////////////////////////////////////////////////////////
+
+FText STableTreeView::GroupBy_GetSelectedTooltipText() const
+{
+	return CurrentGrouping.IsValid() ? CurrentGrouping->GetDescription() : FText::GetEmpty();
+}
+
+////////////////////////////////////////////////////////////////////////////////////////////////////
+// Sorting
+////////////////////////////////////////////////////////////////////////////////////////////////////
+
+void STableTreeView::CreateSortings()
+{
+	AvailableSortings.Reset();
+	ColumnToSortingMap.Reset();
+	CurrentSorting = nullptr;
+
+	TSharedPtr<FTreeNodeSorting> SortingByName = MakeShareable(new FTreeNodeSortingByName());
+	AvailableSortings.Add(SortingByName);
+
+	TSharedPtr<FTreeNodeSorting> SortingByType = MakeShareable(new FTreeNodeSortingByType());
+	AvailableSortings.Add(SortingByType);
+
+	const Trace::ITableLayout& TableLayout = Table->GetSourceTable()->GetLayout();
+
+	for (const TSharedPtr<FTableColumn> ColumnPtr : Table->GetColumns())
+	{
+		if (ColumnPtr->CanBeSorted())
+		{
+			if (ColumnPtr->IsHierarchy())
+			{
+				ColumnToSortingMap.Add(ColumnPtr->GetId(), SortingByName);
+			}
+			else
+			{
+				TSharedPtr<FTreeNodeSorting> Sorting = nullptr;
+
+				switch (TableLayout.GetColumnType(ColumnPtr->GetIndex()))
+				{
+				case Trace::TableColumnType_Bool:
+					Sorting = MakeShareable(new FTreeNodeSortingByBoolValue(ColumnPtr));
+					break;
+				case Trace::TableColumnType_Int:
+					Sorting = MakeShareable(new FTreeNodeSortingByIntValue(ColumnPtr));
+					break;
+				case Trace::TableColumnType_Float:
+					Sorting = MakeShareable(new FTreeNodeSortingByFloatValue(ColumnPtr));
+					break;
+				case Trace::TableColumnType_Double:
+					Sorting = MakeShareable(new FTreeNodeSortingByDoubleValue(ColumnPtr));
+					break;
+				case Trace::TableColumnType_CString:
+					Sorting = MakeShareable(new FTreeNodeSortingByCStringValue(ColumnPtr));
+					break;
+				}
+
+				if (Sorting)
+				{
+					AvailableSortings.Add(Sorting);
+					ColumnToSortingMap.Add(ColumnPtr->GetId(), Sorting);
+				}
+			}
+		}
+	}
+}
+
+////////////////////////////////////////////////////////////////////////////////////////////////////
+
+void STableTreeView::UpdateCurrentSortingByColumn()
+{
+	TSharedPtr<FTreeNodeSorting>* SorterPtrPtr = ColumnToSortingMap.Find(ColumnBeingSorted);
+	CurrentSorting = SorterPtrPtr ? *SorterPtrPtr : nullptr;
+}
+
+////////////////////////////////////////////////////////////////////////////////////////////////////
+
+void STableTreeView::SortTreeNodes()
+{
+	if (CurrentSorting.IsValid())
+	{
+		const Insights::ITreeNodeSorting* Sorter = CurrentSorting.Get();
+
+		const int32 NumGroups = GroupNodes.Num();
+
+		if (ColumnSortMode == EColumnSortMode::Type::Descending)
+		{
+			for (int32 GroupIndex = 0; GroupIndex < NumGroups; ++GroupIndex)
+			{
+				GroupNodes[GroupIndex]->SortChildrenDescending(Sorter);
+			}
+		}
+		else
+		{
+			for (int32 GroupIndex = 0; GroupIndex < NumGroups; ++GroupIndex)
+			{
+				GroupNodes[GroupIndex]->SortChildrenAscending(Sorter);
+			}
+		}
+	}
+}
+
+////////////////////////////////////////////////////////////////////////////////////////////////////
+
+EColumnSortMode::Type STableTreeView::GetSortModeForColumn(const FName ColumnId) const
+{
+	if (ColumnBeingSorted != ColumnId)
+	{
+		return EColumnSortMode::None;
+	}
+
+	return ColumnSortMode;
+}
+
+////////////////////////////////////////////////////////////////////////////////////////////////////
+
+void STableTreeView::SetSortModeForColumn(const FName& ColumnId, const EColumnSortMode::Type SortMode)
+{
+	ColumnBeingSorted = ColumnId;
+	ColumnSortMode = SortMode;
+	UpdateCurrentSortingByColumn();
+
+	SortTreeNodes();
+	ApplyFiltering();
+}
+
+////////////////////////////////////////////////////////////////////////////////////////////////////
+
+void STableTreeView::OnSortModeChanged(const EColumnSortPriority::Type SortPriority, const FName& ColumnId, const EColumnSortMode::Type SortMode)
+{
+	SetSortModeForColumn(ColumnId, SortMode);
+	TreeView_Refresh();
+}
+
+////////////////////////////////////////////////////////////////////////////////////////////////////
+// SortMode action (HeaderMenu)
+////////////////////////////////////////////////////////////////////////////////////////////////////
+
+bool STableTreeView::HeaderMenu_SortMode_IsChecked(const FName ColumnId, const EColumnSortMode::Type InSortMode)
+{
+	return ColumnBeingSorted == ColumnId && ColumnSortMode == InSortMode;
+}
+
+////////////////////////////////////////////////////////////////////////////////////////////////////
+
+bool STableTreeView::HeaderMenu_SortMode_CanExecute(const FName ColumnId, const EColumnSortMode::Type InSortMode) const
+{
+	FTableColumn& Column = *Table->FindColumnChecked(ColumnId);
+	const bool bIsValid = Column.CanBeSorted();
+
+	bool bCanExecute = ColumnBeingSorted != ColumnId ? true : ColumnSortMode != InSortMode;
+	bCanExecute = bCanExecute && bIsValid;
+
+	return bCanExecute;
+}
+
+////////////////////////////////////////////////////////////////////////////////////////////////////
+
+void STableTreeView::HeaderMenu_SortMode_Execute(const FName ColumnId, const EColumnSortMode::Type InSortMode)
+{
+	SetSortModeForColumn(ColumnId, InSortMode);
+	TreeView_Refresh();
+}
+
+////////////////////////////////////////////////////////////////////////////////////////////////////
+// SortMode action (ContextMenu)
+////////////////////////////////////////////////////////////////////////////////////////////////////
+
+bool STableTreeView::ContextMenu_SortMode_IsChecked(const EColumnSortMode::Type InSortMode)
+{
+	return ColumnSortMode == InSortMode;
+}
+
+////////////////////////////////////////////////////////////////////////////////////////////////////
+
+bool STableTreeView::ContextMenu_SortMode_CanExecute(const EColumnSortMode::Type InSortMode) const
+{
+	return true; //ColumnSortMode != InSortMode;
+}
+
+////////////////////////////////////////////////////////////////////////////////////////////////////
+
+void STableTreeView::ContextMenu_SortMode_Execute(const EColumnSortMode::Type InSortMode)
+{
+	SetSortModeForColumn(ColumnBeingSorted, InSortMode);
+	TreeView_Refresh();
+}
+
+////////////////////////////////////////////////////////////////////////////////////////////////////
+// SortByColumn action (ContextMenu)
+////////////////////////////////////////////////////////////////////////////////////////////////////
+
+bool STableTreeView::ContextMenu_SortByColumn_IsChecked(const FName ColumnId)
+{
+	return ColumnId == ColumnBeingSorted;
+}
+
+////////////////////////////////////////////////////////////////////////////////////////////////////
+
+bool STableTreeView::ContextMenu_SortByColumn_CanExecute(const FName ColumnId) const
+{
+	return true; //ColumnId != ColumnBeingSorted;
+}
+
+////////////////////////////////////////////////////////////////////////////////////////////////////
+
+void STableTreeView::ContextMenu_SortByColumn_Execute(const FName ColumnId)
+{
+	SetSortModeForColumn(ColumnId, EColumnSortMode::Descending);
+	TreeView_Refresh();
+}
+
+////////////////////////////////////////////////////////////////////////////////////////////////////
+// ShowColumn action
+////////////////////////////////////////////////////////////////////////////////////////////////////
+
+bool STableTreeView::CanShowColumn(const FName ColumnId) const
+{
+	return true;
+}
+
+////////////////////////////////////////////////////////////////////////////////////////////////////
+
+void STableTreeView::ShowColumn(const FName ColumnId)
+{
+	FTableColumn& Column = *Table->FindColumnChecked(ColumnId);
+	Column.Show();
+
+	SHeaderRow::FColumn::FArguments& ColumnArgs = TreeViewHeaderColumnArgs.FindChecked(ColumnId);
+
+	const int32 NumColumns = TreeViewHeaderRow->GetColumns().Num();
+	int32 ColumnIndex = 0;
+	for (; ColumnIndex < NumColumns; ColumnIndex++)
+	{
+		const SHeaderRow::FColumn& CurrentColumn = TreeViewHeaderRow->GetColumns()[ColumnIndex];
+		const FTableColumn& CurrentTableColumn = *Table->FindColumnChecked(CurrentColumn.ColumnId);
+		if (Column.GetOrder() < CurrentTableColumn.GetOrder())
+		{
+			break;
+		}
+	}
+
+	TreeViewHeaderRow->InsertColumn(ColumnArgs, ColumnIndex);
+}
+
+////////////////////////////////////////////////////////////////////////////////////////////////////
+// HideColumn action
+////////////////////////////////////////////////////////////////////////////////////////////////////
+
+bool STableTreeView::CanHideColumn(const FName ColumnId) const
+{
+	FTableColumn& Column = *Table->FindColumnChecked(ColumnId);
+	return Column.CanBeHidden();
+}
+
+////////////////////////////////////////////////////////////////////////////////////////////////////
+
+void STableTreeView::HideColumn(const FName ColumnId)
+{
+	FTableColumn& Column = *Table->FindColumnChecked(ColumnId);
+	Column.Hide();
+
+	TreeViewHeaderRow->RemoveColumn(ColumnId);
+}
+
+////////////////////////////////////////////////////////////////////////////////////////////////////
+// ToggleColumn action
+////////////////////////////////////////////////////////////////////////////////////////////////////
+
+bool STableTreeView::IsColumnVisible(const FName ColumnId)
+{
+	FTableColumn& Column = *Table->FindColumnChecked(ColumnId);
+	return Column.IsVisible();
+}
+
+////////////////////////////////////////////////////////////////////////////////////////////////////
+
+bool STableTreeView::CanToggleColumnVisibility(const FName ColumnId) const
+{
+	FTableColumn& Column = *Table->FindColumnChecked(ColumnId);
+	return !Column.IsVisible() || Column.CanBeHidden();
+}
+
+////////////////////////////////////////////////////////////////////////////////////////////////////
+
+void STableTreeView::ToggleColumnVisibility(const FName ColumnId)
+{
+	FTableColumn& Column = *Table->FindColumnChecked(ColumnId);
+	if (Column.IsVisible())
+	{
+		HideColumn(ColumnId);
+	}
+	else
+	{
+		ShowColumn(ColumnId);
+	}
+}
+
+////////////////////////////////////////////////////////////////////////////////////////////////////
+// "Show All Columns" action (ContextMenu)
+////////////////////////////////////////////////////////////////////////////////////////////////////
+
+bool STableTreeView::ContextMenu_ShowAllColumns_CanExecute() const
+{
+	return true;
+}
+
+////////////////////////////////////////////////////////////////////////////////////////////////////
+
+void STableTreeView::ContextMenu_ShowAllColumns_Execute()
+{
+	ColumnBeingSorted = DefaultColumnBeingSorted;
+	ColumnSortMode = DefaultColumnSortMode;
+	UpdateCurrentSortingByColumn();
+
+	for (const TSharedPtr<FTableColumn>& ColumnPtr : Table->GetColumns())
+	{
+		const FTableColumn& Column = *ColumnPtr;
+
+		if (!Column.IsVisible())
+		{
+			ShowColumn(Column.GetId());
+		}
+	}
+}
+
+////////////////////////////////////////////////////////////////////////////////////////////////////
+// ResetColumns action (ContextMenu)
+////////////////////////////////////////////////////////////////////////////////////////////////////
+
+bool STableTreeView::ContextMenu_ResetColumns_CanExecute() const
+{
+	return true;
+}
+
+////////////////////////////////////////////////////////////////////////////////////////////////////
+
+void STableTreeView::ContextMenu_ResetColumns_Execute()
+{
+	ColumnBeingSorted = DefaultColumnBeingSorted;
+	ColumnSortMode = DefaultColumnSortMode;
+	UpdateCurrentSortingByColumn();
+
+	for (const TSharedPtr<FTableColumn>& ColumnPtr : Table->GetColumns())
+	{
+		const FTableColumn& Column = *ColumnPtr;
+
+		if (Column.ShouldBeVisible() && !Column.IsVisible())
+		{
+			ShowColumn(Column.GetId());
+		}
+		else if (!Column.ShouldBeVisible() && Column.IsVisible())
+		{
+			HideColumn(Column.GetId());
+		}
+	}
+}
+
+////////////////////////////////////////////////////////////////////////////////////////////////////
+
+void STableTreeView::RebuildTree(bool bResync)
+{
+	bool bListHasChanged = false;
+
+	if (bResync)
+	{
+		const int32 PreviousNodeCount = TableTreeNodes.Num();
+		TableTreeNodes.Empty(PreviousNodeCount);
+		//TableTreeNodesMap.Empty(PreviousNodeCount);
+		TableTreeNodesIdMap.Empty(PreviousNodeCount);
+		bListHasChanged = true;
+	}
+
+	TSharedPtr<Trace::IUntypedTable> SourceTable = Table->GetSourceTable();
+	TSharedPtr<Trace::IUntypedTableReader> TableReader = Table->GetTableReader();
+
+	if (Session.IsValid() && SourceTable.IsValid() && TableReader.IsValid())
+	{
+		int32 TotalRowCount = SourceTable->GetRowCount();
+
+		if (TotalRowCount != TableTreeNodes.Num())
+		{
+			bResync = true;
+		}
+
+		if (bResync)
+		{
+			const int32 PreviousNodeCount = TableTreeNodes.Num();
+			TableTreeNodes.Empty(PreviousNodeCount);
+			//TableTreeNodesMap.Empty(PreviousNodeCount);
+			TableTreeNodesIdMap.Empty(PreviousNodeCount);
+			bListHasChanged = true;
+
+			/*
+			const Trace::ITableLayout& TableLayout = SourceTable->GetLayout();
+			const int32 ColumnCount = TableLayout.GetColumnCount();
+
+			for (; TableReader->IsValid(); TableReader->NextRow())
+			{
+				FString Line;
+				for (int32 ColumnIndex = 0; ColumnIndex < ColumnCount; ++ColumnIndex)
+				{
+					switch (TableLayout.GetColumnType(ColumnIndex))
+					{
+					case Trace::TableColumnType_Bool:
+						Line += TableReader->GetValueBool(ColumnIndex) ? "true" : "false";
+						break;
+					case Trace::TableColumnType_Int:
+						Line += FString::Printf(TEXT("%lld"), TableReader->GetValueInt(ColumnIndex));
+						break;
+					case Trace::TableColumnType_Float:
+						Line += FString::Printf(TEXT("%f"), TableReader->GetValueFloat(ColumnIndex));
+						break;
+					case Trace::TableColumnType_Double:
+						Line += FString::Printf(TEXT("%f"), TableReader->GetValueDouble(ColumnIndex));
+						break;
+					case Trace::TableColumnType_CString:
+						FString ValueString = TableReader->GetValueCString(ColumnIndex);
+						ValueString.ReplaceInline(TEXT(","), TEXT(" "));
+						Line += ValueString;
+						break;
+					}
+					if (ColumnIndex < ColumnCount - 1)
+					{
+						Line += TEXT(",");
+					}
+					else
+					{
+						Line += TEXT("\n");
+					}
+				}
+				//UE_LOG(TimingProfiler, Log, TEXT("%s"), *Line.ToString());
+			}
+			*/
+
+			int32 NameColumnIndex = 0;
+
+			for (int32 RowIndex = 0; RowIndex < TotalRowCount; ++RowIndex)
+			{
+				TableReader->SetRowIndex(RowIndex);
+				uint64 NodeId = static_cast<uint64>(RowIndex);
+				FName NodeName(TableReader->GetValueCString(NameColumnIndex));
+				FTableTreeNodePtr NodePtr = MakeShareable(new FTableTreeNode(NodeId, NodeName, Table, RowIndex));
+				TableTreeNodes.Add(NodePtr);
+				//TableTreeNodesMap.Add(Name, NodePtr);
+				TableTreeNodesIdMap.Add(NodeId, NodePtr);
+			}
+		}
+	}
+
+	if (bListHasChanged)
+	{
+		UpdateTree();
+	}
+}
+
+////////////////////////////////////////////////////////////////////////////////////////////////////
+
+void STableTreeView::SelectNodeByNodeId(uint64 Id)
+{
+	FTableTreeNodePtr* NodePtrPtr = TableTreeNodesIdMap.Find(Id);
+	if (NodePtrPtr != nullptr)
+	{
+		FTableTreeNodePtr NodePtr = *NodePtrPtr;
+
+		//UE_LOG(TimingProfiler, Log, TEXT("Select and RequestScrollIntoView %s"), *(*TimerNodePtrPtr)->GetName().ToString());
+		TreeView->SetSelection(NodePtr);
+		TreeView->RequestScrollIntoView(NodePtr);
+	}
+}
+
+////////////////////////////////////////////////////////////////////////////////////////////////////
+
+void STableTreeView::SelectNodeByTableRowIndex(int32 RowIndex)
+{
+	if (RowIndex >= 0 && RowIndex < TableTreeNodes.Num())
+	{
+		FTableTreeNodePtr NodePtr = TableTreeNodes[RowIndex];
+
+		//UE_LOG(TimingProfiler, Log, TEXT("Select and RequestScrollIntoView %s"), *(*TimerNodePtrPtr)->GetName().ToString());
+		TreeView->SetSelection(NodePtr);
+		TreeView->RequestScrollIntoView(NodePtr);
+	}
+}
+
+////////////////////////////////////////////////////////////////////////////////////////////////////
+
+} // namespace Insights
+
+#undef LOCTEXT_NAMESPACE
diff --git a/Engine/Source/Developer/TraceInsights/Private/Insights/Table/Widgets/STableTreeView.h b/Engine/Source/Developer/TraceInsights/Private/Insights/Table/Widgets/STableTreeView.h
new file mode 100644
--- /dev/null
+++ b/Engine/Source/Developer/TraceInsights/Private/Insights/Table/Widgets/STableTreeView.h
@@ -0,0 +1,326 @@
+// Copyright 1998-2019 Epic Games, Inc. All Rights Reserved.
+
+#pragma once
+
+#include "CoreMinimal.h"
+#include "Misc/FilterCollection.h"
+#include "Misc/TextFilter.h"
+#include "SlateFwd.h"
+#include "Widgets/DeclarativeSyntaxSupport.h"
+#include "Widgets/Input/SComboBox.h"
+#include "Widgets/SCompoundWidget.h"
+#include "Widgets/SWidget.h"
+#include "Widgets/Views/SHeaderRow.h"
+#include "Widgets/Views/STableRow.h"
+#include "Widgets/Views/STableViewBase.h"
+#include "Widgets/Views/STreeView.h"
+
+// Insights
+#include "Insights/Table/ViewModels/Table.h"
+#include "Insights/Table/ViewModels/TableColumn.h"
+#include "Insights/Table/ViewModels/TableTreeNode.h"
+
+class FMenuBuilder;
+
+namespace Trace
+{
+	class IAnalysisSession;
+	struct FTimelineEvent;
+}
+
+namespace Insights
+{
+
+class FTreeNodeSorting;
+class FTreeNodeGrouping;
+
+////////////////////////////////////////////////////////////////////////////////////////////////////
+
+/** The filter collection - used for updating the list of tree nodes. */
+typedef TFilterCollection<const FTableTreeNodePtr&> FTableTreeNodeFilterCollection;
+
+/** The text based filter - used for updating the list of tree nodes. */
+typedef TTextFilter<const FTableTreeNodePtr&> FTableTreeNodeTextFilter;
+
+////////////////////////////////////////////////////////////////////////////////////////////////////
+/**
+ * A custom widget used to display the list of tree nodes.
+ */
+class STableTreeView : public SCompoundWidget
+{
+public:
+	/** Default constructor. */
+	STableTreeView();
+
+	/** Virtual destructor. */
+	virtual ~STableTreeView();
+
+	SLATE_BEGIN_ARGS(STableTreeView) {}
+	SLATE_END_ARGS()
+
+	/**
+	 * Construct this widget
+	 * @param InArgs - The declaration data for this widget
+	 */
+	void Construct(const FArguments& InArgs, TSharedPtr<Insights::FTable> InTablePtr);
+
+	TSharedPtr<Insights::FTable> GetTable() const { return Table; }
+
+	/**
+	 * Rebuilds the tree (if necessary).
+	 * @param bResync - If true, it forces a resync even if the list did not changed since last sync.
+	 */
+	virtual void RebuildTree(bool bResync = true);
+
+	const FTableTreeNodePtr* FindNode(uint64 Id) const { return TableTreeNodesIdMap.Find(Id); }
+
+	void SelectNodeByNodeId(uint64 Id);
+	void SelectNodeByTableRowIndex(int32 RowIndex);
+
+protected:
+	void UpdateTree();
+
+	/** Called when the analysis session has changed. */
+	void InsightsManager_OnSessionChanged();
+
+	/**
+	 * Populates OutSearchStrings with the strings that should be used in searching.
+	 *
+	 * @param GroupOrStatNodePtr - the group and stat node to get a text description from.
+	 * @param OutSearchStrings   - an array of strings to use in searching.
+	 *
+	 */
+	void HandleItemToStringArray(const FTableTreeNodePtr& GroupOrStatNodePtr, TArray<FString>& OutSearchStrings) const;
+
+	////////////////////////////////////////////////////////////////////////////////////////////////////
+	// Tree View - Context Menu
+
+	TSharedPtr<SWidget> TreeView_GetMenuContent();
+	void TreeView_BuildSortByMenu(FMenuBuilder& MenuBuilder);
+	void TreeView_BuildViewColumnMenu(FMenuBuilder& MenuBuilder);
+
+	////////////////////////////////////////////////////////////////////////////////////////////////////
+	// Tree View - Columns' Header
+
+	void InitializeAndShowHeaderColumns();
+	void TreeViewHeaderRow_CreateColumnArgs(const FTableColumn& Column);
+	TSharedRef<SWidget> TreeViewHeaderRow_GenerateColumnMenu(const FTableColumn& Column);
+
+	////////////////////////////////////////////////////////////////////////////////////////////////////
+	// Tree View - Misc
+
+	void TreeView_Refresh();
+
+	/**
+	 * Called by STreeView to retrieves the children for the specified parent item.
+	 * @param InParent    - The parent node to retrieve the children from.
+	 * @param OutChildren - List of children for the parent node.
+	 */
+	void TreeView_OnGetChildren(FTableTreeNodePtr InParent, TArray<FTableTreeNodePtr>& OutChildren);
+
+	/** Called by STreeView when selection has changed. */
+	void TreeView_OnSelectionChanged(FTableTreeNodePtr SelectedItem, ESelectInfo::Type SelectInfo);
+
+	/** Called by STreeView when a tree item is double clicked. */
+	void TreeView_OnMouseButtonDoubleClick(FTableTreeNodePtr TreeNode);
+
+	////////////////////////////////////////////////////////////////////////////////////////////////////
+	// Tree View - Table Row
+
+	/** Called by STreeView to generate a table row for the specified item. */
+	TSharedRef<ITableRow> TreeView_OnGenerateRow(FTableTreeNodePtr TreeNode, const TSharedRef<STableViewBase>& OwnerTable);
+
+	void TableRow_SetHoveredCell(TSharedPtr<FTable> TablePtr, TSharedPtr<FTableColumn> ColumnPtr, const FTableTreeNodePtr NodePtr);
+	EHorizontalAlignment TableRow_GetColumnOutlineHAlignment(const FName ColumnId) const;
+
+	FText TableRow_GetHighlightText() const;
+	FName TableRow_GetHighlightedNodeName() const;
+
+	bool TableRow_ShouldBeEnabled(const uint32 TreeNodeId) const;
+
+	////////////////////////////////////////////////////////////////////////////////////////////////////
+	// Filtering
+
+	/** Populates the group and stat tree with items based on the current data. */
+	void ApplyFiltering();
+
+	bool SearchBox_IsEnabled() const;
+	void SearchBox_OnTextChanged(const FText& InFilterText);
+
+	////////////////////////////////////////////////////////////////////////////////////////////////////
+	// Grouping
+
+	void CreateGroupings();
+	void CreateGroups();
+
+	void GroupBy_OnSelectionChanged(TSharedPtr<FTreeNodeGrouping> NewGrouping, ESelectInfo::Type SelectInfo);
+
+	TSharedRef<SWidget> GroupBy_OnGenerateWidget(TSharedPtr<FTreeNodeGrouping> InGrouping) const;
+
+	FText GroupBy_GetSelectedText() const;
+
+	FText GroupBy_GetSelectedTooltipText() const;
+
+	////////////////////////////////////////////////////////////////////////////////////////////////////
+	// Sorting
+
+	void CreateSortings();
+
+	void UpdateCurrentSortingByColumn();
+	void SortTreeNodes();
+
+	EColumnSortMode::Type GetSortModeForColumn(const FName ColumnId) const;
+	void SetSortModeForColumn(const FName& ColumnId, EColumnSortMode::Type SortMode);
+	void OnSortModeChanged(const EColumnSortPriority::Type SortPriority, const FName& ColumnId, const EColumnSortMode::Type SortMode);
+
+	////////////////////////////////////////////////////////////////////////////////////////////////////
+	// Sorting actions
+
+	// SortMode (HeaderMenu)
+	bool HeaderMenu_SortMode_IsChecked(const FName ColumnId, const EColumnSortMode::Type InSortMode);
+	bool HeaderMenu_SortMode_CanExecute(const FName ColumnId, const EColumnSortMode::Type InSortMode) const;
+	void HeaderMenu_SortMode_Execute(const FName ColumnId, const EColumnSortMode::Type InSortMode);
+
+	// SortMode (ContextMenu)
+	bool ContextMenu_SortMode_IsChecked(const EColumnSortMode::Type InSortMode);
+	bool ContextMenu_SortMode_CanExecute(const EColumnSortMode::Type InSortMode) const;
+	void ContextMenu_SortMode_Execute(const EColumnSortMode::Type InSortMode);
+
+	// SortByColumn (ContextMenu)
+	bool ContextMenu_SortByColumn_IsChecked(const FName ColumnId);
+	bool ContextMenu_SortByColumn_CanExecute(const FName ColumnId) const;
+	void ContextMenu_SortByColumn_Execute(const FName ColumnId);
+
+	////////////////////////////////////////////////////////////////////////////////////////////////////
+	// Column visibility actions
+
+	// ShowColumn
+	bool CanShowColumn(const FName ColumnId) const;
+	void ShowColumn(const FName ColumnId);
+
+	// HideColumn
+	bool CanHideColumn(const FName ColumnId) const;
+	void HideColumn(const FName ColumnId);
+
+	// ToggleColumnVisibility
+	bool IsColumnVisible(const FName ColumnId);
+	bool CanToggleColumnVisibility(const FName ColumnId) const;
+	void ToggleColumnVisibility(const FName ColumnId);
+
+	// ShowAllColumns (ContextMenu)
+	bool ContextMenu_ShowAllColumns_CanExecute() const;
+	void ContextMenu_ShowAllColumns_Execute();
+
+	// ResetColumns (ContextMenu)
+	bool ContextMenu_ResetColumns_CanExecute() const;
+	void ContextMenu_ResetColumns_Execute();
+
+	////////////////////////////////////////////////////////////////////////////////////////////////////
+
+protected:
+	/** Table view model. */
+	TSharedPtr<Insights::FTable> Table;
+
+	/** A weak pointer to the profiler session used to populate this widget. */
+	TSharedPtr<const Trace::IAnalysisSession>/*Weak*/ Session;
+
+	//////////////////////////////////////////////////
+	// Tree View, Columns
+
+	/** The child STreeView widget. */
+	TSharedPtr<STreeView<FTableTreeNodePtr>> TreeView;
+
+	/** Column arguments used to initialize a new header column in the tree view, stored as column name to column arguments mapping. */
+	TMap<FName, SHeaderRow::FColumn::FArguments> TreeViewHeaderColumnArgs;
+
+	/** Holds the tree view header row widget which display all columns in the tree view. */
+	TSharedPtr<SHeaderRow> TreeViewHeaderRow;
+
+	/** External scrollbar used to synchronize tree view position. */
+	TSharedPtr<SScrollBar> ExternalScrollbar;
+
+	//////////////////////////////////////////////////
+	// Hovered Column, Hovered Tree Node
+
+	/** Name of the column currently being hovered by the mouse. */
+	FName HoveredColumnId;
+
+	/** A shared pointer to the tree node currently being hovered by the mouse. */
+	FTableTreeNodePtr HoveredNodePtr;
+
+	/** Name of the tree node that should be drawn as highlighted. */
+	FName HighlightedNodeName;
+
+	//////////////////////////////////////////////////
+	// Tree Nodes
+
+	/** All nodes. Each node corresponds to a table row. */
+	TArray<FTableTreeNodePtr> TableTreeNodes;
+
+	/** The array of group nodes. */
+	TArray<FTableTreeNodePtr> GroupNodes;
+
+	/** A filtered array of group and nodes to be displayed in the tree widget. */
+	TArray<FTableTreeNodePtr> FilteredGroupNodes;
+
+	/** All nodes, stored as Id -> FTableTreeNodePtr. */
+	TMap<uint64, FTableTreeNodePtr> TableTreeNodesIdMap;
+
+	/** Currently expanded group nodes. */
+	TSet<FTableTreeNodePtr> ExpandedNodes;
+
+	/** If true, the expanded nodes have been saved before applying a text filter. */
+	bool bExpansionSaved;
+
+	//////////////////////////////////////////////////
+	// Search box and filters
+
+	/** The search box widget used to filter items displayed in the stats and groups tree. */
+	TSharedPtr<SSearchBox> SearchBox;
+
+	/** The text based filter. */
+	TSharedPtr<FTableTreeNodeTextFilter> TextFilter;
+
+	/** The filter collection. */
+	TSharedPtr<FTableTreeNodeFilterCollection> Filters;
+
+	//////////////////////////////////////////////////
+	// Grouping
+
+	TArray<TSharedPtr<FTreeNodeGrouping>> AvailableGroupings;
+
+	TSharedPtr<SComboBox<TSharedPtr<FTreeNodeGrouping>>> GroupByComboBox;
+
+	/** How we group the tree nodes? */
+	TSharedPtr<FTreeNodeGrouping> CurrentGrouping;
+
+	//////////////////////////////////////////////////
+	// Sorting
+
+	/** All available sortings. */
+	TArray<TSharedPtr<FTreeNodeSorting>> AvailableSortings;
+
+	/** Sortings associated with a column. */
+	TMap<FName, TSharedPtr<FTreeNodeSorting>> ColumnToSortingMap;
+
+	/** Current sorting. It is nullptr if sorting is disabled. */
+	TSharedPtr<FTreeNodeSorting> CurrentSorting;
+
+	/** Name of the column currently being sorted. Can be NAME_None if sorting is disabled (CurrentSorting == nullptr) or if a complex sorting is used (CurrentSorting != nullptr). */
+	FName ColumnBeingSorted;
+
+	/** How we sort the nodes? Ascending or Descending. */
+	EColumnSortMode::Type ColumnSortMode;
+
+	static const FName DefaultColumnBeingSorted;
+	static const EColumnSortMode::Type DefaultColumnSortMode;
+
+	//////////////////////////////////////////////////
+
+	double StatsStartTime;
+	double StatsEndTime;
+};
+
+////////////////////////////////////////////////////////////////////////////////////////////////////
+
+} // namespace Insights
diff --git a/Engine/Source/Developer/TraceInsights/Private/Insights/Table/Widgets/STableTreeViewCell.cpp b/Engine/Source/Developer/TraceInsights/Private/Insights/Table/Widgets/STableTreeViewCell.cpp
new file mode 100644
--- /dev/null
+++ b/Engine/Source/Developer/TraceInsights/Private/Insights/Table/Widgets/STableTreeViewCell.cpp
@@ -0,0 +1,130 @@
+// Copyright 1998-2019 Epic Games, Inc. All Rights Reserved.
+
+#include "STableTreeViewCell.h"
+
+#include "EditorStyleSet.h"
+#include "SlateOptMacros.h"
+#include "Widgets/Images/SImage.h"
+#include "Widgets/Input/SButton.h"
+#include "Widgets/SBoxPanel.h"
+#include "Widgets/Text/STextBlock.h"
+#include "Widgets/Views/SExpanderArrow.h"
+
+// Insights
+#include "Insights/Table/ViewModels/Table.h"
+#include "Insights/Table/ViewModels/TableColumn.h"
+#include "Insights/Table/Widgets/STableTreeViewTooltip.h"
+
+#define LOCTEXT_NAMESPACE "STableTreeView"
+
+namespace Insights
+{
+
+////////////////////////////////////////////////////////////////////////////////////////////////////
+
+BEGIN_SLATE_FUNCTION_BUILD_OPTIMIZATION
+
+void STableTreeViewCell::Construct(const FArguments& InArgs, const TSharedRef<ITableRow>& TableRow)
+{
+	TablePtr = InArgs._TablePtr;
+	ColumnPtr = InArgs._ColumnPtr;
+	TableTreeNodePtr = InArgs._TableTreeNodePtr;
+
+	SetHoveredCellDelegate = InArgs._OnSetHoveredCell;
+
+	ChildSlot
+		[
+			GenerateWidgetForColumn(InArgs, TableRow)
+		];
+}
+
+////////////////////////////////////////////////////////////////////////////////////////////////////
+
+TSharedRef<SWidget> STableTreeViewCell::GenerateWidgetForColumn(const FArguments& InArgs, const TSharedRef<ITableRow>& TableRow)
+{
+	if (InArgs._IsNameColumn)
+	{
+		return GenerateWidgetForNameColumn(InArgs, TableRow);
+	}
+	else
+	{
+		return GenerateWidgetForTableColumn(InArgs, TableRow);
+	}
+}
+
+////////////////////////////////////////////////////////////////////////////////////////////////////
+
+TSharedRef<SWidget> STableTreeViewCell::GenerateWidgetForNameColumn(const FArguments& InArgs, const TSharedRef<ITableRow>& TableRow)
+{
+	return
+		SNew(SHorizontalBox)
+
+		+ SHorizontalBox::Slot()
+		.AutoWidth()
+		.HAlign(HAlign_Right)
+		.VAlign(VAlign_Center)
+		[
+			SNew(SExpanderArrow, TableRow)
+		]
+
+		// Info icon + tooltip
+		+ SHorizontalBox::Slot()
+		.AutoWidth()
+		.HAlign(HAlign_Center)
+		.VAlign(VAlign_Center)
+		[
+			SNew(SImage)
+			.Visibility(this, &STableTreeViewCell::GetHintIconVisibility)
+			.Image(FEditorStyle::GetBrush("Profiler.Tooltip.HintIcon10"))
+			.ToolTip(STableTreeViewTooltip::GetCellTooltip(TableTreeNodePtr, ColumnPtr))
+		]
+
+		// Name
+		+ SHorizontalBox::Slot()
+		.AutoWidth()
+		.VAlign(VAlign_Center)
+		.HAlign(ColumnPtr->GetHorizontalAlignment())
+		.Padding(FMargin(2.0f, 0.0f))
+		[
+			SNew(STextBlock)
+			.Text(this, &STableTreeViewCell::GetDisplayName)
+			.HighlightText(InArgs._HighlightText)
+			.TextStyle(FEditorStyle::Get(), TEXT("Profiler.Tooltip"))
+			.ColorAndOpacity(this, &STableTreeViewCell::GetColorAndOpacity)
+			.ShadowColorAndOpacity(this, &STableTreeViewCell::GetShadowColorAndOpacity)
+		]
+	;
+}
+
+////////////////////////////////////////////////////////////////////////////////////////////////////
+
+TSharedRef<SWidget> STableTreeViewCell::GenerateWidgetForTableColumn(const FArguments& InArgs, const TSharedRef<ITableRow>& TableRow)
+{
+	const FText FormattedValue = ColumnPtr->GetValueAsText(TableTreeNodePtr->GetRowId());
+
+	return
+		SNew(SHorizontalBox)
+
+		// Value
+		+ SHorizontalBox::Slot()
+		.FillWidth(1.0f)
+		.VAlign(VAlign_Center)
+		.HAlign(ColumnPtr->GetHorizontalAlignment())
+		.Padding(FMargin(2.0f, 0.0f))
+		[
+			SNew(STextBlock)
+			.Text(FormattedValue)
+			.TextStyle(FEditorStyle::Get(), TEXT("Profiler.Tooltip"))
+			.ColorAndOpacity(this, &STableTreeViewCell::GetStatsColorAndOpacity)
+			.ShadowColorAndOpacity(this, &STableTreeViewCell::GetShadowColorAndOpacity)
+		]
+	;
+}
+
+END_SLATE_FUNCTION_BUILD_OPTIMIZATION
+
+////////////////////////////////////////////////////////////////////////////////////////////////////
+
+} // namespace Insights
+
+#undef LOCTEXT_NAMESPACE
diff --git a/Engine/Source/Developer/TraceInsights/Private/Insights/Table/Widgets/STableTreeViewCell.h b/Engine/Source/Developer/TraceInsights/Private/Insights/Table/Widgets/STableTreeViewCell.h
new file mode 100644
--- /dev/null
+++ b/Engine/Source/Developer/TraceInsights/Private/Insights/Table/Widgets/STableTreeViewCell.h
@@ -0,0 +1,148 @@
+// Copyright 1998-2019 Epic Games, Inc. All Rights Reserved.
+
+#pragma once
+
+#include "CoreMinimal.h"
+#include "Widgets/DeclarativeSyntaxSupport.h"
+#include "Widgets/SCompoundWidget.h"
+
+// Insights
+#include "Insights/Table/ViewModels/TableTreeNode.h"
+
+class ITableRow;
+
+namespace Insights
+{
+
+class FTable;
+class FTableColumn;
+
+DECLARE_DELEGATE_ThreeParams(FSetHoveredTableTreeViewCell, TSharedPtr<FTable> /*TablePtr*/, TSharedPtr<FTableColumn> /*ColumnPtr*/, const FTableTreeNodePtr /*NodePtr*/);
+
+class STableTreeViewCell : public SCompoundWidget
+{
+public:
+	SLATE_BEGIN_ARGS(STableTreeViewCell) {}
+		SLATE_EVENT(FSetHoveredTableTreeViewCell, OnSetHoveredCell)
+		SLATE_ATTRIBUTE(FText, HighlightText)
+		SLATE_ARGUMENT(TSharedPtr<FTable>, TablePtr)
+		SLATE_ARGUMENT(TSharedPtr<FTableColumn>, ColumnPtr)
+		SLATE_ARGUMENT(FTableTreeNodePtr, TableTreeNodePtr)
+		SLATE_ARGUMENT(bool, IsNameColumn)
+	SLATE_END_ARGS()
+
+	void Construct(const FArguments& InArgs, const TSharedRef<ITableRow>& TableRow);
+
+protected:
+	TSharedRef<SWidget> GenerateWidgetForColumn(const FArguments& InArgs, const TSharedRef<ITableRow>& TableRow);
+	TSharedRef<SWidget> GenerateWidgetForNameColumn(const FArguments& InArgs, const TSharedRef<ITableRow>& TableRow);
+	TSharedRef<SWidget> GenerateWidgetForTableColumn(const FArguments& InArgs, const TSharedRef<ITableRow>& TableRow);
+
+	/**
+	 * The system will use this event to notify a widget that the cursor has entered it. This event is NOT bubbled.
+	 *
+	 * @param MyGeometry The Geometry of the widget receiving the event
+	 * @param MouseEvent Information about the input event
+	 */
+	virtual void OnMouseEnter(const FGeometry& MyGeometry, const FPointerEvent& MouseEvent) override
+	{
+		SCompoundWidget::OnMouseEnter(MyGeometry, MouseEvent);
+		SetHoveredCellDelegate.ExecuteIfBound(TablePtr, ColumnPtr, TableTreeNodePtr);
+	}
+
+	/**
+	 * The system will use this event to notify a widget that the cursor has left it. This event is NOT bubbled.
+	 *
+	 * @param MouseEvent Information about the input event
+	 */
+	virtual void OnMouseLeave(const FPointerEvent& MouseEvent) override
+	{
+		SCompoundWidget::OnMouseLeave(MouseEvent);
+		SetHoveredCellDelegate.ExecuteIfBound(nullptr, nullptr, nullptr);
+	}
+
+	/**
+	 * Called during drag and drop when the drag enters a widget.
+	 *
+	 * Enter/Leave events in slate are meant as lightweight notifications.
+	 * So we do not want to capture mouse or set focus in response to these.
+	 * However, OnDragEnter must also support external APIs (e.g. OLE Drag/Drop)
+	 * Those require that we let them know whether we can handle the content
+	 * being dragged OnDragEnter.
+	 *
+	 * The concession is to return a can_handled/cannot_handle
+	 * boolean rather than a full FReply.
+	 *
+	 * @param MyGeometry      The geometry of the widget receiving the event.
+	 * @param DragDropEvent   The drag and drop event.
+	 *
+	 * @return A reply that indicated whether the contents of the DragDropEvent can potentially be processed by this widget.
+	 */
+	virtual void OnDragEnter(const FGeometry& MyGeometry, const FDragDropEvent& DragDropEvent) override
+	{
+		SCompoundWidget::OnDragEnter(MyGeometry, DragDropEvent);
+		SetHoveredCellDelegate.ExecuteIfBound(TablePtr, ColumnPtr, TableTreeNodePtr);
+	}
+
+	/**
+	 * Called during drag and drop when the drag leaves a widget.
+	 *
+	 * @param DragDropEvent   The drag and drop event.
+	 */
+	virtual void OnDragLeave(const FDragDropEvent& DragDropEvent)  override
+	{
+		SCompoundWidget::OnDragLeave(DragDropEvent);
+		SetHoveredCellDelegate.ExecuteIfBound(nullptr, nullptr, nullptr);
+	}
+
+	EVisibility GetHintIconVisibility() const
+	{
+		return IsHovered() ? EVisibility::Visible : EVisibility::Hidden;
+	}
+
+	FText GetDisplayName() const
+	{
+		return TableTreeNodePtr->GetDisplayName();
+	}
+
+	FSlateColor GetColorAndOpacity() const
+	{
+		const FLinearColor TextColor =
+			TableTreeNodePtr->IsFiltered() ?
+				FLinearColor(1.0f, 1.0f, 1.0f, 0.5f) :
+				FLinearColor(1.0f, 1.0f, 1.0f, 1.0f);
+		return TextColor;
+	}
+
+	FSlateColor GetStatsColorAndOpacity() const
+	{
+		const FLinearColor TextColor =
+			TableTreeNodePtr->IsFiltered() ?
+				FLinearColor(1.0f, 1.0f, 1.0f, 0.5f) :
+				FLinearColor(1.0f, 1.0f, 1.0f, 1.0f);
+		return TextColor;
+	}
+
+	FLinearColor GetShadowColorAndOpacity() const
+	{
+		const FLinearColor ShadowColor =
+			TableTreeNodePtr->IsFiltered() ?
+				FLinearColor(0.f, 0.f, 0.f, 0.25f) :
+				FLinearColor(0.0f, 0.0f, 0.0f, 0.5f);
+		return ShadowColor;
+	}
+
+protected:
+	/** A shared pointer to the table view model. */
+	TSharedPtr<FTable> TablePtr;
+
+	/** A shared pointer to the table column view model. */
+	TSharedPtr<FTableColumn> ColumnPtr;
+
+	/** A shared pointer to the tree node. */
+	FTableTreeNodePtr TableTreeNodePtr;
+
+	FSetHoveredTableTreeViewCell SetHoveredCellDelegate;
+};
+
+} // namespace Insights
diff --git a/Engine/Source/Developer/TraceInsights/Private/Insights/Table/Widgets/STableTreeViewRow.cpp b/Engine/Source/Developer/TraceInsights/Private/Insights/Table/Widgets/STableTreeViewRow.cpp
new file mode 100644
--- /dev/null
+++ b/Engine/Source/Developer/TraceInsights/Private/Insights/Table/Widgets/STableTreeViewRow.cpp
@@ -0,0 +1,191 @@
+// Copyright 1998-2019 Epic Games, Inc. All Rights Reserved.
+
+#include "STableTreeViewRow.h"
+
+#include "EditorStyleSet.h"
+#include "SlateOptMacros.h"
+#include "Widgets/SOverlay.h"
+#include "Widgets/Images/SImage.h"
+
+// Insights
+#include "Insights/Common/TimeUtils.h"
+#include "Insights/Table/ViewModels/Table.h"
+#include "Insights/Table/ViewModels/TableColumn.h"
+#include "Insights/Table/Widgets/STableTreeViewCell.h"
+
+#define LOCTEXT_NAMESPACE "STableTreeView"
+
+namespace Insights
+{
+
+////////////////////////////////////////////////////////////////////////////////////////////////////
+
+BEGIN_SLATE_FUNCTION_BUILD_OPTIMIZATION
+
+void STableTreeViewRow::Construct(const FArguments& InArgs, const TSharedRef<STableViewBase>& InOwnerTableView)
+{
+	OnShouldBeEnabled = InArgs._OnShouldBeEnabled;
+	IsColumnVisibleDelegate = InArgs._OnIsColumnVisible;
+	SetHoveredCellDelegate = InArgs._OnSetHoveredCell;
+	GetColumnOutlineHAlignmentDelegate = InArgs._OnGetColumnOutlineHAlignmentDelegate;
+
+	HighlightText = InArgs._HighlightText;
+	HighlightedNodeName = InArgs._HighlightedNodeName;
+
+	TablePtr = InArgs._TablePtr;
+	TableTreeNodePtr = InArgs._TableTreeNodePtr;
+
+	SetEnabled(TAttribute<bool>(this, &STableTreeViewRow::HandleShouldBeEnabled));
+
+	SMultiColumnTableRow<FTableTreeNodePtr>::Construct(SMultiColumnTableRow<FTableTreeNodePtr>::FArguments(), InOwnerTableView);
+}
+
+////////////////////////////////////////////////////////////////////////////////////////////////////
+
+TSharedRef<SWidget> STableTreeViewRow::GenerateWidgetForColumn(const FName& ColumnId)
+{
+	TSharedPtr<FTableColumn> ColumnPtr = TablePtr->FindColumnChecked(ColumnId);
+
+	return
+		SNew(SOverlay)
+		.Visibility(EVisibility::SelfHitTestInvisible)
+
+		+SOverlay::Slot()
+		.Padding(0.0f)
+		[
+			SNew(SImage)
+			.Image(FEditorStyle::GetBrush("Profiler.LineGraphArea"))
+			.ColorAndOpacity(this, &STableTreeViewRow::GetBackgroundColorAndOpacity)
+		]
+
+		+SOverlay::Slot()
+		.Padding(0.0f)
+		[
+			SNew(SImage)
+			.Image(this, &STableTreeViewRow::GetOutlineBrush, ColumnId)
+			.ColorAndOpacity(this, &STableTreeViewRow::GetOutlineColorAndOpacity)
+		]
+
+		+SOverlay::Slot()
+		[
+			SNew(STableTreeViewCell, SharedThis(this))
+			.Visibility(this, &STableTreeViewRow::IsColumnVisible, ColumnId)
+			.TablePtr(TablePtr)
+			.ColumnPtr(ColumnPtr)
+			.TableTreeNodePtr(TableTreeNodePtr)
+			.HighlightText(HighlightText)
+			.IsNameColumn(ColumnPtr->IsHierarchy())
+			.OnSetHoveredCell(this, &STableTreeViewRow::OnSetHoveredCell)
+		];
+}
+
+END_SLATE_FUNCTION_BUILD_OPTIMIZATION
+
+////////////////////////////////////////////////////////////////////////////////////////////////////
+
+FReply STableTreeViewRow::OnDragDetected(const FGeometry& MyGeometry, const FPointerEvent& MouseEvent)
+{
+	return SMultiColumnTableRow<FTableTreeNodePtr>::OnDragDetected(MyGeometry, MouseEvent);
+}
+
+////////////////////////////////////////////////////////////////////////////////////////////////////
+
+FSlateColor STableTreeViewRow::GetBackgroundColorAndOpacity() const
+{
+	return FLinearColor(0.0f, 0.0f, 0.0f, 1.0f);
+}
+
+////////////////////////////////////////////////////////////////////////////////////////////////////
+
+FSlateColor STableTreeViewRow::GetBackgroundColorAndOpacity(double Time) const
+{
+	const FLinearColor Color =	Time > TimeUtils::Second      ? FLinearColor(0.3f, 0.0f, 0.0f, 1.0f) :
+								Time > TimeUtils::Milisecond  ? FLinearColor(0.3f, 0.1f, 0.0f, 1.0f) :
+								Time > TimeUtils::Microsecond ? FLinearColor(0.0f, 0.1f, 0.0f, 1.0f) :
+																FLinearColor(0.0f, 0.0f, 0.0f, 1.0f);
+	return Color;
+}
+
+////////////////////////////////////////////////////////////////////////////////////////////////////
+
+FSlateColor STableTreeViewRow::GetOutlineColorAndOpacity() const
+{
+	const FLinearColor NoColor(0.0f, 0.0f, 0.0f, 0.0f);
+	const bool bShouldBeHighlighted = TableTreeNodePtr->GetName() == HighlightedNodeName.Get();
+	const FLinearColor OutlineColorAndOpacity = bShouldBeHighlighted ? FLinearColor(FColorList::SlateBlue) : NoColor;
+	return OutlineColorAndOpacity;
+}
+
+////////////////////////////////////////////////////////////////////////////////////////////////////
+
+const FSlateBrush* STableTreeViewRow::GetOutlineBrush(const FName ColumnId) const
+{
+	EHorizontalAlignment Result = HAlign_Center;
+	if (IsColumnVisibleDelegate.IsBound())
+	{
+		Result = GetColumnOutlineHAlignmentDelegate.Execute(ColumnId);
+	}
+
+	const FSlateBrush* Brush = nullptr;
+	if (Result == HAlign_Left)
+	{
+		Brush = FEditorStyle::GetBrush("Profiler.EventGraph.Border.L");
+	}
+	else if(Result == HAlign_Right)
+	{
+		Brush = FEditorStyle::GetBrush("Profiler.EventGraph.Border.R");
+	}
+	else
+	{
+		Brush = FEditorStyle::GetBrush("Profiler.EventGraph.Border.TB");
+	}
+	return Brush;
+}
+
+////////////////////////////////////////////////////////////////////////////////////////////////////
+
+bool STableTreeViewRow::HandleShouldBeEnabled() const
+{
+	bool bResult = false;
+
+	if (TableTreeNodePtr->IsGroup())
+	{
+		bResult = true;
+	}
+	else
+	{
+		if (OnShouldBeEnabled.IsBound())
+		{
+			bResult = OnShouldBeEnabled.Execute(TableTreeNodePtr->GetId());
+		}
+	}
+
+	return bResult;
+}
+
+////////////////////////////////////////////////////////////////////////////////////////////////////
+
+EVisibility STableTreeViewRow::IsColumnVisible(const FName ColumnId) const
+{
+	EVisibility Result = EVisibility::Collapsed;
+
+	if (IsColumnVisibleDelegate.IsBound())
+	{
+		Result = IsColumnVisibleDelegate.Execute(ColumnId) ? EVisibility::Visible : EVisibility::Collapsed;
+	}
+
+	return Result;
+}
+
+////////////////////////////////////////////////////////////////////////////////////////////////////
+
+void STableTreeViewRow::OnSetHoveredCell(TSharedPtr<FTable> InTablePtr, TSharedPtr<FTableColumn> InColumnPtr, const FTableTreeNodePtr InTreeNodePtr)
+{
+	SetHoveredCellDelegate.ExecuteIfBound(InTablePtr, InColumnPtr, InTreeNodePtr);
+}
+
+////////////////////////////////////////////////////////////////////////////////////////////////////
+
+} // namespace Insights
+
+#undef LOCTEXT_NAMESPACE
diff --git a/Engine/Source/Developer/TraceInsights/Private/Insights/Table/Widgets/STableTreeViewRow.h b/Engine/Source/Developer/TraceInsights/Private/Insights/Table/Widgets/STableTreeViewRow.h
new file mode 100644
--- /dev/null
+++ b/Engine/Source/Developer/TraceInsights/Private/Insights/Table/Widgets/STableTreeViewRow.h
@@ -0,0 +1,90 @@
+// Copyright 1998-2019 Epic Games, Inc. All Rights Reserved.
+
+#pragma once
+
+#include "CoreMinimal.h"
+#include "Widgets/DeclarativeSyntaxSupport.h"
+#include "Widgets/Views/SHeaderRow.h"
+#include "Widgets/Views/STableViewBase.h"
+#include "Widgets/Views/STableRow.h"
+#include "Widgets/Views/STreeView.h"
+
+// Insights
+#include "Insights/Table/ViewModels/TableTreeNode.h"
+
+namespace Insights
+{
+
+class FTable;
+class FTableColumn;
+
+DECLARE_DELEGATE_RetVal_OneParam(bool, FShouldBeEnabledDelegate, const uint32 /*TreeNodeId*/);
+DECLARE_DELEGATE_RetVal_OneParam(bool, FIsColumnVisibleDelegate, const FName /*ColumnId*/);
+DECLARE_DELEGATE_RetVal_OneParam(EHorizontalAlignment, FGetColumnOutlineHAlignmentDelegate, const FName /*ColumnId*/);
+DECLARE_DELEGATE_ThreeParams(FSetHoveredTableTreeViewCell, TSharedPtr<FTable> /*TablePtr*/, TSharedPtr<FTableColumn> /*ColumnPtr*/, const FTableTreeNodePtr /*NodePtr*/);
+
+/** Widget that represents a table row in the tree control. Generates widgets for each column on demand. */
+class STableTreeViewRow : public SMultiColumnTableRow<FTableTreeNodePtr>
+{
+public:
+	SLATE_BEGIN_ARGS(STableTreeViewRow) {}
+		SLATE_EVENT(FShouldBeEnabledDelegate, OnShouldBeEnabled)
+		SLATE_EVENT(FIsColumnVisibleDelegate, OnIsColumnVisible)
+		SLATE_EVENT(FGetColumnOutlineHAlignmentDelegate, OnGetColumnOutlineHAlignmentDelegate)
+		SLATE_EVENT(FSetHoveredTableTreeViewCell, OnSetHoveredCell)
+		SLATE_ATTRIBUTE(FText, HighlightText)
+		SLATE_ATTRIBUTE(FName, HighlightedNodeName)
+		SLATE_ARGUMENT(TSharedPtr<FTable>, TablePtr)
+		SLATE_ARGUMENT(FTableTreeNodePtr, TableTreeNodePtr)
+	SLATE_END_ARGS()
+
+public:
+	void Construct(const FArguments& InArgs, const TSharedRef<STableViewBase>& InOwnerTableView);
+
+	virtual TSharedRef<SWidget> GenerateWidgetForColumn(const FName& ColumnId) override;
+
+	/**
+	 * Called when Slate detects that a widget started to be dragged.
+	 * Usage:
+	 * A widget can ask Slate to detect a drag.
+	 * OnMouseDown() reply with FReply::Handled().DetectDrag(SharedThis(this)).
+	 * Slate will either send an OnDragDetected() event or do nothing.
+	 * If the user releases a mouse button or leaves the widget before
+	 * a drag is triggered (maybe user started at the very edge) then no event will be
+	 * sent.
+	 *
+	 * @param  InMyGeometry  Widget geometry
+	 * @param  InMouseEvent  MouseMove that triggered the drag
+	 *
+	 */
+	virtual FReply OnDragDetected(const FGeometry& MyGeometry, const FPointerEvent& MouseEvent) override;
+
+protected:
+	FSlateColor GetBackgroundColorAndOpacity() const;
+	FSlateColor GetBackgroundColorAndOpacity(double Time) const;
+	FSlateColor GetOutlineColorAndOpacity() const;
+	const FSlateBrush* GetOutlineBrush(const FName ColumnId) const;
+	bool HandleShouldBeEnabled() const;
+	EVisibility IsColumnVisible(const FName ColumnId) const;
+	void OnSetHoveredCell(TSharedPtr<FTable> InTablePtr, TSharedPtr<FTableColumn> InColumnPtr, const FTableTreeNodePtr InTreeNodePtr);
+
+protected:
+	/** A shared pointer to the table view model. */
+	TSharedPtr<FTable> TablePtr;
+
+	/** Data context for this table row. */
+	FTableTreeNodePtr TableTreeNodePtr;
+
+	FShouldBeEnabledDelegate OnShouldBeEnabled;
+	FIsColumnVisibleDelegate IsColumnVisibleDelegate;
+	FSetHoveredTableTreeViewCell SetHoveredCellDelegate;
+	FGetColumnOutlineHAlignmentDelegate GetColumnOutlineHAlignmentDelegate;
+
+	/** Text to be highlighted on tree node's name. */
+	TAttribute<FText> HighlightText;
+
+	/** Name of the tree node that should be drawn as highlighted. */
+	TAttribute<FName> HighlightedNodeName;
+};
+
+} // namespace Insights
