[Insights]

diff --git a/Engine/Source/Developer/TraceInsights/Private/Insights/Table/Widgets/STableTreeViewTooltip.cpp b/Engine/Source/Developer/TraceInsights/Private/Insights/Table/Widgets/STableTreeViewTooltip.cpp
new file mode 100644
--- /dev/null
+++ b/Engine/Source/Developer/TraceInsights/Private/Insights/Table/Widgets/STableTreeViewTooltip.cpp
@@ -0,0 +1,195 @@
+// Copyright 1998-2019 Epic Games, Inc. All Rights Reserved.
+
+#include "STableTreeViewTooltip.h"
+
+#include "EditorStyleSet.h"
+#include "SlateOptMacros.h"
+#include "Widgets/Images/SImage.h"
+#include "Widgets/Layout/SGridPanel.h"
+#include "Widgets/Layout/SSeparator.h"
+#include "Widgets/SBoxPanel.h"
+#include "Widgets/SToolTip.h"
+#include "Widgets/Text/STextBlock.h"
+
+// Insights
+#include "Insights/Common/TimeUtils.h"
+#include "Insights/Table/ViewModels/TableTreeNode.h"
+#include "Insights/Table/ViewModels/Table.h"
+#include "Insights/Table/ViewModels/TableColumn.h"
+
+#define LOCTEXT_NAMESPACE "STableTreeView"
+
+namespace Insights
+{
+
+////////////////////////////////////////////////////////////////////////////////////////////////////
+
+BEGIN_SLATE_FUNCTION_BUILD_OPTIMIZATION
+
+TSharedPtr<SToolTip> STableTreeViewTooltip::GetTableTooltip(const FTable& Table)
+{
+	TSharedPtr<SToolTip> ColumnTooltip =
+		SNew(SToolTip)
+		[
+			SNew(SVerticalBox)
+
+			+ SVerticalBox::Slot()
+			.AutoHeight()
+			.Padding(2.0f)
+			[
+				SNew(STextBlock)
+				.Text(Table.GetDisplayName())
+				.TextStyle(FEditorStyle::Get(), TEXT("Profiler.TooltipBold"))
+			]
+
+			+ SVerticalBox::Slot()
+			.AutoHeight()
+			.Padding(2.0f)
+			[
+				SNew(STextBlock)
+				.Text(Table.GetDescription())
+				.TextStyle(FEditorStyle::Get(), TEXT("Profiler.Tooltip"))
+			]
+		];
+
+	return ColumnTooltip;
+}
+
+////////////////////////////////////////////////////////////////////////////////////////////////////
+
+TSharedPtr<SToolTip> STableTreeViewTooltip::GetColumnTooltip(const FTableColumn& Column)
+{
+	TSharedPtr<SToolTip> ColumnTooltip =
+		SNew(SToolTip)
+		[
+			SNew(SVerticalBox)
+
+			+ SVerticalBox::Slot()
+			.AutoHeight()
+			.Padding(2.0f)
+			[
+				SNew(STextBlock)
+				.Text(Column.GetTitleName())
+				.TextStyle(FEditorStyle::Get(), TEXT("Profiler.TooltipBold"))
+			]
+
+			+ SVerticalBox::Slot()
+			.AutoHeight()
+			.Padding(2.0f)
+			[
+				SNew(STextBlock)
+				.Text(Column.GetDescription())
+				.TextStyle(FEditorStyle::Get(), TEXT("Profiler.Tooltip"))
+			]
+		];
+
+	return ColumnTooltip;
+}
+
+////////////////////////////////////////////////////////////////////////////////////////////////////
+
+TSharedPtr<SToolTip> STableTreeViewTooltip::GetCellTooltip(const TSharedPtr<FTableTreeNode> TreeNodePtr, const TSharedPtr<FTableColumn> ColumnPtr)
+{
+	const FLinearColor TextColor(0.8f, 0.8f, 0.8f, 1.0f);
+	const FLinearColor ValueColor(1.0f, 1.0f, 1.0f, 1.0f);
+
+	//const int32 NumDigits = 5;
+	//const FText TimeText = FText::FromString(TimeUtils::FormatTimeAuto(TotalTime));
+	//const FText TimeText = FText::FromString(TimeUtils::FormatTimeMs(Time, NumDigits, true));
+
+	TSharedPtr<SGridPanel> GridPanel;
+	TSharedPtr<SHorizontalBox> HBox;
+
+	TSharedPtr<SToolTip> TableCellTooltip =
+		SNew(SToolTip)
+		[
+			SAssignNew(HBox, SHorizontalBox)
+
+			+SHorizontalBox::Slot()
+			.AutoWidth()
+			[
+				SNew(SVerticalBox)
+
+				+SVerticalBox::Slot()
+				.AutoHeight()
+				.Padding(2.0f)
+				[
+					SNew(SSeparator)
+					.Orientation(Orient_Horizontal)
+				]
+
+				+SVerticalBox::Slot()
+				.AutoHeight()
+				.Padding(2.0f)
+				[
+					SNew(SGridPanel)
+
+					// Id: [Id]
+					+SGridPanel::Slot(0, 0)
+					.Padding(2.0f)
+					[
+						SNew(STextBlock)
+						.Text(LOCTEXT("TT_Id", "Id:"))
+						.TextStyle(FEditorStyle::Get(), TEXT("Profiler.TooltipBold"))
+					]
+					+SGridPanel::Slot(1, 0)
+					.Padding(2.0f)
+					[
+						SNew(STextBlock)
+						.Text(FText::AsNumber(TreeNodePtr->GetId()))
+						.TextStyle(FEditorStyle::Get(), TEXT("Profiler.Tooltip"))
+					]
+
+					// Name: [Name]
+					+SGridPanel::Slot(0, 1)
+					.Padding(2.0f)
+					[
+						SNew(STextBlock)
+						.Text(LOCTEXT("TT_Name", "Name:"))
+						.TextStyle(FEditorStyle::Get(), TEXT("Profiler.TooltipBold"))
+					]
+					+SGridPanel::Slot(1, 1)
+					.Padding(2.0f)
+					[
+						SNew(STextBlock)
+						.Text(FText::FromName(TreeNodePtr->GetName()))
+						.TextStyle(FEditorStyle::Get(), TEXT("Profiler.Tooltip"))
+					]
+
+					// Type: [Type]
+					+ SGridPanel::Slot(0, 3)
+					.Padding(2.0f)
+					[
+						SNew(STextBlock)
+						.Text(LOCTEXT("TT_Type", "Type:"))
+						.TextStyle(FEditorStyle::Get(), TEXT("Profiler.TooltipBold"))
+					]
+					+ SGridPanel::Slot(1, 3)
+					.Padding(2.0f)
+					[
+						SNew(STextBlock)
+						.Text(FText::FromName(TreeNodePtr->GetTypeId()))
+						.TextStyle(FEditorStyle::Get(), TEXT("Profiler.Tooltip"))
+					]
+				]
+
+				+SVerticalBox::Slot()
+				.AutoHeight()
+				.Padding(2.0f)
+				[
+					SNew(SSeparator)
+					.Orientation(Orient_Horizontal)
+				]
+			]
+		];
+
+	return TableCellTooltip;
+}
+
+END_SLATE_FUNCTION_BUILD_OPTIMIZATION
+
+////////////////////////////////////////////////////////////////////////////////////////////////////
+
+} // namespace Insights
+
+#undef LOCTEXT_NAMESPACE
diff --git a/Engine/Source/Developer/TraceInsights/Private/Insights/Table/Widgets/STableTreeViewTooltip.h b/Engine/Source/Developer/TraceInsights/Private/Insights/Table/Widgets/STableTreeViewTooltip.h
new file mode 100644
--- /dev/null
+++ b/Engine/Source/Developer/TraceInsights/Private/Insights/Table/Widgets/STableTreeViewTooltip.h
@@ -0,0 +1,33 @@
+// Copyright 1998-2019 Epic Games, Inc. All Rights Reserved.
+
+#pragma once
+
+#include "CoreMinimal.h"
+#include "EditorStyleSet.h"
+#include "SlateOptMacros.h"
+#include "TraceServices/AnalysisService.h"
+#include "Widgets/Images/SImage.h"
+#include "Widgets/Layout/SGridPanel.h"
+#include "Widgets/Layout/SSeparator.h"
+#include "Widgets/SToolTip.h"
+#include "Widgets/Text/STextBlock.h"
+
+namespace Insights
+{
+
+class FTable;
+class FTableColumn;
+class FTableTreeNode;
+
+/** Tooltip for STableTreeView widget. */
+class STableTreeViewTooltip
+{
+public:
+	STableTreeViewTooltip() = delete;
+
+	static TSharedPtr<SToolTip> GetTableTooltip(const FTable& Table);
+	static TSharedPtr<SToolTip> GetColumnTooltip(const FTableColumn& Column);
+	static TSharedPtr<SToolTip> GetCellTooltip(const TSharedPtr<FTableTreeNode> TreeNodePtr, const TSharedPtr<FTableColumn> ColumnPtr);
+};
+
+} // namespace Insights
diff --git a/Engine/Source/Developer/TraceInsights/Private/Insights/TimingProfilerManager.cpp b/Engine/Source/Developer/TraceInsights/Private/Insights/TimingProfilerManager.cpp
--- a/Engine/Source/Developer/TraceInsights/Private/Insights/TimingProfilerManager.cpp
+++ b/Engine/Source/Developer/TraceInsights/Private/Insights/TimingProfilerManager.cpp
@@ -3,7 +3,6 @@
 #include "TimingProfilerManager.h"
 
 #include "Modules/ModuleManager.h"
-#include "Templates/UniquePtr.h"
 #include "TraceServices/AnalysisService.h"
 #include "TraceServices/SessionService.h"
 
@@ -124,30 +123,7 @@ void FTimingProfilerManager::OnSessionChanged()
 	TSharedPtr<STimingProfilerWindow> Wnd = GetProfilerWindow();
 	if (Wnd.IsValid())
 	{
-		if (Wnd->FrameTrack)
-		{
-			Wnd->FrameTrack->Reset();
-		}
-		if (Wnd->GraphTrack)
-		{
-			Wnd->GraphTrack->Reset();
-		}
-		if (Wnd->TimingView)
-		{
-			Wnd->TimingView->Reset();
-		}
-		if (Wnd->TimersView)
-		{
-			Wnd->TimersView->RebuildTree();
-		}
-		if (Wnd->StatsView)
-		{
-			Wnd->StatsView->RebuildTree();
-		}
-		if (Wnd->LogView)
-		{
-			Wnd->LogView->Reset();
-		}
+		Wnd->Reset();
 	}
 }
 
@@ -155,8 +131,6 @@ void FTimingProfilerManager::OnSessionChanged()
 
 void FTimingProfilerManager::ShowHideFramesTrack(const bool bFramesTrackVisibleState)
 {
-	bool bWasFramesTrackVisible = bIsFramesTrackVisible;
-
 	bIsFramesTrackVisible = bFramesTrackVisibleState;
 
 	TSharedPtr<STimingProfilerWindow> Wnd = GetProfilerWindow();
@@ -170,8 +144,6 @@ void FTimingProfilerManager::ShowHideFramesTrack(const bool bFramesTrackVisibleS
 
 void FTimingProfilerManager::ShowHideGraphTrack(const bool bGraphTrackVisibleState)
 {
-	bool bWasGraphTrackVisible = bIsGraphTrackVisible;
-
 	bIsGraphTrackVisible = bGraphTrackVisibleState;
 
 	TSharedPtr<STimingProfilerWindow> Wnd = GetProfilerWindow();
@@ -185,8 +157,6 @@ void FTimingProfilerManager::ShowHideGraphTrack(const bool bGraphTrackVisibleSta
 
 void FTimingProfilerManager::ShowHideTimingView(const bool bTimingViewVisibleState)
 {
-	bool bWasTimingViewVisible = bIsTimingViewVisible;
-
 	bIsTimingViewVisible = bTimingViewVisibleState;
 
 	TSharedPtr<STimingProfilerWindow> Wnd = GetProfilerWindow();
@@ -200,8 +170,6 @@ void FTimingProfilerManager::ShowHideTimingView(const bool bTimingViewVisibleSta
 
 void FTimingProfilerManager::ShowHideTimersView(const bool bTimersViewVisibleState)
 {
-	bool bWasTimersViewVisible = bIsTimersViewVisible;
-
 	bIsTimersViewVisible = bTimersViewVisibleState;
 
 	TSharedPtr<STimingProfilerWindow> Wnd = GetProfilerWindow();
@@ -215,8 +183,6 @@ void FTimingProfilerManager::ShowHideTimersView(const bool bTimersViewVisibleSta
 
 void FTimingProfilerManager::ShowHideStatsCountersView(const bool bStatsCountersViewVisibleState)
 {
-	bool bWasStatsCountersViewVisible = bIsStatsCountersViewVisible;
-
 	bIsStatsCountersViewVisible = bStatsCountersViewVisibleState;
 
 	TSharedPtr<STimingProfilerWindow> Wnd = GetProfilerWindow();
@@ -230,8 +196,6 @@ void FTimingProfilerManager::ShowHideStatsCountersView(const bool bStatsCounters
 
 void FTimingProfilerManager::ShowHideLogView(const bool bLogViewVisibleState)
 {
-	bool bWasLogViewVisible = bIsLogViewVisible;
-
 	bIsLogViewVisible = bLogViewVisibleState;
 
 	TSharedPtr<STimingProfilerWindow> Wnd = GetProfilerWindow();
@@ -247,15 +211,19 @@ const FTimerNodePtr FTimingProfilerManager::GetTimerNode(uint64 TypeId) const
 {
 	const FTimerNodePtr* TimerNodePtrPtr = nullptr;
 	TSharedPtr<STimingProfilerWindow> Wnd = GetProfilerWindow();
-	if (Wnd.IsValid() && Wnd->TimersView)
+	if (Wnd.IsValid())
 	{
-		TimerNodePtrPtr = Wnd->TimersView->GetTimerNode(TypeId);
-		if (TimerNodePtrPtr == nullptr)
+		TSharedPtr<STimersView> TimersView = Wnd->GetTimersView();
+		if (TimersView.IsValid())
 		{
-			// List of timers in TimersView not up to date?
-			// Refresh and try again.
-			Wnd->TimersView->RebuildTree(false);
-			TimerNodePtrPtr = Wnd->TimersView->GetTimerNode(TypeId);
+			TimerNodePtrPtr = TimersView->GetTimerNode(TypeId);
+			if (TimerNodePtrPtr == nullptr)
+			{
+				// List of timers in TimersView not up to date?
+				// Refresh and try again.
+				TimersView->RebuildTree(false);
+				TimerNodePtrPtr = TimersView->GetTimerNode(TypeId);
+			}
 		}
 	}
 	return TimerNodePtrPtr ? *TimerNodePtrPtr : nullptr;
diff --git a/Engine/Source/Developer/TraceInsights/Private/Insights/TimingProfilerManager.h b/Engine/Source/Developer/TraceInsights/Private/Insights/TimingProfilerManager.h
--- a/Engine/Source/Developer/TraceInsights/Private/Insights/TimingProfilerManager.h
+++ b/Engine/Source/Developer/TraceInsights/Private/Insights/TimingProfilerManager.h
@@ -64,7 +64,7 @@ protected:
 
 public:
 	/**
-	 * @return the global instance of the Timing Profiler manager (FTimingProfilerManager).
+	 * @return the global instance of the Timing Profiler (Timing Insights) manager.
 	 * This is an internal singleton and cannot be used outside ProfilerModule.
 	 * For external use:
 	 *     IProfilerModule& ProfilerModule = FModuleManager::Get().LoadModuleChecked<IProfilerModule>("Profiler");
@@ -148,7 +148,7 @@ protected:
 	/** List of UI commands for this manager. This will be filled by this and corresponding classes. */
 	TSharedRef<FUICommandList> CommandList;
 
-	/** An instance of the Timing Manager action manager. */
+	/** An instance of the Timing Profiler action manager. */
 	FTimingProfilerActionManager ActionManager;
 
 	/** A weak pointer to the Timing Profiler window. */
diff --git a/Engine/Source/Developer/TraceInsights/Private/Insights/TraceInsightsModule.cpp b/Engine/Source/Developer/TraceInsights/Private/Insights/TraceInsightsModule.cpp
--- a/Engine/Source/Developer/TraceInsights/Private/Insights/TraceInsightsModule.cpp
+++ b/Engine/Source/Developer/TraceInsights/Private/Insights/TraceInsightsModule.cpp
@@ -9,23 +9,20 @@
 #include "Widgets/DeclarativeSyntaxSupport.h"
 #include "Widgets/Docking/SDockTab.h"
 #include "Widgets/SWidget.h"
-#include "Widgets/Text/STextBlock.h"
-//#include "WorkspaceMenuStructure.h"
-//#include "WorkspaceMenuStructureModule.h"
 
 // Insights
 #include "Insights/InsightsManager.h"
 #include "Insights/InsightsStyle.h"
-#include "Insights/IoProfilerManager.h"
 #include "Insights/IUnrealInsightsModule.h"
+#include "Insights/LoadingProfiler/LoadingProfilerManager.h"
+#include "Insights/LoadingProfiler/Widgets/SLoadingProfilerWindow.h"
 #include "Insights/TimingProfilerManager.h"
-#include "Insights/Widgets/SIoProfilerWindow.h"
 #include "Insights/Widgets/SStartPageWindow.h"
 #include "Insights/Widgets/STimingProfilerWindow.h"
 
 ////////////////////////////////////////////////////////////////////////////////////////////////////
 /**
- * Implements the Trace Insights module (TimingProfiler, IoProfiler, etc.).
+ * Implements the Trace Insights module (TimingProfiler, LoadingProfiler, etc.).
  */
 class FTraceInsightsModule : public IUnrealInsightsModule
 {
@@ -48,8 +45,8 @@ protected:
 	/** Callback called when the Timing Profiler major tab is closed. */
 	void OnTimingProfilerTabBeingClosed(TSharedRef<SDockTab> TabBeingClosed);
 
-	/** Callback called when the Io Profiler major tab is closed. */
-	void OnIoProfilerTabBeingClosed(TSharedRef<SDockTab> TabBeingClosed);
+	/** Callback called when the Loading Profiler major tab is closed. */
+	void OnLoadingProfilerTabBeingClosed(TSharedRef<SDockTab> TabBeingClosed);
 
 	/** Start Page */
 	TSharedRef<SDockTab> SpawnStartPageTab(const FSpawnTabArgs& Args);
@@ -57,8 +54,8 @@ protected:
 	/** Timing Profiler */
 	TSharedRef<SDockTab> SpawnTimingProfilerTab(const FSpawnTabArgs& Args);
 
-	/** I/O Profiler */
-	TSharedRef<SDockTab> SpawnIoProfilerTab(const FSpawnTabArgs& Args);
+	/** Loading Profiler */
+	TSharedRef<SDockTab> SpawnLoadingProfilerTab(const FSpawnTabArgs& Args);
 
 protected:
 	TSharedPtr<Trace::IAnalysisService> TraceAnalysisService;
@@ -88,7 +85,7 @@ void FTraceInsightsModule::StartupModule()
 
 	FInsightsManager::Initialize(TraceAnalysisService.ToSharedRef(), TraceSessionService.ToSharedRef(), TraceModuleService.ToSharedRef());
 	FTimingProfilerManager::Initialize();
-	FIoProfilerManager::Initialize();
+	FLoadingProfilerManager::Initialize();
 
 	// Register tab spawner for the Start Page.
 	auto& StartPageTabSpawnerEntry = FGlobalTabmanager::Get()->RegisterTabSpawner(FInsightsManagerTabs::StartPageTabId,
@@ -102,7 +99,7 @@ void FTraceInsightsModule::StartupModule()
 		FOnSpawnTab::CreateRaw(this, &FTraceInsightsModule::SpawnTimingProfilerTab))
 		.SetDisplayName(NSLOCTEXT("FTraceInsightsModule", "TimingProfilerTabTitle", "Timing Insights"))
 		.SetTooltipText(NSLOCTEXT("FTraceInsightsModule", "TimingProfilerTooltipText", "Open the Timing Insights tab."))
-		.SetIcon(FSlateIcon(FInsightsStyle::GetStyleSetName(), "TimingInsights.Icon.Small"));
+		.SetIcon(FSlateIcon(FInsightsStyle::GetStyleSetName(), "TimingProfiler.Icon.Small"));
 
 //#if WITH_EDITOR
 //	TimingProfilerTabSpawnerEntry.SetGroup(WorkspaceMenu::GetMenuStructure().GetDeveloperToolsMiscCategory());
@@ -111,16 +108,16 @@ void FTraceInsightsModule::StartupModule()
 //#endif
 
 	// Register tab spawner for the Asset Loading Insights.
-	auto& IoProfilerTabSpawnerEntry = FGlobalTabmanager::Get()->RegisterTabSpawner(FInsightsManagerTabs::IoProfilerTabId,
-		FOnSpawnTab::CreateRaw(this, &FTraceInsightsModule::SpawnIoProfilerTab))
-		.SetDisplayName(NSLOCTEXT("FTraceInsightsModule", "IoProfilerTabTitle", "Asset Loading Insights"))
-		.SetTooltipText(NSLOCTEXT("FTraceInsightsModule", "IoProfilerTooltipText", "Open the Asset Loading Insights tab."))
-		.SetIcon(FSlateIcon(FInsightsStyle::GetStyleSetName(), "AssetLoadingInsights.Icon.Small"));
+	auto& LoadingProfilerTabSpawnerEntry = FGlobalTabmanager::Get()->RegisterTabSpawner(FInsightsManagerTabs::LoadingProfilerTabId,
+		FOnSpawnTab::CreateRaw(this, &FTraceInsightsModule::SpawnLoadingProfilerTab))
+		.SetDisplayName(NSLOCTEXT("FTraceInsightsModule", "LoadingProfilerTabTitle", "Asset Loading Insights"))
+		.SetTooltipText(NSLOCTEXT("FTraceInsightsModule", "LoadingProfilerTooltipText", "Open the Asset Loading Insights tab."))
+		.SetIcon(FSlateIcon(FInsightsStyle::GetStyleSetName(), "LoadingProfiler.Icon.Small"));
 
 //#if WITH_EDITOR
-//	IoProfilerTabSpawnerEntry.SetGroup(WorkspaceMenu::GetMenuStructure().GetDeveloperToolsMiscCategory());
+//	LoadingProfilerTabSpawnerEntry.SetGroup(WorkspaceMenu::GetMenuStructure().GetDeveloperToolsMiscCategory());
 //#else
-//	IoProfilerTabSpawnerEntry.SetGroup(WorkspaceMenu::GetMenuStructure().GetToolsCategory());
+//	LoadingProfilerTabSpawnerEntry.SetGroup(WorkspaceMenu::GetMenuStructure().GetToolsCategory());
 //#endif
 }
 
@@ -132,18 +129,18 @@ void FTraceInsightsModule::ShutdownModule()
 
 	FGlobalTabmanager::Get()->UnregisterNomadTabSpawner(FInsightsManagerTabs::StartPageTabId);
 	FGlobalTabmanager::Get()->UnregisterNomadTabSpawner(FInsightsManagerTabs::TimingProfilerTabId);
-	FGlobalTabmanager::Get()->UnregisterNomadTabSpawner(FInsightsManagerTabs::IoProfilerTabId);
+	FGlobalTabmanager::Get()->UnregisterNomadTabSpawner(FInsightsManagerTabs::LoadingProfilerTabId);
 
 	if (FTimingProfilerManager::Get().IsValid())
 	{
-		// Shutdown the Timing Profiler manager.
+		// Shutdown the TimingProfiler (Timing Insights) manager.
 		FTimingProfilerManager::Get()->Shutdown();
 	}
 
-	if (FIoProfilerManager::Get().IsValid())
+	if (FLoadingProfilerManager::Get().IsValid())
 	{
-		// Shutdown the I/O Profiler manager.
-		FIoProfilerManager::Get()->Shutdown();
+		// Shutdown the LoadingProfiler (Asset Loading Insights) manager.
+		FLoadingProfilerManager::Get()->Shutdown();
 	}
 
 	if (FInsightsManager::Get().IsValid())
@@ -166,7 +163,7 @@ void FTraceInsightsModule::OnNewLayout(TSharedRef<FTabManager::FLayout> NewLayou
 			FTabManager::NewStack()
 			->AddTab(FInsightsManagerTabs::StartPageTabId, ETabState::OpenedTab)
 			->AddTab(FInsightsManagerTabs::TimingProfilerTabId, ETabState::ClosedTab)
-			->AddTab(FInsightsManagerTabs::IoProfilerTabId, ETabState::ClosedTab)
+			->AddTab(FInsightsManagerTabs::LoadingProfilerTabId, ETabState::ClosedTab)
 			->SetForegroundTab(FTabId(FInsightsManagerTabs::StartPageTabId))
 		)
 	);
@@ -220,17 +217,17 @@ void FTraceInsightsModule::OnTimingProfilerTabBeingClosed(TSharedRef<SDockTab> T
 
 ////////////////////////////////////////////////////////////////////////////////////////////////////
 
-TSharedRef<SDockTab> FTraceInsightsModule::SpawnIoProfilerTab(const FSpawnTabArgs& Args)
+TSharedRef<SDockTab> FTraceInsightsModule::SpawnLoadingProfilerTab(const FSpawnTabArgs& Args)
 {
 	const TSharedRef<SDockTab> DockTab = SNew(SDockTab)
 		.TabRole(ETabRole::NomadTab);
 
 	// Register OnTabClosed to handle I/O profiler manager shutdown.
-	DockTab->SetOnTabClosed(SDockTab::FOnTabClosedCallback::CreateRaw(this, &FTraceInsightsModule::OnIoProfilerTabBeingClosed));
+	DockTab->SetOnTabClosed(SDockTab::FOnTabClosedCallback::CreateRaw(this, &FTraceInsightsModule::OnLoadingProfilerTabBeingClosed));
 
-	// Create the SIoProfilerWindow widget.
-	TSharedRef<SIoProfilerWindow> Window = SNew(SIoProfilerWindow);
-	FIoProfilerManager::Get()->AssignProfilerWindow(Window);
+	// Create the SLoadingProfilerWindow widget.
+	TSharedRef<SLoadingProfilerWindow> Window = SNew(SLoadingProfilerWindow, DockTab, Args.GetOwnerWindow());
+	FLoadingProfilerManager::Get()->AssignProfilerWindow(Window);
 	DockTab->SetContent(Window);
 
 	return DockTab;
@@ -238,7 +235,7 @@ TSharedRef<SDockTab> FTraceInsightsModule::SpawnIoProfilerTab(const FSpawnTabArg
 
 ////////////////////////////////////////////////////////////////////////////////////////////////////
 
-void FTraceInsightsModule::OnIoProfilerTabBeingClosed(TSharedRef<SDockTab> TabBeingClosed)
+void FTraceInsightsModule::OnLoadingProfilerTabBeingClosed(TSharedRef<SDockTab> TabBeingClosed)
 {
 	// Disable TabClosed delegate.
 	TabBeingClosed->SetOnTabClosed(SDockTab::FOnTabClosedCallback());
diff --git a/Engine/Source/Developer/TraceInsights/Private/Insights/ViewModels/StatsNode.cpp b/Engine/Source/Developer/TraceInsights/Private/Insights/ViewModels/StatsNode.cpp
--- a/Engine/Source/Developer/TraceInsights/Private/Insights/ViewModels/StatsNode.cpp
+++ b/Engine/Source/Developer/TraceInsights/Private/Insights/ViewModels/StatsNode.cpp
@@ -152,40 +152,4 @@ const FText FStatsNode::GetTextForAggregatedStatsUpperQuartile() const
 
 ////////////////////////////////////////////////////////////////////////////////////////////////////
 
-const FText FStatsNode::GetNameEx() const
-{
-	FText Text = FText::GetEmpty();
-
-	if (IsGroup())
-	{
-		const int32 NumChildren = Children.Num();
-		const int32 NumFilteredChildren = FilteredChildren.Num();
-
-		if (NumFilteredChildren == NumChildren)
-		{
-			Text = FText::Format(LOCTEXT("StatsNodeGroupTextFmt1", "{0} ({1})"), FText::FromName(Name), FText::AsNumber(NumChildren));
-		}
-		else
-		{
-			Text = FText::Format(LOCTEXT("StatsNodeGroupTextFmt2", "{0} ({1} / {2})"), FText::FromName(Name), FText::AsNumber(NumFilteredChildren), FText::AsNumber(NumChildren));
-		}
-	}
-	else
-	{
-		//const bool bIsStatTracked = FProfilerManager::Get()->IsStatTracked(GroupOrStatNode->GetStatID());
-		//if (bIsStatTracked)
-		//{
-		//	Text = FText::Format(LOCTEXT("StatsNodeTrackedTextFmt", "{0}*"), FText::FromName(Name));
-		//}
-		//else
-		{
-			Text = FText::FromName(Name);
-		}
-	}
-
-	return Text;
-}
-
-////////////////////////////////////////////////////////////////////////////////////////////////////
-
 #undef LOCTEXT_NAMESPACE
diff --git a/Engine/Source/Developer/TraceInsights/Private/Insights/ViewModels/StatsNode.h b/Engine/Source/Developer/TraceInsights/Private/Insights/ViewModels/StatsNode.h
--- a/Engine/Source/Developer/TraceInsights/Private/Insights/ViewModels/StatsNode.h
+++ b/Engine/Source/Developer/TraceInsights/Private/Insights/ViewModels/StatsNode.h
@@ -4,19 +4,10 @@
 
 #include "CoreMinimal.h"
 
-class FStatsNode;
-
-/** Type definition for shared pointers to instances of FStatsNode. */
-typedef TSharedPtr<class FStatsNode> FStatsNodePtr;
-
-/** Type definition for shared references to instances of FStatsNode. */
-typedef TSharedRef<class FStatsNode> FStatsNodeRef;
-
-/** Type definition for shared references to const instances of FStatsNode. */
-typedef TSharedRef<const class FStatsNode> FStatsNodeRefConst;
+// Insights
+#include "Insights/Table/ViewModels/BaseTreeNode.h"
 
-/** Type definition for weak references to instances of FStatsNode. */
-typedef TWeakPtr<class FStatsNode> FStatsNodeWeak;
+////////////////////////////////////////////////////////////////////////////////////////////////////
 
 enum class EStatsNodeType
 {
@@ -33,6 +24,8 @@ enum class EStatsNodeType
 	InvalidOrMax,
 };
 
+////////////////////////////////////////////////////////////////////////////////////////////////////
+
 template<typename Type>
 struct TAggregatedStats
 {
@@ -74,25 +67,40 @@ struct TAggregatedStats
 typedef TAggregatedStats<double> FAggregatedStats;
 typedef TAggregatedStats<int64> FAggregatedIntegerStats;
 
+////////////////////////////////////////////////////////////////////////////////////////////////////
+
+class FStatsNode;
+
+/** Type definition for shared pointers to instances of FStatsNode. */
+typedef TSharedPtr<class FStatsNode> FStatsNodePtr;
+
+/** Type definition for shared references to instances of FStatsNode. */
+typedef TSharedRef<class FStatsNode> FStatsNodeRef;
+
+/** Type definition for shared references to const instances of FStatsNode. */
+typedef TSharedRef<const class FStatsNode> FStatsNodeRefConst;
+
+/** Type definition for weak references to instances of FStatsNode. */
+typedef TWeakPtr<class FStatsNode> FStatsNodeWeak;
+
+////////////////////////////////////////////////////////////////////////////////////////////////////
 /**
- * Class used to store information about a timer node (used in timers' tree view).
+ * Class used to store information about a stats counter node (used in SStatsView).
  */
-class FStatsNode : public TSharedFromThis<FStatsNode>
+class FStatsNode : public Insights::FBaseTreeNode
 {
 public:
 	static const uint64 InvalidId = -1;
 
 public:
-	/** Initialization constructor for the timer node. */
+	/** Initialization constructor for the stats node. */
 	FStatsNode(uint64 InId, const FName InName, const FName InMetaGroupName, EStatsNodeType InType)
-		: Id(InId)
-		, Name(InName)
+		: FBaseTreeNode(InId, InName, InType == EStatsNodeType::Group)
 		, MetaGroupName(InMetaGroupName)
 		, Type(InType)
 		, bIsAddedToGraph(false)
-		, bForceExpandGroupNode(false)
 	{
-		const uint32 HashColor = Id * 0x2c2c57ed;
+		const uint32 HashColor = GetId() * 0x2c2c57ed;
 		Color.R = ((HashColor >> 16) & 0xFF) / 255.0f;
 		Color.G = ((HashColor >> 8) & 0xFF) / 255.0f;
 		Color.B = ((HashColor) & 0xFF) / 255.0f;
@@ -103,39 +111,16 @@ public:
 
 	/** Initialization constructor for the group node. */
 	FStatsNode(const FName InGroupName)
-		: Id(0)
-		, Name(InGroupName)
+		: FBaseTreeNode(0, InGroupName, true)
 		, Type(EStatsNodeType::Group)
 		, Color(0.0, 0.0, 0.0, 1.0)
 		, bIsAddedToGraph(false)
-		, bForceExpandGroupNode(false)
 	{
 		ResetAggregatedStats();
 	}
 
 	/**
-	 * @return an Id of this timer, valid only for timer nodes.
-	 */
-	const uint64 GetId() const
-	{
-		return Id;
-	}
-
-	/**
-	 * @return a name of this node, group or timer.
-	 */
-	const FName& GetName() const
-	{
-		return Name;
-	}
-
-	/**
-	 * @return a name of this node, group or timer + addditional info, to display in Stats tree view.
-	*/
-	const FText GetNameEx() const;
-
-	/**
-	 * @return a name of the group that this timer node belongs to, taken from the metadata.
+	 * @return a name of the meta group that this stats node belongs to, taken from the metadata.
 	 */
 	const FName& GetMetaGroupName() const
 	{
@@ -143,21 +128,13 @@ public:
 	}
 
 	/**
-	 * @return a type of this timer or EStatsNodeType::Group for group nodes.
+	 * @return a type of this stats node or EStatsNodeType::Group for group nodes.
 	 */
 	const EStatsNodeType& GetType() const
 	{
 		return Type;
 	}
 
-	/**
-	 * @return true, if this node is a group node.
-	 */
-	bool IsGroup() const
-	{
-		return Type == EStatsNodeType::Group;
-	}
-
 	/**
 	 * @return color of the node. Used when showing a graph series for a stats counter.
 	 */
@@ -167,7 +144,7 @@ public:
 	}
 
 	/**
-	 * @return the aggregated stats of this counter (if counter is a "float number" type).
+	 * @return the aggregated stats of this stats counter (if counter is a "float number" type).
 	 */
 	const FAggregatedStats& GetAggregatedStats() const
 	{
@@ -179,7 +156,7 @@ public:
 	void SetAggregatedStats(FAggregatedStats& AggregatedStats);
 
 	/**
-	 * @return the aggregated stats of this counter  (if counter is an "integer number" type).
+	 * @return the aggregated stats of this stats counter (if counter is an "integer number" type).
 	 */
 	const FAggregatedIntegerStats& GetAggregatedIntegerStats() const
 	{
@@ -203,49 +180,15 @@ public:
 	const FText GetTextForAggregatedStatsLowerQuartile() const;
 	const FText GetTextForAggregatedStatsUpperQuartile() const;
 
-	/**
-	 * @return a const reference to the child nodes of this group.
-	 */
-	FORCEINLINE_DEBUGGABLE const TArray<FStatsNodePtr>& GetChildren() const
-	{
-		return Children;
-	}
-
-	/**
-	 * @return a const reference to the child nodes that should be visible to the UI based on filtering.
-	 */
-	FORCEINLINE_DEBUGGABLE const TArray<FStatsNodePtr>& GetFilteredChildren() const
-	{
-		return FilteredChildren;
-	}
-
-	/**
-	 * @return a weak reference to the group of this timer node, may be invalid.
-	 */
-	FStatsNodeWeak GetGroupPtr() const
-	{
-		return GroupPtr;
-	}
-
-	/**
-	 * @return a name of the fake group that this timer node belongs to.
-	 */
-	const FName& GetGroupName() const
-	{
-		return GroupPtr.Pin()->Name;
-	}
-
-	/**
-	 * @return a name of the fake group that this timer node belongs to.
-	 */
-	//const FName& GetGroupNameSafe() const
-	//{
-	//	return GroupPtr.IsValid() ? GroupPtr.Pin()->Name : NAME_None;
-	//}
-
-	bool IsFiltered() const
+	/** Sorts children using the specified class instance. */
+	template<typename TSortingClass>
+	void SortChildren(const TSortingClass& Instance)
 	{
-		return false; // TODO
+		auto Projection = [](Insights::FBaseTreeNodePtr Node) -> FStatsNodePtr
+		{
+			return StaticCastSharedPtr<FStatsNode, Insights::FBaseTreeNode>(Node);
+		};
+		Algo::SortBy(GetChildrenMutable(), Projection, Instance);
 	}
 
 	bool IsAddedToGraph() const
@@ -258,44 +201,11 @@ public:
 		bIsAddedToGraph = bOnOff;
 	}
 
-public:
-	/** Sorts children using the specified class instance. */
-	template<typename TSortingClass>
-	void SortChildren(const TSortingClass& Instance)
-	{
-		Children.Sort(Instance);
-	}
-
-	/** Adds specified child to the children and sets group for it. */
-	FORCEINLINE_DEBUGGABLE void AddChildAndSetGroupPtr(const FStatsNodePtr& ChildPtr)
-	{
-		ChildPtr->GroupPtr = AsShared();
-		Children.Add(ChildPtr);
-	}
-
-	/** Adds specified child to the filtered children. */
-	FORCEINLINE_DEBUGGABLE void AddFilteredChild(const FStatsNodePtr& ChildPtr)
-	{
-		FilteredChildren.Add(ChildPtr);
-	}
-
-	/** Clears filtered children. */
-	void ClearFilteredChildren()
-	{
-		FilteredChildren.Reset();
-	}
-
-protected:
-	/** The Id of this timer or group. */
-	const uint64 Id;
-
-	/** The name of this timer or group. */
-	const FName Name;
-
-	/** The name of the group that this timer belongs to, based on the timer's metadata; only valid for timer nodes. */
+private:
+	/** The name of the meta group that this stats counter belongs to, based on the stats' metadata; only valid for stats counter nodes. */
 	const FName MetaGroupName;
 
-	/** Holds the type of this node; for the group, this is Group. */
+	/** Holds the type of this stats counter. */
 	const EStatsNodeType Type;
 
 	/** Color of the node. */
@@ -308,17 +218,6 @@ protected:
 
 	/** Aggregated stats (int64). */
 	FAggregatedIntegerStats AggregatedIntegerStats;
-
-	/** Children of this node. */
-	TArray<FStatsNodePtr> Children;
-
-	/** Filtered children of this node. */
-	TArray<FStatsNodePtr> FilteredChildren;
-
-	/** A weak pointer to the group/parent of this node. */
-	FStatsNodeWeak GroupPtr;
-
-public:
-	/** Whether this group node should be expanded when the text filtering is enabled. */
-	bool bForceExpandGroupNode;
 };
+
+////////////////////////////////////////////////////////////////////////////////////////////////////
diff --git a/Engine/Source/Developer/TraceInsights/Private/Insights/ViewModels/TimerNode.cpp b/Engine/Source/Developer/TraceInsights/Private/Insights/ViewModels/TimerNode.cpp
--- a/Engine/Source/Developer/TraceInsights/Private/Insights/ViewModels/TimerNode.cpp
+++ b/Engine/Source/Developer/TraceInsights/Private/Insights/ViewModels/TimerNode.cpp
@@ -20,40 +20,4 @@ void FTimerNode::SetAggregatedStats(const Trace::FAggregatedTimingStats& InAggre
 
 ////////////////////////////////////////////////////////////////////////////////////////////////////
 
-const FText FTimerNode::GetNameEx() const
-{
-	FText Text = FText::GetEmpty();
-
-	if (IsGroup())
-	{
-		const int32 NumChildren = Children.Num();
-		const int32 NumFilteredChildren = FilteredChildren.Num();
-
-		if (NumFilteredChildren == NumChildren)
-		{
-			Text = FText::Format(LOCTEXT("TimerNodeGroupTextFmt1", "{0} ({1})"), FText::FromName(Name), FText::AsNumber(NumChildren));
-		}
-		else
-		{
-			Text = FText::Format(LOCTEXT("TimerNodeGroupTextFmt2", "{0} ({1} / {2})"), FText::FromName(Name), FText::AsNumber(NumFilteredChildren), FText::AsNumber(NumChildren));
-		}
-	}
-	else
-	{
-		//const bool bIsStatTracked = FProfilerManager::Get()->IsStatTracked(GroupOrStatNode->GetStatID());
-		//if (bIsStatTracked)
-		//{
-		//	Text = FText::Format(LOCTEXT("TimerNodeTrackedTextFmt", "{0}*"), FText::FromName(Name));
-		//}
-		//else
-		{
-			Text = FText::FromName(Name);
-		}
-	}
-
-	return Text;
-}
-
-////////////////////////////////////////////////////////////////////////////////////////////////////
-
 #undef LOCTEXT_NAMESPACE
diff --git a/Engine/Source/Developer/TraceInsights/Private/Insights/ViewModels/TimerNode.h b/Engine/Source/Developer/TraceInsights/Private/Insights/ViewModels/TimerNode.h
--- a/Engine/Source/Developer/TraceInsights/Private/Insights/ViewModels/TimerNode.h
+++ b/Engine/Source/Developer/TraceInsights/Private/Insights/ViewModels/TimerNode.h
@@ -3,21 +3,12 @@
 #pragma once
 
 #include "CoreMinimal.h"
-#include "TraceServices/AnalysisService.h"
+#include "TraceServices/Containers/Timelines.h"
 
-class FTimerNode;
-
-/** Type definition for shared pointers to instances of FTimerNode. */
-typedef TSharedPtr<class FTimerNode> FTimerNodePtr;
+// Insights
+#include "Insights/Table/ViewModels/BaseTreeNode.h"
 
-/** Type definition for shared references to instances of FTimerNode. */
-typedef TSharedRef<class FTimerNode> FTimerNodeRef;
-
-/** Type definition for shared references to const instances of FTimerNode. */
-typedef TSharedRef<const class FTimerNode> FTimerNodeRefConst;
-
-/** Type definition for weak references to instances of FTimerNode. */
-typedef TWeakPtr<class FTimerNode> FTimerNodeWeak;
+////////////////////////////////////////////////////////////////////////////////////////////////////
 
 enum class ETimerNodeType
 {
@@ -37,10 +28,27 @@ enum class ETimerNodeType
 	InvalidOrMax,
 };
 
+////////////////////////////////////////////////////////////////////////////////////////////////////
+
+class FTimerNode;
+
+/** Type definition for shared pointers to instances of FTimerNode. */
+typedef TSharedPtr<class FTimerNode> FTimerNodePtr;
+
+/** Type definition for shared references to instances of FTimerNode. */
+typedef TSharedRef<class FTimerNode> FTimerNodeRef;
+
+/** Type definition for shared references to const instances of FTimerNode. */
+typedef TSharedRef<const class FTimerNode> FTimerNodeRefConst;
+
+/** Type definition for weak references to instances of FTimerNode. */
+typedef TWeakPtr<class FTimerNode> FTimerNodeWeak;
+
+////////////////////////////////////////////////////////////////////////////////////////////////////
 /**
- * Class used to store information about a timer node (used in timers' tree view).
+ * Class used to store information about a timer node (used in the STimersView).
  */
-class FTimerNode : public TSharedFromThis<FTimerNode>
+class FTimerNode : public Insights::FBaseTreeNode
 {
 public:
 	static const uint64 InvalidId = -1;
@@ -48,48 +56,23 @@ public:
 public:
 	/** Initialization constructor for the timer node. */
 	FTimerNode(uint64 InId, const FName InName, const FName InMetaGroupName, ETimerNodeType InType)
-		: Id(InId)
-		, Name(InName)
+		: FBaseTreeNode(InId, InName, InType == ETimerNodeType::Group)
 		, MetaGroupName(InMetaGroupName)
 		, Type(InType)
-		, bForceExpandGroupNode(false)
 	{
 		ResetAggregatedStats();
 	}
 
 	/** Initialization constructor for the group node. */
 	FTimerNode(const FName InGroupName)
-		: Id(0)
-		, Name(InGroupName)
+		: FBaseTreeNode(0, InGroupName, true)
 		, Type(ETimerNodeType::Group)
-		, bForceExpandGroupNode(false)
 	{
 		ResetAggregatedStats();
 	}
 
 	/**
-	 * @return an Id of this timer, valid only for timer nodes.
-	 */
-	const uint64 GetId() const
-	{
-		return Id;
-	}
-
-	/**
-	 * @return a name of this node, group or timer.
-	 */
-	const FName& GetName() const
-	{
-		return Name;
-	}
-
-	/**
-	 * @return a name of this node, group or timer + addditional info, to display in Timers tree view.
-	*/
-	const FText GetNameEx() const;
-
-	/**
-	 * @return a name of the group that this timer node belongs to, taken from the metadata.
+	 * @return a name of the meta group that this timer node belongs to, taken from the metadata.
 	 */
 	const FName& GetMetaGroupName() const
 	{
@@ -97,21 +80,13 @@ public:
 	}
 
 	/**
-	 * @return a type of this timer or ETimerNodeType::Group for group nodes.
+	 * @return a type of this timer node or ETimerNodeType::Group for group nodes.
 	 */
 	const ETimerNodeType& GetType() const
 	{
 		return Type;
 	}
 
-	/**
-	 * @return true, if this node is a group node.
-	 */
-	bool IsGroup() const
-	{
-		return Type == ETimerNodeType::Group;
-	}
-
 	/**
 	 * @return the aggregated stats for this timer.
 	 */
@@ -124,104 +99,26 @@ public:
 
 	void SetAggregatedStats(const Trace::FAggregatedTimingStats& AggregatedStats);
 
-	/**
-	 * @return a const reference to the child nodes of this group.
-	 */
-	FORCEINLINE_DEBUGGABLE const TArray<FTimerNodePtr>& GetChildren() const
-	{
-		return Children;
-	}
-
-	/**
-	 * @return a const reference to the child nodes that should be visible to the UI based on filtering.
-	 */
-	FORCEINLINE_DEBUGGABLE const TArray<FTimerNodePtr>& GetFilteredChildren() const
-	{
-		return FilteredChildren;
-	}
-
-	/**
-	 * @return a weak reference to the group of this timer node, may be invalid.
-	 */
-	FTimerNodeWeak GetGroupPtr() const
-	{
-		return GroupPtr;
-	}
-
-	/**
-	 * @return a name of the fake group that this timer node belongs to.
-	 */
-	const FName& GetGroupName() const
-	{
-		return GroupPtr.Pin()->Name;
-	}
-
-	/**
-	 * @return a name of the fake group that this timer node belongs to.
-	 */
-	//const FName& GetGroupNameSafe() const
-	//{
-	//	return GroupPtr.IsValid() ? GroupPtr.Pin()->Name : NAME_None;
-	//}
-
-	bool IsFiltered() const
-	{
-		return false; // TODO
-	}
-
-public:
 	/** Sorts children using the specified class instance. */
 	template<typename TSortingClass>
 	void SortChildren(const TSortingClass& Instance)
 	{
-		Children.Sort(Instance);
-	}
-
-	/** Adds specified child to the children and sets group for it. */
-	FORCEINLINE_DEBUGGABLE void AddChildAndSetGroupPtr(const FTimerNodePtr& ChildPtr)
-	{
-		ChildPtr->GroupPtr = AsShared();
-		Children.Add(ChildPtr);
-	}
-
-	/** Adds specified child to the filtered children. */
-	FORCEINLINE_DEBUGGABLE void AddFilteredChild(const FTimerNodePtr& ChildPtr)
-	{
-		FilteredChildren.Add(ChildPtr);
+		auto Projection = [](Insights::FBaseTreeNodePtr Node) -> FTimerNodePtr
+		{
+			return StaticCastSharedPtr<FTimerNode, Insights::FBaseTreeNode>(Node);
+		};
+		Algo::SortBy(GetChildrenMutable(), Projection, Instance);
 	}
 
-	/** Clears filtered children. */
-	void ClearFilteredChildren()
-	{
-		FilteredChildren.Reset();
-	}
-
-protected:
-	/** The Id of this timer or group. */
-	const uint64 Id;
-
-	/** The name of this timer or group. */
-	const FName Name;
-
-	/** The name of the group that this timer belongs to, based on the timer's metadata; only valid for timer nodes. */
+private:
+	/** The name of the meta group that this timer belongs to, based on the timer's metadata; only valid for timer nodes. */
 	const FName MetaGroupName;
 
-	/** Holds the type of this timer; for the group, this is Group. */
+	/** Holds the type of this timer. */
 	const ETimerNodeType Type;
 
 	/** Aggregated stats. */
 	Trace::FAggregatedTimingStats AggregatedStats;
-
-	/** Children of this node. */
-	TArray<FTimerNodePtr> Children;
-
-	/** Filtered children of this node. */
-	TArray<FTimerNodePtr> FilteredChildren;
-
-	/** A weak pointer to the group/parent of this node. */
-	FTimerNodeWeak GroupPtr;
-
-public:
-	/** Whether this group node should be expanded when the text filtering is enabled. */
-	bool bForceExpandGroupNode;
 };
+
+////////////////////////////////////////////////////////////////////////////////////////////////////
diff --git a/Engine/Source/Developer/TraceInsights/Private/Insights/Widgets/SFrameTrack.cpp b/Engine/Source/Developer/TraceInsights/Private/Insights/Widgets/SFrameTrack.cpp
--- a/Engine/Source/Developer/TraceInsights/Private/Insights/Widgets/SFrameTrack.cpp
+++ b/Engine/Source/Developer/TraceInsights/Private/Insights/Widgets/SFrameTrack.cpp
@@ -351,11 +351,15 @@ void SFrameTrack::SelectFrameAtMousePosition(float X, float Y)
 		TSharedPtr<STimingProfilerWindow> Window = FTimingProfilerManager::Get()->GetProfilerWindow();
 		if (Window.IsValid())
 		{
-			const double StartTime = SampleRef.Sample->LargestFrameStartTime;
-			const double Duration = SampleRef.Sample->LargestFrameDuration;
-			Window->TimingView->CenterOnTimeInterval(StartTime, Duration);
-			Window->TimingView->SelectTimeInterval(StartTime, Duration);
-			FSlateApplication::Get().SetKeyboardFocus(Window->TimingView);
+			TSharedPtr<STimingView> TimingView = Window->GetTimingView();
+			if (TimingView.IsValid())
+			{
+				const double StartTime = SampleRef.Sample->LargestFrameStartTime;
+				const double Duration = SampleRef.Sample->LargestFrameDuration;
+				TimingView->CenterOnTimeInterval(StartTime, Duration);
+				TimingView->SelectTimeInterval(StartTime, Duration);
+				FSlateApplication::Get().SetKeyboardFocus(TimingView);
+			}
 		}
 	}
 }
@@ -444,20 +448,24 @@ int32 SFrameTrack::OnPaint(const FPaintArgs& Args, const FGeometry& AllottedGeom
 		}
 
 		TSharedPtr<STimingProfilerWindow> Window = FTimingProfilerManager::Get()->GetProfilerWindow();
-		if (Window && Window->TimingView)
+		if (Window)
 		{
-			const FFrameTrackTimeline* TimelinePtr = nullptr;
-			for (const TPair<uint64, FFrameTrackTimeline>& KeyValuePair : CachedTimelines)
-			{
-				TimelinePtr = &KeyValuePair.Value;
-				break; // stop at first enumerated timeline
-			}
-			if (TimelinePtr)
+			TSharedPtr<STimingView> TimingView = Window->GetTimingView();
+			if (TimingView)
 			{
-				// Highlight the area corresponding to viewport of Timing View.
-				const double StartTime = Window->TimingView->GetViewport().StartTime;
-				const double EndTime = Window->TimingView->GetViewport().EndTime;
-				Helper.DrawHighlightedInterval(*TimelinePtr, StartTime, EndTime);
+				const FFrameTrackTimeline* TimelinePtr = nullptr;
+				for (const TPair<uint64, FFrameTrackTimeline>& KeyValuePair : CachedTimelines)
+				{
+					TimelinePtr = &KeyValuePair.Value;
+					break; // stop at first enumerated timeline
+				}
+				if (TimelinePtr)
+				{
+					// Highlight the area corresponding to viewport of Timing View.
+					const double StartTime = TimingView->GetViewport().StartTime;
+					const double EndTime = TimingView->GetViewport().EndTime;
+					Helper.DrawHighlightedInterval(*TimelinePtr, StartTime, EndTime);
+				}
 			}
 		}
 
diff --git a/Engine/Source/Developer/TraceInsights/Private/Insights/Widgets/SLogView.cpp b/Engine/Source/Developer/TraceInsights/Private/Insights/Widgets/SLogView.cpp
--- a/Engine/Source/Developer/TraceInsights/Private/Insights/Widgets/SLogView.cpp
+++ b/Engine/Source/Developer/TraceInsights/Private/Insights/Widgets/SLogView.cpp
@@ -154,11 +154,15 @@ public:
 				const double Time = CacheEntry.Time;
 
 				TSharedPtr<STimingProfilerWindow> Window = FTimingProfilerManager::Get()->GetProfilerWindow();
-				if (Window && Window->TimingView)
+				if (Window)
 				{
-					if (Window->TimingView->IsTimeSelectedInclusive(Time))
+					TSharedPtr<STimingView> TimingView = Window->GetTimingView();
+					if (TimingView)
 					{
-						return FSlateColor(FLinearColor(0.25f, 0.5f, 1.0f, 0.25f));
+						if (TimingView->IsTimeSelectedInclusive(Time))
+						{
+							return FSlateColor(FLinearColor(0.25f, 0.5f, 1.0f, 0.25f));
+						}
 					}
 				}
 			}
@@ -185,19 +189,23 @@ public:
 			const double Time = CacheEntry.Time;
 
 			TSharedPtr<STimingProfilerWindow> Window = FTimingProfilerManager::Get()->GetProfilerWindow();
-			if (Window && Window->TimingView)
+			if (Window)
 			{
-				if (Window->TimingView->IsTimeSelectedInclusive(Time))
+				TSharedPtr<STimingView> TimingView = Window->GetTimingView();
+				if (TimingView)
 				{
-					if (IsSelected)
+					if (TimingView->IsTimeSelectedInclusive(Time))
 					{
-						//return FSlateColor(FLinearColor(0.0f, 0.1f, 0.5f, 1.0f));
-						return FSlateColor(FLinearColor(0.0f, 0.05f, 0.2f, 1.0f));
-					}
-					else
-					{
-						//return FSlateColor(FLinearColor(0.2f, 0.4f, 0.8f, 1.0f));
-						return FSlateColor(FLinearColor(0.4f, 0.8f, 1.6f, 1.0f));
+						if (IsSelected)
+						{
+							//return FSlateColor(FLinearColor(0.0f, 0.1f, 0.5f, 1.0f));
+							return FSlateColor(FLinearColor(0.0f, 0.05f, 0.2f, 1.0f));
+						}
+						else
+						{
+							//return FSlateColor(FLinearColor(0.2f, 0.4f, 0.8f, 1.0f));
+							return FSlateColor(FLinearColor(0.4f, 0.8f, 1.6f, 1.0f));
+						}
 					}
 				}
 			}
@@ -930,21 +938,25 @@ void SLogView::OnSelectionChanged(TSharedPtr<FLogMessage> LogMessage, ESelectInf
 {
 	if (SelectInfo != ESelectInfo::Direct)
 	{
-		TSharedPtr<STimingProfilerWindow> Window = FTimingProfilerManager::Get()->GetProfilerWindow();
-		if (Window && Window->TimingView)
+		// Single item selection.
+		if (LogMessage.IsValid())
 		{
-			// Single item selection.
-			if (LogMessage.IsValid())
+			TSharedPtr<STimingProfilerWindow> Window = FTimingProfilerManager::Get()->GetProfilerWindow();
+			if (Window)
 			{
-				const double Time = Cache.Get(LogMessage->GetIndex()).Time;
-
-				if (FSlateApplication::Get().GetModifierKeys().IsShiftDown())
-				{
-					Window->TimingView->SelectToTimeMarker(Time);
-				}
-				else
+				TSharedPtr<STimingView> TimingView = Window->GetTimingView();
+				if (TimingView)
 				{
-					Window->TimingView->SetAndCenterOnTimeMarker(Time);
+					const double Time = Cache.Get(LogMessage->GetIndex()).Time;
+
+					if (FSlateApplication::Get().GetModifierKeys().IsShiftDown())
+					{
+						TimingView->SelectToTimeMarker(Time);
+					}
+					else
+					{
+						TimingView->SetAndCenterOnTimeMarker(Time);
+					}
 				}
 			}
 		}
diff --git a/Engine/Source/Developer/TraceInsights/Private/Insights/Widgets/SLogView.h b/Engine/Source/Developer/TraceInsights/Private/Insights/Widgets/SLogView.h
--- a/Engine/Source/Developer/TraceInsights/Private/Insights/Widgets/SLogView.h
+++ b/Engine/Source/Developer/TraceInsights/Private/Insights/Widgets/SLogView.h
@@ -6,6 +6,7 @@
 #include "Logging/LogVerbosity.h"
 #include "Misc/ScopeLock.h"
 #include "SlateFwd.h"
+#include "Templates/UniquePtr.h"
 #include "Widgets/DeclarativeSyntaxSupport.h"
 #include "Widgets/Input/SSearchBox.h"
 #include "Widgets/SCompoundWidget.h"
diff --git a/Engine/Source/Developer/TraceInsights/Private/Insights/Widgets/SStatsTableCell.cpp b/Engine/Source/Developer/TraceInsights/Private/Insights/Widgets/SStatsTableCell.cpp
--- a/Engine/Source/Developer/TraceInsights/Private/Insights/Widgets/SStatsTableCell.cpp
+++ b/Engine/Source/Developer/TraceInsights/Private/Insights/Widgets/SStatsTableCell.cpp
@@ -39,7 +39,7 @@ void SStatsTableCell::Construct(const FArguments& InArgs, const TSharedRef<class
 
 TSharedRef<SWidget> SStatsTableCell::GenerateWidgetForColumn(const FArguments& InArgs, const TSharedRef<class ITableRow>& TableRow)
 {
-	if (InArgs._IsStatsNameColumn)
+	if (InArgs._IsNameColumn)
 	{
 		return GenerateWidgetForNameColumn(InArgs, TableRow);
 	}
@@ -66,7 +66,7 @@ TSharedRef<SWidget> SStatsTableCell::GenerateWidgetForNameColumn(const FArgument
 			SNew(SExpanderArrow, TableRow)
 		]
 
-		// Info icon + tooltip.
+		// Info icon + tooltip
 		+ SHorizontalBox::Slot()
 		.AutoWidth()
 		.HAlign(HAlign_Center)
@@ -78,7 +78,7 @@ TSharedRef<SWidget> SStatsTableCell::GenerateWidgetForNameColumn(const FArgument
 			.ToolTip(SStatsViewTooltip::GetTableCellTooltip(StatsNodePtr))
 		]
 
-		// Color box.
+		// Color box
 		+ SHorizontalBox::Slot()
 		.AutoWidth()
 		.HAlign(HAlign_Center)
@@ -105,34 +105,13 @@ TSharedRef<SWidget> SStatsTableCell::GenerateWidgetForNameColumn(const FArgument
 		.Padding(FMargin(2.0f, 0.0f))
 		[
 			SNew(STextBlock)
-			//.Text(StatsNodePtr->GetNameEx())
-			.Text(this, &SStatsTableCell::GetNameEx)
+			.Text(this, &SStatsTableCell::GetDisplayName)
 			.HighlightText(InArgs._HighlightText)
 			.TextStyle(FEditorStyle::Get(), TEXT("Profiler.Tooltip"))
 			.ColorAndOpacity(this, &SStatsTableCell::GetColorAndOpacity)
 			.ShadowColorAndOpacity(this, &SStatsTableCell::GetShadowColorAndOpacity)
 		]
-
-		/*
-		// Culled children warning icon + tooltip.
-		+ SHorizontalBox::Slot()
-		.AutoWidth()
-		.HAlign(HAlign_Center)
-		.VAlign(VAlign_Center)
-		[
-			SNew(SButton)
-			.ButtonStyle(FEditorStyle::Get(), TEXT("HoverHintOnly"))
-			.ContentPadding(0.0f)
-			.IsFocusable(false)
-			//.OnClicked(this, &SStatsTableCell::ExpandCulledEvents_OnClicked)
-			[
-				SNew(SImage)
-				.Visibility(this, &SStatsTableCell::GetCulledEventsIconVisibility)
-				.Image(FEditorStyle::GetBrush("Profiler.EventGraph.HasCulledEventsSmall"))
-				.ToolTipText(LOCTEXT("HasCulledEvents_TT", "This event contains culled children, if you want to see all children, please disable culling or use function details, or press this icon"))
-			]
-		]*/
-		;
+	;
 }
 
 ////////////////////////////////////////////////////////////////////////////////////////////////////
@@ -145,6 +124,7 @@ TSharedRef<SWidget> SStatsTableCell::GenerateWidgetForStatsColumn(const FArgumen
 	return
 		SNew(SHorizontalBox)
 
+		// Value
 		+ SHorizontalBox::Slot()
 		.FillWidth(1.0f)
 		.VAlign(VAlign_Center)
@@ -156,7 +136,8 @@ TSharedRef<SWidget> SStatsTableCell::GenerateWidgetForStatsColumn(const FArgumen
 			.TextStyle(FEditorStyle::Get(), TEXT("Profiler.Tooltip"))
 			.ColorAndOpacity(this, &SStatsTableCell::GetStatsColorAndOpacity)
 			.ShadowColorAndOpacity(this, &SStatsTableCell::GetShadowColorAndOpacity)
-		];
+		]
+	;
 }
 
 END_SLATE_FUNCTION_BUILD_OPTIMIZATION
diff --git a/Engine/Source/Developer/TraceInsights/Private/Insights/Widgets/SStatsTableCell.h b/Engine/Source/Developer/TraceInsights/Private/Insights/Widgets/SStatsTableCell.h
--- a/Engine/Source/Developer/TraceInsights/Private/Insights/Widgets/SStatsTableCell.h
+++ b/Engine/Source/Developer/TraceInsights/Private/Insights/Widgets/SStatsTableCell.h
@@ -10,7 +10,6 @@
 #include "Insights/ViewModels/StatsNode.h"
 
 DECLARE_DELEGATE_TwoParams(FSetHoveredStatsTableCell, const FName /*ColumnId*/, const FStatsNodePtr /*SamplePtr*/);
-DECLARE_DELEGATE_RetVal_OneParam(bool, FIsColumnVisibleDelegate, const FName /*ColumnId*/);
 DECLARE_DELEGATE_RetVal_OneParam(EHorizontalAlignment, FGetColumnOutlineHAlignmentDelegate, const FName /*ColumnId*/);
 
 class SStatsTableCell : public SCompoundWidget
@@ -21,7 +20,7 @@ public:
 		SLATE_ATTRIBUTE(FText, HighlightText)
 		SLATE_ARGUMENT(FStatsNodePtr, StatsNodePtr)
 		SLATE_ARGUMENT(FName, ColumnId)
-		SLATE_ARGUMENT(bool, IsStatsNameColumn)
+		SLATE_ARGUMENT(bool, IsNameColumn)
 	SLATE_END_ARGS()
 
 	void Construct(const FArguments& InArgs, const TSharedRef<class ITableRow>& TableRow);
@@ -98,15 +97,9 @@ protected:
 		return IsHovered() ? EVisibility::Visible : EVisibility::Hidden;
 	}
 
-	EVisibility GetCulledEventsIconVisibility() const
+	FText GetDisplayName() const
 	{
-		//return StatsNodePtr->HasCulledChildren() ? EVisibility::Visible : EVisibility::Collapsed;
-		return EVisibility::Collapsed;
-	}
-
-	FText GetNameEx() const
-	{
-		return StatsNodePtr->GetNameEx();
+		return StatsNodePtr->GetDisplayName();
 	}
 
 	FSlateColor GetColorAndOpacity() const
diff --git a/Engine/Source/Developer/TraceInsights/Private/Insights/Widgets/SStatsTableRow.cpp b/Engine/Source/Developer/TraceInsights/Private/Insights/Widgets/SStatsTableRow.cpp
--- a/Engine/Source/Developer/TraceInsights/Private/Insights/Widgets/SStatsTableRow.cpp
+++ b/Engine/Source/Developer/TraceInsights/Private/Insights/Widgets/SStatsTableRow.cpp
@@ -2,6 +2,8 @@
 
 #include "SStatsTableRow.h"
 
+#include "EditorStyleSet.h"
+#include "SlateOptMacros.h"
 #include "Widgets/SOverlay.h"
 #include "Widgets/Images/SImage.h"
 
@@ -13,6 +15,8 @@
 
 #define LOCTEXT_NAMESPACE "SStatsView"
 
+////////////////////////////////////////////////////////////////////////////////////////////////////
+
 BEGIN_SLATE_FUNCTION_BUILD_OPTIMIZATION
 
 void SStatsTableRow::Construct(const FArguments& InArgs, const TSharedRef<STableViewBase>& InOwnerTableView)
@@ -23,109 +27,55 @@ void SStatsTableRow::Construct(const FArguments& InArgs, const TSharedRef<STable
 	GetColumnOutlineHAlignmentDelegate = InArgs._OnGetColumnOutlineHAlignmentDelegate;
 
 	HighlightText = InArgs._HighlightText;
-	HighlightedStatsName = InArgs._HighlightedStatsName;
+	HighlightedNodeName = InArgs._HighlightedNodeName;
 
 	StatsNodePtr = InArgs._StatsNodePtr;
 
 	SetEnabled(TAttribute<bool>(this, &SStatsTableRow::HandleShouldBeEnabled));
 
 	SMultiColumnTableRow<FStatsNodePtr>::Construct(SMultiColumnTableRow<FStatsNodePtr>::FArguments(), InOwnerTableView);
-
-	/*
-	const FSlateBrush* const NodeIcon = StatsNodeTypeHelper::GetIconForStatsNodeType(InStatsNode->GetType());
-	const TSharedRef<SToolTip> Tooltip = InStatsNode->IsGroup() ? SNew(SToolTip) : SStatsViewTooltip(InStatsNode->GetId()).GetTooltip();
-
-	ChildSlot
-		[
-			SNew(SHorizontalBox)
-
-			// Expander arrow.
-			+ SHorizontalBox::Slot()
-			.AutoWidth()
-			.HAlign(HAlign_Right)
-			.VAlign(VAlign_Center)
-			[
-				SNew(SExpanderArrow, SharedThis(this))
-			]
-
-			// Icon to visualize group or timer type.
-			+ SHorizontalBox::Slot()
-			.AutoWidth()
-			.Padding(0.0f, 0.0f, 8.0f, 0.0f)
-			[
-				SNew(SImage)
-				.Image(NodeIcon)
-				.ToolTip(Tooltip)
-			]
-
-			// Description text.
-			+ SHorizontalBox::Slot()
-			.AutoWidth()
-			.VAlign(VAlign_Center)
-			.HAlign(HAlign_Left)
-			.Padding(FMargin(2.0f, 0.0f))
-			[
-				SNew(STextBlock)
-				.Text(this, &SStatsTableRow::GetText)
-				.HighlightText(InArgs._HighlightText)
-				.TextStyle(FEditorStyle::Get(), TEXT("Profiler.Tooltip"))
-				.ColorAndOpacity(this, &SStatsTableRow::GetColorAndOpacity)
-			]
-
-			+ SHorizontalBox::Slot()
-			.AutoWidth()
-			.HAlign(HAlign_Center)
-			.VAlign(VAlign_Top)
-			.Padding(0.0f, 1.0f, 0.0f, 0.0f)
-			[
-				SNew(SImage)
-				.Visibility(!InStatsNode->IsGroup() ? EVisibility::Visible : EVisibility::Collapsed)
-				.Image(FEditorStyle::GetBrush("Profiler.Tooltip.HintIcon10"))
-				.ToolTip(Tooltip)
-			]
-		];
-
-	STableRow<FStatsNodePtr>::ConstructInternal(STableRow::FArguments().ShowSelection(true), InOwnerTableView);
-	*/
 }
 
+////////////////////////////////////////////////////////////////////////////////////////////////////
+
 TSharedRef<SWidget> SStatsTableRow::GenerateWidgetForColumn(const FName& ColumnId)
 {
 	return
+		SNew(SOverlay)
+		.Visibility(EVisibility::SelfHitTestInvisible)
 
-	SNew(SOverlay)
-	.Visibility(EVisibility::SelfHitTestInvisible)
-
-	+SOverlay::Slot()
-	.Padding(0.0f)
-	[
-		SNew(SImage)
-		.Image(FEditorStyle::GetBrush("Profiler.LineGraphArea"))
-		.ColorAndOpacity(this, &SStatsTableRow::GetBackgroundColorAndOpacity)
-	]
-
-	+SOverlay::Slot()
-	.Padding(0.0f)
-	[
-		SNew(SImage)
-		.Image(this, &SStatsTableRow::GetOutlineBrush, ColumnId)
-		.ColorAndOpacity(this, &SStatsTableRow::GetOutlineColorAndOpacity)
-	]
-
-	+SOverlay::Slot()
-	[
-		SNew(SStatsTableCell, SharedThis(this))
-		.Visibility(this, &SStatsTableRow::IsColumnVisible, ColumnId)
-		.StatsNodePtr(StatsNodePtr)
-		.ColumnId(ColumnId)
-		.HighlightText(HighlightText)
-		.IsStatsNameColumn(ColumnId == FStatsViewColumnFactory::Get().Collection[0]->Id) // name column
-		.OnSetHoveredTableCell(this, &SStatsTableRow::OnSetHoveredTableCell)
-	];
+		+SOverlay::Slot()
+		.Padding(0.0f)
+		[
+			SNew(SImage)
+			.Image(FEditorStyle::GetBrush("Profiler.LineGraphArea"))
+			.ColorAndOpacity(this, &SStatsTableRow::GetBackgroundColorAndOpacity)
+		]
+
+		+SOverlay::Slot()
+		.Padding(0.0f)
+		[
+			SNew(SImage)
+			.Image(this, &SStatsTableRow::GetOutlineBrush, ColumnId)
+			.ColorAndOpacity(this, &SStatsTableRow::GetOutlineColorAndOpacity)
+		]
+
+		+SOverlay::Slot()
+		[
+			SNew(SStatsTableCell, SharedThis(this))
+			.Visibility(this, &SStatsTableRow::IsColumnVisible, ColumnId)
+			.StatsNodePtr(StatsNodePtr)
+			.ColumnId(ColumnId)
+			.HighlightText(HighlightText)
+			.IsNameColumn(ColumnId == FStatsViewColumnFactory::Get().Collection[0]->Id) // name column
+			.OnSetHoveredTableCell(this, &SStatsTableRow::OnSetHoveredTableCell)
+		];
 }
 
 END_SLATE_FUNCTION_BUILD_OPTIMIZATION
 
+////////////////////////////////////////////////////////////////////////////////////////////////////
+
 FReply SStatsTableRow::OnDragDetected(const FGeometry& MyGeometry, const FPointerEvent& MouseEvent)
 {
 	//im:TODO
@@ -155,37 +105,7 @@ FReply SStatsTableRow::OnDragDetected(const FGeometry& MyGeometry, const FPointe
 	return SMultiColumnTableRow<FStatsNodePtr>::OnDragDetected(MyGeometry, MouseEvent);
 }
 
-FText SStatsTableRow::GetText() const
-{
-	FText Text = FText::GetEmpty();
-
-	if (StatsNodePtr->IsGroup())
-	{
-		Text = FText::Format(LOCTEXT("StatsNode_GroupNodeTextFmt", "{0} ({1})"), FText::FromName(StatsNodePtr->GetName()), FText::AsNumber(StatsNodePtr->GetChildren().Num()));
-	}
-	else
-	{
-		Text = FText::FromName(StatsNodePtr->GetName());
-	}
-
-	return Text;
-}
-
-FSlateFontInfo SStatsTableRow::GetFont() const
-{
-	const bool bIsStatTracked = false;//im:TODO: FTimingProfilerManager::Get()->IsStatTracked(StatsNodePtr->GetStatsId());
-	const FSlateFontInfo FontInfo = bIsStatTracked ? FEditorStyle::GetFontStyle("BoldFont") : FEditorStyle::GetFontStyle("NormalFont");
-	return FontInfo;
-}
-
-FSlateColor SStatsTableRow::GetColorAndOpacity() const
-{
-	//const bool bIsStatTracked = FTimingProfilerManager::Get()->IsStatTracked(StatsNodePtr->GetId());
-	//const FSlateColor Color = bIsStatTracked ? FTimingProfilerManager::Get()->GetColorForStatsId(StatsNodePtr->GetId()) : FLinearColor::White;
-	//return Color;
-	//FTimingProfilerManager::GetSettings().GetColorForStat(StatName)
-	return FLinearColor::White;
-}
+////////////////////////////////////////////////////////////////////////////////////////////////////
 
 FSlateColor SStatsTableRow::GetBackgroundColorAndOpacity() const
 {
@@ -193,6 +113,8 @@ FSlateColor SStatsTableRow::GetBackgroundColorAndOpacity() const
 	return FLinearColor(0.0f, 0.0f, 0.0f, 1.0f);
 }
 
+////////////////////////////////////////////////////////////////////////////////////////////////////
+
 FSlateColor SStatsTableRow::GetBackgroundColorAndOpacity(double Time) const
 {
 	const FLinearColor Color =	Time > TimeUtils::Second      ? FLinearColor(0.3f, 0.0f, 0.0f, 1.0f) :
@@ -202,14 +124,18 @@ FSlateColor SStatsTableRow::GetBackgroundColorAndOpacity(double Time) const
 	return Color;
 }
 
+////////////////////////////////////////////////////////////////////////////////////////////////////
+
 FSlateColor SStatsTableRow::GetOutlineColorAndOpacity() const
 {
 	const FLinearColor NoColor(0.0f, 0.0f, 0.0f, 0.0f);
-	const bool bShouldBeHighlighted = StatsNodePtr->GetName() == HighlightedStatsName.Get();
+	const bool bShouldBeHighlighted = StatsNodePtr->GetName() == HighlightedNodeName.Get();
 	const FLinearColor OutlineColorAndOpacity = bShouldBeHighlighted ? FLinearColor(FColorList::SlateBlue) : NoColor;
 	return OutlineColorAndOpacity;
 }
 
+////////////////////////////////////////////////////////////////////////////////////////////////////
+
 const FSlateBrush* SStatsTableRow::GetOutlineBrush(const FName ColumnId) const
 {
 	EHorizontalAlignment Result = HAlign_Center;
@@ -234,6 +160,8 @@ const FSlateBrush* SStatsTableRow::GetOutlineBrush(const FName ColumnId) const
 	return Brush;
 }
 
+////////////////////////////////////////////////////////////////////////////////////////////////////
+
 bool SStatsTableRow::HandleShouldBeEnabled() const
 {
 	bool bResult = false;
@@ -253,6 +181,8 @@ bool SStatsTableRow::HandleShouldBeEnabled() const
 	return bResult;
 }
 
+////////////////////////////////////////////////////////////////////////////////////////////////////
+
 EVisibility SStatsTableRow::IsColumnVisible(const FName ColumnId) const
 {
 	EVisibility Result = EVisibility::Collapsed;
@@ -265,9 +195,13 @@ EVisibility SStatsTableRow::IsColumnVisible(const FName ColumnId) const
 	return Result;
 }
 
+////////////////////////////////////////////////////////////////////////////////////////////////////
+
 void SStatsTableRow::OnSetHoveredTableCell(const FName InColumnId, const FStatsNodePtr InSamplePtr)
 {
 	SetHoveredTableCellDelegate.ExecuteIfBound(InColumnId, InSamplePtr);
 }
 
+////////////////////////////////////////////////////////////////////////////////////////////////////
+
 #undef LOCTEXT_NAMESPACE
diff --git a/Engine/Source/Developer/TraceInsights/Private/Insights/Widgets/SStatsTableRow.h b/Engine/Source/Developer/TraceInsights/Private/Insights/Widgets/SStatsTableRow.h
--- a/Engine/Source/Developer/TraceInsights/Private/Insights/Widgets/SStatsTableRow.h
+++ b/Engine/Source/Developer/TraceInsights/Private/Insights/Widgets/SStatsTableRow.h
@@ -27,7 +27,7 @@ public:
 		SLATE_EVENT(FSetHoveredStatsTableCell, OnSetHoveredTableCell)
 		SLATE_EVENT(FGetColumnOutlineHAlignmentDelegate, OnGetColumnOutlineHAlignmentDelegate)
 		SLATE_ATTRIBUTE(FText, HighlightText)
-		SLATE_ATTRIBUTE(FName, HighlightedStatsName)
+		SLATE_ATTRIBUTE(FName, HighlightedNodeName)
 		SLATE_ARGUMENT(FStatsNodePtr, StatsNodePtr)
 	SLATE_END_ARGS()
 
@@ -53,21 +53,6 @@ public:
 	virtual FReply OnDragDetected(const FGeometry& MyGeometry, const FPointerEvent& MouseEvent) override;
 
 protected:
-	/**
-	 * @return a text which describes this table row, refers to both groups and stats
-	 */
-	FText GetText() const;
-
-	/**
-	 * @return a font style which is used to draw this table row, refers to both groups and stats
-	 */
-	FSlateFontInfo GetFont() const;
-
-	/**
-	 * @return a color and opacity value used to draw this table row, refers to both groups and stats
-	 */
-	FSlateColor GetColorAndOpacity() const;
-
 	FSlateColor GetBackgroundColorAndOpacity() const;
 	FSlateColor GetBackgroundColorAndOpacity(double Time) const;
 	FSlateColor GetOutlineColorAndOpacity() const;
@@ -88,6 +73,6 @@ protected:
 	/** Text to be highlighted on stats name. */
 	TAttribute<FText> HighlightText;
 
-	/** Name of the stats that should be drawn as highlighted. */
-	TAttribute<FName> HighlightedStatsName;
+	/** Name of the stats node that should be drawn as highlighted. */
+	TAttribute<FName> HighlightedNodeName;
 };
diff --git a/Engine/Source/Developer/TraceInsights/Private/Insights/Widgets/SStatsView.cpp b/Engine/Source/Developer/TraceInsights/Private/Insights/Widgets/SStatsView.cpp
--- a/Engine/Source/Developer/TraceInsights/Private/Insights/Widgets/SStatsView.cpp
+++ b/Engine/Source/Developer/TraceInsights/Private/Insights/Widgets/SStatsView.cpp
@@ -93,10 +93,10 @@ void SStatsView::Construct(const FArguments& InArgs)
 				.AutoHeight()
 				[
 					SAssignNew(SearchBox, SSearchBox)
-					.HintText(LOCTEXT("SearchBoxHint", "Search timers or groups"))
+					.HintText(LOCTEXT("SearchBoxHint", "Search stats or groups"))
 					.OnTextChanged(this, &SStatsView::SearchBox_OnTextChanged)
 					.IsEnabled(this, &SStatsView::SearchBox_IsEnabled)
-					.ToolTipText(LOCTEXT("FilterSearchHint", "Type here to search timer or group"))
+					.ToolTipText(LOCTEXT("FilterSearchHint", "Type here to search stats or group"))
 				]
 
 				// Group by
@@ -131,7 +131,7 @@ void SStatsView::Construct(const FArguments& InArgs)
 					]
 				]
 
-				// Check boxes for: GpuScope, ComputeScope, CpuScope
+				// Check boxes for: Int64, Float
 				+ SVerticalBox::Slot()
 				.VAlign(VAlign_Center)
 				.Padding(2.0f)
@@ -642,7 +642,6 @@ void SStatsView::InsightsManager_OnSessionChanged()
 
 void SStatsView::UpdateTree()
 {
-	// Create groups, sort timers within the group and apply filtering.
 	CreateGroups();
 	SortStats();
 	ApplyFiltering();
@@ -662,13 +661,13 @@ void SStatsView::ApplyFiltering()
 		GroupPtr->ClearFilteredChildren();
 		const bool bIsGroupVisible = Filters->PassesAllFilters(GroupPtr);
 
-		const TArray<FStatsNodePtr>& GroupChildren = GroupPtr->GetChildren();
+		const TArray<Insights::FBaseTreeNodePtr>& GroupChildren = GroupPtr->GetChildren();
 		const int32 NumChildren = GroupChildren.Num();
 		int32 NumVisibleChildren = 0;
 		for (int32 Cx = 0; Cx < NumChildren; ++Cx)
 		{
 			// Add a child.
-			const FStatsNodePtr& NodePtr = GroupChildren[Cx];
+			const FStatsNodePtr& NodePtr = StaticCastSharedPtr<FStatsNode, Insights::FBaseTreeNode>(GroupChildren[Cx]);
 			const bool bIsChildVisible = Filters->PassesAllFilters(NodePtr) && bStatsNodeIsVisible[static_cast<int>(NodePtr->GetType())];
 			if (bIsChildVisible)
 			{
@@ -681,15 +680,15 @@ void SStatsView::ApplyFiltering()
 		{
 			// Add a group.
 			FilteredGroupNodes.Add(GroupPtr);
-			GroupPtr->bForceExpandGroupNode = true;
+			GroupPtr->SetExpansion(true);
 		}
 		else
 		{
-			GroupPtr->bForceExpandGroupNode = false;
+			GroupPtr->SetExpansion(false);
 		}
 	}
 
-	// Only expand timer nodes if we have a text filter.
+	// Only expand stats nodes if we have a text filter.
 	const bool bNonEmptyTextFilter = !TextFilter->GetRawFilterText().IsEmpty();
 	if (bNonEmptyTextFilter)
 	{
@@ -703,7 +702,7 @@ void SStatsView::ApplyFiltering()
 		for (int32 Fx = 0; Fx < FilteredGroupNodes.Num(); Fx++)
 		{
 			const FStatsNodePtr& GroupPtr = FilteredGroupNodes[Fx];
-			TreeView->SetItemExpansion(GroupPtr, GroupPtr->bForceExpandGroupNode);
+			TreeView->SetItemExpansion(GroupPtr, GroupPtr->IsExpanded());
 		}
 	}
 	else
@@ -800,11 +799,11 @@ void SStatsView::TreeView_OnSelectionChanged(FStatsNodePtr SelectedItem, ESelect
 		TArray<FStatsNodePtr> SelectedItems = TreeView->GetSelectedItems();
 		if (SelectedItems.Num() == 1)
 		{
-			//HighlightedStatsName = SelectedItems[0]->GetName();
+			//HighlightedNodeName = SelectedItems[0]->GetName();
 		}
 		else
 		{
-			//HighlightedStatsName = NAME_None;
+			//HighlightedNodeName = NAME_None;
 		}
 	}
 }
@@ -813,7 +812,12 @@ void SStatsView::TreeView_OnSelectionChanged(FStatsNodePtr SelectedItem, ESelect
 
 void SStatsView::TreeView_OnGetChildren(FStatsNodePtr InParent, TArray<FStatsNodePtr>& OutChildren)
 {
-	OutChildren = InParent->GetFilteredChildren();
+	const TArray<Insights::FBaseTreeNodePtr>& Children = InParent->GetFilteredChildren();
+	OutChildren.Reset(Children.Num());
+	for (const Insights::FBaseTreeNodePtr& Child : Children)
+	{
+		OutChildren.Add(StaticCastSharedPtr<FStatsNode, Insights::FBaseTreeNode>(Child));
+	}
 }
 
 ////////////////////////////////////////////////////////////////////////////////////////////////////
@@ -825,15 +829,18 @@ void SStatsView::UpdateStatsNode(FStatsNodePtr StatsNode)
 	if (!StatsNode->IsGroup())
 	{
 		TSharedPtr<STimingProfilerWindow> Wnd = FTimingProfilerManager::Get()->GetProfilerWindow();
-		TSharedPtr<STimingView> TimingView = Wnd.IsValid() ? Wnd->TimingView : nullptr;
-		if (TimingView.IsValid())
+		if (Wnd.IsValid())
 		{
-			TSharedPtr<FTimingGraphTrack> GraphTrack = TimingView->GetMainTimingGraphTrack();
-			if (GraphTrack.IsValid())
+			TSharedPtr<STimingView> TimingView = Wnd->GetTimingView();
+			if (TimingView.IsValid())
 			{
-				uint32 StatsCounterId = static_cast<uint32>(StatsNode->GetId());
-				TSharedPtr<FTimingGraphSeries> Series = GraphTrack->GetStatsCounterSeries(StatsCounterId);
-				bAddedToGraphFlag = Series.IsValid();
+				TSharedPtr<FTimingGraphTrack> GraphTrack = TimingView->GetMainTimingGraphTrack();
+				if (GraphTrack.IsValid())
+				{
+					uint32 StatsCounterId = static_cast<uint32>(StatsNode->GetId());
+					TSharedPtr<FTimingGraphSeries> Series = GraphTrack->GetStatsCounterSeries(StatsCounterId);
+					bAddedToGraphFlag = Series.IsValid();
+				}
 			}
 		}
 	}
@@ -848,31 +855,34 @@ void SStatsView::TreeView_OnMouseButtonDoubleClick(FStatsNodePtr StatsNode)
 	if (!StatsNode->IsGroup())
 	{
 		TSharedPtr<STimingProfilerWindow> Wnd = FTimingProfilerManager::Get()->GetProfilerWindow();
-		TSharedPtr<STimingView> TimingView = Wnd.IsValid() ? Wnd->TimingView : nullptr;
-		if (TimingView.IsValid())
+		if (Wnd.IsValid())
 		{
-			TSharedPtr<FTimingGraphTrack> GraphTrack = TimingView->GetMainTimingGraphTrack();
-			if (GraphTrack.IsValid())
+			TSharedPtr<STimingView> TimingView = Wnd->GetTimingView();
+			if (TimingView.IsValid())
 			{
-				uint32 StatsCounterId = static_cast<uint32>(StatsNode->GetId());
-				TSharedPtr<FTimingGraphSeries> Series = GraphTrack->GetStatsCounterSeries(StatsCounterId);
-				if (Series.IsValid())
-				{
-					GraphTrack->RemoveStatsCounterSeries(StatsCounterId);
-					GraphTrack->SetDirtyFlag();
-					StatsNode->SetAddedToGraphFlag(false);
-				}
-				else
+				TSharedPtr<FTimingGraphTrack> GraphTrack = TimingView->GetMainTimingGraphTrack();
+				if (GraphTrack.IsValid())
 				{
-					GraphTrack->AddStatsCounterSeries(StatsCounterId, StatsNode->GetColor());
-					GraphTrack->SetDirtyFlag();
-					StatsNode->SetAddedToGraphFlag(true);
+					uint32 StatsCounterId = static_cast<uint32>(StatsNode->GetId());
+					TSharedPtr<FTimingGraphSeries> Series = GraphTrack->GetStatsCounterSeries(StatsCounterId);
+					if (Series.IsValid())
+					{
+						GraphTrack->RemoveStatsCounterSeries(StatsCounterId);
+						GraphTrack->SetDirtyFlag();
+						StatsNode->SetAddedToGraphFlag(false);
+					}
+					else
+					{
+						GraphTrack->AddStatsCounterSeries(StatsCounterId, StatsNode->GetColor());
+						GraphTrack->SetDirtyFlag();
+						StatsNode->SetAddedToGraphFlag(true);
+					}
 				}
 			}
 		}
 
-		//im:TODO: const bool bIsStatsCounterTracked = FTimingProfilerManager::Get()->IsStatsCounterTracked(StatsNode->GetId());
-	//	if (!bIsStatsCounterTracked)
+		//im:TODO: const bool bIsTracked = FTimingProfilerManager::Get()->IsStatsCounterTracked(StatsNode->GetId());
+	//	if (!bIsTracked)
 	//	{
 	//		// Add a new graph series.
 	//		FTimingProfilerManager::Get()->TrackStatsCounter(StatsNode->GetId());
@@ -903,7 +913,7 @@ TSharedRef<ITableRow> SStatsView::TreeView_OnGenerateRow(FStatsNodePtr StatsNode
 		.OnSetHoveredTableCell(this, &SStatsView::TableRow_SetHoveredTableCell)
 		.OnGetColumnOutlineHAlignmentDelegate(this, &SStatsView::TableRow_GetColumnOutlineHAlignment)
 		.HighlightText(this, &SStatsView::TableRow_GetHighlightText)
-		.HighlightedStatsName(this, &SStatsView::TableRow_GetHighlightedStatsName)
+		.HighlightedNodeName(this, &SStatsView::TableRow_GetHighlightedNodeName)
 		.StatsNodePtr(StatsNodePtr);
 
 	return TableRow;
@@ -965,9 +975,9 @@ FText SStatsView::TableRow_GetHighlightText() const
 
 ////////////////////////////////////////////////////////////////////////////////////////////////////
 
-FName SStatsView::TableRow_GetHighlightedStatsName() const
+FName SStatsView::TableRow_GetHighlightedNodeName() const
 {
-	return HighlightedStatsName;
+	return HighlightedNodeName;
 }
 
 ////////////////////////////////////////////////////////////////////////////////////////////////////
@@ -1105,7 +1115,6 @@ void SStatsView::GroupBy_OnSelectionChanged(TSharedPtr<EStatsGroupingMode> NewGr
 	{
 		GroupingMode = *NewGroupingMode;
 
-		// Create groups, sort timers within the group and apply filtering.
 		CreateGroups();
 		SortStats();
 		ApplyFiltering();
@@ -1196,7 +1205,7 @@ void SStatsView::SetSortModeForColumn(const FName& ColumnId, const EColumnSortMo
 	ColumnBeingSorted = ColumnId;
 	ColumnSortMode = SortMode;
 
-	// Sort timers and apply filtering.
+	// Sort stats and apply filtering.
 	SortStats();
 	ApplyFiltering();
 }
@@ -1454,9 +1463,10 @@ void SStatsView::RebuildTree(bool bResync)
 
 	if (bResync)
 	{
-		StatsNodes.Empty(StatsNodes.Num());
-		//StatsNodesMap.Empty(StatsNodesMap.Num());
-		StatsNodesIdMap.Empty(StatsNodesIdMap.Num());
+		const int32 PreviousNodeCount = StatsNodes.Num();
+		StatsNodes.Empty(PreviousNodeCount);
+		//StatsNodesMap.Empty(PreviousNodeCount);
+		StatsNodesIdMap.Empty(PreviousNodeCount);
 		bListHasChanged = true;
 	}
 
@@ -1466,9 +1476,9 @@ void SStatsView::RebuildTree(bool bResync)
 
 		const Trace::ICounterProvider& CountersProvider = Trace::ReadCounterProvider(*Session.Get());
 
-		if (!bResync)
+		if (CountersProvider.GetCounterCount() != StatsNodes.Num())
 		{
-			bResync = (CountersProvider.GetCounterCount() != StatsNodes.Num());
+			bResync = true;
 		}
 
 		if (bResync)
@@ -1476,9 +1486,10 @@ void SStatsView::RebuildTree(bool bResync)
 			// Save selection.
 			TreeView->GetSelectedItems(SelectedItems);
 
-			StatsNodes.Empty(StatsNodes.Num());
-			//StatsNodesMap.Empty(StatsNodesMap.Num());
-			StatsNodesIdMap.Empty(StatsNodesIdMap.Num());
+			const int32 PreviousNodeCount = StatsNodes.Num();
+			StatsNodes.Empty(PreviousNodeCount);
+			//StatsNodesMap.Empty(PreviousNodeCount);
+			StatsNodesIdMap.Empty(PreviousNodeCount);
 			bListHasChanged = true;
 
 			CountersProvider.EnumerateCounters([this](const Trace::ICounter& Counter)
diff --git a/Engine/Source/Developer/TraceInsights/Private/Insights/Widgets/SStatsView.h b/Engine/Source/Developer/TraceInsights/Private/Insights/Widgets/SStatsView.h
--- a/Engine/Source/Developer/TraceInsights/Private/Insights/Widgets/SStatsView.h
+++ b/Engine/Source/Developer/TraceInsights/Private/Insights/Widgets/SStatsView.h
@@ -153,7 +153,7 @@ protected:
 	EHorizontalAlignment TableRow_GetColumnOutlineHAlignment(const FName ColumnId) const;
 
 	FText TableRow_GetHighlightText() const;
-	FName TableRow_GetHighlightedStatsName() const;
+	FName TableRow_GetHighlightedNodeName() const;
 
 	bool TableRow_ShouldBeEnabled(const uint32 StatsId) const;
 
@@ -276,7 +276,7 @@ protected:
 	FStatsNodePtr HoveredStatsNodePtr;
 
 	/** Name of the stats that should be drawn as highlighted. */
-	FName HighlightedStatsName;
+	FName HighlightedNodeName;
 
 	//////////////////////////////////////////////////
 	// Stats Nodes
diff --git a/Engine/Source/Developer/TraceInsights/Private/Insights/Widgets/STimerTableCell.cpp b/Engine/Source/Developer/TraceInsights/Private/Insights/Widgets/STimerTableCell.cpp
--- a/Engine/Source/Developer/TraceInsights/Private/Insights/Widgets/STimerTableCell.cpp
+++ b/Engine/Source/Developer/TraceInsights/Private/Insights/Widgets/STimerTableCell.cpp
@@ -38,7 +38,7 @@ void STimerTableCell::Construct(const FArguments& InArgs, const TSharedRef<class
 
 TSharedRef<SWidget> STimerTableCell::GenerateWidgetForColumn(const FArguments& InArgs, const TSharedRef<class ITableRow>& TableRow)
 {
-	if (InArgs._IsTimerNameColumn)
+	if (InArgs._IsNameColumn)
 	{
 		return GenerateWidgetForNameColumn(InArgs, TableRow);
 	}
@@ -65,7 +65,7 @@ TSharedRef<SWidget> STimerTableCell::GenerateWidgetForNameColumn(const FArgument
 			SNew(SExpanderArrow, TableRow)
 		]
 
-		// Event info icon + tooltip.
+		// Info icon + tooltip
 		+ SHorizontalBox::Slot()
 		.AutoWidth()
 		.HAlign(HAlign_Center)
@@ -85,34 +85,13 @@ TSharedRef<SWidget> STimerTableCell::GenerateWidgetForNameColumn(const FArgument
 		.Padding(FMargin(2.0f, 0.0f))
 		[
 			SNew(STextBlock)
-			//.Text(TimerNodePtr->GetNameEx())
-			.Text(this, &STimerTableCell::GetNameEx)
+			.Text(this, &STimerTableCell::GetDisplayName)
 			.HighlightText(InArgs._HighlightText)
 			.TextStyle(FEditorStyle::Get(), TEXT("Profiler.Tooltip"))
 			.ColorAndOpacity(this, &STimerTableCell::GetColorAndOpacity)
 			.ShadowColorAndOpacity(this, &STimerTableCell::GetShadowColorAndOpacity)
 		]
-
-		/*
-		// Culled children warning icon + tooltip.
-		+ SHorizontalBox::Slot()
-		.AutoWidth()
-		.HAlign(HAlign_Center)
-		.VAlign(VAlign_Center)
-		[
-			SNew(SButton)
-			.ButtonStyle(FEditorStyle::Get(), TEXT("HoverHintOnly"))
-			.ContentPadding(0.0f)
-			.IsFocusable(false)
-			//.OnClicked(this, &STimerTableCell::ExpandCulledEvents_OnClicked)
-			[
-				SNew(SImage)
-				.Visibility(this, &STimerTableCell::GetCulledEventsIconVisibility)
-				.Image(FEditorStyle::GetBrush("Profiler.EventGraph.HasCulledEventsSmall"))
-				.ToolTipText(LOCTEXT("HasCulledEvents_TT", "This event contains culled children, if you want to see all children, please disable culling or use function details, or press this icon"))
-			]
-		]*/
-		;
+	;
 }
 
 ////////////////////////////////////////////////////////////////////////////////////////////////////
@@ -125,6 +104,7 @@ TSharedRef<SWidget> STimerTableCell::GenerateWidgetForStatsColumn(const FArgumen
 	return
 		SNew(SHorizontalBox)
 
+		// Value
 		+ SHorizontalBox::Slot()
 		.FillWidth(1.0f)
 		.VAlign(VAlign_Center)
@@ -136,7 +116,8 @@ TSharedRef<SWidget> STimerTableCell::GenerateWidgetForStatsColumn(const FArgumen
 			.TextStyle(FEditorStyle::Get(), TEXT("Profiler.Tooltip"))
 			.ColorAndOpacity(this, &STimerTableCell::GetStatsColorAndOpacity)
 			.ShadowColorAndOpacity(this, &STimerTableCell::GetShadowColorAndOpacity)
-		];
+		]
+	;
 }
 
 END_SLATE_FUNCTION_BUILD_OPTIMIZATION
diff --git a/Engine/Source/Developer/TraceInsights/Private/Insights/Widgets/STimerTableCell.h b/Engine/Source/Developer/TraceInsights/Private/Insights/Widgets/STimerTableCell.h
--- a/Engine/Source/Developer/TraceInsights/Private/Insights/Widgets/STimerTableCell.h
+++ b/Engine/Source/Developer/TraceInsights/Private/Insights/Widgets/STimerTableCell.h
@@ -10,7 +10,6 @@
 #include "Insights/ViewModels/TimerNode.h"
 
 DECLARE_DELEGATE_TwoParams(FSetHoveredTimerTableCell, const FName /*ColumnId*/, const FTimerNodePtr /*SamplePtr*/);
-DECLARE_DELEGATE_RetVal_OneParam(bool, FIsColumnVisibleDelegate, const FName /*ColumnId*/);
 DECLARE_DELEGATE_RetVal_OneParam(EHorizontalAlignment, FGetColumnOutlineHAlignmentDelegate, const FName /*ColumnId*/);
 
 class STimerTableCell : public SCompoundWidget
@@ -21,7 +20,7 @@ public:
 		SLATE_ATTRIBUTE(FText, HighlightText)
 		SLATE_ARGUMENT(FTimerNodePtr, TimerNodePtr)
 		SLATE_ARGUMENT(FName, ColumnId)
-		SLATE_ARGUMENT(bool, IsTimerNameColumn)
+		SLATE_ARGUMENT(bool, IsNameColumn)
 	SLATE_END_ARGS()
 
 	void Construct(const FArguments& InArgs, const TSharedRef<class ITableRow>& TableRow);
@@ -93,15 +92,9 @@ protected:
 		return IsHovered() ? EVisibility::Visible : EVisibility::Hidden;
 	}
 
-	EVisibility GetCulledEventsIconVisibility() const
+	FText GetDisplayName() const
 	{
-		//return TimerNodePtr->HasCulledChildren() ? EVisibility::Visible : EVisibility::Collapsed;
-		return EVisibility::Collapsed;
-	}
-
-	FText GetNameEx() const
-	{
-		return TimerNodePtr->GetNameEx();
+		return TimerNodePtr->GetDisplayName();
 	}
 
 	FSlateColor GetColorAndOpacity() const
diff --git a/Engine/Source/Developer/TraceInsights/Private/Insights/Widgets/STimerTableRow.cpp b/Engine/Source/Developer/TraceInsights/Private/Insights/Widgets/STimerTableRow.cpp
--- a/Engine/Source/Developer/TraceInsights/Private/Insights/Widgets/STimerTableRow.cpp
+++ b/Engine/Source/Developer/TraceInsights/Private/Insights/Widgets/STimerTableRow.cpp
@@ -2,6 +2,8 @@
 
 #include "STimerTableRow.h"
 
+#include "EditorStyleSet.h"
+#include "SlateOptMacros.h"
 #include "Widgets/SOverlay.h"
 #include "Widgets/Images/SImage.h"
 
@@ -13,6 +15,8 @@
 
 #define LOCTEXT_NAMESPACE "STimersView"
 
+////////////////////////////////////////////////////////////////////////////////////////////////////
+
 BEGIN_SLATE_FUNCTION_BUILD_OPTIMIZATION
 
 void STimerTableRow::Construct(const FArguments& InArgs, const TSharedRef<STableViewBase>& InOwnerTableView)
@@ -23,109 +27,55 @@ void STimerTableRow::Construct(const FArguments& InArgs, const TSharedRef<STable
 	GetColumnOutlineHAlignmentDelegate = InArgs._OnGetColumnOutlineHAlignmentDelegate;
 
 	HighlightText = InArgs._HighlightText;
-	HighlightedTimerName = InArgs._HighlightedTimerName;
+	HighlightedNodeName = InArgs._HighlightedNodeName;
 
 	TimerNodePtr = InArgs._TimerNodePtr;
 
 	SetEnabled(TAttribute<bool>(this, &STimerTableRow::HandleShouldBeEnabled));
 
 	SMultiColumnTableRow<FTimerNodePtr>::Construct(SMultiColumnTableRow<FTimerNodePtr>::FArguments(), InOwnerTableView);
-
-	/*
-	const FSlateBrush* const NodeIcon = TimerNodeTypeHelper::GetIconForTimerNodeType(InTimerNode->GetType());
-	const TSharedRef<SToolTip> Tooltip = InTimerNode->IsGroup() ? SNew(SToolTip) : STimersViewTooltip(InTimerNode->GetId()).GetTooltip();
-
-	ChildSlot
-		[
-			SNew(SHorizontalBox)
-
-			// Expander arrow.
-			+ SHorizontalBox::Slot()
-			.AutoWidth()
-			.HAlign(HAlign_Right)
-			.VAlign(VAlign_Center)
-			[
-				SNew(SExpanderArrow, SharedThis(this))
-			]
-
-			// Icon to visualize group or timer type.
-			+ SHorizontalBox::Slot()
-			.AutoWidth()
-			.Padding(0.0f, 0.0f, 8.0f, 0.0f)
-			[
-				SNew(SImage)
-				.Image(NodeIcon)
-				.ToolTip(Tooltip)
-			]
-
-			// Description text.
-			+ SHorizontalBox::Slot()
-			.AutoWidth()
-			.VAlign(VAlign_Center)
-			.HAlign(HAlign_Left)
-			.Padding(FMargin(2.0f, 0.0f))
-			[
-				SNew(STextBlock)
-				.Text(this, &STimerTableRow::GetText)
-				.HighlightText(InArgs._HighlightText)
-				.TextStyle(FEditorStyle::Get(), TEXT("Profiler.Tooltip"))
-				.ColorAndOpacity(this, &STimerTableRow::GetColorAndOpacity)
-			]
-
-			+ SHorizontalBox::Slot()
-			.AutoWidth()
-			.HAlign(HAlign_Center)
-			.VAlign(VAlign_Top)
-			.Padding(0.0f, 1.0f, 0.0f, 0.0f)
-			[
-				SNew(SImage)
-				.Visibility(!InTimerNode->IsGroup() ? EVisibility::Visible : EVisibility::Collapsed)
-				.Image(FEditorStyle::GetBrush("Profiler.Tooltip.HintIcon10"))
-				.ToolTip(Tooltip)
-			]
-		];
-
-	STableRow<FTimerNodePtr>::ConstructInternal(STableRow::FArguments().ShowSelection(true), InOwnerTableView);
-	*/
 }
 
+////////////////////////////////////////////////////////////////////////////////////////////////////
+
 TSharedRef<SWidget> STimerTableRow::GenerateWidgetForColumn(const FName& ColumnId)
 {
 	return
+		SNew(SOverlay)
+		.Visibility(EVisibility::SelfHitTestInvisible)
 
-	SNew(SOverlay)
-	.Visibility(EVisibility::SelfHitTestInvisible)
-
-	+SOverlay::Slot()
-	.Padding(0.0f)
-	[
-		SNew(SImage)
-		.Image(FEditorStyle::GetBrush("Profiler.LineGraphArea"))
-		.ColorAndOpacity(this, &STimerTableRow::GetBackgroundColorAndOpacity)
-	]
-
-	+SOverlay::Slot()
-	.Padding(0.0f)
-	[
-		SNew(SImage)
-		.Image(this, &STimerTableRow::GetOutlineBrush, ColumnId)
-		.ColorAndOpacity(this, &STimerTableRow::GetOutlineColorAndOpacity)
-	]
-
-	+SOverlay::Slot()
-	[
-		SNew(STimerTableCell, SharedThis(this))
-		.Visibility(this, &STimerTableRow::IsColumnVisible, ColumnId)
-		.TimerNodePtr(TimerNodePtr)
-		.ColumnId(ColumnId)
-		.HighlightText(HighlightText)
-		.IsTimerNameColumn(ColumnId == FTimersViewColumnFactory::Get().Collection[0]->Id) // name column
-		.OnSetHoveredTableCell(this, &STimerTableRow::OnSetHoveredTableCell)
-	];
+		+SOverlay::Slot()
+		.Padding(0.0f)
+		[
+			SNew(SImage)
+			.Image(FEditorStyle::GetBrush("Profiler.LineGraphArea"))
+			.ColorAndOpacity(this, &STimerTableRow::GetBackgroundColorAndOpacity)
+		]
+
+		+SOverlay::Slot()
+		.Padding(0.0f)
+		[
+			SNew(SImage)
+			.Image(this, &STimerTableRow::GetOutlineBrush, ColumnId)
+			.ColorAndOpacity(this, &STimerTableRow::GetOutlineColorAndOpacity)
+		]
+
+		+SOverlay::Slot()
+		[
+			SNew(STimerTableCell, SharedThis(this))
+			.Visibility(this, &STimerTableRow::IsColumnVisible, ColumnId)
+			.TimerNodePtr(TimerNodePtr)
+			.ColumnId(ColumnId)
+			.HighlightText(HighlightText)
+			.IsNameColumn(ColumnId == FTimersViewColumnFactory::Get().Collection[0]->Id) // name column
+			.OnSetHoveredTableCell(this, &STimerTableRow::OnSetHoveredTableCell)
+		];
 }
 
 END_SLATE_FUNCTION_BUILD_OPTIMIZATION
 
+////////////////////////////////////////////////////////////////////////////////////////////////////
+
 FReply STimerTableRow::OnDragDetected(const FGeometry& MyGeometry, const FPointerEvent& MouseEvent)
 {
 	//im:TODO
@@ -155,37 +105,7 @@ FReply STimerTableRow::OnDragDetected(const FGeometry& MyGeometry, const FPointe
 	return SMultiColumnTableRow<FTimerNodePtr>::OnDragDetected(MyGeometry, MouseEvent);
 }
 
-FText STimerTableRow::GetText() const
-{
-	FText Text = FText::GetEmpty();
-
-	if (TimerNodePtr->IsGroup())
-	{
-		Text = FText::Format(LOCTEXT("TimerNode_GroupNodeTextFmt", "{0} ({1})"), FText::FromName(TimerNodePtr->GetName()), FText::AsNumber(TimerNodePtr->GetChildren().Num()));
-	}
-	else
-	{
-		Text = FText::FromName(TimerNodePtr->GetName());
-	}
-
-	return Text;
-}
-
-FSlateFontInfo STimerTableRow::GetFont() const
-{
-	const bool bIsStatTracked = false;//im:TODO: FTimingProfilerManager::Get()->IsStatTracked(TimerNodePtr->GetTimerId());
-	const FSlateFontInfo FontInfo = bIsStatTracked ? FEditorStyle::GetFontStyle("BoldFont") : FEditorStyle::GetFontStyle("NormalFont");
-	return FontInfo;
-}
-
-FSlateColor STimerTableRow::GetColorAndOpacity() const
-{
-	//const bool bIsStatTracked = FTimingProfilerManager::Get()->IsStatTracked(TimerNodePtr->GetId());
-	//const FSlateColor Color = bIsStatTracked ? FTimingProfilerManager::Get()->GetColorForTimerId(TimerNodePtr->GetId()) : FLinearColor::White;
-	//return Color;
-	//FTimingProfilerManager::GetSettings().GetColorForStat(StatName)
-	return FLinearColor::White;
-}
+////////////////////////////////////////////////////////////////////////////////////////////////////
 
 FSlateColor STimerTableRow::GetBackgroundColorAndOpacity() const
 {
@@ -193,6 +113,8 @@ FSlateColor STimerTableRow::GetBackgroundColorAndOpacity() const
 	//return FLinearColor(0.0f, 0.0f, 0.0f, 1.0f)
 }
 
+////////////////////////////////////////////////////////////////////////////////////////////////////
+
 FSlateColor STimerTableRow::GetBackgroundColorAndOpacity(double Time) const
 {
 	const FLinearColor Color =	Time > TimeUtils::Second      ? FLinearColor(0.3f, 0.0f, 0.0f, 1.0f) :
@@ -202,14 +124,18 @@ FSlateColor STimerTableRow::GetBackgroundColorAndOpacity(double Time) const
 	return Color;
 }
 
+////////////////////////////////////////////////////////////////////////////////////////////////////
+
 FSlateColor STimerTableRow::GetOutlineColorAndOpacity() const
 {
 	const FLinearColor NoColor(0.0f, 0.0f, 0.0f, 0.0f);
-	const bool bShouldBeHighlighted = TimerNodePtr->GetName() == HighlightedTimerName.Get();
+	const bool bShouldBeHighlighted = TimerNodePtr->GetName() == HighlightedNodeName.Get();
 	const FLinearColor OutlineColorAndOpacity = bShouldBeHighlighted ? FLinearColor(FColorList::SlateBlue) : NoColor;
 	return OutlineColorAndOpacity;
 }
 
+////////////////////////////////////////////////////////////////////////////////////////////////////
+
 const FSlateBrush* STimerTableRow::GetOutlineBrush(const FName ColumnId) const
 {
 	EHorizontalAlignment Result = HAlign_Center;
@@ -234,6 +160,8 @@ const FSlateBrush* STimerTableRow::GetOutlineBrush(const FName ColumnId) const
 	return Brush;
 }
 
+////////////////////////////////////////////////////////////////////////////////////////////////////
+
 bool STimerTableRow::HandleShouldBeEnabled() const
 {
 	bool bResult = false;
@@ -253,6 +181,8 @@ bool STimerTableRow::HandleShouldBeEnabled() const
 	return bResult;
 }
 
+////////////////////////////////////////////////////////////////////////////////////////////////////
+
 EVisibility STimerTableRow::IsColumnVisible(const FName ColumnId) const
 {
 	EVisibility Result = EVisibility::Collapsed;
@@ -265,9 +195,13 @@ EVisibility STimerTableRow::IsColumnVisible(const FName ColumnId) const
 	return Result;
 }
 
+////////////////////////////////////////////////////////////////////////////////////////////////////
+
 void STimerTableRow::OnSetHoveredTableCell(const FName InColumnId, const FTimerNodePtr InSamplePtr)
 {
 	SetHoveredTableCellDelegate.ExecuteIfBound(InColumnId, InSamplePtr);
 }
 
+////////////////////////////////////////////////////////////////////////////////////////////////////
+
 #undef LOCTEXT_NAMESPACE
diff --git a/Engine/Source/Developer/TraceInsights/Private/Insights/Widgets/STimerTableRow.h b/Engine/Source/Developer/TraceInsights/Private/Insights/Widgets/STimerTableRow.h
--- a/Engine/Source/Developer/TraceInsights/Private/Insights/Widgets/STimerTableRow.h
+++ b/Engine/Source/Developer/TraceInsights/Private/Insights/Widgets/STimerTableRow.h
@@ -27,7 +27,7 @@ public:
 		SLATE_EVENT(FSetHoveredTimerTableCell, OnSetHoveredTableCell)
 		SLATE_EVENT(FGetColumnOutlineHAlignmentDelegate, OnGetColumnOutlineHAlignmentDelegate)
 		SLATE_ATTRIBUTE(FText, HighlightText)
-		SLATE_ATTRIBUTE(FName, HighlightedTimerName)
+		SLATE_ATTRIBUTE(FName, HighlightedNodeName)
 		SLATE_ARGUMENT(FTimerNodePtr, TimerNodePtr)
 	SLATE_END_ARGS()
 
@@ -53,21 +53,6 @@ public:
 	virtual FReply OnDragDetected(const FGeometry& MyGeometry, const FPointerEvent& MouseEvent) override;
 
 protected:
-	/**
-	 * @return a text which describes this table row, refers to both groups and timers
-	 */
-	FText GetText() const;
-
-	/**
-	 * @return a font style which is used to draw this table row, refers to both groups and timers
-	 */
-	FSlateFontInfo GetFont() const;
-
-	/**
-	 * @return a color and opacity value used to draw this table row, refers to both groups and timers
-	 */
-	FSlateColor GetColorAndOpacity() const;
-
 	FSlateColor GetBackgroundColorAndOpacity() const;
 	FSlateColor GetBackgroundColorAndOpacity(double Time) const;
 	FSlateColor GetOutlineColorAndOpacity() const;
@@ -88,6 +73,6 @@ protected:
 	/** Text to be highlighted on timer name. */
 	TAttribute<FText> HighlightText;
 
-	/** Name of the timer that should be drawn as highlighted. */
-	TAttribute<FName> HighlightedTimerName;
+	/** Name of the timer node that should be drawn as highlighted. */
+	TAttribute<FName> HighlightedNodeName;
 };
diff --git a/Engine/Source/Developer/TraceInsights/Private/Insights/Widgets/STimersView.cpp b/Engine/Source/Developer/TraceInsights/Private/Insights/Widgets/STimersView.cpp
--- a/Engine/Source/Developer/TraceInsights/Private/Insights/Widgets/STimersView.cpp
+++ b/Engine/Source/Developer/TraceInsights/Private/Insights/Widgets/STimersView.cpp
@@ -5,6 +5,7 @@
 #include "EditorStyleSet.h"
 #include "Framework/MultiBox/MultiBoxBuilder.h"
 #include "SlateOptMacros.h"
+#include "Templates/UniquePtr.h"
 #include "TraceServices/AnalysisService.h"
 #include "Widgets/Input/SCheckBox.h"
 #include "Widgets/Layout/SScrollBox.h"
@@ -648,7 +649,6 @@ void STimersView::InsightsManager_OnSessionChanged()
 
 void STimersView::UpdateTree()
 {
-	// Create groups, sort timers within the group and apply filtering.
 	CreateGroups();
 	SortTimers();
 	ApplyFiltering();
@@ -668,13 +668,13 @@ void STimersView::ApplyFiltering()
 		GroupPtr->ClearFilteredChildren();
 		const bool bIsGroupVisible = Filters->PassesAllFilters(GroupPtr);
 
-		const TArray<FTimerNodePtr>& GroupChildren = GroupPtr->GetChildren();
+		const TArray<Insights::FBaseTreeNodePtr>& GroupChildren = GroupPtr->GetChildren();
 		const int32 NumChildren = GroupChildren.Num();
 		int32 NumVisibleChildren = 0;
 		for (int32 Cx = 0; Cx < NumChildren; ++Cx)
 		{
 			// Add a child.
-			const FTimerNodePtr& NodePtr = GroupChildren[Cx];
+			const FTimerNodePtr& NodePtr = StaticCastSharedPtr<FTimerNode, Insights::FBaseTreeNode>(GroupChildren[Cx]);
 			const bool bIsChildVisible = Filters->PassesAllFilters(NodePtr) && bTimerNodeIsVisible[static_cast<int>(NodePtr->GetType())];
 			if (bIsChildVisible)
 			{
@@ -687,11 +687,11 @@ void STimersView::ApplyFiltering()
 		{
 			// Add a group.
 			FilteredGroupNodes.Add(GroupPtr);
-			GroupPtr->bForceExpandGroupNode = true;
+			GroupPtr->SetExpansion(true);
 		}
 		else
 		{
-			GroupPtr->bForceExpandGroupNode = false;
+			GroupPtr->SetExpansion(false);
 		}
 	}
 
@@ -709,7 +709,7 @@ void STimersView::ApplyFiltering()
 		for (int32 Fx = 0; Fx < FilteredGroupNodes.Num(); Fx++)
 		{
 			const FTimerNodePtr& GroupPtr = FilteredGroupNodes[Fx];
-			TreeView->SetItemExpansion(GroupPtr, GroupPtr->bForceExpandGroupNode);
+			TreeView->SetItemExpansion(GroupPtr, GroupPtr->IsExpanded());
 		}
 	}
 	else
@@ -806,11 +806,11 @@ void STimersView::TreeView_OnSelectionChanged(FTimerNodePtr SelectedItem, ESelec
 		TArray<FTimerNodePtr> SelectedItems = TreeView->GetSelectedItems();
 		if (SelectedItems.Num() == 1)
 		{
-			//HighlightedTimerName = SelectedItems[0]->GetName();
+			//HighlightedNodeName = SelectedItems[0]->GetName();
 		}
 		else
 		{
-			//HighlightedTimerName = NAME_None;
+			//HighlightedNodeName = NAME_None;
 		}
 	}
 }
@@ -819,7 +819,12 @@ void STimersView::TreeView_OnSelectionChanged(FTimerNodePtr SelectedItem, ESelec
 
 void STimersView::TreeView_OnGetChildren(FTimerNodePtr InParent, TArray<FTimerNodePtr>& OutChildren)
 {
-	OutChildren = InParent->GetFilteredChildren();
+	const TArray<Insights::FBaseTreeNodePtr>& Children = InParent->GetFilteredChildren();
+	OutChildren.Reset(Children.Num());
+	for (const Insights::FBaseTreeNodePtr& Child : Children)
+	{
+		OutChildren.Add(StaticCastSharedPtr<FTimerNode, Insights::FBaseTreeNode>(Child));
+	}
 }
 
 ////////////////////////////////////////////////////////////////////////////////////////////////////
@@ -828,16 +833,16 @@ void STimersView::TreeView_OnMouseButtonDoubleClick(FTimerNodePtr TimerNode)
 {
 	if (!TimerNode->IsGroup())
 	{
-		//im:TODO: const bool bIsStatTracked = FTimingProfilerManager::Get()->IsStatTracked(TimerNode->GetId());
-	//	if (!bIsStatTracked)
+		//im:TODO: const bool bIsTracked = FTimingProfilerManager::Get()->IsTimerTracked(TimerNode->GetId());
+	//	if (!bIsTracked)
 	//	{
-	//		// Add a new graph.
-	//		FTimingProfilerManager::Get()->TrackStat(TimerNode->GetId());
+	//		// Add a new graph series.
+	//		FTimingProfilerManager::Get()->TrackTimer(TimerNode->GetId());
 	//	}
 	//	else
 	//	{
-	//		// Remove a graph
-	//		FTimingProfilerManager::Get()->UntrackStat(TimerNode->GetId());
+	//		// Remove the corresponding graph series.
+	//		FTimingProfilerManager::Get()->UntrackTimer(TimerNode->GetId());
 	//	}
 	}
 	else
@@ -860,7 +865,7 @@ TSharedRef<ITableRow> STimersView::TreeView_OnGenerateRow(FTimerNodePtr TimerNod
 		.OnSetHoveredTableCell(this, &STimersView::TableRow_SetHoveredTableCell)
 		.OnGetColumnOutlineHAlignmentDelegate(this, &STimersView::TableRow_GetColumnOutlineHAlignment)
 		.HighlightText(this, &STimersView::TableRow_GetHighlightText)
-		.HighlightedTimerName(this, &STimersView::TableRow_GetHighlightedTimerName)
+		.HighlightedNodeName(this, &STimersView::TableRow_GetHighlightedNodeName)
 		.TimerNodePtr(TimerNodePtr);
 
 	return TableRow;
@@ -922,9 +927,9 @@ FText STimersView::TableRow_GetHighlightText() const
 
 ////////////////////////////////////////////////////////////////////////////////////////////////////
 
-FName STimersView::TableRow_GetHighlightedTimerName() const
+FName STimersView::TableRow_GetHighlightedNodeName() const
 {
-	return HighlightedTimerName;
+	return HighlightedNodeName;
 }
 
 ////////////////////////////////////////////////////////////////////////////////////////////////////
@@ -953,7 +958,7 @@ bool STimersView::SearchBox_IsEnabled() const
 }
 
 ////////////////////////////////////////////////////////////////////////////////////////////////////
-// GroupBy
+// Grouping
 ////////////////////////////////////////////////////////////////////////////////////////////////////
 
 void STimersView::CreateGroups()
@@ -1077,7 +1082,6 @@ void STimersView::GroupBy_OnSelectionChanged(TSharedPtr<ETimerGroupingMode> NewG
 	{
 		GroupingMode = *NewGroupingMode;
 
-		// Create groups, sort timers within the group and apply filtering.
 		CreateGroups();
 		SortTimers();
 		ApplyFiltering();
@@ -1108,7 +1112,7 @@ FText STimersView::GroupBy_GetSelectedTooltipText() const
 }
 
 ////////////////////////////////////////////////////////////////////////////////////////////////////
-// SortBy
+// Sorting
 ////////////////////////////////////////////////////////////////////////////////////////////////////
 
 void STimersView::SortTimers()
@@ -1139,7 +1143,7 @@ void STimersView::SortTimers()
 	else CHECK_AND_SORT_COLUMN(FTimersViewColumns::TypeColumnID,          ByType)
 	else CHECK_AND_SORT_COLUMN(FTimersViewColumns::InstanceCountColumnID, ByInstanceCount)
 	// Inclusive Time columns
-	else CHECK_AND_SORT_COLUMN(FTimersViewColumns::TotalInclusiveTimeColumnID, ByTotalInclusiveTime)
+	else CHECK_AND_SORT_COLUMN(FTimersViewColumns::TotalInclusiveTimeColumnID,   ByTotalInclusiveTime)
 	else CHECK_AND_SORT_COLUMN(FTimersViewColumns::MaxInclusiveTimeColumnID,     ByMaxInclusiveTime)
 	else CHECK_AND_SORT_COLUMN(FTimersViewColumns::AverageInclusiveTimeColumnID, ByAverageInclusiveTime)
 	else CHECK_AND_SORT_COLUMN(FTimersViewColumns::MedianInclusiveTimeColumnID,  ByMedianInclusiveTime)
@@ -1173,7 +1177,6 @@ void STimersView::SetSortModeForColumn(const FName& ColumnId, const EColumnSortM
 	ColumnBeingSorted = ColumnId;
 	ColumnSortMode = SortMode;
 
-	// Sort timers and apply filtering.
 	SortTimers();
 	ApplyFiltering();
 }
@@ -1434,9 +1437,10 @@ void STimersView::RebuildTree(bool bResync)
 
 	if (bResync)
 	{
-		TimerNodes.Empty(TimerNodes.Num());
-		//TimerNodesMap.Empty(TimerNodesMap.Num());
-		TimerNodesIdMap.Empty(TimerNodesIdMap.Num());
+		const int32 PreviousNodeCount = TimerNodes.Num();
+		TimerNodes.Empty(PreviousNodeCount);
+		//TimerNodesMap.Empty(PreviousNodeCount);
+		TimerNodesIdMap.Empty(PreviousNodeCount);
 		bListHasChanged = true;
 	}
 
@@ -1448,9 +1452,9 @@ void STimersView::RebuildTree(bool bResync)
 
 		TimingProfilerProvider.ReadTimers([this, &bResync, &SelectedItems, &bListHasChanged](const Trace::FTimingProfilerTimer* Timers, uint64 TimersCount)
 		{
-			if (!bResync)
+			if (TimersCount != TimerNodes.Num())
 			{
-				bResync = (TimersCount != TimerNodes.Num());
+				bResync = true;
 			}
 
 			if (bResync)
@@ -1458,9 +1462,10 @@ void STimersView::RebuildTree(bool bResync)
 				// Save selection.
 				TreeView->GetSelectedItems(SelectedItems);
 
-				TimerNodes.Empty(TimerNodes.Num());
-				//TimerNodesMap.Empty(TimerNodesMap.Num());
-				TimerNodesIdMap.Empty(TimerNodesIdMap.Num());
+				const int32 PreviousNodeCount = TimerNodes.Num();
+				TimerNodes.Empty(PreviousNodeCount);
+				//TimerNodesMap.Empty(PreviousNodeCount);
+				TimerNodesIdMap.Empty(PreviousNodeCount);
 				bListHasChanged = true;
 
 				for (uint64 TimerIndex = 0; TimerIndex < TimersCount; ++TimerIndex)
@@ -1517,7 +1522,7 @@ void STimersView::UpdateStats(double StartTime, double EndTime)
 	if (Session.IsValid() && StartTime < EndTime && Trace::ReadTimingProfilerProvider(*Session.Get()))
 	{
 		TSharedPtr<STimingProfilerWindow> Wnd = FTimingProfilerManager::Get()->GetProfilerWindow();
-		TSharedPtr<STimingView> TimingView = Wnd.IsValid() ? Wnd->TimingView : nullptr;
+		TSharedPtr<STimingView> TimingView = Wnd.IsValid() ? Wnd->GetTimingView() : nullptr;
 
 		auto ThreadFilter = [&TimingView](uint32 ThreadId)
 		{
diff --git a/Engine/Source/Developer/TraceInsights/Private/Insights/Widgets/STimersView.h b/Engine/Source/Developer/TraceInsights/Private/Insights/Widgets/STimersView.h
--- a/Engine/Source/Developer/TraceInsights/Private/Insights/Widgets/STimersView.h
+++ b/Engine/Source/Developer/TraceInsights/Private/Insights/Widgets/STimersView.h
@@ -133,7 +133,7 @@ protected:
 	EHorizontalAlignment TableRow_GetColumnOutlineHAlignment(const FName ColumnId) const;
 
 	FText TableRow_GetHighlightText() const;
-	FName TableRow_GetHighlightedTimerName() const;
+	FName TableRow_GetHighlightedNodeName() const;
 
 	bool TableRow_ShouldBeEnabled(const uint32 TimerId) const;
 
@@ -252,7 +252,7 @@ protected:
 	FTimerNodePtr HoveredTimerNodePtr;
 
 	/** Name of the timer that should be drawn as highlighted. */
-	FName HighlightedTimerName;
+	FName HighlightedNodeName;
 
 	//////////////////////////////////////////////////
 	// Timer Nodes
