[Insights]

diff --git a/Engine/Source/Developer/TraceInsights/Private/Insights/Widgets/STimingProfilerWindow.cpp b/Engine/Source/Developer/TraceInsights/Private/Insights/Widgets/STimingProfilerWindow.cpp
--- a/Engine/Source/Developer/TraceInsights/Private/Insights/Widgets/STimingProfilerWindow.cpp
+++ b/Engine/Source/Developer/TraceInsights/Private/Insights/Widgets/STimingProfilerWindow.cpp
@@ -13,8 +13,6 @@
 #include "Widgets/Notifications/SNotificationList.h"
 #include "Widgets/SBoxPanel.h"
 #include "Widgets/Text/STextBlock.h"
-//#include "WorkspaceMenuStructure.h"
-//#include "WorkspaceMenuStructureModule.h"
 
 #if WITH_EDITOR
 	#include "EngineAnalytics.h"
@@ -64,13 +62,43 @@ STimingProfilerWindow::~STimingProfilerWindow()
 #if WITH_EDITOR
 	if (DurationActive > 0.0f && FEngineAnalytics::IsAvailable())
 	{
-		FEngineAnalytics::GetProvider().RecordEvent(TEXT("Editor.Usage.Profiler"), FAnalyticsEventAttribute(TEXT("Duration"), DurationActive));
+		FEngineAnalytics::GetProvider().RecordEvent(TEXT("Insights.Usage.TimingProfiler"), FAnalyticsEventAttribute(TEXT("Duration"), DurationActive));
 	}
 #endif // WITH_EDITOR
 }
 
 ////////////////////////////////////////////////////////////////////////////////////////////////////
 
+void STimingProfilerWindow::Reset()
+{
+	if (FrameTrack)
+	{
+		FrameTrack->Reset();
+	}
+	if (GraphTrack)
+	{
+		GraphTrack->Reset();
+	}
+	if (TimingView)
+	{
+		TimingView->Reset();
+	}
+	if (TimersView)
+	{
+		TimersView->RebuildTree(true);
+	}
+	if (StatsView)
+	{
+		StatsView->RebuildTree(true);
+	}
+	if (LogView)
+	{
+		LogView->Reset();
+	}
+}
+
+////////////////////////////////////////////////////////////////////////////////////////////////////
+
 BEGIN_SLATE_FUNCTION_BUILD_OPTIMIZATION
 
 TSharedRef<SDockTab> STimingProfilerWindow::SpawnTab_Toolbar(const FSpawnTabArgs& Args)
@@ -78,7 +106,6 @@ TSharedRef<SDockTab> STimingProfilerWindow::SpawnTab_Toolbar(const FSpawnTabArgs
 	const TSharedRef<SDockTab> DockTab = SNew(SDockTab)
 		.ShouldAutosize(true)
 		.TabRole(ETabRole::PanelTab)
-		//.IsEnabled(this, &STimingProfilerWindow::IsProfilerEnabled)
 		[
 			SNew(STimingProfilerToolbar)
 		];
@@ -103,7 +130,6 @@ TSharedRef<SDockTab> STimingProfilerWindow::SpawnTab_FramesTrack(const FSpawnTab
 	const TSharedRef<SDockTab> DockTab = SNew(SDockTab)
 		.ShouldAutosize(false)
 		.TabRole(ETabRole::PanelTab)
-		//.IsEnabled(this, &STimingProfilerWindow::IsProfilerEnabled)
 		[
 			SAssignNew(FrameTrack, SFrameTrack)
 		];
@@ -118,6 +144,7 @@ TSharedRef<SDockTab> STimingProfilerWindow::SpawnTab_FramesTrack(const FSpawnTab
 void STimingProfilerWindow::OnFramesTrackTabClosed(TSharedRef<SDockTab> TabBeingClosed)
 {
 	FTimingProfilerManager::Get()->SetFramesTrackVisible(false);
+	FrameTrack = nullptr;
 }
 
 ////////////////////////////////////////////////////////////////////////////////////////////////////
@@ -129,7 +156,6 @@ TSharedRef<SDockTab> STimingProfilerWindow::SpawnTab_GraphTrack(const FSpawnTabA
 	const TSharedRef<SDockTab> DockTab = SNew(SDockTab)
 		.ShouldAutosize(false)
 		.TabRole(ETabRole::PanelTab)
-		//.IsEnabled(this, &STimingProfilerWindow::IsProfilerEnabled)
 		[
 			SAssignNew(GraphTrack, SGraphTrack)
 		];
@@ -144,6 +170,7 @@ TSharedRef<SDockTab> STimingProfilerWindow::SpawnTab_GraphTrack(const FSpawnTabA
 void STimingProfilerWindow::OnGraphTrackTabClosed(TSharedRef<SDockTab> TabBeingClosed)
 {
 	FTimingProfilerManager::Get()->SetGraphTrackVisible(false);
+	GraphTrack = nullptr;
 }
 
 ////////////////////////////////////////////////////////////////////////////////////////////////////
@@ -155,7 +182,6 @@ TSharedRef<SDockTab> STimingProfilerWindow::SpawnTab_TimingView(const FSpawnTabA
 	const TSharedRef<SDockTab> DockTab = SNew(SDockTab)
 		.ShouldAutosize(false)
 		.TabRole(ETabRole::PanelTab)
-		//.IsEnabled(this, &STimingProfilerWindow::IsProfilerEnabled)
 		[
 			SAssignNew(TimingView, STimingView)
 		];
@@ -170,6 +196,7 @@ TSharedRef<SDockTab> STimingProfilerWindow::SpawnTab_TimingView(const FSpawnTabA
 void STimingProfilerWindow::OnTimingViewTabClosed(TSharedRef<SDockTab> TabBeingClosed)
 {
 	FTimingProfilerManager::Get()->SetTimingViewVisible(false);
+	TimingView = nullptr;
 }
 
 ////////////////////////////////////////////////////////////////////////////////////////////////////
@@ -181,7 +208,6 @@ TSharedRef<SDockTab> STimingProfilerWindow::SpawnTab_Timers(const FSpawnTabArgs&
 	const TSharedRef<SDockTab> DockTab = SNew(SDockTab)
 		.ShouldAutosize(false)
 		.TabRole(ETabRole::PanelTab)
-		//.IsEnabled(this, &STimingProfilerWindow::IsProfilerEnabled)
 		[
 			SAssignNew(TimersView, STimersView)
 		];
@@ -196,6 +222,7 @@ TSharedRef<SDockTab> STimingProfilerWindow::SpawnTab_Timers(const FSpawnTabArgs&
 void STimingProfilerWindow::OnTimersTabClosed(TSharedRef<SDockTab> TabBeingClosed)
 {
 	FTimingProfilerManager::Get()->SetTimersViewVisible(false);
+	TimersView = nullptr;
 }
 
 ////////////////////////////////////////////////////////////////////////////////////////////////////
@@ -207,7 +234,6 @@ TSharedRef<SDockTab> STimingProfilerWindow::SpawnTab_StatsCounters(const FSpawnT
 	const TSharedRef<SDockTab> DockTab = SNew(SDockTab)
 		.ShouldAutosize(false)
 		.TabRole(ETabRole::PanelTab)
-		//.IsEnabled(this, &STimingProfilerWindow::IsProfilerEnabled)
 		[
 			SAssignNew(StatsView, SStatsView)
 		];
@@ -222,6 +248,7 @@ TSharedRef<SDockTab> STimingProfilerWindow::SpawnTab_StatsCounters(const FSpawnT
 void STimingProfilerWindow::OnStatsCountersTabClosed(TSharedRef<SDockTab> TabBeingClosed)
 {
 	FTimingProfilerManager::Get()->SetStatsCountersViewVisible(false);
+	StatsView = nullptr;
 }
 
 ////////////////////////////////////////////////////////////////////////////////////////////////////
@@ -233,7 +260,6 @@ TSharedRef<SDockTab> STimingProfilerWindow::SpawnTab_LogView(const FSpawnTabArgs
 	const TSharedRef<SDockTab> DockTab = SNew(SDockTab)
 		.ShouldAutosize(false)
 		.TabRole(ETabRole::PanelTab)
-		//.IsEnabled(this, &STimingProfilerWindow::IsProfilerEnabled)
 		[
 			SAssignNew(LogView, SLogView)
 		];
@@ -248,6 +274,7 @@ TSharedRef<SDockTab> STimingProfilerWindow::SpawnTab_LogView(const FSpawnTabArgs
 void STimingProfilerWindow::OnLogViewTabClosed(TSharedRef<SDockTab> TabBeingClosed)
 {
 	FTimingProfilerManager::Get()->SetLogViewVisible(false);
+	LogView = nullptr;
 }
 
 ////////////////////////////////////////////////////////////////////////////////////////////////////
@@ -384,7 +411,7 @@ void STimingProfilerWindow::Construct(const FArguments& InArgs, const TSharedRef
 						.ColorAndOpacity(FLinearColor(0.15f, 0.15f, 0.15f, 1.0f))
 				]
 
-			// Overlay slot for the main profiler window area
+			// Overlay slot for the main window area
 			+ SOverlay::Slot()
 				.HAlign(HAlign_Fill)
 				.VAlign(VAlign_Fill)
diff --git a/Engine/Source/Developer/TraceInsights/Private/Insights/Widgets/STimingProfilerWindow.h b/Engine/Source/Developer/TraceInsights/Private/Insights/Widgets/STimingProfilerWindow.h
--- a/Engine/Source/Developer/TraceInsights/Private/Insights/Widgets/STimingProfilerWindow.h
+++ b/Engine/Source/Developer/TraceInsights/Private/Insights/Widgets/STimingProfilerWindow.h
@@ -52,9 +52,11 @@ public:
 	/** Virtual destructor. */
 	virtual ~STimingProfilerWindow();
 
-	SLATE_BEGIN_ARGS(STimingProfilerWindow){}
+	SLATE_BEGIN_ARGS(STimingProfilerWindow) {}
 	SLATE_END_ARGS()
 
+	void Reset();
+
 	/** Constructs this widget. */
 	void Construct(const FArguments& InArgs, const TSharedRef<SDockTab>& ConstructUnderMajorTab, const TSharedPtr<SWindow>& ConstructUnderWindow);
 
@@ -64,22 +66,27 @@ public:
 
 	TSharedPtr<FTabManager> GetTabManager() const { return TabManager; }
 
+	TSharedPtr<STimingView> GetTimingView() const { return TimingView; }
+	TSharedPtr<STimersView> GetTimersView() const { return TimersView; }
+	TSharedPtr<SStatsView> GetStatsView() const { return StatsView; }
+	TSharedPtr<SLogView> GetLogView() const { return LogView; }
+
 private:
 	TSharedRef<SDockTab> SpawnTab_Toolbar(const FSpawnTabArgs& Args);
 	void OnToolbarTabClosed(TSharedRef<SDockTab> TabBeingClosed);
-	
+
 	TSharedRef<SDockTab> SpawnTab_FramesTrack(const FSpawnTabArgs& Args);
 	void OnFramesTrackTabClosed(TSharedRef<SDockTab> TabBeingClosed);
-	
+
 	TSharedRef<SDockTab> SpawnTab_GraphTrack(const FSpawnTabArgs& Args);
 	void OnGraphTrackTabClosed(TSharedRef<SDockTab> TabBeingClosed);
-	
+
 	TSharedRef<SDockTab> SpawnTab_TimingView(const FSpawnTabArgs& Args);
 	void OnTimingViewTabClosed(TSharedRef<SDockTab> TabBeingClosed);
-	
+
 	TSharedRef<SDockTab> SpawnTab_Timers(const FSpawnTabArgs& Args);
 	void OnTimersTabClosed(TSharedRef<SDockTab> TabBeingClosed);
-	
+
 	TSharedRef<SDockTab> SpawnTab_StatsCounters(const FSpawnTabArgs& Args);
 	void OnStatsCountersTabClosed(TSharedRef<SDockTab> TabBeingClosed);
 
@@ -157,32 +164,31 @@ private:
 	 */
 	virtual FReply OnDragOver(const FGeometry& MyGeometry, const FDragDropEvent& DragDropEvent)  override;
 
-public:
-	/** Widget for the frame track */
+private:
+	/** The Frame track widget */
 	TSharedPtr<SFrameTrack> FrameTrack;
 
-	/** Widget for the graph track */
+	/** The graph track widget */
 	TSharedPtr<SGraphTrack> GraphTrack;
 
-	/** Widget for the timing track */
+	/** The Timing view (multi-track) widget */
 	TSharedPtr<STimingView> TimingView;
 
-	/** Holds the Timers view widget. */
+	/** The Timers view widget */
 	TSharedPtr<STimersView> TimersView;
 
-	/** Holds the Stats (Counters) view widget. */
+	/** The Stats (Counters) view widget */
 	TSharedPtr<SStatsView> StatsView;
 
-	/** Widget for the log view */
+	/** The Log view widget */
 	TSharedPtr<SLogView> LogView;
 
-	/** The number of seconds the profiler has been active */
-	float DurationActive;
-
-private:
 	/** Holds the tab manager that manages the front-end's tabs. */
 	TSharedPtr<FTabManager> TabManager;
 
 	/** The handle to the active update duration tick */
 	TWeakPtr<FActiveTimerHandle> ActiveTimerHandle;
+
+	/** The number of seconds the profiler has been active */
+	float DurationActive;
 };
diff --git a/Engine/Source/Developer/TraceInsights/Private/Insights/Widgets/STimingView.cpp b/Engine/Source/Developer/TraceInsights/Private/Insights/Widgets/STimingView.cpp
--- a/Engine/Source/Developer/TraceInsights/Private/Insights/Widgets/STimingView.cpp
+++ b/Engine/Source/Developer/TraceInsights/Private/Insights/Widgets/STimingView.cpp
@@ -35,6 +35,9 @@
 #include "Insights/Common/TimeUtils.h"
 #include "Insights/InsightsManager.h"
 #include "Insights/InsightsStyle.h"
+#include "Insights/LoadingProfiler/LoadingProfilerManager.h"
+#include "Insights/LoadingProfiler/Widgets/SLoadingProfilerWindow.h"
+#include "Insights/Table/Widgets/STableTreeView.h"
 #include "Insights/TimingProfilerCommon.h"
 #include "Insights/TimingProfilerManager.h"
 #include "Insights/ViewModels/BaseTimingTrack.h"
@@ -404,13 +407,15 @@ void STimingView::Tick(const FGeometry& AllottedGeometry, const double InCurrent
 		TSharedPtr<STimingProfilerWindow> Wnd = FTimingProfilerManager::Get()->GetProfilerWindow();
 		if (Wnd)
 		{
-			if (Wnd->TimersView)
+			TSharedPtr<STimersView> TimersView = Wnd->GetTimersView();
+			if (TimersView)
 			{
-				Wnd->TimersView->RebuildTree(false);
+				TimersView->RebuildTree(false);
 			}
-			if (Wnd->StatsView)
+			TSharedPtr<SStatsView> StatsView = Wnd->GetStatsView();
+			if (StatsView)
 			{
-				Wnd->StatsView->RebuildTree(false);
+				StatsView->RebuildTree(false);
 			}
 		}
 	}
@@ -2191,18 +2196,29 @@ void STimingView::UpdateAggregatedStats()
 	TSharedPtr<STimingProfilerWindow> Wnd = FTimingProfilerManager::Get()->GetProfilerWindow();
 	if (Wnd)
 	{
-		if (Wnd->TimersView)
+		TSharedPtr<STimersView> TimersView = Wnd->GetTimersView();
+		if (TimersView)
 		{
-			Wnd->TimersView->UpdateStats(SelectionStartTime, SelectionEndTime);
+			TimersView->UpdateStats(SelectionStartTime, SelectionEndTime);
 		}
-		if (Wnd->StatsView)
+		TSharedPtr<SStatsView> StatsView = Wnd->GetStatsView();
+		if (StatsView)
 		{
-			Wnd->StatsView->UpdateStats(SelectionStartTime, SelectionEndTime);
+			StatsView->UpdateStats(SelectionStartTime, SelectionEndTime);
 		}
 	}
 
 	if (bAssetLoadingMode)
 	{
+		TSharedPtr<Insights::STableTreeView> EventAggregationTreeView;
+		TSharedPtr<Insights::STableTreeView> ObjectTypeAggregationTreeView;
+		TSharedPtr<SLoadingProfilerWindow> LoadingProfilerWnd = FLoadingProfilerManager::Get()->GetProfilerWindow();
+		if (LoadingProfilerWnd)
+		{
+			EventAggregationTreeView = LoadingProfilerWnd->GetEventAggregationTreeView();
+			ObjectTypeAggregationTreeView = LoadingProfilerWnd->GetObjectTypeAggregationTreeView();
+		}
+
 		TSharedPtr<const Trace::IAnalysisSession> Session = FInsightsManager::Get()->GetSession();
 		if (Session.IsValid() && Trace::ReadLoadTimeProfilerProvider(*Session.Get()))
 		{
@@ -2287,15 +2303,21 @@ void STimingView::UpdateAggregatedStats()
 			const Trace::ILoadTimeProfilerProvider& LoadTimeProfilerProvider = *Trace::ReadLoadTimeProfilerProvider(*Session.Get());
 
 			{
+				//TODO: LoadTimeProfilerProvider.UpdateEventAggregation(EventAggregationTreeView->GetTable(), SelectionStartTime, SelectionEndTime)
 				Trace::ITable<Trace::FLoadTimeProfilerAggregatedStats>* EventAggregationTable = LoadTimeProfilerProvider.CreateEventAggregation(SelectionStartTime, SelectionEndTime);
+				EventAggregationTreeView->GetTable()->UpdateSourceTable(MakeShareable(EventAggregationTable));
+				EventAggregationTreeView->RebuildTree(true);
 				ReadTable(EventAggregationTable, EventAggregation, EventAggregationTotalCount);
-				delete EventAggregationTable;
+				//delete EventAggregationTable;
 			}
 
 			{
+				//TODO: LoadTimeProfilerProvider.UpdateObjectTypeAggregation(EventAggregationTreeView->GetTable(), SelectionStartTime, SelectionEndTime)
 				Trace::ITable<Trace::FLoadTimeProfilerAggregatedStats>* ObjectTypeAggregationTable = LoadTimeProfilerProvider.CreateObjectTypeAggregation(SelectionStartTime, SelectionEndTime);
+				ObjectTypeAggregationTreeView->GetTable()->UpdateSourceTable(MakeShareable(ObjectTypeAggregationTable));
+				ObjectTypeAggregationTreeView->RebuildTree(true);
 				ReadTable(ObjectTypeAggregationTable, ObjectTypeAggregation, ObjectTypeAggregationTotalCount);
-				delete ObjectTypeAggregationTable;
+				//delete ObjectTypeAggregationTable;
 			}
 		}
 	}
@@ -2442,9 +2464,10 @@ void STimingView::OnSelectedTimingEventChanged()
 		TSharedPtr<STimingProfilerWindow> Wnd = FTimingProfilerManager::Get()->GetProfilerWindow();
 		if (Wnd)
 		{
-			if (Wnd->TimersView)
+			TSharedPtr<STimersView> TimersView = Wnd->GetTimersView();
+			if (TimersView)
 			{
-				Wnd->TimersView->SelectTimerNode(SelectedTimingEvent.TypeId);
+				TimersView->SelectTimerNode(SelectedTimingEvent.TypeId);
 			}
 		}
 	}
diff --git a/Engine/Source/Developer/TraceInsights/Private/Insights/Widgets/STimingView.h b/Engine/Source/Developer/TraceInsights/Private/Insights/Widgets/STimingView.h
--- a/Engine/Source/Developer/TraceInsights/Private/Insights/Widgets/STimingView.h
+++ b/Engine/Source/Developer/TraceInsights/Private/Insights/Widgets/STimingView.h
@@ -249,6 +249,9 @@ public:
 
 	const FTimingTrackViewport& GetViewport() { return Viewport; }
 
+	double GetSelectionStartTime() const { return SelectionStartTime; }
+	double GetSelectionEndTime() const { return SelectionEndTime; }
+
 	bool IsTimeSelected(double Time) const { return Time >= SelectionStartTime && Time < SelectionEndTime; }
 	bool IsTimeSelectedInclusive(double Time) const { return Time >= SelectionStartTime && Time <= SelectionEndTime; }
 
diff --git a/Engine/Source/Developer/TraceInsights/Public/Insights/Version.h b/Engine/Source/Developer/TraceInsights/Public/Insights/Version.h
--- a/Engine/Source/Developer/TraceInsights/Public/Insights/Version.h
+++ b/Engine/Source/Developer/TraceInsights/Public/Insights/Version.h
@@ -5,10 +5,10 @@
 // See https://semver.org/
 
 #define UNREAL_INSIGHTS_VERSION_MAJOR 0
-#define UNREAL_INSIGHTS_VERSION_MINOR 8
+#define UNREAL_INSIGHTS_VERSION_MINOR 9
 #define UNREAL_INSIGHTS_VERSION_PATCH 0
 
-#define UNREAL_INSIGHTS_VERSION_STRING "0.08"
+#define UNREAL_INSIGHTS_VERSION_STRING "0.09"
 
 /** Extra identifier string (like: "alpha", "beta.1", "beta.2", "test", "featureX", etc.). */
 //#define UNREAL_INSIGHTS_VERSION_ID_STRING "alpha"
