[Insights] Fixed AreTableLayoutsEqual in Table.cpp.

diff --git a/Engine/Source/Developer/TraceInsights/Private/Insights/Table/ViewModels/Table.cpp b/Engine/Source/Developer/TraceInsights/Private/Insights/Table/ViewModels/Table.cpp
--- a/Engine/Source/Developer/TraceInsights/Private/Insights/Table/ViewModels/Table.cpp
+++ b/Engine/Source/Developer/TraceInsights/Private/Insights/Table/ViewModels/Table.cpp
@@ -70,11 +70,11 @@ bool AreTableLayoutsEqual(const Trace::ITableLayout& TableLayoutA, const Trace::
 	int32 ColumnCount = TableLayoutA.GetColumnCount();
 	for (int32 ColumnIndex = 0; ColumnIndex < ColumnCount; ++ColumnIndex)
 	{
-		if (TableLayoutA.GetColumnType(ColumnIndex) != TableLayoutA.GetColumnType(ColumnIndex))
+		if (TableLayoutA.GetColumnType(ColumnIndex) != TableLayoutB.GetColumnType(ColumnIndex))
 		{
 			return false;
 		}
-		if (FCString::Strcmp(TableLayoutA.GetColumnName(ColumnIndex), TableLayoutA.GetColumnName(ColumnIndex)) != 0)
+		if (FCString::Strcmp(TableLayoutA.GetColumnName(ColumnIndex), TableLayoutB.GetColumnName(ColumnIndex)) != 0)
 		{
 			return false;
 		}
@@ -114,7 +114,7 @@ void FTable::CreateHierarchyColumn(const TCHAR* ColumnName)
 
 	TSharedPtr<FTableColumn> ColumnPtr = MakeShareable(new FTableColumn
 	(
-		0, // Order
+		-1, // Order
 		0,
 		FName(ColumnName), // Id
 		ColumnNameText, // Short Name
@@ -183,7 +183,7 @@ void FTable::CreateColumns()
 		float MinColumnWidth = 0.0f;
 		float MaxColumnWidth = FLT_MAX;
 
-		FTableColumn::FGetValueAsTextFunction GetUnknownValueAsTextFn = [](const FTable& Table, const FTableColumn& Column, const FTableRowId& RowId)->FText
+		FTableColumn::FGetValueAsTextFunction GetUnknownValueAsTextFn = [](const FTable& Table, const FTableColumn& Column, const FTableRowId& RowId) -> FText
 		{
 			return LOCTEXT("UnknownValue", "!?");
 		};
@@ -281,7 +281,7 @@ void FTable::CreateColumns()
 
 		TSharedPtr<FTableColumn> ColumnPtr = MakeShareable(new FTableColumn
 		(
-			ColumnIndex + 1, // Order
+			ColumnIndex, // Order
 			ColumnIndex,
 			FName(ColumnName), // Id
 			ColumnNameText, // Short Name
