[Insights]
  - StartPage: Changed units for session file sizes to be consistently KiB (the actual value in bytes and in KiB/MiB/GiB can be seen in tooltip).
  - StartPage: Fixed file size for live sessions (i.e. now updates as file size changes).
  - Generic Table Tree View: Improved default value formatting for "float" and "double". Also added all column values in the tooltip (where double values also show as "time duration").

diff --git a/Engine/Source/Developer/TraceInsights/Private/Insights/Table/ViewModels/Table.cpp b/Engine/Source/Developer/TraceInsights/Private/Insights/Table/ViewModels/Table.cpp
--- a/Engine/Source/Developer/TraceInsights/Private/Insights/Table/ViewModels/Table.cpp
+++ b/Engine/Source/Developer/TraceInsights/Private/Insights/Table/ViewModels/Table.cpp
@@ -4,6 +4,7 @@
 #include "TraceServices/Containers/Tables.h"
 
 // Insights
+#include "Insights/Common/TimeUtils.h"
 #include "Insights/Table/ViewModels/TableColumn.h"
 
 #define LOCTEXT_NAMESPACE "Insights_Table"
@@ -107,11 +108,6 @@ void FTable::CreateHierarchyColumn(const TCHAR* ColumnName)
 	constexpr float MinColumnWidth = 0.0f;
 	constexpr float MaxColumnWidth = FLT_MAX;
 
-	FTableColumn::FGetValueAsTextFunction GetValueAsTextFn = [](const FTable& Table, const FTableColumn& Column, const FTableRowId& RowId)->FText
-	{
-		return FText();
-	};
-
 	TSharedPtr<FTableColumn> ColumnPtr = MakeShareable(new FTableColumn
 	(
 		-1, // Order
@@ -124,8 +120,7 @@ void FTable::CreateHierarchyColumn(const TCHAR* ColumnName)
 		HorizontalAlignment,
 		InitialColumnWidth,
 		MinColumnWidth,
-		MaxColumnWidth,
-		GetValueAsTextFn
+		MaxColumnWidth
 	));
 
 	AddColumn(ColumnPtr);
@@ -219,7 +214,33 @@ void FTable::CreateColumns()
 			{
 				Reader->SetRowIndex(RowId.RowIndex);
 				float Value = Reader->GetValueFloat(Column.GetIndex());
-				return FText::AsNumber(Value);
+				if (Value == 0.0f)
+				{
+					return FText::FromString(TEXT("0"));
+				}
+				else
+				{
+					return FText::FromString(FString::Printf(TEXT("%f"), Value));
+				}
+			}
+			return FText::GetEmpty();
+		};
+
+		FTableColumn::FGetValueAsTextFunction GetFloatValueAsTextFn2 = [](const FTable& Table, const FTableColumn& Column, const FTableRowId& RowId) -> FText
+		{
+			TSharedPtr<Trace::IUntypedTableReader> Reader = Table.GetTableReader();
+			if (Reader.IsValid() && RowId.HasValidIndex())
+			{
+				Reader->SetRowIndex(RowId.RowIndex);
+				float Value = Reader->GetValueFloat(Column.GetIndex());
+				if (Value == 0.0f)
+				{
+					return FText::FromString(TEXT("0"));
+				}
+				else
+				{
+					return FText::FromString(FString::Printf(TEXT("%f (%s)"), Value, *TimeUtils::FormatTimeAuto(static_cast<double>(Value))));
+				}
 			}
 			return FText::GetEmpty();
 		};
@@ -231,7 +252,33 @@ void FTable::CreateColumns()
 			{
 				Reader->SetRowIndex(RowId.RowIndex);
 				double Value = Reader->GetValueDouble(Column.GetIndex());
-				return FText::AsNumber(Value);
+				if (Value == 0.0)
+				{
+					return FText::FromString(TEXT("0"));
+				}
+				else
+				{
+					return FText::FromString(FString::Printf(TEXT("%f"), Value));
+				}
+			}
+			return FText::GetEmpty();
+		};
+
+		FTableColumn::FGetValueAsTextFunction GetDoubleValueAsTextFn2 = [](const FTable& Table, const FTableColumn& Column, const FTableRowId& RowId) -> FText
+		{
+			TSharedPtr<Trace::IUntypedTableReader> Reader = Table.GetTableReader();
+			if (Reader.IsValid() && RowId.HasValidIndex())
+			{
+				Reader->SetRowIndex(RowId.RowIndex);
+				double Value = Reader->GetValueDouble(Column.GetIndex());
+				if (Value == 0.0)
+				{
+					return FText::FromString(TEXT("0"));
+				}
+				else
+				{
+					return FText::FromString(FString::Printf(TEXT("%f (%s)"), Value, *TimeUtils::FormatTimeAuto(Value)));
+				}
 			}
 			return FText::GetEmpty();
 		};
@@ -249,6 +296,7 @@ void FTable::CreateColumns()
 		};
 
 		FTableColumn::FGetValueAsTextFunction GetValueAsTextFn = GetUnknownValueAsTextFn;
+		FTableColumn::FGetValueAsTextFunction GetValueAsTooltipTextFn = GetUnknownValueAsTextFn;
 
 		switch (ColumnType)
 		{
@@ -256,26 +304,31 @@ void FTable::CreateColumns()
 			HorizontalAlignment = HAlign_Right;
 			InitialColumnWidth = 40.0f;
 			GetValueAsTextFn = GetBoolValueAsTextFn;
+			GetValueAsTooltipTextFn = GetBoolValueAsTextFn;
 			break;
 		case Trace::TableColumnType_Int:
 			HorizontalAlignment = HAlign_Right;
 			InitialColumnWidth = 60.0f;
 			GetValueAsTextFn = GetIntValueAsTextFn;
+			GetValueAsTooltipTextFn = GetIntValueAsTextFn;
 			break;
 		case Trace::TableColumnType_Float:
 			HorizontalAlignment = HAlign_Right;
 			InitialColumnWidth = 60.0f;
 			GetValueAsTextFn = GetFloatValueAsTextFn;
+			GetValueAsTooltipTextFn = GetFloatValueAsTextFn2;
 			break;
 		case Trace::TableColumnType_Double:
 			HorizontalAlignment = HAlign_Right;
 			InitialColumnWidth = 80.0f;
 			GetValueAsTextFn = GetDoubleValueAsTextFn;
+			GetValueAsTooltipTextFn = GetDoubleValueAsTextFn2;
 			break;
 		case Trace::TableColumnType_CString:
 			HorizontalAlignment = HAlign_Left;
 			InitialColumnWidth = FMath::Max(120.0f, 6.0f * ColumnNameStr.Len());
 			GetValueAsTextFn = GetStringValueAsTextFn;
+			GetValueAsTooltipTextFn = GetStringValueAsTextFn;
 			break;
 		}
 
@@ -291,10 +344,12 @@ void FTable::CreateColumns()
 			HorizontalAlignment,
 			InitialColumnWidth,
 			MinColumnWidth,
-			MaxColumnWidth,
-			GetValueAsTextFn
+			MaxColumnWidth
 		));
 
+		ColumnPtr->SetValueAsTextFunction(GetValueAsTextFn);
+		ColumnPtr->SetValueAsTooltipTextFunction(GetValueAsTooltipTextFn);
+
 		AddColumn(ColumnPtr);
 	}
 }
diff --git a/Engine/Source/Developer/TraceInsights/Private/Insights/Table/ViewModels/TableColumn.h b/Engine/Source/Developer/TraceInsights/Private/Insights/Table/ViewModels/TableColumn.h
--- a/Engine/Source/Developer/TraceInsights/Private/Insights/Table/ViewModels/TableColumn.h
+++ b/Engine/Source/Developer/TraceInsights/Private/Insights/Table/ViewModels/TableColumn.h
@@ -53,8 +53,7 @@ public:
 		const EHorizontalAlignment InHorizontalAlignment,
 		const float InInitialWidth,
 		const float InMinWidth,
-		const float InMaxWidth,
-		const FGetValueAsTextFunction InGetValueAsTextFn
+		const float InMaxWidth
 	)
 		: Order(InOrder)
 		, Index(InIndex)
@@ -67,9 +66,15 @@ public:
 		, InitialWidth(InInitialWidth)
 		, MinWidth(InMinWidth)
 		, MaxWidth(InMaxWidth)
-		, GetValueAsTextFn(InGetValueAsTextFn)
+		, GetValueAsTextFn(nullptr)
+		, GetValueAsTooltipTextFn(nullptr)
 		, ParentTable(nullptr)
 	{
+		GetValueAsTextFn = [](const FTable& Table, const FTableColumn& Column, const FTableRowId& RowId) -> FText
+		{
+			return FText();
+		};
+		GetValueAsTooltipTextFn = GetValueAsTextFn;
 	}
 
 	int32 GetOrder() const { return Order; }
@@ -118,6 +123,14 @@ public:
 		return GetValueAsTextFn(*ParentTable.Pin(), *this, InRowId);
 	}
 
+	FText GetValueAsTooltipText(const FTableRowId& InRowId) const
+	{
+		return GetValueAsTooltipTextFn(*ParentTable.Pin(), *this, InRowId);
+	}
+
+	void SetValueAsTextFunction(const FGetValueAsTextFunction InGetValueAsTextFn) { GetValueAsTextFn = InGetValueAsTextFn; }
+	void SetValueAsTooltipTextFunction(const FGetValueAsTextFunction InGetValueAsTooltipTextFn) { GetValueAsTooltipTextFn = InGetValueAsTooltipTextFn; }
+
 	TWeakPtr<FTable> GetParentTable() const { return ParentTable; }
 	void SetParentTable(TWeakPtr<FTable> InParentTable) { ParentTable = InParentTable; }
 
@@ -154,9 +167,12 @@ private:
 	float MaxWidth; /**< Maximum column width. */
 	//float Width; /**< Current column width. */
 
-	/** Custom function used to format (as an FText) the value displayed by this column. */
+	/** Custom function used to format (as an FText) the value displayed by this column (usually in table views). */
 	FGetValueAsTextFunction GetValueAsTextFn;
 
+	/** Custom function used to format (as an FText) the value displayed by this column (usually in tooltips; i.e. with increased precission). */
+	FGetValueAsTextFunction GetValueAsTooltipTextFn;
+
 	/* Parent table. Only one table instance can own this column. */
 	TWeakPtr<FTable> ParentTable;
 };
diff --git a/Engine/Source/Developer/TraceInsights/Private/Insights/Table/Widgets/STableTreeViewTooltip.cpp b/Engine/Source/Developer/TraceInsights/Private/Insights/Table/Widgets/STableTreeViewTooltip.cpp
--- a/Engine/Source/Developer/TraceInsights/Private/Insights/Table/Widgets/STableTreeViewTooltip.cpp
+++ b/Engine/Source/Developer/TraceInsights/Private/Insights/Table/Widgets/STableTreeViewTooltip.cpp
@@ -173,6 +173,23 @@ TSharedPtr<SToolTip> STableTreeViewTooltip::GetCellTooltip(const TSharedPtr<FTab
 					]
 				]
 
+				+ SVerticalBox::Slot()
+					.AutoHeight()
+					.Padding(2.0f)
+					[
+						SNew(SSeparator)
+						.Orientation(Orient_Horizontal)
+					]
+
+				+ SVerticalBox::Slot()
+					.AutoHeight()
+					.Padding(2.0f)
+					[
+						SAssignNew(GridPanel, SGridPanel)
+
+						// Values for each table column are added here.
+					]
+
 				+SVerticalBox::Slot()
 				.AutoHeight()
 				.Padding(2.0f)
@@ -183,9 +200,47 @@ TSharedPtr<SToolTip> STableTreeViewTooltip::GetCellTooltip(const TSharedPtr<FTab
 			]
 		];
 
+	TSharedPtr<FTable> Table = TreeNodePtr->GetParentTable().Pin();
+	if (Table.IsValid())
+	{
+		int32 Row = 0;
+		for (const TSharedPtr<FTableColumn>& Column : Table->GetColumns())
+		{
+			if (!Column->IsHierarchy() && TreeNodePtr->GetRowId().HasValidIndex())
+			{
+				FText Name = FText::Format(LOCTEXT("TooltipValueFormat", "{0}:"), Column->GetTitleName());
+				AddGridRow(GridPanel, Row, Name, Column->GetValueAsTooltipText(TreeNodePtr->GetRowId()));
+			}
+		}
+	}
+
 	return TableCellTooltip;
 }
 
+////////////////////////////////////////////////////////////////////////////////////////////////////
+
+void STableTreeViewTooltip::AddGridRow(TSharedPtr<SGridPanel> Grid, int32& Row, const FText& Name, const FText& Value)
+{
+	Grid->AddSlot(0, Row)
+		.Padding(2.0f)
+		[
+			SNew(STextBlock)
+			.Text(Name)
+			.TextStyle(FEditorStyle::Get(), TEXT("Profiler.TooltipBold"))
+		];
+
+	Grid->AddSlot(1, Row)
+		.Padding(2.0f)
+		.HAlign(HAlign_Left)
+		[
+			SNew(STextBlock)
+			.Text(Value)
+			.TextStyle(FEditorStyle::Get(), TEXT("Profiler.Tooltip"))
+		];
+
+	Row++;
+}
+
 END_SLATE_FUNCTION_BUILD_OPTIMIZATION
 
 ////////////////////////////////////////////////////////////////////////////////////////////////////
diff --git a/Engine/Source/Developer/TraceInsights/Private/Insights/Table/Widgets/STableTreeViewTooltip.h b/Engine/Source/Developer/TraceInsights/Private/Insights/Table/Widgets/STableTreeViewTooltip.h
--- a/Engine/Source/Developer/TraceInsights/Private/Insights/Table/Widgets/STableTreeViewTooltip.h
+++ b/Engine/Source/Developer/TraceInsights/Private/Insights/Table/Widgets/STableTreeViewTooltip.h
@@ -28,6 +28,9 @@ public:
 	static TSharedPtr<SToolTip> GetTableTooltip(const FTable& Table);
 	static TSharedPtr<SToolTip> GetColumnTooltip(const FTableColumn& Column);
 	static TSharedPtr<SToolTip> GetCellTooltip(const TSharedPtr<FTableTreeNode> TreeNodePtr, const TSharedPtr<FTableColumn> ColumnPtr);
+
+private:
+	static void AddGridRow(TSharedPtr<SGridPanel> Grid, int32& Row, const FText& Name, const FText& Value);
 };
 
 } // namespace Insights
diff --git a/Engine/Source/Developer/TraceInsights/Private/Insights/Widgets/SStartPageWindow.cpp b/Engine/Source/Developer/TraceInsights/Private/Insights/Widgets/SStartPageWindow.cpp
--- a/Engine/Source/Developer/TraceInsights/Private/Insights/Widgets/SStartPageWindow.cpp
+++ b/Engine/Source/Developer/TraceInsights/Private/Insights/Widgets/SStartPageWindow.cpp
@@ -161,7 +161,11 @@ public:
 		TSharedPtr<FTraceSession> TraceSessionPin = WeakTraceSession.Pin();
 		if (TraceSessionPin.IsValid())
 		{
-			return FText::AsMemory(TraceSessionPin->Size);
+			//FNumberFormattingOptions FormattingOptions;
+			//FormattingOptions.MinimumFractionalDigits = 1;
+			//FormattingOptions.MaximumFractionalDigits = 1;
+			//return FText::AsMemory(TraceSessionPin->Size, &FormattingOptions);
+			return FText::Format(LOCTEXT("SessionFileSizeFormatKiB", "{0} KiB"), TraceSessionPin->Size / 1024);
 		}
 		else
 		{
@@ -222,13 +226,28 @@ public:
 		TSharedPtr<FTraceSession> TraceSessionPin = WeakTraceSession.Pin();
 		if (TraceSessionPin.IsValid())
 		{
-			return FText::Format(LOCTEXT("TraceSessionTooltipFormat", "{0}\n{1}\nTimestamp: {2}\nFile Size: {3} bytes\nStatus: {4}"),
-				TraceSessionPin->Name,
-				TraceSessionPin->Uri,
-				FText::AsDateTime(TraceSessionPin->TimeStamp),
-				FText::AsNumber(TraceSessionPin->Size),
-				TraceSessionPin->bIsLive ? LOCTEXT("LiveTraceSessionStatus", "LIVE") : LOCTEXT("OfflineTraceSessionStatus", "Offline")
-			);
+			const FText Status = TraceSessionPin->bIsLive ? LOCTEXT("LiveTraceSessionStatus", "LIVE") : LOCTEXT("OfflineTraceSessionStatus", "Offline");
+			if (TraceSessionPin->Size > 1024)
+			{
+				return FText::Format(LOCTEXT("TraceSessionTooltipFormat1", "{0}\n{1}\nTimestamp: {2}\nFile Size: {3} bytes ({4})\nStatus: {5}"),
+					TraceSessionPin->Name,
+					TraceSessionPin->Uri,
+					FText::AsDateTime(TraceSessionPin->TimeStamp),
+					FText::AsNumber(TraceSessionPin->Size),
+					FText::AsMemory(TraceSessionPin->Size),
+					Status
+				);
+			}
+			else
+			{
+				return FText::Format(LOCTEXT("TraceSessionTooltipFormat2", "{0}\n{1}\nTimestamp: {2}\nFile Size: {3} bytes\nStatus: {4}"),
+					TraceSessionPin->Name,
+					TraceSessionPin->Uri,
+					FText::AsDateTime(TraceSessionPin->TimeStamp),
+					FText::AsNumber(TraceSessionPin->Size),
+					Status
+				);
+			}
 		}
 		else
 		{
@@ -762,7 +781,7 @@ void SStartPageWindow::RefreshTraceSessionList()
 
 	AvailableSessionCount = AvailableSessions.Num();
 
-	// Count number of live sessions.
+	// Count number of live sessions and update file sizes.
 	int32 OldLiveSessionCount = LiveSessionCount;
 	LiveSessionCount = 0;
 	for (Trace::FSessionHandle SessionHandle : AvailableSessions)
@@ -773,6 +792,12 @@ void SStartPageWindow::RefreshTraceSessionList()
 		{
 			++LiveSessionCount;
 		}
+
+		TSharedPtr<FTraceSession>* TraceSessionPtrPtr = TraceSessionsMap.Find(SessionHandle);
+		if (TraceSessionPtrPtr)
+		{
+			(*TraceSessionPtrPtr)->Size = SessionInfo.Size;
+		}
 	}
 
 	// If session list has changed on analysis side, recreate the TraceSessions list view widget.
@@ -783,18 +808,21 @@ void SStartPageWindow::RefreshTraceSessionList()
 		TSharedPtr<FTraceSession> NewSelectedTraceSession;
 
 		TraceSessions.Reset();
+		TraceSessionsMap.Reset();
 
 		for (Trace::FSessionHandle SessionHandle : AvailableSessions)
 		{
 			Trace::FSessionInfo SessionInfo;
 			SessionService->GetSessionInfo(SessionHandle, SessionInfo);
 
-			TraceSessions.Add(MakeShareable(new FTraceSession(SessionHandle, SessionInfo)));
+			const TSharedPtr<FTraceSession> TraceSessionPtr = MakeShareable(new FTraceSession(SessionHandle, SessionInfo));
+			TraceSessions.Add(TraceSessionPtr);
+			TraceSessionsMap.Add(SessionHandle, TraceSessionPtr);
 
 			// Identify the previously selected session (if stil available) to ensure selection remains unchanged.
 			if (SelectedTraceSession && SelectedTraceSession->SessionHandle == SessionHandle)
 			{
-				NewSelectedTraceSession = TraceSessions.Last();
+				NewSelectedTraceSession = TraceSessionPtr;
 			}
 		}
 		Algo::SortBy(TraceSessions, &FTraceSession::TimeStamp);
diff --git a/Engine/Source/Developer/TraceInsights/Private/Insights/Widgets/SStartPageWindow.h b/Engine/Source/Developer/TraceInsights/Private/Insights/Widgets/SStartPageWindow.h
--- a/Engine/Source/Developer/TraceInsights/Private/Insights/Widgets/SStartPageWindow.h
+++ b/Engine/Source/Developer/TraceInsights/Private/Insights/Widgets/SStartPageWindow.h
@@ -204,6 +204,7 @@ private:
 
 	TSharedPtr<SListView<TSharedPtr<FTraceSession>>> TraceSessionsListView;
 	TArray<TSharedPtr<FTraceSession>> TraceSessions;
+	TMap<Trace::FSessionHandle, TSharedPtr<FTraceSession>> TraceSessionsMap;
 	TSharedPtr<SEditableTextBox> HostTextBox;
 	TSharedPtr<FTraceSession> SelectedTraceSession;
 
