[Insights]
  - Added initial layout for the new Networking Insights profiler.

diff --git a/Engine/Source/Developer/TraceInsights/Private/Insights/InsightsManager.cpp b/Engine/Source/Developer/TraceInsights/Private/Insights/InsightsManager.cpp
--- a/Engine/Source/Developer/TraceInsights/Private/Insights/InsightsManager.cpp
+++ b/Engine/Source/Developer/TraceInsights/Private/Insights/InsightsManager.cpp
@@ -8,6 +8,7 @@
 
 // Insights
 #include "Insights/LoadingProfiler/LoadingProfilerManager.h"
+#include "Insights/NetworkingProfiler/NetworkingProfilerManager.h"
 #include "Insights/TimingProfilerManager.h"
 #include "Insights/Widgets/SStartPageWindow.h"
 #include "Insights/Widgets/STimingProfilerWindow.h"
@@ -19,6 +20,7 @@
 const FName FInsightsManagerTabs::StartPageTabId(TEXT("StartPage"));
 const FName FInsightsManagerTabs::TimingProfilerTabId(TEXT("TimingProfiler"));
 const FName FInsightsManagerTabs::LoadingProfilerTabId(TEXT("LoadingProfiler"));
+const FName FInsightsManagerTabs::NetworkingProfilerTabId(TEXT("NetworkingProfiler"));
 
 TSharedPtr<FInsightsManager> FInsightsManager::Instance = nullptr;
 
@@ -154,6 +156,12 @@ void FInsightsManager::OnSessionChanged()
 		LoadingProfilerManager->OnSessionChanged();
 	}
 
+	if (TSharedPtr<FNetworkingProfilerManager> NetworkingProfilerManager = FNetworkingProfilerManager::Get())
+	{
+		// FIXME: make NetworkingProfilerManager to register to SessionChangedEvent instead
+		NetworkingProfilerManager->OnSessionChanged();
+	}
+
 	SessionChangedEvent.Broadcast();
 }
 
@@ -170,6 +178,10 @@ void FInsightsManager::SpawnAndActivateTabs()
 	{
 		FGlobalTabmanager::Get()->InvokeTab(FInsightsManagerTabs::LoadingProfilerTabId);
 	}
+	if (FGlobalTabmanager::Get()->HasTabSpawner(FInsightsManagerTabs::NetworkingProfilerTabId))
+	{
+		FGlobalTabmanager::Get()->InvokeTab(FInsightsManagerTabs::NetworkingProfilerTabId);
+	}
 
 	// Ensure Timing Insights / Timing View is the active tab / view.
 	if (TSharedPtr<SDockTab> TimingInsightsTab = FGlobalTabmanager::Get()->FindExistingLiveTab(FInsightsManagerTabs::TimingProfilerTabId))
diff --git a/Engine/Source/Developer/TraceInsights/Private/Insights/InsightsManager.h b/Engine/Source/Developer/TraceInsights/Private/Insights/InsightsManager.h
--- a/Engine/Source/Developer/TraceInsights/Private/Insights/InsightsManager.h
+++ b/Engine/Source/Developer/TraceInsights/Private/Insights/InsightsManager.h
@@ -24,6 +24,7 @@ struct FInsightsManagerTabs
 	static const FName StartPageTabId;
 	static const FName TimingProfilerTabId;
 	static const FName LoadingProfilerTabId;
+	static const FName NetworkingProfilerTabId;
 };
 
 ////////////////////////////////////////////////////////////////////////////////////////////////////
diff --git a/Engine/Source/Developer/TraceInsights/Private/Insights/InsightsStyle.cpp b/Engine/Source/Developer/TraceInsights/Private/Insights/InsightsStyle.cpp
--- a/Engine/Source/Developer/TraceInsights/Private/Insights/InsightsStyle.cpp
+++ b/Engine/Source/Developer/TraceInsights/Private/Insights/InsightsStyle.cpp
@@ -51,6 +51,9 @@ TSharedRef<FSlateStyleSet> FInsightsStyle::Create()
 	Style.Set("LoadingProfiler.Icon.Large", new IMAGE_BRUSH("Icons/Profiler/profiler_stats_40x", FVector2D(32.0f, 32.0f)));
 	Style.Set("LoadingProfiler.Icon.Small", new IMAGE_BRUSH("Icons/Profiler/profiler_stats_40x", FVector2D(16.0f, 16.0f)));
 
+	Style.Set("NetworkingProfiler.Icon.Large", new IMAGE_BRUSH("Icons/Profiler/profiler_stats_40x", FVector2D(32.0f, 32.0f)));
+	Style.Set("NetworkingProfiler.Icon.Small", new IMAGE_BRUSH("Icons/Profiler/profiler_stats_40x", FVector2D(16.0f, 16.0f)));
+
 	Style.Set("Toolbar.Icon.Large", new IMAGE_BRUSH("/Icons/icon_tab_Tools_16x", FVector2D(32.0f, 32.0f)));
 	Style.Set("Toolbar.Icon.Small", new IMAGE_BRUSH("/Icons/icon_tab_Tools_16x", FVector2D(16.0f, 16.0f)));
 
@@ -87,6 +90,18 @@ TSharedRef<FSlateStyleSet> FInsightsStyle::Create()
 	Style.Set("TableTreeView.Icon.Large", new IMAGE_BRUSH("Icons/Profiler/Profiler_Data_Capture_40x", FVector2D(32.0f, 32.0f)));
 	Style.Set("TableTreeView.Icon.Small", new IMAGE_BRUSH("Icons/Profiler/Profiler_Data_Capture_40x", FVector2D(16.0f, 16.0f)));
 
+	//////////////////////////////////////////////////
+	// Networking Insights
+
+	Style.Set("PacketSizesView.Icon.Large", new IMAGE_BRUSH("Icons/Profiler/profiler_stats_40x", FVector2D(32.0f, 32.0f)));
+	Style.Set("PacketSizesView.Icon.Small", new IMAGE_BRUSH("Icons/Profiler/profiler_stats_40x", FVector2D(16.0f, 16.0f)));
+
+	Style.Set("PacketBreakdownView.Icon.Large", new IMAGE_BRUSH("Icons/Profiler/profiler_stats_40x", FVector2D(32.0f, 32.0f)));
+	Style.Set("PacketBreakdownView.Icon.Small", new IMAGE_BRUSH("Icons/Profiler/profiler_stats_40x", FVector2D(16.0f, 16.0f)));
+
+	Style.Set("DataStreamBreakdownView.Icon.Large", new IMAGE_BRUSH("Icons/Profiler/profiler_stats_40x", FVector2D(32.0f, 32.0f)));
+	Style.Set("DataStreamBreakdownView.Icon.Small", new IMAGE_BRUSH("Icons/Profiler/profiler_stats_40x", FVector2D(16.0f, 16.0f)));
+
 	//////////////////////////////////////////////////
 
 	return StyleRef;
diff --git a/Engine/Source/Developer/TraceInsights/Private/Insights/NetworkingProfiler/NetworkingProfilerCommands.cpp b/Engine/Source/Developer/TraceInsights/Private/Insights/NetworkingProfiler/NetworkingProfilerCommands.cpp
new file mode 100644
--- /dev/null
+++ b/Engine/Source/Developer/TraceInsights/Private/Insights/NetworkingProfiler/NetworkingProfilerCommands.cpp
@@ -0,0 +1,106 @@
+// Copyright 1998-2019 Epic Games, Inc. All Rights Reserved.
+
+#include "NetworkingProfilerCommands.h"
+
+#include "DesktopPlatformModule.h"
+#include "EditorStyleSet.h"
+#include "Framework/Application/SlateApplication.h"
+#include "Framework/MultiBox/MultiBoxBuilder.h"
+
+// Insights
+#include "Insights/InsightsManager.h"
+#include "Insights/NetworkingProfiler/NetworkingProfilerManager.h"
+
+////////////////////////////////////////////////////////////////////////////////////////////////////
+
+#define LOCTEXT_NAMESPACE "FNetworkingProfilerCommands"
+
+////////////////////////////////////////////////////////////////////////////////////////////////////
+// FNetworkingProfilerMenuBuilder
+////////////////////////////////////////////////////////////////////////////////////////////////////
+
+void FNetworkingProfilerMenuBuilder::AddMenuEntry(FMenuBuilder& MenuBuilder, const TSharedPtr< FUICommandInfo >& UICommandInfo, const FUIAction& UIAction)
+{
+	MenuBuilder.AddMenuEntry
+	(
+		UICommandInfo->GetLabel(),
+		UICommandInfo->GetDescription(),
+		UICommandInfo->GetIcon(),
+		UIAction,
+		NAME_None,
+		UICommandInfo->GetUserInterfaceType()
+	);
+}
+
+////////////////////////////////////////////////////////////////////////////////////////////////////
+// FNetworkingProfilerCommands
+////////////////////////////////////////////////////////////////////////////////////////////////////
+
+FNetworkingProfilerCommands::FNetworkingProfilerCommands()
+	: TCommands<FNetworkingProfilerCommands>(
+		TEXT("NetworkingProfilerCommand"), // Context name for fast lookup
+		NSLOCTEXT("Contexts", "NetworkingProfilerCommand", "Networking Insights"), // Localized context name for displaying
+		NAME_None, // Parent
+		FEditorStyle::GetStyleSetName() // Icon Style Set
+	)
+{
+}
+
+////////////////////////////////////////////////////////////////////////////////////////////////////
+
+// UI_COMMAND takes long for the compiler to optimize
+PRAGMA_DISABLE_OPTIMIZATION
+void FNetworkingProfilerCommands::RegisterCommands()
+{
+	UI_COMMAND(TogglePacketSizesViewVisibility, "Packet Sizes", "Toggles the visibility of the Packet Sizes view", EUserInterfaceActionType::ToggleButton, FInputChord());
+	UI_COMMAND(TogglePacketBreakdownViewVisibility, "Packet Breakdown", "Toggles the visibility of the Packet Breakdown view", EUserInterfaceActionType::ToggleButton, FInputChord());
+	UI_COMMAND(ToggleDataStreamBreakdownViewVisibility, "Data Stream Breakdown", "Toggles the visibility of the Data Stream Breakdown view", EUserInterfaceActionType::ToggleButton, FInputChord());
+}
+PRAGMA_ENABLE_OPTIMIZATION
+
+////////////////////////////////////////////////////////////////////////////////////////////////////
+// Toggle Commands
+////////////////////////////////////////////////////////////////////////////////////////////////////
+
+#define IMPLEMENT_TOGGLE_COMMAND(CmdName, IsEnabled, SetIsEnabled) \
+	\
+	void FNetworkingProfilerActionManager::Map_##CmdName##_Global()\
+	{\
+		This->CommandList->MapAction(This->GetCommands().CmdName, CmdName##_Custom());\
+	}\
+	\
+	const FUIAction FNetworkingProfilerActionManager::CmdName##_Custom() \
+	{\
+		FUIAction UIAction;\
+		UIAction.ExecuteAction = FExecuteAction::CreateRaw(this, &FNetworkingProfilerActionManager::CmdName##_Execute);\
+		UIAction.CanExecuteAction = FCanExecuteAction::CreateRaw(this, &FNetworkingProfilerActionManager::CmdName##_CanExecute);\
+		UIAction.GetActionCheckState = FGetActionCheckState::CreateRaw(this, &FNetworkingProfilerActionManager::CmdName##_GetCheckState);\
+		return UIAction;\
+	}\
+	\
+	void FNetworkingProfilerActionManager::CmdName##_Execute()\
+	{\
+		const bool b##IsEnabled = !This->IsEnabled();\
+		This->SetIsEnabled(b##IsEnabled);\
+	}\
+	\
+	bool FNetworkingProfilerActionManager::CmdName##_CanExecute() const\
+	{\
+		return FInsightsManager::Get()->GetSession().IsValid();\
+	}\
+	\
+	ECheckBoxState FNetworkingProfilerActionManager::CmdName##_GetCheckState() const\
+	{\
+		const bool b##IsEnabled = This->IsEnabled();\
+		return b##IsEnabled ? ECheckBoxState::Checked : ECheckBoxState::Unchecked;\
+	}
+
+IMPLEMENT_TOGGLE_COMMAND(TogglePacketSizesViewVisibility, IsPacketSizesViewVisible, ShowHidePacketSizesView)
+IMPLEMENT_TOGGLE_COMMAND(TogglePacketBreakdownViewVisibility, IsPacketBreakdownViewVisible, ShowHidePacketBreakdownView)
+IMPLEMENT_TOGGLE_COMMAND(ToggleDataStreamBreakdownViewVisibility, IsDataStreamBreakdownViewVisible, ShowHideDataStreamBreakdownView)
+
+#undef IMPLEMENT_TOGGLE_COMMAND
+
+////////////////////////////////////////////////////////////////////////////////////////////////////
+
+#undef LOCTEXT_NAMESPACE
diff --git a/Engine/Source/Developer/TraceInsights/Private/Insights/NetworkingProfiler/NetworkingProfilerCommands.h b/Engine/Source/Developer/TraceInsights/Private/Insights/NetworkingProfiler/NetworkingProfilerCommands.h
new file mode 100644
--- /dev/null
+++ b/Engine/Source/Developer/TraceInsights/Private/Insights/NetworkingProfiler/NetworkingProfilerCommands.h
@@ -0,0 +1,101 @@
+// Copyright 1998-2019 Epic Games, Inc. All Rights Reserved.
+
+#pragma once
+
+#include "CoreMinimal.h"
+#include "Framework/Commands/UIAction.h"
+#include "Framework/Commands/Commands.h"
+
+class FMenuBuilder;
+
+////////////////////////////////////////////////////////////////////////////////////////////////////
+/**
+ * Class that holds all profiler commands.
+ */
+class FNetworkingProfilerCommands : public TCommands<FNetworkingProfilerCommands>
+{
+public:
+	/** Default constructor. */
+	FNetworkingProfilerCommands();
+
+	/** Initialize commands. */
+	virtual void RegisterCommands() override;
+
+public:
+	//////////////////////////////////////////////////
+	// Global commands need to implement following method:
+	//     void Map_<CommandName>_Global();
+	// Custom commands needs to implement also the following method:
+	//     const FUIAction <CommandName>_Custom(...) const;
+	//////////////////////////////////////////////////
+
+	/** Toggles visibility for the Packet Sizes view. Global and custom command. */
+	TSharedPtr<FUICommandInfo> TogglePacketSizesViewVisibility;
+
+	/** Toggles visibility for the Packet Breakdown view. Global and custom command. */
+	TSharedPtr<FUICommandInfo> TogglePacketBreakdownViewVisibility;
+
+	/** Toggles visibility for the Data Stream Breakdown view. Global and custom command. */
+	TSharedPtr<FUICommandInfo> ToggleDataStreamBreakdownViewVisibility;
+};
+
+////////////////////////////////////////////////////////////////////////////////////////////////////
+/**
+ * Menu builder. Helper class for adding a customized menu entry using the global UI command info.
+ */
+class FNetworkingProfilerMenuBuilder
+{
+public:
+	/**
+	 * Helper method for adding a customized menu entry using the global UI command info.
+	 * FUICommandInfo cannot be executed with custom parameters, so we need to create a custom FUIAction,
+	 * but sometime we have global and local version for the UI command, so reuse data from the global UI command info.
+	 * Ex:
+	 *     SessionInstance_ToggleCapture          - Global version will toggle capture process for all active session instances
+	 *     SessionInstance_ToggleCapture_OneParam - Local version will toggle capture process only for the specified session instance
+	 *
+	 * @param MenuBuilder The menu to add items to
+	 * @param FUICommandInfo A shared pointer to the UI command info
+	 * @param UIAction Customized version of the UI command info stored in an UI action
+	 */
+	static void AddMenuEntry(FMenuBuilder& MenuBuilder, const TSharedPtr<FUICommandInfo>& UICommandInfo, const FUIAction& UIAction);
+};
+
+////////////////////////////////////////////////////////////////////////////////////////////////////
+/**
+ * Class that provides helper functions for the commands to avoid cluttering profiler manager with many small functions.
+ * Can't contain any variables. Directly operates on the profiler manager instance.
+ */
+class FNetworkingProfilerActionManager
+{
+	friend class FNetworkingProfilerManager;
+
+private:
+	/** Private constructor. */
+	FNetworkingProfilerActionManager(class FNetworkingProfilerManager* Instance)
+		: This(Instance)
+	{}
+
+	//////////////////////////////////////////////////
+	// Toggle Commands
+
+#define DECLARE_TOGGLE_COMMAND(CmdName)\
+public:\
+	void Map_##CmdName##_Global(); /**< Maps UI command info CmdName with the specified UI command list. */\
+	const FUIAction CmdName##_Custom(); /**< UI action for CmdName command. */\
+protected:\
+	void CmdName##_Execute(); /**< Handles FExecuteAction for CmdName. */\
+	bool CmdName##_CanExecute() const; /**< Handles FCanExecuteAction for CmdName. */\
+	ECheckBoxState CmdName##_GetCheckState() const; /**< Handles FGetActionCheckState for CmdName. */
+
+	DECLARE_TOGGLE_COMMAND(TogglePacketSizesViewVisibility)
+	DECLARE_TOGGLE_COMMAND(TogglePacketBreakdownViewVisibility)
+	DECLARE_TOGGLE_COMMAND(ToggleDataStreamBreakdownViewVisibility)
+#undef DECLARE_TOGGLE_COMMAND
+
+	//////////////////////////////////////////////////
+
+protected:
+	/** Reference to the global instance of the profiler manager. */
+	class FNetworkingProfilerManager* This;
+};
diff --git a/Engine/Source/Developer/TraceInsights/Private/Insights/NetworkingProfiler/NetworkingProfilerManager.cpp b/Engine/Source/Developer/TraceInsights/Private/Insights/NetworkingProfiler/NetworkingProfilerManager.cpp
new file mode 100644
--- /dev/null
+++ b/Engine/Source/Developer/TraceInsights/Private/Insights/NetworkingProfiler/NetworkingProfilerManager.cpp
@@ -0,0 +1,167 @@
+// Copyright 1998-2019 Epic Games, Inc. All Rights Reserved.
+
+#include "NetworkingProfilerManager.h"
+
+#include "Modules/ModuleManager.h"
+#include "TraceServices/AnalysisService.h"
+#include "TraceServices/SessionService.h"
+
+// Insights
+#include "Insights/InsightsManager.h"
+#include "Insights/NetworkingProfiler/Widgets/SDataStreamBreakdownView.h"
+#include "Insights/NetworkingProfiler/Widgets/SNetworkingProfilerWindow.h"
+#include "Insights/NetworkingProfiler/Widgets/SPacketBreakdownView.h"
+#include "Insights/NetworkingProfiler/Widgets/SPacketSizesView.h"
+
+////////////////////////////////////////////////////////////////////////////////////////////////////
+
+#define LOCTEXT_NAMESPACE "NetworkingProfilerManager"
+
+//DEFINE_LOG_CATEGORY(NetworkingProfiler);
+//
+//DEFINE_STAT(STAT_FT_OnPaint);
+//DEFINE_STAT(STAT_GT_OnPaint);
+//DEFINE_STAT(STAT_TT_OnPaint);
+//DEFINE_STAT(STAT_IOPM_Tick);
+
+TSharedPtr<FNetworkingProfilerManager> FNetworkingProfilerManager::Instance = nullptr;
+
+////////////////////////////////////////////////////////////////////////////////////////////////////
+
+FNetworkingProfilerManager::FNetworkingProfilerManager(TSharedRef<FUICommandList> InCommandList)
+	: CommandList(InCommandList)
+	, ActionManager(this)
+	, ProfilerWindow(nullptr)
+	, bIsPacketSizesViewVisible(true)
+	, bIsPacketBreakdownViewVisible(true)
+	, bIsDataStreamBreakdownViewVisible(true)
+{
+}
+
+////////////////////////////////////////////////////////////////////////////////////////////////////
+
+void FNetworkingProfilerManager::PostConstructor()
+{
+	// Register tick functions.
+	//OnTick = FTickerDelegate::CreateSP(this, &FNetworkingProfilerManager::Tick);
+	//OnTickHandle = FTicker::GetCoreTicker().AddTicker(OnTick, 1.0f);
+
+	FNetworkingProfilerCommands::Register();
+	BindCommands();
+}
+
+////////////////////////////////////////////////////////////////////////////////////////////////////
+
+void FNetworkingProfilerManager::BindCommands()
+{
+	ActionManager.Map_TogglePacketSizesViewVisibility_Global();
+	ActionManager.Map_TogglePacketBreakdownViewVisibility_Global();
+	ActionManager.Map_ToggleDataStreamBreakdownViewVisibility_Global();
+}
+
+////////////////////////////////////////////////////////////////////////////////////////////////////
+
+FNetworkingProfilerManager::~FNetworkingProfilerManager()
+{
+	FNetworkingProfilerCommands::Unregister();
+
+	// Unregister tick function.
+	//FTicker::GetCoreTicker().RemoveTicker(OnTickHandle);
+}
+
+////////////////////////////////////////////////////////////////////////////////////////////////////
+
+TSharedPtr<FNetworkingProfilerManager> FNetworkingProfilerManager::Get()
+{
+	return FNetworkingProfilerManager::Instance;
+}
+
+////////////////////////////////////////////////////////////////////////////////////////////////////
+
+const TSharedRef<FUICommandList> FNetworkingProfilerManager::GetCommandList() const
+{
+	return CommandList;
+}
+
+////////////////////////////////////////////////////////////////////////////////////////////////////
+
+const FNetworkingProfilerCommands& FNetworkingProfilerManager::GetCommands()
+{
+	return FNetworkingProfilerCommands::Get();
+}
+
+////////////////////////////////////////////////////////////////////////////////////////////////////
+
+FNetworkingProfilerActionManager& FNetworkingProfilerManager::GetActionManager()
+{
+	return FNetworkingProfilerManager::Instance->ActionManager;
+}
+
+////////////////////////////////////////////////////////////////////////////////////////////////////
+
+bool FNetworkingProfilerManager::Tick(float DeltaTime)
+{
+	//SCOPE_CYCLE_COUNTER(STAT_IOPM_Tick);
+
+	return true;
+}
+
+////////////////////////////////////////////////////////////////////////////////////////////////////
+
+void FNetworkingProfilerManager::OnSessionChanged()
+{
+	TSharedPtr<SNetworkingProfilerWindow> Wnd = GetProfilerWindow();
+	if (Wnd.IsValid())
+	{
+		Wnd->Reset();
+	}
+}
+
+////////////////////////////////////////////////////////////////////////////////////////////////////
+
+void FNetworkingProfilerManager::ShowHidePacketSizesView(const bool bPacketSizesViewVisibleState)
+{
+	//bool bWasPacketSizesViewVisible = bIsPacketSizesViewVisible;
+
+	bIsPacketSizesViewVisible = bPacketSizesViewVisibleState;
+
+	TSharedPtr<SNetworkingProfilerWindow> Wnd = GetProfilerWindow();
+	if (Wnd.IsValid())
+	{
+		Wnd->ShowHideTab(FNetworkingProfilerTabs::PacketSizesViewID, bIsPacketSizesViewVisible);
+	}
+}
+
+////////////////////////////////////////////////////////////////////////////////////////////////////
+
+void FNetworkingProfilerManager::ShowHidePacketBreakdownView(const bool bPacketBreakdownViewVisibleState)
+{
+	//bool bWasPacketBreakdownViewVisible = bIsPacketBreakdownViewVisible;
+
+	bIsPacketBreakdownViewVisible = bPacketBreakdownViewVisibleState;
+
+	TSharedPtr<SNetworkingProfilerWindow> Wnd = GetProfilerWindow();
+	if (Wnd.IsValid())
+	{
+		Wnd->ShowHideTab(FNetworkingProfilerTabs::PacketBreakdownViewID, bIsPacketBreakdownViewVisible);
+	}
+}
+
+////////////////////////////////////////////////////////////////////////////////////////////////////
+
+void FNetworkingProfilerManager::ShowHideDataStreamBreakdownView(const bool bDataStreamBreakdownViewVisibleState)
+{
+	//bool bWasDataStreamBreakdownViewVisible = bIsDataStreamBreakdownViewVisible;
+
+	bIsDataStreamBreakdownViewVisible = bDataStreamBreakdownViewVisibleState;
+
+	TSharedPtr<SNetworkingProfilerWindow> Wnd = GetProfilerWindow();
+	if (Wnd.IsValid())
+	{
+		Wnd->ShowHideTab(FNetworkingProfilerTabs::DataStreamBreakdownViewID, bIsDataStreamBreakdownViewVisible);
+	}
+}
+
+////////////////////////////////////////////////////////////////////////////////////////////////////
+
+#undef LOCTEXT_NAMESPACE
diff --git a/Engine/Source/Developer/TraceInsights/Private/Insights/NetworkingProfiler/NetworkingProfilerManager.h b/Engine/Source/Developer/TraceInsights/Private/Insights/NetworkingProfiler/NetworkingProfilerManager.h
new file mode 100644
--- /dev/null
+++ b/Engine/Source/Developer/TraceInsights/Private/Insights/NetworkingProfiler/NetworkingProfilerManager.h
@@ -0,0 +1,150 @@
+// Copyright 1998-2019 Epic Games, Inc. All Rights Reserved.
+
+#pragma once
+
+#include "Containers/Ticker.h"
+#include "CoreMinimal.h"
+#include "Framework/Commands/UICommandList.h"
+
+// Insights
+#include "Insights/InsightsManager.h"
+#include "Insights/NetworkingProfiler/NetworkingProfilerCommands.h"
+
+class SNetworkingProfilerWindow;
+
+namespace Trace
+{
+	class ISessionService;
+	class IAnalysisService;
+	class IAnalysisSession;
+}
+
+////////////////////////////////////////////////////////////////////////////////////////////////////
+/**
+ * This class manages the Timing Profiler state and settings.
+ */
+class FNetworkingProfilerManager : public TSharedFromThis<FNetworkingProfilerManager>
+{
+	friend class FNetworkingProfilerActionManager;
+
+public:
+	/** Creates the Timing Profiler manager, only one instance can exist. */
+	FNetworkingProfilerManager(TSharedRef<FUICommandList> InCommandList);
+
+	/** Virtual destructor. */
+	virtual ~FNetworkingProfilerManager();
+
+	/** Creates an instance of the profiler manager. */
+	static TSharedPtr<FNetworkingProfilerManager> Initialize()
+	{
+		if (FNetworkingProfilerManager::Instance.IsValid())
+		{
+			FNetworkingProfilerManager::Instance.Reset();
+		}
+
+		FNetworkingProfilerManager::Instance = MakeShareable(new FNetworkingProfilerManager(FInsightsManager::Get()->GetCommandList()));
+		FNetworkingProfilerManager::Instance->PostConstructor();
+
+		return FNetworkingProfilerManager::Instance;
+	}
+
+	/** Shutdowns the Timing Profiler manager. */
+	void Shutdown()
+	{
+		FNetworkingProfilerManager::Instance.Reset();
+	}
+
+protected:
+	/** Finishes initialization of the profiler manager. */
+	void PostConstructor();
+
+	/** Binds our UI commands to delegates. */
+	void BindCommands();
+
+public:
+	/**
+	 * @return the global instance of the Networking Profiler (Asset Networking Insights) manager.
+	 * This is an internal singleton and cannot be used outside ProfilerModule.
+	 * For external use:
+	 *     IProfilerModule& ProfilerModule = FModuleManager::Get().LoadModuleChecked<IProfilerModule>("Profiler");
+	 *     ProfilerModule.GetProfilerManager();
+	 */
+	static TSharedPtr<FNetworkingProfilerManager> Get();
+
+	/** @returns UI command list for the Timing Profiler manager. */
+	const TSharedRef<FUICommandList> GetCommandList() const;
+
+	/** @return an instance of the Timing Profiler commands. */
+	static const FNetworkingProfilerCommands& GetCommands();
+
+	/** @return an instance of the Timing Profiler action manager. */
+	static FNetworkingProfilerActionManager& GetActionManager();
+
+	void AssignProfilerWindow(const TSharedRef<SNetworkingProfilerWindow>& InProfilerWindow)
+	{
+		ProfilerWindow = InProfilerWindow;
+	}
+
+	/**
+	 * Converts profiler window weak pointer to a shared pointer and returns it.
+	 * Make sure the returned pointer is valid before trying to dereference it.
+	 */
+	TSharedPtr<class SNetworkingProfilerWindow> GetProfilerWindow() const
+	{
+		return ProfilerWindow.Pin();
+	}
+
+	////////////////////////////////////////////////////////////////////////////////////////////////////
+	// Getters and setters used by Toggle Commands.
+
+	/** @return true, if Packet Sizes view is visible */
+	const bool IsPacketSizesViewVisible() const { return bIsPacketSizesViewVisible; }
+	void SetPacketSizesViewVisible(const bool bPacketSizesViewVisibleState) { bIsPacketSizesViewVisible = bPacketSizesViewVisibleState; }
+	void ShowHidePacketSizesView(const bool bPacketSizesViewVisibleState);
+
+	/** @return true, if Packet Breakdown view is visible */
+	const bool IsPacketBreakdownViewVisible() const { return bIsPacketBreakdownViewVisible; }
+	void SetPacketBreakdownViewVisible(const bool bVisibleState) { bIsPacketBreakdownViewVisible = bVisibleState; }
+	void ShowHidePacketBreakdownView(const bool bVisibleState);
+
+	/** @return true, if Data Stream Breakdown view is visible */
+	const bool IsDataStreamBreakdownViewVisible() const { return bIsDataStreamBreakdownViewVisible; }
+	void SetDataStreamBreakdownViewVisible(const bool bVisibleState) { bIsDataStreamBreakdownViewVisible = bVisibleState; }
+	void ShowHideDataStreamBreakdownView(const bool bVisibleState);
+
+	////////////////////////////////////////////////////////////////////////////////////////////////////
+
+	void OnSessionChanged();
+
+protected:
+	/** Updates this manager, done through FCoreTicker. */
+	bool Tick(float DeltaTime);
+
+protected:
+	/** The delegate to be invoked when this manager ticks. */
+	FTickerDelegate OnTick;
+
+	/** Handle to the registered OnTick. */
+	FDelegateHandle OnTickHandle;
+
+	/** List of UI commands for this manager. This will be filled by this and corresponding classes. */
+	TSharedRef<FUICommandList> CommandList;
+
+	/** An instance of the Networking Profiler action manager. */
+	FNetworkingProfilerActionManager ActionManager;
+
+	/** A weak pointer to the Networking Profiler window. */
+	TWeakPtr<class SNetworkingProfilerWindow> ProfilerWindow;
+
+	/** If the Packet Sizes view is visible or hidden. */
+	bool bIsPacketSizesViewVisible;
+
+	/** If the Packet Breakdown view is visible or hidden. */
+	bool bIsPacketBreakdownViewVisible;
+
+	/** If the Data Stream Breakdown tree view is visible or hidden. */
+	bool bIsDataStreamBreakdownViewVisible;
+
+	/** A shared pointer to the global instance of the profiler manager. */
+	static TSharedPtr<FNetworkingProfilerManager> Instance;
+};
diff --git a/Engine/Source/Developer/TraceInsights/Private/Insights/NetworkingProfiler/Widgets/SDataStreamBreakdownView.cpp b/Engine/Source/Developer/TraceInsights/Private/Insights/NetworkingProfiler/Widgets/SDataStreamBreakdownView.cpp
new file mode 100644
--- /dev/null
+++ b/Engine/Source/Developer/TraceInsights/Private/Insights/NetworkingProfiler/Widgets/SDataStreamBreakdownView.cpp
@@ -0,0 +1,34 @@
+// Copyright 1998-2019 Epic Games, Inc. All Rights Reserved.
+
+#include "SDataStreamBreakdownView.h"
+
+// Insights
+
+////////////////////////////////////////////////////////////////////////////////////////////////////
+
+#define LOCTEXT_NAMESPACE "SDataStreamBreakdownView"
+
+////////////////////////////////////////////////////////////////////////////////////////////////////
+
+SDataStreamBreakdownView::SDataStreamBreakdownView()
+{
+}
+
+////////////////////////////////////////////////////////////////////////////////////////////////////
+
+SDataStreamBreakdownView::~SDataStreamBreakdownView()
+{
+}
+
+////////////////////////////////////////////////////////////////////////////////////////////////////
+
+void SDataStreamBreakdownView::Construct(const FArguments& InArgs)
+{
+	//ChildSlot
+	//[
+	//];
+}
+
+////////////////////////////////////////////////////////////////////////////////////////////////////
+
+#undef LOCTEXT_NAMESPACE
diff --git a/Engine/Source/Developer/TraceInsights/Private/Insights/NetworkingProfiler/Widgets/SDataStreamBreakdownView.h b/Engine/Source/Developer/TraceInsights/Private/Insights/NetworkingProfiler/Widgets/SDataStreamBreakdownView.h
new file mode 100644
--- /dev/null
+++ b/Engine/Source/Developer/TraceInsights/Private/Insights/NetworkingProfiler/Widgets/SDataStreamBreakdownView.h
@@ -0,0 +1,28 @@
+// Copyright 1998-2019 Epic Games, Inc. All Rights Reserved.
+
+#pragma once
+
+#include "CoreMinimal.h"
+#include "Styling/SlateTypes.h"
+#include "Widgets/DeclarativeSyntaxSupport.h"
+#include "Widgets/SCompoundWidget.h"
+
+class SDataStreamBreakdownView : public SCompoundWidget
+{
+public:
+	/** Default constructor. */
+	SDataStreamBreakdownView();
+
+	/** Virtual destructor. */
+	virtual ~SDataStreamBreakdownView();
+
+	SLATE_BEGIN_ARGS(SDataStreamBreakdownView) {}
+	SLATE_END_ARGS()
+
+	/**
+	 * Construct this widget
+	 *
+	 * @param	InArgs	The declaration data for this widget
+	 */
+	void Construct(const FArguments& InArgs);
+};
diff --git a/Engine/Source/Developer/TraceInsights/Private/Insights/NetworkingProfiler/Widgets/SNetworkingProfilerToolbar.cpp b/Engine/Source/Developer/TraceInsights/Private/Insights/NetworkingProfiler/Widgets/SNetworkingProfilerToolbar.cpp
new file mode 100644
--- /dev/null
+++ b/Engine/Source/Developer/TraceInsights/Private/Insights/NetworkingProfiler/Widgets/SNetworkingProfilerToolbar.cpp
@@ -0,0 +1,112 @@
+// Copyright 1998-2019 Epic Games, Inc. All Rights Reserved.
+
+#include "SNetworkingProfilerToolbar.h"
+
+#include "EditorStyleSet.h"
+#include "Framework/Application/SlateApplication.h"
+#include "Framework/Docking/TabManager.h"
+#include "Framework/MultiBox/MultiBoxDefs.h"
+#include "Framework/MultiBox/MultiBoxBuilder.h"
+#include "Widgets/Docking/SDockTab.h"
+#include "Widgets/Input/SCheckBox.h"
+#include "Widgets/Layout/SBorder.h"
+#include "Widgets/SBoxPanel.h"
+#include "Widgets/SToolTip.h"
+
+// Insights
+#include "Insights/InsightsCommands.h"
+#include "Insights/InsightsManager.h"
+#include "Insights/NetworkingProfiler/NetworkingProfilerCommands.h"
+#include "Insights/NetworkingProfiler/NetworkingProfilerManager.h"
+
+////////////////////////////////////////////////////////////////////////////////////////////////////
+
+#define LOCTEXT_NAMESPACE "SNetworkingProfilerToolbar"
+
+////////////////////////////////////////////////////////////////////////////////////////////////////
+
+SNetworkingProfilerToolbar::SNetworkingProfilerToolbar()
+{
+}
+
+////////////////////////////////////////////////////////////////////////////////////////////////////
+
+SNetworkingProfilerToolbar::~SNetworkingProfilerToolbar()
+{
+}
+
+////////////////////////////////////////////////////////////////////////////////////////////////////
+
+void SNetworkingProfilerToolbar::Construct(const FArguments& InArgs)
+{
+	struct Local
+	{
+		static void FillViewToolbar(FToolBarBuilder& ToolbarBuilder)
+		{
+			ToolbarBuilder.BeginSection("View");
+			{
+				ToolbarBuilder.AddToolBarButton(FNetworkingProfilerCommands::Get().TogglePacketSizesViewVisibility);
+				ToolbarBuilder.AddToolBarButton(FNetworkingProfilerCommands::Get().TogglePacketBreakdownViewVisibility);
+				ToolbarBuilder.AddToolBarButton(FNetworkingProfilerCommands::Get().ToggleDataStreamBreakdownViewVisibility);
+			}
+			ToolbarBuilder.EndSection();
+		}
+
+		static void FillRightSideToolbar(FToolBarBuilder& ToolbarBuilder)
+		{
+			ToolbarBuilder.BeginSection("Debug");
+			{
+				ToolbarBuilder.AddToolBarButton(FInsightsCommands::Get().ToggleDebugInfo);
+			}
+			ToolbarBuilder.EndSection();
+		}
+	};
+
+	TSharedPtr<FUICommandList> CommandList = FInsightsManager::Get()->GetCommandList();
+
+	FToolBarBuilder ToolbarBuilder(CommandList.ToSharedRef(), FMultiBoxCustomization::None);
+	Local::FillViewToolbar(ToolbarBuilder);
+
+	FToolBarBuilder RightSideToolbarBuilder(CommandList.ToSharedRef(), FMultiBoxCustomization::None);
+	Local::FillRightSideToolbar(RightSideToolbarBuilder);
+
+	// Create the tool bar!
+	ChildSlot
+	[
+		SNew(SHorizontalBox)
+
+		+SHorizontalBox::Slot()
+		.HAlign(HAlign_Left)
+		.VAlign(VAlign_Center)
+		.FillWidth(1.0)
+		.Padding(0.0f)
+		[
+			SNew(SBorder)
+			.Padding(0)
+			.BorderImage(FEditorStyle::GetBrush("NoBorder"))
+			.IsEnabled(FSlateApplication::Get().GetNormalExecutionAttribute())
+			[
+				ToolbarBuilder.MakeWidget()
+			]
+		]
+
+		+SHorizontalBox::Slot()
+		.HAlign(HAlign_Right)
+		.VAlign(VAlign_Center)
+		.FillWidth(1.0)
+		.Padding(0.0f)
+		[
+			SNew(SBorder)
+			.Padding(0)
+			.BorderImage(FEditorStyle::GetBrush("NoBorder"))
+			.IsEnabled(FSlateApplication::Get().GetNormalExecutionAttribute())
+			[
+				RightSideToolbarBuilder.MakeWidget()
+			]
+		]
+	];
+}
+
+////////////////////////////////////////////////////////////////////////////////////////////////////
+
+#undef LOCTEXT_NAMESPACE
diff --git a/Engine/Source/Developer/TraceInsights/Private/Insights/NetworkingProfiler/Widgets/SNetworkingProfilerToolbar.h b/Engine/Source/Developer/TraceInsights/Private/Insights/NetworkingProfiler/Widgets/SNetworkingProfilerToolbar.h
new file mode 100644
--- /dev/null
+++ b/Engine/Source/Developer/TraceInsights/Private/Insights/NetworkingProfiler/Widgets/SNetworkingProfilerToolbar.h
@@ -0,0 +1,28 @@
+// Copyright 1998-2019 Epic Games, Inc. All Rights Reserved.
+
+#pragma once
+
+#include "CoreMinimal.h"
+#include "Styling/SlateTypes.h"
+#include "Widgets/DeclarativeSyntaxSupport.h"
+#include "Widgets/SCompoundWidget.h"
+
+class SNetworkingProfilerToolbar : public SCompoundWidget
+{
+public:
+	/** Default constructor. */
+	SNetworkingProfilerToolbar();
+
+	/** Virtual destructor. */
+	virtual ~SNetworkingProfilerToolbar();
+
+	SLATE_BEGIN_ARGS(SNetworkingProfilerToolbar) {}
+	SLATE_END_ARGS()
+
+	/**
+	 * Construct this widget
+	 *
+	 * @param	InArgs	The declaration data for this widget
+	 */
+	void Construct(const FArguments& InArgs);
+};
diff --git a/Engine/Source/Developer/TraceInsights/Private/Insights/NetworkingProfiler/Widgets/SNetworkingProfilerWindow.cpp b/Engine/Source/Developer/TraceInsights/Private/Insights/NetworkingProfiler/Widgets/SNetworkingProfilerWindow.cpp
new file mode 100644
--- /dev/null
+++ b/Engine/Source/Developer/TraceInsights/Private/Insights/NetworkingProfiler/Widgets/SNetworkingProfilerWindow.cpp
@@ -0,0 +1,471 @@
+// Copyright 1998-2019 Epic Games, Inc. All Rights Reserved.
+
+#include "SNetworkingProfilerWindow.h"
+
+#include "EditorStyleSet.h"
+#include "Framework/Docking/WorkspaceItem.h"
+#include "SlateOptMacros.h"
+#include "Widgets/Docking/SDockTab.h"
+#include "Widgets/Images/SImage.h"
+#include "Widgets/Layout/SBorder.h"
+#include "Widgets/Layout/SBox.h"
+#include "Widgets/Layout/SSpacer.h"
+#include "Widgets/Notifications/SNotificationList.h"
+#include "Widgets/SBoxPanel.h"
+#include "Widgets/Text/STextBlock.h"
+
+#if WITH_EDITOR
+	#include "EngineAnalytics.h"
+	#include "Runtime/Analytics/Analytics/Public/AnalyticsEventAttribute.h"
+	#include "Runtime/Analytics/Analytics/Public/Interfaces/IAnalyticsProvider.h"
+#endif // WITH_EDITOR
+
+// Insights
+#include "Insights/InsightsManager.h"
+#include "Insights/InsightsStyle.h"
+#include "Insights/NetworkingProfiler/NetworkingProfilerManager.h"
+#include "Insights/NetworkingProfiler/Widgets/SDataStreamBreakdownView.h"
+#include "Insights/NetworkingProfiler/Widgets/SNetworkingProfilerToolbar.h"
+#include "Insights/NetworkingProfiler/Widgets/SPacketBreakdownView.h"
+#include "Insights/NetworkingProfiler/Widgets/SPacketSizesView.h"
+#include "Insights/Version.h"
+#include "Insights/Widgets/SInsightsSettings.h"
+
+////////////////////////////////////////////////////////////////////////////////////////////////////
+
+#define LOCTEXT_NAMESPACE "SNetworkingProfilerWindow"
+
+////////////////////////////////////////////////////////////////////////////////////////////////////
+
+const FName FNetworkingProfilerTabs::ToolbarID(TEXT("Toolbar"));
+const FName FNetworkingProfilerTabs::PacketSizesViewID(TEXT("PacketSizesView"));
+const FName FNetworkingProfilerTabs::PacketBreakdownViewID(TEXT("PacketBreakdownView"));
+const FName FNetworkingProfilerTabs::DataStreamBreakdownViewID(TEXT("DataStreamBreakdownView"));
+
+////////////////////////////////////////////////////////////////////////////////////////////////////
+
+SNetworkingProfilerWindow::SNetworkingProfilerWindow()
+	: DurationActive(0.0f)
+{
+}
+
+////////////////////////////////////////////////////////////////////////////////////////////////////
+
+SNetworkingProfilerWindow::~SNetworkingProfilerWindow()
+{
+#if WITH_EDITOR
+	if (DurationActive > 0.0f && FEngineAnalytics::IsAvailable())
+	{
+		FEngineAnalytics::GetProvider().RecordEvent(TEXT("Insights.Usage.NetworkingProfiler"), FAnalyticsEventAttribute(TEXT("Duration"), DurationActive));
+	}
+#endif // WITH_EDITOR
+}
+
+////////////////////////////////////////////////////////////////////////////////////////////////////
+
+void SNetworkingProfilerWindow::Reset()
+{
+	if (PacketSizesView)
+	{
+		//PacketSizesView->Reset();
+	}
+
+	if (PacketBreakdownView)
+	{
+		//PacketBreakdownView->Reset();
+	}
+
+	if (DataStreamBreakdownView)
+	{
+		//DataStreamBreakdownView->Reset();
+	}
+}
+
+////////////////////////////////////////////////////////////////////////////////////////////////////
+
+BEGIN_SLATE_FUNCTION_BUILD_OPTIMIZATION
+
+TSharedRef<SDockTab> SNetworkingProfilerWindow::SpawnTab_Toolbar(const FSpawnTabArgs& Args)
+{
+	const TSharedRef<SDockTab> DockTab = SNew(SDockTab)
+		.ShouldAutosize(true)
+		.TabRole(ETabRole::PanelTab)
+		[
+			SNew(SNetworkingProfilerToolbar)
+		];
+
+	DockTab->SetOnTabClosed(SDockTab::FOnTabClosedCallback::CreateRaw(this, &SNetworkingProfilerWindow::OnToolbarTabClosed));
+
+	return DockTab;
+}
+
+////////////////////////////////////////////////////////////////////////////////////////////////////
+
+void SNetworkingProfilerWindow::OnToolbarTabClosed(TSharedRef<SDockTab> TabBeingClosed)
+{
+}
+
+////////////////////////////////////////////////////////////////////////////////////////////////////
+
+TSharedRef<SDockTab> SNetworkingProfilerWindow::SpawnTab_PacketSizesView(const FSpawnTabArgs& Args)
+{
+	FNetworkingProfilerManager::Get()->SetPacketSizesViewVisible(true);
+
+	const TSharedRef<SDockTab> DockTab = SNew(SDockTab)
+		.ShouldAutosize(false)
+		.TabRole(ETabRole::PanelTab)
+		[
+			SAssignNew(PacketSizesView, SPacketSizesView)
+		];
+
+	DockTab->SetOnTabClosed(SDockTab::FOnTabClosedCallback::CreateRaw(this, &SNetworkingProfilerWindow::OnPacketSizesViewTabClosed));
+
+	return DockTab;
+}
+
+////////////////////////////////////////////////////////////////////////////////////////////////////
+
+void SNetworkingProfilerWindow::OnPacketSizesViewTabClosed(TSharedRef<SDockTab> TabBeingClosed)
+{
+	FNetworkingProfilerManager::Get()->SetPacketSizesViewVisible(false);
+	PacketSizesView = nullptr;
+}
+
+////////////////////////////////////////////////////////////////////////////////////////////////////
+
+TSharedRef<SDockTab> SNetworkingProfilerWindow::SpawnTab_PacketBreakdownView(const FSpawnTabArgs& Args)
+{
+	FNetworkingProfilerManager::Get()->SetPacketBreakdownViewVisible(true);
+
+	const TSharedRef<SDockTab> DockTab = SNew(SDockTab)
+		.ShouldAutosize(false)
+		.TabRole(ETabRole::PanelTab)
+		[
+			SAssignNew(PacketBreakdownView, SPacketBreakdownView)
+		];
+
+	DockTab->SetOnTabClosed(SDockTab::FOnTabClosedCallback::CreateRaw(this, &SNetworkingProfilerWindow::OnPacketBreakdownViewTabClosed));
+
+	return DockTab;
+}
+
+////////////////////////////////////////////////////////////////////////////////////////////////////
+
+void SNetworkingProfilerWindow::OnPacketBreakdownViewTabClosed(TSharedRef<SDockTab> TabBeingClosed)
+{
+	FNetworkingProfilerManager::Get()->SetPacketBreakdownViewVisible(false);
+	PacketBreakdownView = nullptr;
+}
+
+////////////////////////////////////////////////////////////////////////////////////////////////////
+
+TSharedRef<SDockTab> SNetworkingProfilerWindow::SpawnTab_DataStreamBreakdownView(const FSpawnTabArgs& Args)
+{
+	FNetworkingProfilerManager::Get()->SetDataStreamBreakdownViewVisible(true);
+
+	const TSharedRef<SDockTab> DockTab = SNew(SDockTab)
+		.ShouldAutosize(false)
+		.TabRole(ETabRole::PanelTab)
+		[
+			SAssignNew(DataStreamBreakdownView, SDataStreamBreakdownView)
+		];
+
+	DockTab->SetOnTabClosed(SDockTab::FOnTabClosedCallback::CreateRaw(this, &SNetworkingProfilerWindow::OnDataStreamBreakdownViewTabClosed));
+
+	return DockTab;
+}
+
+////////////////////////////////////////////////////////////////////////////////////////////////////
+
+void SNetworkingProfilerWindow::OnDataStreamBreakdownViewTabClosed(TSharedRef<SDockTab> TabBeingClosed)
+{
+	FNetworkingProfilerManager::Get()->SetDataStreamBreakdownViewVisible(false);
+	DataStreamBreakdownView = nullptr;
+}
+
+////////////////////////////////////////////////////////////////////////////////////////////////////
+
+void SNetworkingProfilerWindow::Construct(const FArguments& InArgs, const TSharedRef<SDockTab>& ConstructUnderMajorTab, const TSharedPtr<SWindow>& ConstructUnderWindow)
+{
+	// Create & initialize tab manager.
+	TabManager = FGlobalTabmanager::Get()->NewTabManager(ConstructUnderMajorTab);
+	TSharedRef<FWorkspaceItem> AppMenuGroup = TabManager->AddLocalWorkspaceMenuCategory(LOCTEXT("NetworkingProfilerMenuGroupName", "Timing Insights"));
+
+	TabManager->RegisterTabSpawner(FNetworkingProfilerTabs::ToolbarID, FOnSpawnTab::CreateRaw(this, &SNetworkingProfilerWindow::SpawnTab_Toolbar))
+		.SetDisplayName(LOCTEXT("DeviceToolbarTabTitle", "Toolbar"))
+		.SetIcon(FSlateIcon(FInsightsStyle::GetStyleSetName(), "Toolbar.Icon.Small"))
+		.SetGroup(AppMenuGroup);
+
+	TabManager->RegisterTabSpawner(FNetworkingProfilerTabs::PacketSizesViewID, FOnSpawnTab::CreateRaw(this, &SNetworkingProfilerWindow::SpawnTab_PacketSizesView))
+		.SetDisplayName(LOCTEXT("NetworkingProfiler.PacketSizesViewTabTitle", "Packet Sizes"))
+		.SetIcon(FSlateIcon(FInsightsStyle::GetStyleSetName(), "PacketSizesView.Icon.Small"))
+		.SetGroup(AppMenuGroup);
+
+	TabManager->RegisterTabSpawner(FNetworkingProfilerTabs::PacketBreakdownViewID, FOnSpawnTab::CreateRaw(this, &SNetworkingProfilerWindow::SpawnTab_PacketBreakdownView))
+		.SetDisplayName(LOCTEXT("NetworkingProfiler.PacketBreakdownViewTabTitle", "Packet Breakdown"))
+		.SetIcon(FSlateIcon(FInsightsStyle::GetStyleSetName(), "PacketBreakdownView.Icon.Small"))
+		.SetGroup(AppMenuGroup);
+
+	TabManager->RegisterTabSpawner(FNetworkingProfilerTabs::DataStreamBreakdownViewID, FOnSpawnTab::CreateRaw(this, &SNetworkingProfilerWindow::SpawnTab_DataStreamBreakdownView))
+		.SetDisplayName(LOCTEXT("NetworkingProfiler.DataStreamBreakdownViewTabTitle", "Data Stream Breakdown"))
+		.SetIcon(FSlateIcon(FInsightsStyle::GetStyleSetName(), "DataStreamBreakdownView.Icon.Small"))
+		.SetGroup(AppMenuGroup);
+
+	TSharedPtr<FNetworkingProfilerManager> NetworkingProfilerManager = FNetworkingProfilerManager::Get();
+	ensure(NetworkingProfilerManager.IsValid());
+
+	// Create tab layout.
+	const TSharedRef<FTabManager::FLayout> Layout = FTabManager::NewLayout("InsightsNetworkingProfilerLayout_v1.0")
+		->AddArea
+		(
+			FTabManager::NewPrimaryArea()
+			->SetOrientation(Orient_Vertical)
+			->Split
+			(
+				FTabManager::NewStack()
+				->AddTab(FNetworkingProfilerTabs::ToolbarID, ETabState::OpenedTab)
+				->SetHideTabWell(true)
+			)
+			->Split
+			(
+				FTabManager::NewStack()
+				->SetSizeCoefficient(0.33f)
+				->SetHideTabWell(true)
+				->AddTab(FNetworkingProfilerTabs::PacketSizesViewID, NetworkingProfilerManager->IsPacketSizesViewVisible() ? ETabState::OpenedTab : ETabState::ClosedTab)
+			)
+			->Split
+			(
+				FTabManager::NewStack()
+				->SetSizeCoefficient(0.33f)
+				->AddTab(FNetworkingProfilerTabs::PacketBreakdownViewID, NetworkingProfilerManager->IsPacketBreakdownViewVisible() ? ETabState::OpenedTab : ETabState::ClosedTab)
+			)
+			->Split
+			(
+				FTabManager::NewStack()
+				->SetSizeCoefficient(0.33f)
+				->AddTab(FNetworkingProfilerTabs::DataStreamBreakdownViewID, NetworkingProfilerManager->IsDataStreamBreakdownViewVisible() ? ETabState::OpenedTab : ETabState::ClosedTab)
+			)
+		);
+
+	// Create & initialize main menu.
+	FMenuBarBuilder MenuBarBuilder = FMenuBarBuilder(TSharedPtr<FUICommandList>());
+
+	MenuBarBuilder.AddPullDownMenu(
+		LOCTEXT("MenuLabel", "MENU"),
+		FText::GetEmpty(),
+		FNewMenuDelegate::CreateStatic(&SNetworkingProfilerWindow::FillMenu, TabManager),
+		FName(TEXT("MENU"))
+	);
+
+	ChildSlot
+		[
+			SNew(SOverlay)
+
+			// Version
+			+ SOverlay::Slot()
+				.HAlign(HAlign_Right)
+				.VAlign(VAlign_Top)
+				.Padding(0.0f, -16.0f, 0.0f, 0.0f)
+				[
+					SNew(STextBlock)
+						.Clipping(EWidgetClipping::ClipToBoundsWithoutIntersecting)
+						.Text(LOCTEXT("UnrealInsightsVersion", UNREAL_INSIGHTS_VERSION_STRING_EX))
+						.ColorAndOpacity(FLinearColor(0.15f, 0.15f, 0.15f, 1.0f))
+				]
+
+			// Overlay slot for the main window area
+			+ SOverlay::Slot()
+				.HAlign(HAlign_Fill)
+				.VAlign(VAlign_Fill)
+				[
+					SNew(SVerticalBox)
+
+					+ SVerticalBox::Slot()
+						.AutoHeight()
+						[
+							MenuBarBuilder.MakeWidget()
+						]
+
+					+ SVerticalBox::Slot()
+						.FillHeight(1.0f)
+						[
+							TabManager->RestoreFrom(Layout, ConstructUnderWindow).ToSharedRef()
+						]
+				]
+
+			// Session hint overlay
+			+ SOverlay::Slot()
+				.HAlign(HAlign_Center)
+				.VAlign(VAlign_Center)
+				[
+					SNew(SBorder)
+						.Visibility(this, &SNetworkingProfilerWindow::IsSessionOverlayVisible)
+						.BorderImage(FEditorStyle::GetBrush("NotificationList.ItemBackground"))
+						.Padding(8.0f)
+						[
+							SNew(STextBlock)
+								.Text(LOCTEXT("SelectTraceOverlayText", "Please select a trace."))
+						]
+				]
+		];
+
+	// Tell tab-manager about the global menu bar.
+	TabManager->SetMenuMultiBox(MenuBarBuilder.GetMultiBox());
+}
+
+END_SLATE_FUNCTION_BUILD_OPTIMIZATION
+
+////////////////////////////////////////////////////////////////////////////////////////////////////
+
+void SNetworkingProfilerWindow::FillMenu(FMenuBuilder& MenuBuilder, const TSharedPtr<FTabManager> TabManager)
+{
+	if (!TabManager.IsValid())
+	{
+		return;
+	}
+
+#if !WITH_EDITOR
+	//TODO: FGlobalTabmanager::Get()->PopulateTabSpawnerMenu(MenuBuilder, WorkspaceMenu::GetMenuStructure().GetStructureRoot());
+#endif //!WITH_EDITOR
+
+	TabManager->PopulateLocalTabSpawnerMenu(MenuBuilder);
+}
+
+////////////////////////////////////////////////////////////////////////////////////////////////////
+
+void SNetworkingProfilerWindow::ShowTab(const FName& TabID)
+{
+	if (TabManager->HasTabSpawner(TabID))
+	{
+		TabManager->InvokeTab(TabID);
+	}
+}
+
+////////////////////////////////////////////////////////////////////////////////////////////////////
+
+void SNetworkingProfilerWindow::HideTab(const FName& TabID)
+{
+	TSharedPtr<SDockTab> Tab = TabManager->FindExistingLiveTab(TabID);
+	if (Tab.IsValid())
+	{
+		Tab->RequestCloseTab();
+	}
+}
+
+////////////////////////////////////////////////////////////////////////////////////////////////////
+
+EVisibility SNetworkingProfilerWindow::IsSessionOverlayVisible() const
+{
+	if (FInsightsManager::Get()->GetSession().IsValid())
+	{
+		return EVisibility::Hidden;
+	}
+	else
+	{
+		return EVisibility::Visible;
+	}
+}
+
+////////////////////////////////////////////////////////////////////////////////////////////////////
+
+bool SNetworkingProfilerWindow::IsProfilerEnabled() const
+{
+	return FInsightsManager::Get()->GetSession().IsValid();
+}
+
+////////////////////////////////////////////////////////////////////////////////////////////////////
+
+EActiveTimerReturnType SNetworkingProfilerWindow::UpdateActiveDuration(double InCurrentTime, float InDeltaTime)
+{
+	DurationActive += InDeltaTime;
+
+	// The profiler window will explicitly unregister this active timer when the mouse leaves.
+	return EActiveTimerReturnType::Continue;
+}
+
+////////////////////////////////////////////////////////////////////////////////////////////////////
+
+void SNetworkingProfilerWindow::OnMouseEnter(const FGeometry& MyGeometry, const FPointerEvent& MouseEvent)
+{
+	SCompoundWidget::OnMouseEnter(MyGeometry, MouseEvent);
+
+	if (!ActiveTimerHandle.IsValid())
+	{
+		ActiveTimerHandle = RegisterActiveTimer(0.f, FWidgetActiveTimerDelegate::CreateSP(this, &SNetworkingProfilerWindow::UpdateActiveDuration));
+	}
+}
+
+////////////////////////////////////////////////////////////////////////////////////////////////////
+
+void SNetworkingProfilerWindow::OnMouseLeave(const FPointerEvent& MouseEvent)
+{
+	SCompoundWidget::OnMouseLeave(MouseEvent);
+
+	auto PinnedActiveTimerHandle = ActiveTimerHandle.Pin();
+	if (PinnedActiveTimerHandle.IsValid())
+	{
+		UnRegisterActiveTimer(PinnedActiveTimerHandle.ToSharedRef());
+	}
+}
+
+////////////////////////////////////////////////////////////////////////////////////////////////////
+
+FReply SNetworkingProfilerWindow::OnKeyDown(const FGeometry& MyGeometry, const FKeyEvent& InKeyEvent)
+{
+	return FNetworkingProfilerManager::Get()->GetCommandList()->ProcessCommandBindings(InKeyEvent) ? FReply::Handled() : FReply::Unhandled();
+}
+
+////////////////////////////////////////////////////////////////////////////////////////////////////
+
+FReply SNetworkingProfilerWindow::OnDragOver(const FGeometry& MyGeometry, const FDragDropEvent& DragDropEvent)
+{
+	TSharedPtr<FExternalDragOperation> DragDropOp = DragDropEvent.GetOperationAs<FExternalDragOperation>();
+	if (DragDropOp.IsValid())
+	{
+		if (DragDropOp->HasFiles())
+		{
+			const TArray<FString>& Files = DragDropOp->GetFiles();
+			if (Files.Num() == 1)
+			{
+				const FString DraggedFileExtension = FPaths::GetExtension(Files[0], true);
+				if (DraggedFileExtension == TEXT(".utrace"))
+				{
+					return FReply::Handled();
+				}
+			}
+		}
+	}
+
+	return SCompoundWidget::OnDragOver(MyGeometry,DragDropEvent);
+}
+
+////////////////////////////////////////////////////////////////////////////////////////////////////
+
+FReply SNetworkingProfilerWindow::OnDrop(const FGeometry& MyGeometry, const FDragDropEvent& DragDropEvent)
+{
+	TSharedPtr<FExternalDragOperation> DragDropOp = DragDropEvent.GetOperationAs<FExternalDragOperation>();
+	if (DragDropOp.IsValid())
+	{
+		if (DragDropOp->HasFiles())
+		{
+			// For now, only allow a single file.
+			const TArray<FString>& Files = DragDropOp->GetFiles();
+			if (Files.Num() == 1)
+			{
+				const FString DraggedFileExtension = FPaths::GetExtension(Files[0], true);
+				if (DraggedFileExtension == TEXT(".utrace"))
+				{
+					// Enqueue load operation.
+					FInsightsManager::Get()->LoadTraceFile(Files[0]);
+					return FReply::Handled();
+				}
+			}
+		}
+	}
+
+	return SCompoundWidget::OnDrop(MyGeometry,DragDropEvent);
+}
+
+////////////////////////////////////////////////////////////////////////////////////////////////////
+
+#undef LOCTEXT_NAMESPACE
diff --git a/Engine/Source/Developer/TraceInsights/Private/Insights/NetworkingProfiler/Widgets/SNetworkingProfilerWindow.h b/Engine/Source/Developer/TraceInsights/Private/Insights/NetworkingProfiler/Widgets/SNetworkingProfilerWindow.h
new file mode 100644
--- /dev/null
+++ b/Engine/Source/Developer/TraceInsights/Private/Insights/NetworkingProfiler/Widgets/SNetworkingProfilerWindow.h
@@ -0,0 +1,168 @@
+// Copyright 1998-2019 Epic Games, Inc. All Rights Reserved.
+
+#pragma once
+
+#include "CoreMinimal.h"
+#include "Framework/Docking/TabManager.h"
+#include "Input/Reply.h"
+#include "Layout/Visibility.h"
+#include "Misc/Guid.h"
+#include "SlateFwd.h"
+#include "Widgets/DeclarativeSyntaxSupport.h"
+#include "Widgets/Docking/SDockTab.h"
+#include "Widgets/SCompoundWidget.h"
+
+// Insights
+#include "Insights/InsightsManager.h"
+
+////////////////////////////////////////////////////////////////////////////////////////////////////
+
+class FActiveTimerHandle;
+class SDataStreamBreakdownView;
+class SPacketBreakdownView;
+class SPacketSizesView;
+
+////////////////////////////////////////////////////////////////////////////////////////////////////
+
+struct FNetworkingProfilerTabs
+{
+	// Tab identifiers
+	static const FName ToolbarID;
+	static const FName PacketSizesViewID;
+	static const FName PacketBreakdownViewID;
+	static const FName DataStreamBreakdownViewID;
+};
+
+////////////////////////////////////////////////////////////////////////////////////////////////////
+
+/** Implements the timing profiler window. */
+class SNetworkingProfilerWindow : public SCompoundWidget
+{
+public:
+	/** Default constructor. */
+	SNetworkingProfilerWindow();
+
+	/** Virtual destructor. */
+	virtual ~SNetworkingProfilerWindow();
+
+	SLATE_BEGIN_ARGS(SNetworkingProfilerWindow) {}
+	SLATE_END_ARGS()
+
+	void Reset();
+
+	/** Constructs this widget. */
+	void Construct(const FArguments& InArgs, const TSharedRef<SDockTab>& ConstructUnderMajorTab, const TSharedPtr<SWindow>& ConstructUnderWindow);
+
+	void ShowTab(const FName& TabID);
+	void HideTab(const FName& TabID);
+	void ShowHideTab(const FName& TabID, bool bShow) { bShow ? ShowTab(TabID) : HideTab(TabID); }
+
+	TSharedPtr<FTabManager> GetTabManager() const { return TabManager; }
+
+	TSharedPtr<SPacketSizesView> GetPacketSizesView() const { return PacketSizesView; }
+	TSharedPtr<SPacketBreakdownView> GetPacketBreakdownView() const { return PacketBreakdownView; }
+	TSharedPtr<SDataStreamBreakdownView> GetDataStreamBreakdownView() const { return DataStreamBreakdownView; }
+
+private:
+	TSharedRef<SDockTab> SpawnTab_Toolbar(const FSpawnTabArgs& Args);
+	void OnToolbarTabClosed(TSharedRef<SDockTab> TabBeingClosed);
+
+	TSharedRef<SDockTab> SpawnTab_PacketSizesView(const FSpawnTabArgs& Args);
+	void OnPacketSizesViewTabClosed(TSharedRef<SDockTab> TabBeingClosed);
+
+	TSharedRef<SDockTab> SpawnTab_PacketBreakdownView(const FSpawnTabArgs& Args);
+	void OnPacketBreakdownViewTabClosed(TSharedRef<SDockTab> TabBeingClosed);
+
+	TSharedRef<SDockTab> SpawnTab_DataStreamBreakdownView(const FSpawnTabArgs& Args);
+	void OnDataStreamBreakdownViewTabClosed(TSharedRef<SDockTab> TabBeingClosed);
+
+	/**
+	 * Fill the main menu with menu items.
+	 *
+	 * @param MenuBuilder The multi-box builder that should be filled with content for this pull-down menu.
+	 * @param TabManager A Tab Manager from which to populate tab spawner menu items.
+	 */
+	static void FillMenu(FMenuBuilder& MenuBuilder, const TSharedPtr<FTabManager> TabManager);
+
+	/** Callback for determining the visibility of the 'Select a session' overlay. */
+	EVisibility IsSessionOverlayVisible() const;
+
+	/** Callback for getting the enabled state of the profiler window. */
+	bool IsProfilerEnabled() const;
+
+	/** Updates the amount of time the profiler has been active. */
+	EActiveTimerReturnType UpdateActiveDuration(double InCurrentTime, float InDeltaTime);
+
+	/**
+	 * Ticks this widget. Override in derived classes, but always call the parent implementation.
+	 *
+	 * @param  AllottedGeometry The space allotted for this widget
+	 * @param  InCurrentTime  Current absolute real time
+	 * @param  InDeltaTime  Real time passed since last tick
+	 */
+	//virtual void Tick(const FGeometry& AllottedGeometry, const double InCurrentTime, const float InDeltaTime) override;
+
+	/**
+	 * The system will use this event to notify a widget that the cursor has entered it. This event is NOT bubbled.
+	 *
+	 * @param MyGeometry The Geometry of the widget receiving the event
+	 * @param MouseEvent Information about the input event
+	 */
+	virtual void OnMouseEnter(const FGeometry& MyGeometry, const FPointerEvent& MouseEvent) override;
+
+	/**
+	 * The system will use this event to notify a widget that the cursor has left it. This event is NOT bubbled.
+	 *
+	 * @param MouseEvent Information about the input event
+	 */
+	virtual void OnMouseLeave(const FPointerEvent& MouseEvent) override;
+
+	/**
+	 * Called after a key is pressed when this widget has focus
+	 *
+	 * @param MyGeometry The Geometry of the widget receiving the event
+	 * @param  InKeyEvent  Key event
+	 *
+	 * @return  Returns whether the event was handled, along with other possible actions
+	 */
+	virtual FReply OnKeyDown(const FGeometry& MyGeometry, const FKeyEvent& InKeyEvent) override;
+
+	/**
+	 * Called when the user is dropping something onto a widget; terminates drag and drop.
+	 *
+	 * @param MyGeometry      The geometry of the widget receiving the event.
+	 * @param DragDropEvent   The drag and drop event.
+	 *
+	 * @return A reply that indicated whether this event was handled.
+	 */
+	virtual FReply OnDrop(const FGeometry& MyGeometry, const FDragDropEvent& DragDropEvent) override;
+
+	/**
+	 * Called during drag and drop when the the mouse is being dragged over a widget.
+	 *
+	 * @param MyGeometry      The geometry of the widget receiving the event.
+	 * @param DragDropEvent   The drag and drop event.
+	 *
+	 * @return A reply that indicated whether this event was handled.
+	 */
+	virtual FReply OnDragOver(const FGeometry& MyGeometry, const FDragDropEvent& DragDropEvent)  override;
+
+private:
+	/** The Packet Sizes widget */
+	TSharedPtr<SPacketSizesView> PacketSizesView;
+
+	/** The Packet Breakdown widget */
+	TSharedPtr<SPacketBreakdownView> PacketBreakdownView;
+
+	/** The Data Stream Breakdown widget */
+	TSharedPtr<SDataStreamBreakdownView> DataStreamBreakdownView;
+
+	/** Holds the tab manager that manages the front-end's tabs. */
+	TSharedPtr<FTabManager> TabManager;
+
+	/** The handle to the active update duration tick */
+	TWeakPtr<FActiveTimerHandle> ActiveTimerHandle;
+
+	/** The number of seconds the profiler has been active */
+	float DurationActive;
+};
diff --git a/Engine/Source/Developer/TraceInsights/Private/Insights/NetworkingProfiler/Widgets/SPacketBreakdownView.cpp b/Engine/Source/Developer/TraceInsights/Private/Insights/NetworkingProfiler/Widgets/SPacketBreakdownView.cpp
new file mode 100644
--- /dev/null
+++ b/Engine/Source/Developer/TraceInsights/Private/Insights/NetworkingProfiler/Widgets/SPacketBreakdownView.cpp
@@ -0,0 +1,34 @@
+// Copyright 1998-2019 Epic Games, Inc. All Rights Reserved.
+
+#include "SPacketBreakdownView.h"
+
+// Insights
+
+////////////////////////////////////////////////////////////////////////////////////////////////////
+
+#define LOCTEXT_NAMESPACE "SPacketBreakdownView"
+
+////////////////////////////////////////////////////////////////////////////////////////////////////
+
+SPacketBreakdownView::SPacketBreakdownView()
+{
+}
+
+////////////////////////////////////////////////////////////////////////////////////////////////////
+
+SPacketBreakdownView::~SPacketBreakdownView()
+{
+}
+
+////////////////////////////////////////////////////////////////////////////////////////////////////
+
+void SPacketBreakdownView::Construct(const FArguments& InArgs)
+{
+	//ChildSlot
+	//[
+	//];
+}
+
+////////////////////////////////////////////////////////////////////////////////////////////////////
+
+#undef LOCTEXT_NAMESPACE
diff --git a/Engine/Source/Developer/TraceInsights/Private/Insights/NetworkingProfiler/Widgets/SPacketBreakdownView.h b/Engine/Source/Developer/TraceInsights/Private/Insights/NetworkingProfiler/Widgets/SPacketBreakdownView.h
new file mode 100644
--- /dev/null
+++ b/Engine/Source/Developer/TraceInsights/Private/Insights/NetworkingProfiler/Widgets/SPacketBreakdownView.h
@@ -0,0 +1,28 @@
+// Copyright 1998-2019 Epic Games, Inc. All Rights Reserved.
+
+#pragma once
+
+#include "CoreMinimal.h"
+#include "Styling/SlateTypes.h"
+#include "Widgets/DeclarativeSyntaxSupport.h"
+#include "Widgets/SCompoundWidget.h"
+
+class SPacketBreakdownView : public SCompoundWidget
+{
+public:
+	/** Default constructor. */
+	SPacketBreakdownView();
+
+	/** Virtual destructor. */
+	virtual ~SPacketBreakdownView();
+
+	SLATE_BEGIN_ARGS(SPacketBreakdownView) {}
+	SLATE_END_ARGS()
+
+	/**
+	 * Construct this widget
+	 *
+	 * @param	InArgs	The declaration data for this widget
+	 */
+	void Construct(const FArguments& InArgs);
+};
diff --git a/Engine/Source/Developer/TraceInsights/Private/Insights/NetworkingProfiler/Widgets/SPacketSizesView.cpp b/Engine/Source/Developer/TraceInsights/Private/Insights/NetworkingProfiler/Widgets/SPacketSizesView.cpp
new file mode 100644
--- /dev/null
+++ b/Engine/Source/Developer/TraceInsights/Private/Insights/NetworkingProfiler/Widgets/SPacketSizesView.cpp
@@ -0,0 +1,34 @@
+// Copyright 1998-2019 Epic Games, Inc. All Rights Reserved.
+
+#include "SPacketSizesView.h"
+
+// Insights
+
+////////////////////////////////////////////////////////////////////////////////////////////////////
+
+#define LOCTEXT_NAMESPACE "SPacketSizesView"
+
+////////////////////////////////////////////////////////////////////////////////////////////////////
+
+SPacketSizesView::SPacketSizesView()
+{
+}
+
+////////////////////////////////////////////////////////////////////////////////////////////////////
+
+SPacketSizesView::~SPacketSizesView()
+{
+}
+
+////////////////////////////////////////////////////////////////////////////////////////////////////
+
+void SPacketSizesView::Construct(const FArguments& InArgs)
+{
+	//ChildSlot
+	//[
+	//];
+}
+
+////////////////////////////////////////////////////////////////////////////////////////////////////
+
+#undef LOCTEXT_NAMESPACE
diff --git a/Engine/Source/Developer/TraceInsights/Private/Insights/NetworkingProfiler/Widgets/SPacketSizesView.h b/Engine/Source/Developer/TraceInsights/Private/Insights/NetworkingProfiler/Widgets/SPacketSizesView.h
new file mode 100644
--- /dev/null
+++ b/Engine/Source/Developer/TraceInsights/Private/Insights/NetworkingProfiler/Widgets/SPacketSizesView.h
@@ -0,0 +1,28 @@
+// Copyright 1998-2019 Epic Games, Inc. All Rights Reserved.
+
+#pragma once
+
+#include "CoreMinimal.h"
+#include "Styling/SlateTypes.h"
+#include "Widgets/DeclarativeSyntaxSupport.h"
+#include "Widgets/SCompoundWidget.h"
+
+class SPacketSizesView : public SCompoundWidget
+{
+public:
+	/** Default constructor. */
+	SPacketSizesView();
+
+	/** Virtual destructor. */
+	virtual ~SPacketSizesView();
+
+	SLATE_BEGIN_ARGS(SPacketSizesView) {}
+	SLATE_END_ARGS()
+
+	/**
+	 * Construct this widget
+	 *
+	 * @param	InArgs	The declaration data for this widget
+	 */
+	void Construct(const FArguments& InArgs);
+};
diff --git a/Engine/Source/Developer/TraceInsights/Private/Insights/TraceInsightsModule.cpp b/Engine/Source/Developer/TraceInsights/Private/Insights/TraceInsightsModule.cpp
--- a/Engine/Source/Developer/TraceInsights/Private/Insights/TraceInsightsModule.cpp
+++ b/Engine/Source/Developer/TraceInsights/Private/Insights/TraceInsightsModule.cpp
@@ -16,6 +16,8 @@
 #include "Insights/IUnrealInsightsModule.h"
 #include "Insights/LoadingProfiler/LoadingProfilerManager.h"
 #include "Insights/LoadingProfiler/Widgets/SLoadingProfilerWindow.h"
+#include "Insights/NetworkingProfiler/NetworkingProfilerManager.h"
+#include "Insights/NetworkingProfiler/Widgets/SNetworkingProfilerWindow.h"
 #include "Insights/TimingProfilerManager.h"
 #include "Insights/Widgets/SStartPageWindow.h"
 #include "Insights/Widgets/STimingProfilerWindow.h"
@@ -48,6 +50,9 @@ protected:
 	/** Callback called when the Loading Profiler major tab is closed. */
 	void OnLoadingProfilerTabBeingClosed(TSharedRef<SDockTab> TabBeingClosed);
 
+	/** Callback called when the Networking Profiler major tab is closed. */
+	void OnNetworkingProfilerTabBeingClosed(TSharedRef<SDockTab> TabBeingClosed);
+
 	/** Start Page */
 	TSharedRef<SDockTab> SpawnStartPageTab(const FSpawnTabArgs& Args);
 
@@ -57,6 +62,9 @@ protected:
 	/** Loading Profiler */
 	TSharedRef<SDockTab> SpawnLoadingProfilerTab(const FSpawnTabArgs& Args);
 
+	/** Networking Profiler */
+	TSharedRef<SDockTab> SpawnNetworkingProfilerTab(const FSpawnTabArgs& Args);
+
 protected:
 	TSharedPtr<Trace::IAnalysisService> TraceAnalysisService;
 	TSharedPtr<Trace::ISessionService> TraceSessionService;
@@ -86,6 +94,7 @@ void FTraceInsightsModule::StartupModule()
 	FInsightsManager::Initialize(TraceAnalysisService.ToSharedRef(), TraceSessionService.ToSharedRef(), TraceModuleService.ToSharedRef());
 	FTimingProfilerManager::Initialize();
 	FLoadingProfilerManager::Initialize();
+	FNetworkingProfilerManager::Initialize();
 
 	// Register tab spawner for the Start Page.
 	auto& StartPageTabSpawnerEntry = FGlobalTabmanager::Get()->RegisterTabSpawner(FInsightsManagerTabs::StartPageTabId,
@@ -118,6 +127,19 @@ void FTraceInsightsModule::StartupModule()
 //	LoadingProfilerTabSpawnerEntry.SetGroup(WorkspaceMenu::GetMenuStructure().GetDeveloperToolsMiscCategory());
 //#else
 //	LoadingProfilerTabSpawnerEntry.SetGroup(WorkspaceMenu::GetMenuStructure().GetToolsCategory());
+//#endif
+
+	// Register tab spawner for the Networking Insights.
+	auto& NetworkingProfilerTabSpawnerEntry = FGlobalTabmanager::Get()->RegisterTabSpawner(FInsightsManagerTabs::NetworkingProfilerTabId,
+		FOnSpawnTab::CreateRaw(this, &FTraceInsightsModule::SpawnNetworkingProfilerTab))
+		.SetDisplayName(NSLOCTEXT("FTraceInsightsModule", "NetworkingProfilerTabTitle", "Networking Insights"))
+		.SetTooltipText(NSLOCTEXT("FTraceInsightsModule", "NetworkingProfilerTooltipText", "Open the Networking Insights tab."))
+		.SetIcon(FSlateIcon(FInsightsStyle::GetStyleSetName(), "NetworkingProfiler.Icon.Small"));
+
+//#if WITH_EDITOR
+//	NetworkingProfilerTabSpawnerEntry.SetGroup(WorkspaceMenu::GetMenuStructure().GetDeveloperToolsMiscCategory());
+//#else
+//	NetworkingProfilerTabSpawnerEntry.SetGroup(WorkspaceMenu::GetMenuStructure().GetToolsCategory());
 //#endif
 }
 
@@ -127,14 +149,15 @@ void FTraceInsightsModule::ShutdownModule()
 {
 	TraceSessionService->StopRecorderServer();
 
-	FGlobalTabmanager::Get()->UnregisterNomadTabSpawner(FInsightsManagerTabs::StartPageTabId);
-	FGlobalTabmanager::Get()->UnregisterNomadTabSpawner(FInsightsManagerTabs::TimingProfilerTabId);
+	FGlobalTabmanager::Get()->UnregisterNomadTabSpawner(FInsightsManagerTabs::NetworkingProfilerTabId);
 	FGlobalTabmanager::Get()->UnregisterNomadTabSpawner(FInsightsManagerTabs::LoadingProfilerTabId);
+	FGlobalTabmanager::Get()->UnregisterNomadTabSpawner(FInsightsManagerTabs::TimingProfilerTabId);
+	FGlobalTabmanager::Get()->UnregisterNomadTabSpawner(FInsightsManagerTabs::StartPageTabId);
 
-	if (FTimingProfilerManager::Get().IsValid())
+	if (FNetworkingProfilerManager::Get().IsValid())
 	{
-		// Shutdown the TimingProfiler (Timing Insights) manager.
-		FTimingProfilerManager::Get()->Shutdown();
+		// Shutdown the NetworkingProfiler (Networking Insights) manager.
+		FNetworkingProfilerManager::Get()->Shutdown();
 	}
 
 	if (FLoadingProfilerManager::Get().IsValid())
@@ -143,6 +166,12 @@ void FTraceInsightsModule::ShutdownModule()
 		FLoadingProfilerManager::Get()->Shutdown();
 	}
 
+	if (FTimingProfilerManager::Get().IsValid())
+	{
+		// Shutdown the TimingProfiler (Timing Insights) manager.
+		FTimingProfilerManager::Get()->Shutdown();
+	}
+
 	if (FInsightsManager::Get().IsValid())
 	{
 		// Shutdown the main manager.
@@ -164,6 +193,7 @@ void FTraceInsightsModule::OnNewLayout(TSharedRef<FTabManager::FLayout> NewLayou
 			->AddTab(FInsightsManagerTabs::StartPageTabId, ETabState::OpenedTab)
 			->AddTab(FInsightsManagerTabs::TimingProfilerTabId, ETabState::ClosedTab)
 			->AddTab(FInsightsManagerTabs::LoadingProfilerTabId, ETabState::ClosedTab)
+			->AddTab(FInsightsManagerTabs::NetworkingProfilerTabId, ETabState::ClosedTab)
 			->SetForegroundTab(FTabId(FInsightsManagerTabs::StartPageTabId))
 		)
 	);
@@ -242,3 +272,29 @@ void FTraceInsightsModule::OnLoadingProfilerTabBeingClosed(TSharedRef<SDockTab>
 }
 
 ////////////////////////////////////////////////////////////////////////////////////////////////////
+
+TSharedRef<SDockTab> FTraceInsightsModule::SpawnNetworkingProfilerTab(const FSpawnTabArgs& Args)
+{
+	const TSharedRef<SDockTab> DockTab = SNew(SDockTab)
+		.TabRole(ETabRole::NomadTab);
+
+	// Register OnTabClosed to handle I/O profiler manager shutdown.
+	DockTab->SetOnTabClosed(SDockTab::FOnTabClosedCallback::CreateRaw(this, &FTraceInsightsModule::OnNetworkingProfilerTabBeingClosed));
+
+	// Create the SNetworkingProfilerWindow widget.
+	TSharedRef<SNetworkingProfilerWindow> Window = SNew(SNetworkingProfilerWindow, DockTab, Args.GetOwnerWindow());
+	FNetworkingProfilerManager::Get()->AssignProfilerWindow(Window);
+	DockTab->SetContent(Window);
+
+	return DockTab;
+}
+
+////////////////////////////////////////////////////////////////////////////////////////////////////
+
+void FTraceInsightsModule::OnNetworkingProfilerTabBeingClosed(TSharedRef<SDockTab> TabBeingClosed)
+{
+	// Disable TabClosed delegate.
+	TabBeingClosed->SetOnTabClosed(SDockTab::FOnTabClosedCallback());
+}
+
+////////////////////////////////////////////////////////////////////////////////////////////////////
