[Insights]
  - Networking Insights: Implemented prototype for the Packet Sizes View widget (rendering mock data, scrolling, zooming horiz & vert, tootlip, horiz. and vert. grid).
  - Frames track: Added horizontal and vertical grid (frame index for the horizontal axis major grid lines and frame duration + fps for the vertical axis major grid lines).
  - Frames track: Added fade in animation for tooltip.
  - Frames track: Fixed tooltip flickering when zooming in/out.
  - Frames track: Fixed missaligned viewport issue after zooming in/out (affecting detection of frame under mouse).
  - Various code refactoring (ex.: split FFrameTrackViewport into FIndexAxisViewport + FValueAxisViewport).

diff --git a/Engine/Source/Developer/TraceInsights/Private/Insights/LoadingProfiler/Widgets/SLoadingProfilerWindow.cpp b/Engine/Source/Developer/TraceInsights/Private/Insights/LoadingProfiler/Widgets/SLoadingProfilerWindow.cpp
--- a/Engine/Source/Developer/TraceInsights/Private/Insights/LoadingProfiler/Widgets/SLoadingProfilerWindow.cpp
+++ b/Engine/Source/Developer/TraceInsights/Private/Insights/LoadingProfiler/Widgets/SLoadingProfilerWindow.cpp
@@ -4,6 +4,7 @@
 
 #include "EditorStyleSet.h"
 #include "Framework/Docking/WorkspaceItem.h"
+#include "Framework/MultiBox/MultiBoxBuilder.h"
 #include "SlateOptMacros.h"
 #include "Widgets/Docking/SDockTab.h"
 #include "Widgets/Images/SImage.h"
diff --git a/Engine/Source/Developer/TraceInsights/Private/Insights/LoadingProfiler/Widgets/SLoadingProfilerWindow.h b/Engine/Source/Developer/TraceInsights/Private/Insights/LoadingProfiler/Widgets/SLoadingProfilerWindow.h
--- a/Engine/Source/Developer/TraceInsights/Private/Insights/LoadingProfiler/Widgets/SLoadingProfilerWindow.h
+++ b/Engine/Source/Developer/TraceInsights/Private/Insights/LoadingProfiler/Widgets/SLoadingProfilerWindow.h
@@ -18,6 +18,8 @@
 ////////////////////////////////////////////////////////////////////////////////////////////////////
 
 class FActiveTimerHandle;
+class FMenuBuilder;
+
 class STimingView;
 
 namespace Insights
diff --git a/Engine/Source/Developer/TraceInsights/Private/Insights/NetworkingProfiler/ViewModels/PacketSizesViewHelper.cpp b/Engine/Source/Developer/TraceInsights/Private/Insights/NetworkingProfiler/ViewModels/PacketSizesViewHelper.cpp
new file mode 100644
--- /dev/null
+++ b/Engine/Source/Developer/TraceInsights/Private/Insights/NetworkingProfiler/ViewModels/PacketSizesViewHelper.cpp
@@ -0,0 +1,270 @@
+// Copyright 1998-2019 Epic Games, Inc. All Rights Reserved.
+
+#include "PacketSizesViewHelper.h"
+
+#include "Brushes/SlateBorderBrush.h"
+#include "Brushes/SlateColorBrush.h"
+#include "EditorStyleSet.h"
+#include "Fonts/FontMeasure.h"
+#include "Fonts/SlateFontInfo.h"
+#include "Framework/Application/SlateApplication.h"
+#include "Rendering/DrawElements.h"
+#include "TraceServices/AnalysisService.h"
+
+// Insights
+#include "Insights/Common/PaintUtils.h"
+#include "Insights/NetworkingProfiler/ViewModels/PacketSizesViewport.h"
+
+#include <limits>
+
+////////////////////////////////////////////////////////////////////////////////////////////////////
+// FNetworkPacketSeriesBuilder
+////////////////////////////////////////////////////////////////////////////////////////////////////
+
+FNetworkPacketSeriesBuilder::FNetworkPacketSeriesBuilder(FNetworkPacketSeries& InSeries, const FPacketSizesViewport& InViewport)
+	: Series(InSeries)
+	, Viewport(InViewport)
+	, NumAddedPackets(0)
+{
+	SampleW = Viewport.GetSampleWidth();
+	PacketsPerSample = Viewport.GetNumPacketsPerSample();
+	NumSamples = FMath::Max(0, FMath::CeilToInt(Viewport.GetWidth() / SampleW));
+	FirstFrameIndex = Viewport.GetFirstFrameIndex();
+
+	Series.NumAggregatedPackets = 0;
+	Series.Samples.Reset();
+	Series.Samples.AddDefaulted(NumSamples);
+}
+
+////////////////////////////////////////////////////////////////////////////////////////////////////
+
+void FNetworkPacketSeriesBuilder::AddPacket(int32 FrameIndex, int64 Size, double TimeSent, double TimeAck, ENetworkPacketStatus Status)
+{
+	NumAddedPackets++;
+
+	int32 SampleIndex = (FrameIndex - FirstFrameIndex) / PacketsPerSample;
+	if (SampleIndex >= 0 && SampleIndex < NumSamples)
+	{
+		FNetworkPacketAggregatedSample& Sample = Series.Samples[SampleIndex];
+		Sample.NumPackets++;
+
+		if (TimeSent < Sample.StartTime)
+		{
+			Sample.StartTime = TimeSent;
+		}
+		if (TimeSent > Sample.EndTime)
+		{
+			Sample.EndTime = TimeSent;
+		}
+
+		if (Size > Sample.LargestPacket.Size)
+		{
+			Sample.LargestPacket.FrameIndex = FrameIndex;
+			Sample.LargestPacket.Size = Size;
+			Sample.LargestPacket.TimeSent = TimeSent;
+			Sample.LargestPacket.TimeAck = TimeAck;
+			Sample.LargestPacket.Status = Status;
+		}
+
+		ensure(Status != ENetworkPacketStatus::Unknown);
+		if (static_cast<int32>(Status) > static_cast<int32>(Sample.AggregatedStatus))
+		{
+			Sample.AggregatedStatus = Status;
+		}
+
+		Series.NumAggregatedPackets++;
+	}
+}
+
+////////////////////////////////////////////////////////////////////////////////////////////////////
+// FPacketSizesViewDrawHelper
+////////////////////////////////////////////////////////////////////////////////////////////////////
+
+FPacketSizesViewDrawHelper::FPacketSizesViewDrawHelper(const FDrawContext& InDrawContext, const FPacketSizesViewport& InViewport)
+	: DrawContext(InDrawContext)
+	, Viewport(InViewport)
+	, WhiteBrush(FCoreStyle::Get().GetBrush("WhiteBrush"))
+	, BorderBrush(FEditorStyle::GetBrush("PlainBorder"))
+	, NumPackets(0)
+	, NumDrawSamples(0)
+{
+}
+
+////////////////////////////////////////////////////////////////////////////////////////////////////
+
+void FPacketSizesViewDrawHelper::DrawBackground() const
+{
+	const FSlateBrush* AreaBrush = WhiteBrush;
+	const FLinearColor ValidAreaColor(0.07f, 0.07f, 0.07f, 1.0f);
+	const FLinearColor InvalidAreaColor(0.1f, 0.07f, 0.07f, 1.0f);
+
+	const FIndexAxisViewport& ViewportX = Viewport.GetHorizontalAxisViewport();
+
+	const float X0 = ViewportX.GetMinPos() - ViewportX.GetPos();
+	const float X1 = ViewportX.GetMaxPos() - ViewportX.GetPos();
+
+	const float W = FMath::CeilToFloat(Viewport.GetWidth());
+	const float H = FMath::CeilToFloat(Viewport.GetHeight());
+
+	if (X0 >= W || X1 <= 0.0f)
+	{
+		// Draw invalid area (entire view).
+		DrawContext.DrawBox(0.0f, 0.0f, W, H, AreaBrush, InvalidAreaColor);
+	}
+	else // X0 < W && X1 > 0
+	{
+		if (X0 > 0.0f)
+		{
+			// Draw invalid area (left).
+			DrawContext.DrawBox(0.0f, 0.0f, X0, H, AreaBrush, InvalidAreaColor);
+		}
+
+		if (X1 < W)
+		{
+			// Draw invalid area (right).
+			DrawContext.DrawBox(X1, 0.0f, W - X1, H, AreaBrush, InvalidAreaColor);
+		}
+
+		const float ValidX0 = FMath::Max(X0, 0.0f);
+		const float ValidX1 = FMath::Min(X1, W);
+
+		if (ValidX1 > ValidX0)
+		{
+			// Draw valid area.
+			DrawContext.DrawBox(ValidX0, 0.0f, ValidX1 - ValidX0, H, AreaBrush, ValidAreaColor);
+		}
+	}
+
+	DrawContext.LayerId++;
+}
+
+////////////////////////////////////////////////////////////////////////////////////////////////////
+
+FLinearColor FPacketSizesViewDrawHelper::GetColorByStatus(ENetworkPacketStatus Status)
+{
+	const float Alpha = 1.0f;
+	switch (Status)
+	{
+	case ENetworkPacketStatus::Unknown:
+		return FLinearColor(0.25f, 0.25f, 0.25f, Alpha);
+
+	case ENetworkPacketStatus::Sent:
+		return FLinearColor(0.75f, 0.75f, 0.75f, Alpha);
+
+	case ENetworkPacketStatus::ConfirmedReceived:
+		return FLinearColor(0.5f, 1.0f, 0.5f, Alpha);
+
+	case ENetworkPacketStatus::ConfirmedLost:
+		return FLinearColor(1.0f, 0.5f, 0.5f, Alpha);
+
+	default:
+		return FLinearColor(1.0f, 0.0f, 0.0f, Alpha);
+	}
+}
+
+////////////////////////////////////////////////////////////////////////////////////////////////////
+
+void FPacketSizesViewDrawHelper::DrawCached(const FNetworkPacketSeries& Series) const
+{
+	if (Series.NumAggregatedPackets == 0)
+	{
+		return;
+	}
+
+	NumPackets += Series.NumAggregatedPackets;
+
+	const float SampleW = Viewport.GetSampleWidth();
+	const int32 NumSamples = Series.Samples.Num();
+
+	const FValueAxisViewport& ViewportY = Viewport.GetVerticalAxisViewport();
+
+	const float ViewHeight = FMath::RoundToFloat(Viewport.GetHeight());
+	const float BaselineY = FMath::RoundToFloat(ViewportY.GetOffsetForValue(0.0));
+
+	for (int32 SampleIndex = 0; SampleIndex < NumSamples; SampleIndex++)
+	{
+		const FNetworkPacketAggregatedSample& Sample = Series.Samples[SampleIndex];
+		if (Sample.NumPackets == 0)
+		{
+			continue;
+		}
+
+		NumDrawSamples++;
+
+		const float X = SampleIndex * SampleW;
+
+		const float ValueY = FMath::RoundToFloat(ViewportY.GetOffsetForValue(static_cast<double>(Sample.LargestPacket.Size)));
+
+		const float H = ValueY - BaselineY;
+		const float Y = ViewHeight - H;
+
+		const FLinearColor ColorFill = GetColorByStatus(Sample.AggregatedStatus);
+		const FLinearColor ColorBorder(ColorFill.R * 0.75f, ColorFill.G * 0.75f, ColorFill.B * 0.75f, 1.0);
+
+		if (SampleW > 2.0f)
+		{
+			DrawContext.DrawBox(X + 1.0f, Y + 1.0f, SampleW - 2.0f, H - 2.0f, WhiteBrush, ColorFill);
+
+			// Draw border.
+			DrawContext.DrawBox(X, Y, 1.0, H, WhiteBrush, ColorBorder);
+			DrawContext.DrawBox(X + SampleW - 1.0f, Y, 1.0, H, WhiteBrush, ColorBorder);
+			DrawContext.DrawBox(X + 1.0f, Y, SampleW - 2.0f, 1.0f, WhiteBrush, ColorBorder);
+			DrawContext.DrawBox(X + 1.0f, Y + H - 1.0f, SampleW - 2.0f, 1.0f, WhiteBrush, ColorBorder);
+		}
+		else
+		{
+			DrawContext.DrawBox(X, Y, SampleW, H, WhiteBrush, ColorBorder);
+		}
+	}
+
+	DrawContext.LayerId++;
+}
+
+////////////////////////////////////////////////////////////////////////////////////////////////////
+
+void FPacketSizesViewDrawHelper::DrawSampleHighlight(const FNetworkPacketAggregatedSample& Sample, EHighlightMode Mode) const
+{
+	const float SampleW = Viewport.GetSampleWidth();
+	const int32 PacketsPerSample = Viewport.GetNumPacketsPerSample();
+	const int32 FirstFrameIndex = Viewport.GetFirstFrameIndex();
+	const int32 SampleIndex = (Sample.LargestPacket.FrameIndex - FirstFrameIndex) / PacketsPerSample;
+	const float X = SampleIndex * SampleW;
+
+	const FValueAxisViewport& ViewportY = Viewport.GetVerticalAxisViewport();
+
+	const float ViewHeight = FMath::RoundToFloat(Viewport.GetHeight());
+	const float BaselineY = FMath::RoundToFloat(ViewportY.GetOffsetForValue(0.0));
+	const float ValueY = FMath::RoundToFloat(ViewportY.GetOffsetForValue(static_cast<double>(Sample.LargestPacket.Size)));
+
+	const float H = ValueY - BaselineY;
+	const float Y = ViewHeight - H;
+
+	if (Mode == EHighlightMode::Hovered)
+	{
+		const FSlateBrush* HighlightBrush = BorderBrush;
+
+		const FLinearColor Color(1.0f, 1.0f, 0.0f, 1.0f); // yellow
+
+		// Draw border around the timing event box.
+		DrawContext.DrawBox(X - 1.0f, Y - 1.0f, SampleW + 2.0f, H + 2.0f, HighlightBrush, Color);
+	}
+	else // EHighlightMode::Selected or EHighlightMode::SelectedAndHovered
+	{
+		const FSlateBrush* SelectedBrush = BorderBrush;
+
+		// Animate color from white (if selected and hovered) or yellow (if only selected) to black, using a squared sine function.
+		const double Time = static_cast<double>(FPlatformTime::Cycles64()) * FPlatformTime::GetSecondsPerCycle64();
+		float S = FMath::Sin(2.0 * Time);
+		S = S * S; // squared, to ensure only positive [0 - 1] values
+		const float Blue = (Mode == EHighlightMode::SelectedAndHovered) ? 0.0f : S;
+		const FLinearColor Color(S, S, Blue, 1.0f);
+
+		// Draw border around the timing event box.
+		DrawContext.DrawBox(X - 1.0f, Y - 1.0f, SampleW + 2.0f, H + 2.0f, SelectedBrush, Color);
+	}
+	DrawContext.LayerId++;
+}
+
+
+
+////////////////////////////////////////////////////////////////////////////////////////////////////
diff --git a/Engine/Source/Developer/TraceInsights/Private/Insights/NetworkingProfiler/ViewModels/PacketSizesViewHelper.h b/Engine/Source/Developer/TraceInsights/Private/Insights/NetworkingProfiler/ViewModels/PacketSizesViewHelper.h
new file mode 100644
--- /dev/null
+++ b/Engine/Source/Developer/TraceInsights/Private/Insights/NetworkingProfiler/ViewModels/PacketSizesViewHelper.h
@@ -0,0 +1,207 @@
+// Copyright 1998-2019 Epic Games, Inc. All Rights Reserved.
+
+#pragma once
+
+#include "CoreMinimal.h"
+#include "Styling/WidgetStyle.h"
+
+#include <limits>
+
+enum class ESlateDrawEffect : uint8;
+
+struct FDrawContext;
+struct FGeometry;
+struct FSlateBrush;
+
+class FPacketSizesViewport;
+class FSlateWindowElementList;
+
+////////////////////////////////////////////////////////////////////////////////////////////////////
+
+enum class ENetworkPacketStatus
+{
+	Unknown = 0,
+	ConfirmedReceived, // Sent + Ack(Received)
+	Sent, // Sent, no Ack event yet
+	ConfirmedLost // Sent + Ack(Lost)
+};
+
+////////////////////////////////////////////////////////////////////////////////////////////////////
+
+struct FNetworkPacket
+{
+	int32 FrameIndex;
+	int64 Size; // [bits]
+	double TimeSent; // time when packet is sent
+	double TimeAck; // time when ACK event is received (either confirmed received or lost)
+	ENetworkPacketStatus Status;
+
+	FNetworkPacket()
+		: FrameIndex(-1)
+		, Size(-1)
+		, TimeSent(std::numeric_limits<double>::infinity())
+		, TimeAck(std::numeric_limits<double>::infinity())
+		, Status(ENetworkPacketStatus::Unknown)
+	{}
+
+	FNetworkPacket(const FNetworkPacket&) = default;
+	FNetworkPacket& operator=(const FNetworkPacket&) = default;
+
+	FNetworkPacket(FNetworkPacket&&) = default;
+	FNetworkPacket& operator=(FNetworkPacket&&) = default;
+
+	bool Equals(const FNetworkPacket& Other) const
+	{
+		return FrameIndex == Other.FrameIndex
+			&& Size == Other.Size
+			/*&& TimeSent == Other.TimeSent
+			&& TimeAck == Other.TimeAck
+			&& Status == Other.Status*/;
+	}
+
+	static bool AreEquals(const FNetworkPacket& A, const FNetworkPacket& B)
+	{
+		return A.Equals(B);
+	}
+};
+
+////////////////////////////////////////////////////////////////////////////////////////////////////
+
+struct FNetworkPacketAggregatedSample
+{
+	int32 NumPackets;
+
+	double StartTime; // min TimeSent of all packets in this aggregated sample
+	double EndTime; // max TimeSent of all packets in this aggregated sample
+
+	/** Aggregated status for packets in this aggregated sample.
+	 *    Unknown  --> all aggregated packets are unknown
+	 *    Sent     --> all aggregated packets are sent; none is received or confirmed lost
+	 *    Received --> at least one packet in the sample set is confirmed received and none are confirmed lost
+	 *    Lost     --> at least one packet in the sample set is confirmed lost
+	**/
+	ENetworkPacketStatus AggregatedStatus;
+
+	FNetworkPacket LargestPacket;
+
+	FNetworkPacketAggregatedSample()
+		: NumPackets(0)
+		, StartTime(DBL_MAX)
+		, EndTime(-DBL_MAX)
+		, AggregatedStatus(ENetworkPacketStatus::Unknown)
+		, LargestPacket()
+	{}
+
+	FNetworkPacketAggregatedSample(const FNetworkPacketAggregatedSample&) = default;
+	FNetworkPacketAggregatedSample& operator=(const FNetworkPacketAggregatedSample&) = default;
+
+	FNetworkPacketAggregatedSample(FNetworkPacketAggregatedSample&&) = default;
+	FNetworkPacketAggregatedSample& operator=(FNetworkPacketAggregatedSample&&) = default;
+
+	bool Equals(const FNetworkPacketAggregatedSample& Other) const
+	{
+		return NumPackets == Other.NumPackets
+			&& StartTime == Other.StartTime
+			&& EndTime == Other.EndTime
+			&& AggregatedStatus == Other.AggregatedStatus
+			&& LargestPacket.Equals(Other.LargestPacket);
+	}
+
+	static bool AreEquals(const FNetworkPacketAggregatedSample& A, const FNetworkPacketAggregatedSample& B)
+	{
+		return A.Equals(B);
+	}
+};
+
+////////////////////////////////////////////////////////////////////////////////////////////////////
+
+struct FNetworkPacketSeries
+{
+	int32 NumAggregatedPackets; // total number of packets aggregated in samples; i.e. sum of all Sample.NumPackets
+
+	TArray<FNetworkPacketAggregatedSample> Samples;
+
+	FNetworkPacketSeries()
+		: NumAggregatedPackets(0)
+		, Samples()
+	{
+	}
+
+	void Reset()
+	{
+		NumAggregatedPackets = 0;
+		Samples.Reset();
+	}
+};
+
+////////////////////////////////////////////////////////////////////////////////////////////////////
+
+class FNetworkPacketSeriesBuilder
+{
+public:
+	FNetworkPacketSeriesBuilder(FNetworkPacketSeries& InSeries, const FPacketSizesViewport& InViewport);
+
+	/**
+	 * Non-copyable
+	 */
+	FNetworkPacketSeriesBuilder(const FNetworkPacketSeriesBuilder&) = delete;
+	FNetworkPacketSeriesBuilder& operator=(const FNetworkPacketSeriesBuilder&) = delete;
+
+	void AddPacket(int32 FrameIndex, int64 Size, double TimeSent, double TimeAck, ENetworkPacketStatus Status);
+
+	int32 GetNumAddedPackets() const { return NumAddedPackets; }
+
+private:
+	FNetworkPacketSeries& Series; // series to update
+	const FPacketSizesViewport& Viewport;
+
+	float SampleW; // width of a sample, in Slate units
+	int32 PacketsPerSample; // number of packets in a sample
+	int32 FirstFrameIndex; // index of first frame in first sample; can be negative
+	int32 NumSamples; // total number of samples
+
+	// Debug stats.
+	int32 NumAddedPackets; // counts total number of added packets; i.e. number of AddPacket() calls
+};
+
+////////////////////////////////////////////////////////////////////////////////////////////////////
+
+class FPacketSizesViewDrawHelper
+{
+public:
+	enum class EHighlightMode : uint32
+	{
+		Hovered = 1,
+		Selected = 2,
+		SelectedAndHovered = 3
+	};
+
+public:
+	FPacketSizesViewDrawHelper(const FDrawContext& InDrawContext, const FPacketSizesViewport& InViewport);
+
+	/**
+	 * Non-copyable
+	 */
+	FPacketSizesViewDrawHelper(const FPacketSizesViewDrawHelper&) = delete;
+	FPacketSizesViewDrawHelper& operator=(const FPacketSizesViewDrawHelper&) = delete;
+
+	void DrawBackground() const;
+	void DrawCached(const FNetworkPacketSeries& Series) const;
+	void DrawSampleHighlight(const FNetworkPacketAggregatedSample& Sample, EHighlightMode Mode) const;
+
+	static FLinearColor GetColorByStatus(ENetworkPacketStatus Status);
+
+	int32 GetNumPackets() const { return NumPackets; }
+	int32 GetNumDrawSamples() const { return NumDrawSamples; }
+
+private:
+	const FDrawContext& DrawContext;
+	const FPacketSizesViewport& Viewport;
+
+	const FSlateBrush* WhiteBrush;
+	const FSlateBrush* BorderBrush;
+
+	// Debug stats.
+	mutable int32 NumPackets;
+	mutable int32 NumDrawSamples;
+};
diff --git a/Engine/Source/Developer/TraceInsights/Private/Insights/NetworkingProfiler/ViewModels/PacketSizesViewport.cpp b/Engine/Source/Developer/TraceInsights/Private/Insights/NetworkingProfiler/ViewModels/PacketSizesViewport.cpp
new file mode 100644
--- /dev/null
+++ b/Engine/Source/Developer/TraceInsights/Private/Insights/NetworkingProfiler/ViewModels/PacketSizesViewport.cpp
@@ -0,0 +1,5 @@
+// Copyright 1998-2019 Epic Games, Inc. All Rights Reserved.
+
+#include "PacketSizesViewport.h"
+
+////////////////////////////////////////////////////////////////////////////////////////////////////
diff --git a/Engine/Source/Developer/TraceInsights/Private/Insights/NetworkingProfiler/ViewModels/PacketSizesViewport.h b/Engine/Source/Developer/TraceInsights/Private/Insights/NetworkingProfiler/ViewModels/PacketSizesViewport.h
new file mode 100644
--- /dev/null
+++ b/Engine/Source/Developer/TraceInsights/Private/Insights/NetworkingProfiler/ViewModels/PacketSizesViewport.h
@@ -0,0 +1,66 @@
+// Copyright 1998-2019 Epic Games, Inc. All Rights Reserved.
+
+#pragma once
+
+#include "CoreMinimal.h"
+
+// Insights
+#include "Insights/ViewModels/FrameTrackViewport.h"
+//#include "Insights/ViewModels/IndexAxisViewport.h"
+//#include "Insights/ViewModels/ValueAxisViewport.h"
+
+////////////////////////////////////////////////////////////////////////////////////////////////////
+
+class FPacketSizesViewport
+{
+private:
+	static constexpr float SLATE_UNITS_TOLERANCE = 0.1f;
+
+public:
+	FPacketSizesViewport()
+	{
+		Reset();
+	}
+
+	void Reset()
+	{
+		HorizontalAxisViewport.Reset();
+		VerticalAxisViewport.Reset();
+	}
+
+	const FIndexAxisViewport& GetHorizontalAxisViewport() const { return HorizontalAxisViewport; }
+	FIndexAxisViewport& GetHorizontalAxisViewport() { return HorizontalAxisViewport; }
+
+	const FValueAxisViewport& GetVerticalAxisViewport() const { return VerticalAxisViewport; }
+	FValueAxisViewport& GetVerticalAxisViewport() { return VerticalAxisViewport; }
+
+	float GetWidth() const { return HorizontalAxisViewport.GetSize(); }
+	float GetHeight() const { return VerticalAxisViewport.GetSize(); }
+
+	bool SetSize(const float InWidth, const float InHeight)
+	{
+		const bool bWidthChanged = HorizontalAxisViewport.SetSize(InWidth);
+		const bool bHeightChanged = VerticalAxisViewport.SetSize(InHeight);
+		if (bWidthChanged || bHeightChanged)
+		{
+			OnSizeChanged();
+			return true;
+		}
+		return false;
+	}
+
+	float GetSampleWidth() const { return HorizontalAxisViewport.GetSampleSize(); }
+	int32 GetNumPacketsPerSample() const { return HorizontalAxisViewport.GetNumIndicesPerSample(); }
+	int32 GetFirstFrameIndex() const { return HorizontalAxisViewport.GetIndexAtOffset(0.0f); }
+
+private:
+	void OnSizeChanged()
+	{
+	}
+
+private:
+	FIndexAxisViewport HorizontalAxisViewport;
+	FValueAxisViewport VerticalAxisViewport;
+};
+
+////////////////////////////////////////////////////////////////////////////////////////////////////
diff --git a/Engine/Source/Developer/TraceInsights/Private/Insights/NetworkingProfiler/Widgets/SNetworkingProfilerWindow.h b/Engine/Source/Developer/TraceInsights/Private/Insights/NetworkingProfiler/Widgets/SNetworkingProfilerWindow.h
--- a/Engine/Source/Developer/TraceInsights/Private/Insights/NetworkingProfiler/Widgets/SNetworkingProfilerWindow.h
+++ b/Engine/Source/Developer/TraceInsights/Private/Insights/NetworkingProfiler/Widgets/SNetworkingProfilerWindow.h
@@ -18,6 +18,8 @@
 ////////////////////////////////////////////////////////////////////////////////////////////////////
 
 class FActiveTimerHandle;
+class FMenuBuilder;
+
 class SDataStreamBreakdownView;
 class SPacketBreakdownView;
 class SPacketSizesView;
diff --git a/Engine/Source/Developer/TraceInsights/Private/Insights/NetworkingProfiler/Widgets/SPacketSizesView.cpp b/Engine/Source/Developer/TraceInsights/Private/Insights/NetworkingProfiler/Widgets/SPacketSizesView.cpp
--- a/Engine/Source/Developer/TraceInsights/Private/Insights/NetworkingProfiler/Widgets/SPacketSizesView.cpp
+++ b/Engine/Source/Developer/TraceInsights/Private/Insights/NetworkingProfiler/Widgets/SPacketSizesView.cpp
@@ -2,7 +2,24 @@
 
 #include "SPacketSizesView.h"
 
+#include "Brushes/SlateBorderBrush.h"
+#include "Brushes/SlateColorBrush.h"
+#include "EditorStyleSet.h"
+#include "Fonts/FontMeasure.h"
+#include "Fonts/SlateFontInfo.h"
+#include "Framework/Application/SlateApplication.h"
+#include "Framework/MultiBox/MultiBoxBuilder.h"
+#include "HAL/PlatformTime.h"
+#include "Rendering/DrawElements.h"
+#include "Widgets/Layout/SScrollBar.h"
+
 // Insights
+#include "Insights/Common/PaintUtils.h"
+#include "Insights/Common/Stopwatch.h"
+#include "Insights/Common/TimeUtils.h"
+#include "Insights/InsightsManager.h"
+
+#include <limits>
 
 ////////////////////////////////////////////////////////////////////////////////////////////////////
 
@@ -11,7 +28,9 @@
 ////////////////////////////////////////////////////////////////////////////////////////////////////
 
 SPacketSizesView::SPacketSizesView()
+	: PacketSeries(MakeShareable(new FNetworkPacketSeries()))
 {
+	Reset();
 }
 
 ////////////////////////////////////////////////////////////////////////////////////////////////////
@@ -22,11 +41,981 @@ SPacketSizesView::~SPacketSizesView()
 
 ////////////////////////////////////////////////////////////////////////////////////////////////////
 
+void SPacketSizesView::Reset()
+{
+	Viewport.Reset();
+
+	FValueAxisViewport& ViewportY = Viewport.GetVerticalAxisViewport();
+	ViewportY.SetScaleLimits(0.0001, 50.0);
+	ViewportY.SetScale(0.08);
+	bIsViewportDirty = true;
+
+	PacketSeries->Reset();
+
+	bIsStateDirty = true;
+
+	bIsAutoZoomEnabled = true;
+
+	AnalysisSyncNextTimestamp = 0;
+
+	MousePosition = FVector2D::ZeroVector;
+
+	MousePositionOnButtonDown = FVector2D::ZeroVector;
+	ViewportPosXOnButtonDown = 0.0f;
+	ViewportPosYOnButtonDown = 0.0f;
+
+	MousePositionOnButtonUp = FVector2D::ZeroVector;
+
+	bIsLMB_Pressed = false;
+	bIsRMB_Pressed = false;
+
+	bIsScrolling = false;
+
+	SelectionStartFrameIndex = 0;
+	SelectionEndFrameIndex = 0;
+
+	SelectedSample.Reset();
+	HoveredSample.Reset();
+	TooltipDesiredOpacity = 0.9f;
+	TooltipOpacity = 0.0f;
+
+	//ThisGeometry
+
+	CursorType = ECursorType::Default;
+
+	NumUpdatedPackets = 0;
+	UpdateDurationHistory.Reset();
+	DrawDurationHistory.Reset();
+	OnPaintDurationHistory.Reset();
+	LastOnPaintTime = FPlatformTime::Cycles64();
+}
+
+////////////////////////////////////////////////////////////////////////////////////////////////////
+
 void SPacketSizesView::Construct(const FArguments& InArgs)
 {
-	//ChildSlot
-	//[
-	//];
+	ChildSlot
+	[
+		SNew(SOverlay)
+		.Visibility(EVisibility::SelfHitTestInvisible)
+
+		+ SOverlay::Slot()
+		.VAlign(VAlign_Top)
+		.Padding(FMargin(0, 0, 0, 0))
+		[
+			SAssignNew(HorizontalScrollBar, SScrollBar)
+			.Orientation(Orient_Horizontal)
+			.AlwaysShowScrollbar(false)
+			.Visibility(EVisibility::Visible)
+			.Thickness(FVector2D(5.0f, 5.0f))
+			.RenderOpacity(0.75)
+			.OnUserScrolled(this, &SPacketSizesView::HorizontalScrollBar_OnUserScrolled)
+		]
+	];
+
+	UpdateHorizontalScrollBar();
+
+	BindCommands();
+}
+
+////////////////////////////////////////////////////////////////////////////////////////////////////
+
+void SPacketSizesView::Tick(const FGeometry& AllottedGeometry, const double InCurrentTime, const float InDeltaTime)
+{
+	if (ThisGeometry != AllottedGeometry || bIsViewportDirty)
+	{
+		bIsViewportDirty = false;
+		const float ViewWidth = AllottedGeometry.GetLocalSize().X;
+		const float ViewHeight = AllottedGeometry.GetLocalSize().Y;
+		Viewport.SetSize(ViewWidth, ViewHeight);
+		bIsStateDirty = true;
+	}
+
+	ThisGeometry = AllottedGeometry;
+
+	FIndexAxisViewport& ViewportX = Viewport.GetHorizontalAxisViewport();
+
+	if (!bIsScrolling)
+	{
+		// Elastic snap to horizontal limits.
+		if (ViewportX.UpdatePosWithinLimits())
+		{
+			bIsStateDirty = true;
+		}
+	}
+
+	uint64 Time = FPlatformTime::Cycles64();
+	if (Time > AnalysisSyncNextTimestamp)
+	{
+		const uint64 WaitTime = static_cast<uint64>(0.1 / FPlatformTime::GetSecondsPerCycle64()); // 100ms
+		AnalysisSyncNextTimestamp = Time + WaitTime;
+
+		//TODO
+	}
+
+	if (bIsStateDirty)
+	{
+		bIsStateDirty = false;
+		UpdateState();
+	}
+}
+
+////////////////////////////////////////////////////////////////////////////////////////////////////
+
+void SPacketSizesView::UpdateState()
+{
+	FStopwatch Stopwatch;
+	Stopwatch.Start();
+
+	// Reset stats.
+	PacketSeries->NumAggregatedPackets = 0;
+	NumUpdatedPackets = 0;
+
+	// Init series with mock data.
+	{
+		FIndexAxisViewport& ViewportX = Viewport.GetHorizontalAxisViewport();
+
+		constexpr int32 NumPackets = 100000;
+
+		ViewportX.SetMinMaxIndexInterval(0, NumPackets);
+
+		const int32 StartIndex = FMath::Max(0, ViewportX.GetIndexAtOffset(0.0f));
+		const int32 EndIndex = FMath::Min(NumPackets, ViewportX.GetIndexAtOffset(ViewportX.GetSize()));
+
+		FNetworkPacketSeriesBuilder Builder(*PacketSeries, Viewport);
+		for (int32 FrameIndex = StartIndex; FrameIndex < EndIndex; ++FrameIndex)
+		{
+			FRandomStream RandomStream((FrameIndex * FrameIndex * FrameIndex) ^ 0x2c2c57ed);
+			int64 Size = RandomStream.RandRange(0, 2000);
+			const float Fraction = RandomStream.GetFraction();
+			ENetworkPacketStatus Status = ENetworkPacketStatus::ConfirmedReceived;
+			if (Fraction < 0.01) // 1%
+			{
+				Status = ENetworkPacketStatus::ConfirmedLost;
+			}
+			else if (Fraction < 0.05) // 4%
+			{
+				Status = ENetworkPacketStatus::Sent;
+			}
+
+			double TimeSent = ((double)FrameIndex * 100.0) / (double)NumPackets + RandomStream.GetFraction() * 0.1;
+			double TimeAck = 0.0;
+			if (Status != ENetworkPacketStatus::Sent)
+			{
+				TimeAck = TimeSent + RandomStream.GetFraction() * 0.1;
+			}
+
+			Builder.AddPacket(FrameIndex, Size, TimeSent, TimeAck, Status);
+		}
+		NumUpdatedPackets += Builder.GetNumAddedPackets();
+	}
+
+	Stopwatch.Stop();
+	UpdateDurationHistory.AddValue(Stopwatch.AccumulatedTime);
+}
+
+////////////////////////////////////////////////////////////////////////////////////////////////////
+
+FNetworkPacketSampleRef SPacketSizesView::GetSampleAtMousePosition(float X, float Y)
+{
+	if (!bIsStateDirty)
+	{
+		float SampleW = Viewport.GetSampleWidth();
+		int32 SampleIndex = FMath::FloorToInt(X / SampleW);
+		if (SampleIndex >= 0)
+		{
+			if (PacketSeries->NumAggregatedPackets > 0 &&
+				SampleIndex < PacketSeries->Samples.Num())
+			{
+				const FNetworkPacketAggregatedSample& Sample = PacketSeries->Samples[SampleIndex];
+				if (Sample.NumPackets > 0)
+				{
+					const FValueAxisViewport& ViewportY = Viewport.GetVerticalAxisViewport();
+
+					const float ViewHeight = FMath::RoundToFloat(Viewport.GetHeight());
+					const float BaselineY = FMath::RoundToFloat(ViewportY.GetOffsetForValue(0.0));
+					const float ValueY = FMath::RoundToFloat(ViewportY.GetOffsetForValue(static_cast<double>(Sample.LargestPacket.Size)));
+
+					constexpr float ToleranceY = 3.0f; // [pixels]
+
+					const float BottomY = FMath::Min(ViewHeight, ViewHeight - BaselineY + ToleranceY);
+					const float TopY = FMath::Max(0.0f, ViewHeight - ValueY - ToleranceY);
+
+					if (Y >= TopY && Y < BottomY)
+					{
+						return FNetworkPacketSampleRef(PacketSeries, MakeShareable(new FNetworkPacketAggregatedSample(Sample)));
+					}
+				}
+			}
+		}
+	}
+	return FNetworkPacketSampleRef(nullptr, nullptr);
+}
+
+////////////////////////////////////////////////////////////////////////////////////////////////////
+
+void SPacketSizesView::SelectSampleAtMousePosition(float X, float Y)
+{
+	FNetworkPacketSampleRef SampleRef = GetSampleAtMousePosition(X, Y);
+	if (!SampleRef.IsValid())
+	{
+		SampleRef = GetSampleAtMousePosition(X - 1.0f, Y);
+	}
+	if (!SampleRef.IsValid())
+	{
+		SampleRef = GetSampleAtMousePosition(X + 1.0f, Y);
+	}
+
+	if (SampleRef.IsValid())
+	{
+		SelectedSample = SampleRef;
+	}
+}
+
+////////////////////////////////////////////////////////////////////////////////////////////////////
+
+const TCHAR* StatusToString(ENetworkPacketStatus Status)
+{
+	switch (Status)
+	{
+		case ENetworkPacketStatus::ConfirmedReceived: return TEXT("Received");
+		case ENetworkPacketStatus::Sent:              return TEXT("Sent");
+		case ENetworkPacketStatus::ConfirmedLost:     return TEXT("Lost");
+		case ENetworkPacketStatus::Unknown:
+		default:                                      return TEXT("Unknown");
+	};
+}
+
+const TCHAR* AggregatedStatusToString(ENetworkPacketStatus Status)
+{
+	switch (Status)
+	{
+		case ENetworkPacketStatus::ConfirmedReceived: return TEXT("all are Received");
+		case ENetworkPacketStatus::Sent:              return TEXT("some are Sent, but have no Ack");
+		case ENetworkPacketStatus::ConfirmedLost:     return TEXT("some are Lost");
+		case ENetworkPacketStatus::Unknown:
+		default:                                      return TEXT("Unknown");
+	};
+}
+
+////////////////////////////////////////////////////////////////////////////////////////////////////
+
+int32 SPacketSizesView::OnPaint(const FPaintArgs& Args, const FGeometry& AllottedGeometry, const FSlateRect& MyCullingRect, FSlateWindowElementList& OutDrawElements, int32 LayerId, const FWidgetStyle& InWidgetStyle, bool bParentEnabled) const
+{
+	const bool bEnabled = ShouldBeEnabled(bParentEnabled);
+	const ESlateDrawEffect DrawEffects = bEnabled ? ESlateDrawEffect::None : ESlateDrawEffect::DisabledEffect;
+	FDrawContext DrawContext(AllottedGeometry, MyCullingRect, InWidgetStyle, DrawEffects, OutDrawElements, LayerId);
+
+	const TSharedRef<FSlateFontMeasure> FontMeasureService = FSlateApplication::Get().GetRenderer()->GetFontMeasureService();
+	FSlateFontInfo SummaryFont = FCoreStyle::GetDefaultFontStyle("Regular", 8);
+
+	const FSlateBrush* WhiteBrush = FCoreStyle::Get().GetBrush("WhiteBrush");
+
+	const float ViewWidth = AllottedGeometry.Size.X;
+	const float ViewHeight = AllottedGeometry.Size.Y;
+
+	int32 NumDrawSamples = 0;
+
+	{
+		FStopwatch Stopwatch;
+		Stopwatch.Start();
+
+		FPacketSizesViewDrawHelper Helper(DrawContext, Viewport);
+
+		Helper.DrawBackground();
+
+		DrawHorizontalAxisGrid(DrawContext, WhiteBrush, SummaryFont);
+
+		Helper.DrawCached(*PacketSeries);
+
+		NumDrawSamples = Helper.GetNumDrawSamples();
+
+		// Highlight the selected and/or hovered sample.
+		bool bIsSelectedAndHovered = SelectedSample.Equals(HoveredSample);
+		if (SelectedSample.IsValid())
+		{
+			Helper.DrawSampleHighlight(*SelectedSample.Sample, bIsSelectedAndHovered ? FPacketSizesViewDrawHelper::EHighlightMode::SelectedAndHovered : FPacketSizesViewDrawHelper::EHighlightMode::Selected);
+		}
+		if (HoveredSample.IsValid() && !bIsSelectedAndHovered)
+		{
+			Helper.DrawSampleHighlight(*HoveredSample.Sample, FPacketSizesViewDrawHelper::EHighlightMode::Hovered);
+		}
+
+		DrawVerticalAxisGrid(DrawContext, WhiteBrush, SummaryFont);
+
+		// Draw tooltip for hovered sample.
+		if (HoveredSample.IsValid())
+		{
+			if (TooltipOpacity < TooltipDesiredOpacity)
+			{
+				TooltipOpacity = TooltipOpacity * 0.9f + TooltipDesiredOpacity * 0.1f;
+			}
+			else
+			{
+				TooltipOpacity = TooltipDesiredOpacity;
+			}
+
+			const double Precision = 0.01; // 10ms
+			int NumLines;
+			FString Text;
+			if (HoveredSample.Sample->NumPackets == 1)
+			{
+				NumLines = 5;
+				Text = FString::Format(TEXT("Frame Index: {0}\n"
+											"Size: {1} bits\n"
+											"Sent Timestamp: {2}\n"
+											"Ack Timestamp: {3}\n"
+											"Status: {4}"),
+					{
+						FText::AsNumber(HoveredSample.Sample->LargestPacket.FrameIndex).ToString(),
+						FText::AsNumber(HoveredSample.Sample->LargestPacket.Size).ToString(),
+						TimeUtils::FormatTimeHMS(HoveredSample.Sample->LargestPacket.TimeSent, Precision),
+						HoveredSample.Sample->LargestPacket.TimeAck == 0 ? TEXT("N/A") : TimeUtils::FormatTimeHMS(HoveredSample.Sample->LargestPacket.TimeAck, Precision),
+						::StatusToString(HoveredSample.Sample->LargestPacket.Status)
+					});
+			}
+			else
+			{
+				NumLines = 8;
+				Text = FString::Format(TEXT("{0} network packets\n"
+											"({1})\n"
+											"Largest Packet\n"
+											"    Frame Index: {2}\n"
+											"    Size: {3} bits\n"
+											"    Sent Timestamp: {4}\n"
+											"    Ack Timestamp: {5}\n"
+											"    Status: {6}"),
+					{
+						HoveredSample.Sample->NumPackets,
+						::AggregatedStatusToString(HoveredSample.Sample->AggregatedStatus),
+						FText::AsNumber(HoveredSample.Sample->LargestPacket.FrameIndex).ToString(),
+						FText::AsNumber(HoveredSample.Sample->LargestPacket.Size).ToString(),
+						TimeUtils::FormatTimeHMS(HoveredSample.Sample->LargestPacket.TimeSent, Precision),
+						HoveredSample.Sample->LargestPacket.TimeAck == 0 ? TEXT("N/A") : TimeUtils::FormatTimeHMS(HoveredSample.Sample->LargestPacket.TimeAck, Precision),
+						::StatusToString(HoveredSample.Sample->LargestPacket.Status)
+					});
+			}
+
+			FVector2D TextSize = FontMeasureService->Measure(Text, SummaryFont);
+
+			const float DX = 2.0f;
+			const float W2 = TextSize.X / 2 + DX;
+
+			const FIndexAxisViewport& ViewportX = Viewport.GetHorizontalAxisViewport();
+
+			float X1 = ViewportX.GetOffsetForIndex(HoveredSample.Sample->LargestPacket.FrameIndex);
+			float CX = X1 + FMath::RoundToFloat(Viewport.GetSampleWidth() / 2);
+			if (CX + W2 > ViewportX.GetSize())
+			{
+				CX = FMath::RoundToFloat(ViewportX.GetSize() - W2);
+			}
+			if (CX - W2 < 0)
+			{
+				CX = W2;
+			}
+
+			const float Y = 10.0f;
+			const float H = 2.0f + 13.0f * NumLines;
+			DrawContext.DrawBox(CX - W2, Y, 2 * W2, H, WhiteBrush, FLinearColor(0.7, 0.7, 0.7, TooltipOpacity));
+			DrawContext.LayerId++;
+			DrawContext.DrawText(CX - W2 + DX, Y + 1.0f, Text, SummaryFont, FLinearColor(0.0, 0.0, 0.0, TooltipOpacity));
+			DrawContext.LayerId++;
+		}
+
+		Stopwatch.Stop();
+		DrawDurationHistory.AddValue(Stopwatch.AccumulatedTime);
+	}
+
+	//////////////////////////////////////////////////
+
+	const bool bShouldDisplayDebugInfo = FInsightsManager::Get()->IsDebugInfoEnabled();
+	if (bShouldDisplayDebugInfo)
+	{
+		const float MaxFontCharHeight = FontMeasureService->Measure(TEXT("!"), SummaryFont).Y;
+		const float DbgDY = MaxFontCharHeight;
+
+		const float DbgW = 280.0f;
+		const float DbgH = DbgDY * 4 + 3.0f;
+		const float DbgX = ViewWidth - DbgW - 20.0f;
+		float DbgY = 7.0f;
+
+		++LayerId;
+		FSlateDrawElement::MakeBox(
+			OutDrawElements,
+			LayerId,
+			MAKE_PAINT_GEOMETRY_RC(AllottedGeometry, DbgX - 2.0f, DbgY - 2.0f, DbgW, DbgH),
+			WhiteBrush,
+			DrawEffects,
+			FLinearColor(1.0, 1.0, 1.0, 0.9)
+		);
+
+		++LayerId;
+		FLinearColor DbgTextColor(0.0, 0.0, 0.0, 0.9);
+
+		// Time interval since last OnPaint call.
+		const uint64 CurrentTime = FPlatformTime::Cycles64();
+		const uint64 OnPaintDuration = CurrentTime - LastOnPaintTime;
+		LastOnPaintTime = CurrentTime;
+		OnPaintDurationHistory.AddValue(OnPaintDuration); // saved for last 32 OnPaint calls
+		const uint64 AvgOnPaintDuration = OnPaintDurationHistory.ComputeAverage();
+		const uint64 AvgOnPaintDurationMs = FStopwatch::Cycles64ToMilliseconds(AvgOnPaintDuration);
+		const double AvgOnPaintFps = AvgOnPaintDurationMs != 0 ? 1.0 / FStopwatch::Cycles64ToSeconds(AvgOnPaintDuration) : 0.0;
+
+		const uint64 AvgUpdateDurationMs = FStopwatch::Cycles64ToMilliseconds(UpdateDurationHistory.ComputeAverage());
+		const uint64 AvgDrawDurationMs = FStopwatch::Cycles64ToMilliseconds(DrawDurationHistory.ComputeAverage());
+
+		// Draw performance info.
+		FSlateDrawElement::MakeText
+		(
+			OutDrawElements,
+			LayerId,
+			MAKE_PAINT_GEOMETRY_PT(AllottedGeometry, DbgX, DbgY),
+			FString::Printf(TEXT("U: %llu ms, D: %llu ms + %llu ms = %llu ms (%d fps)"),
+				AvgUpdateDurationMs, // caching time
+				AvgDrawDurationMs, // drawing time
+				AvgOnPaintDurationMs - AvgDrawDurationMs, // other overhead to OnPaint calls
+				AvgOnPaintDurationMs, // average time between two OnPaint calls
+				FMath::RoundToInt(AvgOnPaintFps)), // framerate of OnPaint calls
+			SummaryFont,
+			DrawEffects,
+			DbgTextColor
+		);
+		DbgY += DbgDY;
+
+		// Draw number of draw calls.
+		FSlateDrawElement::MakeText
+		(
+			OutDrawElements,
+			LayerId,
+			MAKE_PAINT_GEOMETRY_PT(AllottedGeometry, DbgX, DbgY),
+			FString::Format(TEXT("U: {0} packets, D: {1} samples"),
+				{
+					FText::AsNumber(NumUpdatedPackets).ToString(),
+					FText::AsNumber(NumDrawSamples).ToString()
+				}),
+			SummaryFont,
+			DrawEffects,
+			DbgTextColor
+		);
+		DbgY += DbgDY;
+
+		// Draw viewport's horizontal info.
+		FSlateDrawElement::MakeText
+		(
+			OutDrawElements,
+			LayerId,
+			MAKE_PAINT_GEOMETRY_PT(AllottedGeometry, DbgX, DbgY),
+			Viewport.GetHorizontalAxisViewport().ToDebugString(TEXT("VX: "), TEXT("packet")),
+			SummaryFont,
+			DrawEffects,
+			DbgTextColor
+		);
+		DbgY += DbgDY;
+
+		// Draw viewport's vertical info.
+		FSlateDrawElement::MakeText
+		(
+			OutDrawElements,
+			LayerId,
+			MAKE_PAINT_GEOMETRY_PT(AllottedGeometry, DbgX, DbgY),
+			Viewport.GetVerticalAxisViewport().ToDebugString(TEXT("VY: ")),
+			SummaryFont,
+			DrawEffects,
+			DbgTextColor
+		);
+		DbgY += DbgDY;
+	}
+
+	return SCompoundWidget::OnPaint(Args, AllottedGeometry, MyCullingRect, OutDrawElements, LayerId, InWidgetStyle, bParentEnabled && IsEnabled());
+}
+
+////////////////////////////////////////////////////////////////////////////////////////////////////
+
+void SPacketSizesView::DrawVerticalAxisGrid(FDrawContext& DrawContext, const FSlateBrush* Brush, const FSlateFontInfo& Font) const
+{
+	const FValueAxisViewport& ViewportY = Viewport.GetVerticalAxisViewport();
+
+	const float RoundedViewHeight = FMath::RoundToFloat(ViewportY.GetSize());
+
+	constexpr float MinDY = 32.0f; // min vertical distance between horizontal grid lines
+
+	const double TopValue = ViewportY.GetValueForOffset(RoundedViewHeight);
+	const double GridValue = ViewportY.GetValueForOffset(MinDY);
+	const double BottomValue = ViewportY.GetValueForOffset(0.0f);
+	const double Delta = GridValue - BottomValue;
+
+	if (Delta > 0.0)
+	{
+		const int64 DeltaBits = static_cast<int64>(Delta);
+		// Compute rounding based on magnitude of visible range of values (Delta).
+		int64 Power10 = 1;
+		int64 Delta10 = DeltaBits;
+		while (Delta10 > 0)
+		{
+			Delta10 /= 10;
+			Power10 *= 10;
+		}
+		if (Power10 >= 100)
+		{
+			Power10 /= 100;
+		}
+		else
+		{
+			Power10 = 1;
+		}
+
+		const double Grid = static_cast<double>(((DeltaBits + Power10) / Power10) * Power10); // next value divisible with a multiple of 10
+
+		const double StartValue = FMath::GridSnap(BottomValue, Grid);
+
+		const FLinearColor GridColor(0.0f, 0.0f, 0.0f, 0.1f);
+		const FLinearColor TextBgColor(0.05f, 0.05f, 0.05f, 1.0f);
+		const FLinearColor TextColor(1.0f, 1.0f, 1.0f, 1.0f);
+
+		const float ViewWidth = Viewport.GetWidth();
+
+		const TSharedRef<FSlateFontMeasure> FontMeasureService = FSlateApplication::Get().GetRenderer()->GetFontMeasureService();
+
+		for (double Value = StartValue; Value < TopValue; Value += Grid)
+		{
+			const float Y = RoundedViewHeight - FMath::RoundToFloat(ViewportY.GetOffsetForValue(Value));
+
+			// Draw horizontal grid line.
+			DrawContext.DrawBox(0, Y, ViewWidth, 1, Brush, GridColor);
+
+			const int64 ValueBits = static_cast<int64>(Value);
+			const FString Text = (ValueBits == 0) ? TEXT("0") : FString::Format(TEXT("{0} bits"), { FText::AsNumber(ValueBits).ToString() });
+			const FVector2D TextSize = FontMeasureService->Measure(Text, Font);
+			constexpr float TextH = 14.0f;
+
+			// Draw background for value text.
+			DrawContext.DrawBox(ViewWidth - TextSize.X - 4.0f, Y - TextH, TextSize.X + 4.0f, TextH, Brush, TextBgColor);
+
+			// Draw value text.
+			DrawContext.DrawText(ViewWidth - TextSize.X - 2.0f, Y - TextH + 1.0f, Text, Font, TextColor);
+		}
+		DrawContext.LayerId++;
+	}
+}
+
+////////////////////////////////////////////////////////////////////////////////////////////////////
+
+void SPacketSizesView::DrawHorizontalAxisGrid(FDrawContext& DrawContext, const FSlateBrush* Brush, const FSlateFontInfo& Font) const
+{
+	const FIndexAxisViewport& ViewportX = Viewport.GetHorizontalAxisViewport();
+
+	const float RoundedViewWidth = FMath::RoundToFloat(ViewportX.GetSize());
+
+	constexpr float MinDX = 100.0f; // min horizontal distance between vertical grid lines
+
+	int32 LeftIndex = ViewportX.GetIndexAtOffset(0.0f);
+	int32 GridIndex = ViewportX.GetIndexAtOffset(MinDX);
+	int32 RightIndex = ViewportX.GetIndexAtOffset(RoundedViewWidth);
+	int32 Delta = GridIndex - LeftIndex;
+
+	if (Delta > 0)
+	{
+		// Compute rounding based on magnitude of visible range of samples (Delta).
+		int32 Power10 = 1;
+		int32 Delta10 = Delta;
+		while (Delta10 > 0)
+		{
+			Delta10 /= 10;
+			Power10 *= 10;
+		}
+		if (Power10 >= 100)
+		{
+			Power10 /= 100;
+		}
+		else
+		{
+			Power10 = 1;
+		}
+
+		const int32 Grid = ((Delta + Power10) / Power10) * Power10; // next value divisible with a multiple of 10
+
+		const FLinearColor GridColor(0.0f, 0.0f, 0.0f, 0.1f);
+		//const FLinearColor TextBgColor(0.05f, 0.05f, 0.05f, 1.0f);
+		//const FLinearColor TextColor(1.0f, 1.0f, 1.0f, 1.0f);
+		const FLinearColor TopTextColor(1.0f, 1.0f, 1.0f, 0.7f);
+
+		// Skip grid lines for negative indices.
+		double StartIndex = ((LeftIndex + Grid - 1) / Grid) * Grid;
+		while (StartIndex < 0)
+		{
+			StartIndex += Grid;
+		}
+
+		const float ViewHeight = Viewport.GetHeight();
+
+		//const TSharedRef<FSlateFontMeasure> FontMeasureService = FSlateApplication::Get().GetRenderer()->GetFontMeasureService();
+
+		for (int32 Index = StartIndex; Index < RightIndex; Index += Grid)
+		{
+			const float X = FMath::RoundToFloat(ViewportX.GetOffsetForIndex(Index));
+
+			// Draw vertical grid line.
+			DrawContext.DrawBox(X, 0, 1, ViewHeight, Brush, GridColor);
+
+			const FString Text = FText::AsNumber(Index).ToString();
+			//const FVector2D TextSize = FontMeasureService->Measure(Text, Font);
+			//constexpr float TextH = 14.0f;
+
+			// Draw background for index text.
+			//DrawContext.DrawBox(X, ViewHeight - TextH, TextSize.X + 4.0f, TextH, Brush, TextBgColor);
+
+			// Draw index text.
+			//DrawContext.DrawText(X + 2.0f, ViewHeight - TextH + 1.0f, Text, Font, TextColor);
+
+			DrawContext.DrawText(X + 2.0f, 10.0f, Text, Font, TopTextColor);
+		}
+		DrawContext.LayerId++;
+	}
+}
+
+////////////////////////////////////////////////////////////////////////////////////////////////////
+
+FReply SPacketSizesView::OnMouseButtonDown(const FGeometry& MyGeometry, const FPointerEvent& MouseEvent)
+{
+	FReply Reply = FReply::Unhandled();
+
+	MousePositionOnButtonDown = MyGeometry.AbsoluteToLocal(MouseEvent.GetScreenSpacePosition());
+	ViewportPosXOnButtonDown = Viewport.GetHorizontalAxisViewport().GetPos();
+	ViewportPosYOnButtonDown = Viewport.GetVerticalAxisViewport().GetPos();
+
+	if (MouseEvent.GetEffectingButton() == EKeys::LeftMouseButton)
+	{
+		bIsLMB_Pressed = true;
+
+		// Capture mouse.
+		Reply = FReply::Handled().CaptureMouse(SharedThis(this));
+	}
+	else if (MouseEvent.GetEffectingButton() == EKeys::RightMouseButton)
+	{
+		bIsRMB_Pressed = true;
+
+		// Capture mouse, so we can scroll outside this widget.
+		Reply = FReply::Handled().CaptureMouse(SharedThis(this));
+	}
+
+	return Reply;
+}
+
+////////////////////////////////////////////////////////////////////////////////////////////////////
+
+FReply SPacketSizesView::OnMouseButtonUp(const FGeometry& MyGeometry, const FPointerEvent& MouseEvent)
+{
+	FReply Reply = FReply::Unhandled();
+
+	MousePositionOnButtonUp = MyGeometry.AbsoluteToLocal(MouseEvent.GetScreenSpacePosition());
+
+	const bool bIsValidForMouseClick = MousePositionOnButtonUp.Equals(MousePositionOnButtonDown, MOUSE_SNAP_DISTANCE);
+
+	if (MouseEvent.GetEffectingButton() == EKeys::LeftMouseButton)
+	{
+		if (bIsLMB_Pressed)
+		{
+			if (bIsScrolling)
+			{
+				bIsScrolling = false;
+				CursorType = ECursorType::Default;
+			}
+			else if (bIsValidForMouseClick)
+			{
+				SelectSampleAtMousePosition(MousePositionOnButtonUp.X, MousePositionOnButtonUp.Y);
+			}
+
+			bIsLMB_Pressed = false;
+
+			// Release the mouse.
+			Reply = FReply::Handled().ReleaseMouseCapture();
+		}
+	}
+	else if (MouseEvent.GetEffectingButton() == EKeys::RightMouseButton)
+	{
+		if (bIsRMB_Pressed)
+		{
+			if (bIsScrolling)
+			{
+				bIsScrolling = false;
+				CursorType = ECursorType::Default;
+			}
+			else if (bIsValidForMouseClick)
+			{
+				ShowContextMenu(MouseEvent);
+			}
+
+			bIsRMB_Pressed = false;
+
+			// Release mouse as we no longer scroll.
+			Reply = FReply::Handled().ReleaseMouseCapture();
+		}
+	}
+
+	return Reply;
+}
+
+////////////////////////////////////////////////////////////////////////////////////////////////////
+
+FReply SPacketSizesView::OnMouseMove(const FGeometry& MyGeometry, const FPointerEvent& MouseEvent)
+{
+	FReply Reply = FReply::Unhandled();
+
+	MousePosition = MyGeometry.AbsoluteToLocal(MouseEvent.GetScreenSpacePosition());
+
+	if (MouseEvent.IsMouseButtonDown(EKeys::LeftMouseButton))
+	{
+		if (HasMouseCapture() && !MouseEvent.GetCursorDelta().IsZero())
+		{
+			if (!bIsScrolling)
+			{
+				bIsScrolling = true;
+				CursorType = ECursorType::Hand;
+
+				HoveredSample.Reset();
+			}
+
+			FIndexAxisViewport& ViewportX = Viewport.GetHorizontalAxisViewport();
+			const float PosX = ViewportPosXOnButtonDown + (MousePositionOnButtonDown.X - MousePosition.X);
+			ViewportX.ScrollAtIndex(ViewportX.GetIndexAtPos(PosX)); // align pos with at index
+			UpdateHorizontalScrollBar();
+			bIsStateDirty = true;
+
+			Reply = FReply::Handled();
+		}
+	}
+	else if (MouseEvent.IsMouseButtonDown(EKeys::RightMouseButton))
+	{
+		if (HasMouseCapture() && !MouseEvent.GetCursorDelta().IsZero())
+		{
+			if (!bIsScrolling)
+			{
+				bIsScrolling = true;
+				CursorType = ECursorType::Hand;
+
+				HoveredSample.Reset();
+			}
+
+			FIndexAxisViewport& ViewportX = Viewport.GetHorizontalAxisViewport();
+			const float PosX = ViewportPosXOnButtonDown + (MousePositionOnButtonDown.X - MousePosition.X);
+			ViewportX.ScrollAtIndex(ViewportX.GetIndexAtPos(PosX)); // align pos with at index
+			UpdateHorizontalScrollBar();
+			bIsStateDirty = true;
+
+			Reply = FReply::Handled();
+		}
+	}
+	else
+	{
+		if (!HoveredSample.IsValid())
+		{
+			TooltipOpacity = 0.0f;
+		}
+		HoveredSample = GetSampleAtMousePosition(MousePosition.X, MousePosition.Y);
+		if (!HoveredSample.IsValid())
+		{
+			HoveredSample = GetSampleAtMousePosition(MousePosition.X - 1.0f, MousePosition.Y);
+		}
+		if (!HoveredSample.IsValid())
+		{
+			HoveredSample = GetSampleAtMousePosition(MousePosition.X + 1.0f, MousePosition.Y);
+		}
+	}
+
+	return Reply;
+}
+
+////////////////////////////////////////////////////////////////////////////////////////////////////
+
+void SPacketSizesView::OnMouseEnter(const FGeometry& MyGeometry, const FPointerEvent& MouseEvent)
+{
+}
+
+////////////////////////////////////////////////////////////////////////////////////////////////////
+
+void SPacketSizesView::OnMouseLeave(const FPointerEvent& MouseEvent)
+{
+	if (!HasMouseCapture())
+	{
+		bIsLMB_Pressed = false;
+		bIsRMB_Pressed = false;
+
+		HoveredSample.Reset();
+
+		CursorType = ECursorType::Default;
+	}
+}
+
+////////////////////////////////////////////////////////////////////////////////////////////////////
+
+FReply SPacketSizesView::OnMouseWheel(const FGeometry& MyGeometry, const FPointerEvent& MouseEvent)
+{
+	MousePosition = MyGeometry.AbsoluteToLocal(MouseEvent.GetScreenSpacePosition());
+
+	if (MouseEvent.GetModifierKeys().IsShiftDown())
+	{
+		FValueAxisViewport& ViewportY = Viewport.GetVerticalAxisViewport();
+
+		// Zoom in/out vertically.
+		const float Delta = MouseEvent.GetWheelDelta();
+		constexpr float ZoomStep = 0.25f; // as percent
+		float ScaleY;
+
+		if (Delta > 0)
+		{
+			ScaleY = ViewportY.GetScale() * FMath::Pow(1.0f + ZoomStep, Delta);
+		}
+		else
+		{
+			ScaleY = ViewportY.GetScale() * FMath::Pow(1.0f / (1.0f + ZoomStep), -Delta);
+		}
+
+		ViewportY.SetScale(ScaleY);
+		//UpdateVerticalScrollBar();
+	}
+	else //if (MouseEvent.GetModifierKeys().IsControlDown())
+	{
+		// Zoom in/out horizontally.
+		const float Delta = MouseEvent.GetWheelDelta();
+		ZoomHorizontally(Delta, MousePosition.X);
+	}
+
+	return FReply::Handled();
+}
+
+////////////////////////////////////////////////////////////////////////////////////////////////////
+
+void SPacketSizesView::ZoomHorizontally(const float Delta, const float X)
+{
+	FIndexAxisViewport& ViewportX = Viewport.GetHorizontalAxisViewport();
+	ViewportX.RelativeZoomWithFixedOffset(Delta, X);
+	ViewportX.ScrollAtIndex(ViewportX.GetIndexAtPos(ViewportX.GetPos())); // align pos with at index
+	UpdateHorizontalScrollBar();
+	bIsStateDirty = true;
+}
+
+////////////////////////////////////////////////////////////////////////////////////////////////////
+
+FReply SPacketSizesView::OnMouseButtonDoubleClick(const FGeometry& MyGeometry, const FPointerEvent& MouseEvent)
+{
+	return FReply::Unhandled();
+}
+
+////////////////////////////////////////////////////////////////////////////////////////////////////
+
+FCursorReply SPacketSizesView::OnCursorQuery(const FGeometry& MyGeometry, const FPointerEvent& CursorEvent) const
+{
+	FCursorReply CursorReply = FCursorReply::Unhandled();
+
+	if (CursorType == ECursorType::Arrow)
+	{
+		CursorReply = FCursorReply::Cursor(EMouseCursor::ResizeLeftRight);
+	}
+	else if (CursorType == ECursorType::Hand)
+	{
+		CursorReply = FCursorReply::Cursor(EMouseCursor::GrabHand);
+	}
+
+	return CursorReply;
+}
+
+////////////////////////////////////////////////////////////////////////////////////////////////////
+
+void SPacketSizesView::ShowContextMenu(const FPointerEvent& MouseEvent)
+{
+	const bool bShouldCloseWindowAfterMenuSelection = true;
+	FMenuBuilder MenuBuilder(bShouldCloseWindowAfterMenuSelection, NULL);
+
+	MenuBuilder.BeginSection("Misc");
+	{
+		FUIAction Action_AutoZoom
+		(
+			FExecuteAction::CreateSP(this, &SPacketSizesView::ContextMenu_AutoZoom_Execute),
+			FCanExecuteAction::CreateSP(this, &SPacketSizesView::ContextMenu_AutoZoom_CanExecute),
+			FIsActionChecked::CreateSP(this, &SPacketSizesView::ContextMenu_AutoZoom_IsChecked)
+		);
+		MenuBuilder.AddMenuEntry
+		(
+			LOCTEXT("ContextMenu_AutoZoom", "Auto Zoom"),
+			LOCTEXT("ContextMenu_AutoZoom_Desc", "Enable auto zoom. Makes entire graph series to fit into view."),
+			FSlateIcon(),
+			Action_AutoZoom,
+			NAME_None,
+			EUserInterfaceActionType::ToggleButton
+		);
+	}
+	MenuBuilder.EndSection();
+
+	TSharedRef<SWidget> MenuWidget = MenuBuilder.MakeWidget();
+
+	FWidgetPath EventPath = MouseEvent.GetEventPath() != nullptr ? *MouseEvent.GetEventPath() : FWidgetPath();
+	const FVector2D ScreenSpacePosition = MouseEvent.GetScreenSpacePosition();
+	FSlateApplication::Get().PushMenu(SharedThis(this), EventPath, MenuWidget, ScreenSpacePosition, FPopupTransitionEffect::ContextMenu);
+}
+
+////////////////////////////////////////////////////////////////////////////////////////////////////
+
+void SPacketSizesView::ContextMenu_AutoZoom_Execute()
+{
+	FIndexAxisViewport& ViewportX = Viewport.GetHorizontalAxisViewport();
+
+	bIsAutoZoomEnabled = !(bIsAutoZoomEnabled && ViewportX.GetPos() == 0.0f);
+	if (bIsAutoZoomEnabled)
+	{
+		ViewportX.ScrollAtPos(0.0f);
+
+		// Auto zoom in.
+		while (ViewportX.GetMaxPos() - ViewportX.GetMinPos() < ViewportX.GetSize())
+		{
+			ZoomHorizontally(+0.1f, 0.0f);
+			ViewportX.ScrollAtPos(0.0f);
+		}
+
+		// Auto zoom out (until entire session time range fits into view).
+		while (ViewportX.GetMaxPos() - ViewportX.GetMinPos() > ViewportX.GetSize())
+		{
+			ZoomHorizontally(-0.1f, 0.0f);
+			ViewportX.ScrollAtPos(0.0f);
+		}
+
+		UpdateHorizontalScrollBar();
+		bIsStateDirty = true;
+	}
+}
+
+////////////////////////////////////////////////////////////////////////////////////////////////////
+
+bool SPacketSizesView::ContextMenu_AutoZoom_CanExecute()
+{
+	return true;
+}
+
+////////////////////////////////////////////////////////////////////////////////////////////////////
+
+bool SPacketSizesView::ContextMenu_AutoZoom_IsChecked()
+{
+	return bIsAutoZoomEnabled && Viewport.GetHorizontalAxisViewport().GetPos() == 0.0f;
+}
+
+////////////////////////////////////////////////////////////////////////////////////////////////////
+
+void SPacketSizesView::BindCommands()
+{
+}
+
+////////////////////////////////////////////////////////////////////////////////////////////////////
+
+void SPacketSizesView::HorizontalScrollBar_OnUserScrolled(float ScrollOffset)
+{
+	Viewport.GetHorizontalAxisViewport().OnUserScrolled(HorizontalScrollBar, ScrollOffset);
+	bIsStateDirty = true;
+}
+
+////////////////////////////////////////////////////////////////////////////////////////////////////
+
+void SPacketSizesView::UpdateHorizontalScrollBar()
+{
+	Viewport.GetHorizontalAxisViewport().UpdateScrollBar(HorizontalScrollBar);
 }
 
 ////////////////////////////////////////////////////////////////////////////////////////////////////
diff --git a/Engine/Source/Developer/TraceInsights/Private/Insights/NetworkingProfiler/Widgets/SPacketSizesView.h b/Engine/Source/Developer/TraceInsights/Private/Insights/NetworkingProfiler/Widgets/SPacketSizesView.h
--- a/Engine/Source/Developer/TraceInsights/Private/Insights/NetworkingProfiler/Widgets/SPacketSizesView.h
+++ b/Engine/Source/Developer/TraceInsights/Private/Insights/NetworkingProfiler/Widgets/SPacketSizesView.h
@@ -3,12 +3,89 @@
 #pragma once
 
 #include "CoreMinimal.h"
-#include "Styling/SlateTypes.h"
+#include "Input/CursorReply.h"
+#include "Input/Reply.h"
+#include "Layout/Geometry.h"
+#include "Rendering/RenderingCommon.h"
 #include "Widgets/DeclarativeSyntaxSupport.h"
 #include "Widgets/SCompoundWidget.h"
 
+// Insights
+#include "Insights/Common/FixedCircularBuffer.h"
+#include "Insights/NetworkingProfiler/ViewModels/PacketSizesViewport.h"
+#include "Insights/NetworkingProfiler/ViewModels/PacketSizesViewHelper.h"
+
+class SScrollBar;
+
+////////////////////////////////////////////////////////////////////////////////////////////////////
+
+struct FNetworkPacketSampleRef
+{
+	TSharedPtr<FNetworkPacketSeries> Series;
+	TSharedPtr<FNetworkPacketAggregatedSample> Sample;
+
+	FNetworkPacketSampleRef()
+		: Series(), Sample()
+	{
+	}
+
+	FNetworkPacketSampleRef(TSharedPtr<FNetworkPacketSeries> InSeries, TSharedPtr<FNetworkPacketAggregatedSample> InSample)
+		: Series(InSeries), Sample(InSample)
+	{
+	}
+
+	FNetworkPacketSampleRef(const FNetworkPacketSampleRef& Other)
+		: Series(Other.Series), Sample(Other.Sample)
+	{
+	}
+
+	FNetworkPacketSampleRef& operator=(const FNetworkPacketSampleRef& Other)
+	{
+		Series = Other.Series;
+		Sample = Other.Sample;
+		return *this;
+	}
+
+	void Reset()
+	{
+		Series.Reset();
+		Sample.Reset();
+	}
+
+	bool IsValid() const
+	{
+		return Series.IsValid() && Sample.IsValid();
+	}
+
+	bool Equals(const FNetworkPacketSampleRef& Other) const
+	{
+		return Series == Other.Series
+			&& ((Sample == Other.Sample) || (Sample.IsValid() && Other.Sample.IsValid() && Sample->Equals(*Other.Sample)));
+	}
+
+	static bool AreEquals(const FNetworkPacketSampleRef& A, const FNetworkPacketSampleRef& B)
+	{
+		return A.Equals(B);
+	}
+};
+
+////////////////////////////////////////////////////////////////////////////////////////////////////
+/**
+ * Widget used to present frames data in a track.
+ */
 class SPacketSizesView : public SCompoundWidget
 {
+public:
+	/** Number of pixels. */
+	static constexpr float MOUSE_SNAP_DISTANCE = 2.0f;
+
+	enum class ECursorType
+	{
+		Default,
+		Arrow,
+		Hand,
+	};
+
 public:
 	/** Default constructor. */
 	SPacketSizesView();
@@ -16,13 +93,129 @@ public:
 	/** Virtual destructor. */
 	virtual ~SPacketSizesView();
 
-	SLATE_BEGIN_ARGS(SPacketSizesView) {}
+	/** Resets internal widget's data to the default one. */
+	void Reset();
+
+	SLATE_BEGIN_ARGS(SPacketSizesView)
+	{
+		_Clipping = EWidgetClipping::ClipToBounds;
+	}
 	SLATE_END_ARGS()
 
 	/**
 	 * Construct this widget
-	 *
-	 * @param	InArgs	The declaration data for this widget
+	 * @param InArgs The declaration data for this widget
 	 */
 	void Construct(const FArguments& InArgs);
+
+	/**
+	 * Ticks this widget. Override in derived classes, but always call the parent implementation.
+	 * @param  AllottedGeometry The space allotted for this widget
+	 * @param  InCurrentTime  Current absolute real time
+	 * @param  InDeltaTime  Real time passed since last tick
+	 */
+	virtual void Tick(const FGeometry& AllottedGeometry, const double InCurrentTime, const float InDeltaTime) override;
+
+	virtual int32 OnPaint(const FPaintArgs& Args, const FGeometry& AllottedGeometry, const FSlateRect& MyCullingRect, FSlateWindowElementList& OutDrawElements, int32 LayerId, const FWidgetStyle& InWidgetStyle, bool bParentEnabled) const override;
+
+	virtual FReply OnMouseButtonDown(const FGeometry& MyGeometry, const FPointerEvent& MouseEvent) override;
+	virtual FReply OnMouseButtonUp(const FGeometry& MyGeometry, const FPointerEvent& MouseEvent) override;
+	virtual FReply OnMouseMove(const FGeometry& MyGeometry, const FPointerEvent& MouseEvent) override;
+	virtual void OnMouseEnter(const FGeometry& MyGeometry, const FPointerEvent& MouseEvent) override;
+	virtual void OnMouseLeave(const FPointerEvent& MouseEvent) override;
+	virtual FReply OnMouseWheel(const FGeometry& MyGeometry, const FPointerEvent& MouseEvent) override;
+	virtual FReply OnMouseButtonDoubleClick(const FGeometry& MyGeometry, const FPointerEvent& MouseEvent) override;
+
+	virtual FCursorReply OnCursorQuery(const FGeometry& MyGeometry, const FPointerEvent& CursorEvent) const override;
+
+protected:
+	void UpdateState();
+
+	void DrawHorizontalAxisGrid(FDrawContext& DrawContext, const FSlateBrush* Brush, const FSlateFontInfo& Font) const;
+	void DrawVerticalAxisGrid(FDrawContext& DrawContext, const FSlateBrush* Brush, const FSlateFontInfo& Font) const;
+
+	FNetworkPacketSampleRef GetSampleAtMousePosition(float X, float Y);
+	void SelectSampleAtMousePosition(float X, float Y);
+
+	void ShowContextMenu(const FPointerEvent& MouseEvent);
+
+	void ContextMenu_AutoZoom_Execute();
+	bool ContextMenu_AutoZoom_CanExecute();
+	bool ContextMenu_AutoZoom_IsChecked();
+
+	/** Binds our UI commands to delegates. */
+	void BindCommands();
+
+	/**
+	 * Called when the user scrolls the horizontal scrollbar.
+	 * @param ScrollOffset Scroll offset as a fraction between 0 and 1.
+	 */
+	void HorizontalScrollBar_OnUserScrolled(float ScrollOffset);
+	void UpdateHorizontalScrollBar();
+
+	void ZoomHorizontally(const float Delta, const float X);
+
+protected:
+	/** The track's viewport. Encapsulates info about position and scale. */
+	FPacketSizesViewport Viewport;
+	bool bIsViewportDirty;
+
+	/** Cached info for the packet series. */
+	TSharedPtr<FNetworkPacketSeries> PacketSeries;
+	bool bIsStateDirty;
+
+	bool bIsAutoZoomEnabled;
+
+	uint64 AnalysisSyncNextTimestamp;
+
+	//////////////////////////////////////////////////
+
+	TSharedPtr<SScrollBar> HorizontalScrollBar;
+
+	//////////////////////////////////////////////////
+	// Panning and Zooming behaviors
+
+	/** The current mouse position. */
+	FVector2D MousePosition;
+
+	/** Mouse position during the call on mouse button down. */
+	FVector2D MousePositionOnButtonDown;
+	float ViewportPosXOnButtonDown;
+	float ViewportPosYOnButtonDown;
+
+	/** Mouse position during the call on mouse button up. */
+	FVector2D MousePositionOnButtonUp;
+
+	bool bIsLMB_Pressed;
+	bool bIsRMB_Pressed;
+
+	/** True, if the user is currently interactively scrolling the view by holding the right mouse button and dragging. */
+	bool bIsScrolling;
+
+	//////////////////////////////////////////////////
+	// Selection
+
+	int32 SelectionStartFrameIndex;
+	int32 SelectionEndFrameIndex;
+
+	FNetworkPacketSampleRef SelectedSample;
+	FNetworkPacketSampleRef HoveredSample;
+
+	float TooltipDesiredOpacity;
+	mutable float TooltipOpacity;
+
+	//////////////////////////////////////////////////
+	// Misc
+
+	FGeometry ThisGeometry;
+
+	/** Cursor type. */
+	ECursorType CursorType;
+
+	// Debug stats
+	int32 NumUpdatedPackets;
+	TFixedCircularBuffer<uint64, 32> UpdateDurationHistory;
+	mutable TFixedCircularBuffer<uint64, 32> DrawDurationHistory;
+	mutable TFixedCircularBuffer<uint64, 32> OnPaintDurationHistory;
+	mutable uint64 LastOnPaintTime;
 };
diff --git a/Engine/Source/Developer/TraceInsights/Private/Insights/ViewModels/FrameTrackHelper.cpp b/Engine/Source/Developer/TraceInsights/Private/Insights/ViewModels/FrameTrackHelper.cpp
--- a/Engine/Source/Developer/TraceInsights/Private/Insights/ViewModels/FrameTrackHelper.cpp
+++ b/Engine/Source/Developer/TraceInsights/Private/Insights/ViewModels/FrameTrackHelper.cpp
@@ -18,27 +18,27 @@
 #include <limits>
 
 ////////////////////////////////////////////////////////////////////////////////////////////////////
-// FFrameTrackTimelineBuilder
+// FFrameTrackSeriesBuilder
 ////////////////////////////////////////////////////////////////////////////////////////////////////
 
-FFrameTrackTimelineBuilder::FFrameTrackTimelineBuilder(FFrameTrackTimeline& InTimeline, const FFrameTrackViewport& InViewport)
-	: Timeline(InTimeline)
+FFrameTrackSeriesBuilder::FFrameTrackSeriesBuilder(FFrameTrackSeries& InSeries, const FFrameTrackViewport& InViewport)
+	: Series(InSeries)
 	, Viewport(InViewport)
 	, NumAddedFrames(0)
 {
 	SampleW = Viewport.GetSampleWidth();
 	FramesPerSample = Viewport.GetNumFramesPerSample();
-	NumSamples = FMath::Max(0, FMath::CeilToInt(Viewport.Width / SampleW));
+	NumSamples = FMath::Max(0, FMath::CeilToInt(Viewport.GetWidth() / SampleW));
 	FirstFrameIndex = Viewport.GetFirstFrameIndex();
 
-	Timeline.NumAggregatedFrames = 0;
-	Timeline.Samples.Reset();
-	Timeline.Samples.AddDefaulted(NumSamples);
+	Series.NumAggregatedFrames = 0;
+	Series.Samples.Reset();
+	Series.Samples.AddDefaulted(NumSamples);
 }
 
 ////////////////////////////////////////////////////////////////////////////////////////////////////
 
-void FFrameTrackTimelineBuilder::AddFrame(const Trace::FFrame& Frame)
+void FFrameTrackSeriesBuilder::AddFrame(const Trace::FFrame& Frame)
 {
 	NumAddedFrames++;
 
@@ -47,7 +47,7 @@ void FFrameTrackTimelineBuilder::AddFrame(const Trace::FFrame& Frame)
 	int32 SampleIndex = (FrameIndex - FirstFrameIndex) / FramesPerSample;
 	if (SampleIndex >= 0 && SampleIndex < NumSamples)
 	{
-		FFrameTrackSample& Sample = Timeline.Samples[SampleIndex];
+		FFrameTrackSample& Sample = Series.Samples[SampleIndex];
 		Sample.NumFrames++;
 
 		double Duration = Frame.EndTime - Frame.StartTime;
@@ -67,7 +67,7 @@ void FFrameTrackTimelineBuilder::AddFrame(const Trace::FFrame& Frame)
 			Sample.LargestFrameDuration = Duration;
 		}
 
-		Timeline.NumAggregatedFrames++;
+		Series.NumAggregatedFrames++;
 	}
 }
 
@@ -93,10 +93,13 @@ void FFrameTrackDrawHelper::DrawBackground() const
 	const FLinearColor ValidAreaColor(0.07f, 0.07f, 0.07f, 1.0f);
 	const FLinearColor InvalidAreaColor(0.1f, 0.07f, 0.07f, 1.0f);
 
-	const float X0 = Viewport.MinX - Viewport.PosX;
-	const float X1 = Viewport.MaxX - Viewport.PosX;
-	const float W = FMath::CeilToFloat(Viewport.Width);
-	const float H = FMath::CeilToFloat(Viewport.Height);
+	const FIndexAxisViewport& ViewportX = Viewport.GetHorizontalAxisViewport();
+
+	const float X0 = ViewportX.GetMinPos() - ViewportX.GetPos();
+	const float X1 = ViewportX.GetMaxPos() - ViewportX.GetPos();
+
+	const float W = FMath::CeilToFloat(Viewport.GetWidth());
+	const float H = FMath::CeilToFloat(Viewport.GetHeight());
 
 	if (X0 >= W || X1 <= 0.0f)
 	{
@@ -132,19 +135,19 @@ void FFrameTrackDrawHelper::DrawBackground() const
 
 ////////////////////////////////////////////////////////////////////////////////////////////////////
 
-FLinearColor FFrameTrackDrawHelper::GetColorById(int32 Id) const
+FLinearColor FFrameTrackDrawHelper::GetColorByFrameType(int32 FrameType)
 {
-	const float Alpha = 0.9;
-	switch (Id)
+	const float Alpha = 0.9f;
+	switch (FrameType)
 	{
-	case 0: // Game Frames
+	case TraceFrameType_Game:
 		return FLinearColor(0.75, 1.0, 1.0, Alpha);
 
-	case 1: // Rendering Frames
+	case TraceFrameType_Rendering:
 		return FLinearColor(1.0, 0.75, 0.75, Alpha);
 
-	case 2:
-		return FLinearColor(0.75, 0.75, 1.0, Alpha);
+	//case 2:
+	//	return FLinearColor(0.75, 0.75, 1.0, Alpha);
 
 	default:
 		return FLinearColor(1.0, 1.0, 1.0, Alpha);
@@ -153,25 +156,28 @@ FLinearColor FFrameTrackDrawHelper::GetColorById(int32 Id) const
 
 ////////////////////////////////////////////////////////////////////////////////////////////////////
 
-void FFrameTrackDrawHelper::DrawCached(const FFrameTrackTimeline& Timeline) const
+void FFrameTrackDrawHelper::DrawCached(const FFrameTrackSeries& Series) const
 {
-	if (Timeline.NumAggregatedFrames == 0)
+	if (Series.NumAggregatedFrames == 0)
 	{
 		return;
 	}
 
-	NumFrames += Timeline.NumAggregatedFrames;
+	NumFrames += Series.NumAggregatedFrames;
 
-	FLinearColor TimelineColor = GetColorById(static_cast<int32>(Timeline.Id));
+	FLinearColor SeriesColor = GetColorByFrameType(Series.FrameType);
 
 	const float SampleW = Viewport.GetSampleWidth();
-	const int32 NumSamples = Timeline.Samples.Num();
+	const int32 NumSamples = Series.Samples.Num();
+
+	const FValueAxisViewport& ViewportY = Viewport.GetVerticalAxisViewport();
 
-	const float Y1 = FMath::RoundToFloat(Viewport.GetViewportYForValue(0.0));
+	const float ViewHeight = FMath::RoundToFloat(Viewport.GetHeight());
+	const float BaselineY = FMath::RoundToFloat(ViewportY.GetOffsetForValue(0.0));
 
 	for (int32 SampleIndex = 0; SampleIndex < NumSamples; SampleIndex++)
 	{
-		const FFrameTrackSample& Sample = Timeline.Samples[SampleIndex];
+		const FFrameTrackSample& Sample = Series.Samples[SampleIndex];
 		if (Sample.NumFrames == 0)
 		{
 			continue;
@@ -180,20 +186,20 @@ void FFrameTrackDrawHelper::DrawCached(const FFrameTrackTimeline& Timeline) cons
 		NumDrawSamples++;
 
 		const float X = SampleIndex * SampleW;
-		float Y2;
+		float ValueY;
 
-		FLinearColor ColorFill = TimelineColor;
+		FLinearColor ColorFill = SeriesColor;
 
 		if (Sample.LargestFrameDuration == std::numeric_limits<double>::infinity())
 		{
-			Y2 = Viewport.Height;
+			ValueY = ViewHeight;
 			ColorFill.R = 0.0f;
 			ColorFill.G = 0.0f;
 			ColorFill.B = 0.0f;
 		}
 		else
 		{
-			Y2 = FMath::RoundToFloat(Viewport.GetViewportYForValue(Sample.LargestFrameDuration));
+			ValueY = FMath::RoundToFloat(ViewportY.GetOffsetForValue(Sample.LargestFrameDuration));
 			if (Sample.LargestFrameDuration > 1.0 / 30.0)
 			{
 				ColorFill.G *= 0.5f;
@@ -205,8 +211,8 @@ void FFrameTrackDrawHelper::DrawCached(const FFrameTrackTimeline& Timeline) cons
 			}
 		}
 
-		const float Y = FMath::RoundToFloat(Viewport.Height) - Y2;
-		const float H = Y2 - Y1;
+		const float H = ValueY - BaselineY;
+		const float Y = ViewHeight - H;
 
 		const FLinearColor ColorBorder(ColorFill.R * 0.75f, ColorFill.G * 0.75f, ColorFill.B * 0.75f, 1.0);
 
@@ -239,18 +245,22 @@ void FFrameTrackDrawHelper::DrawHoveredSample(const FFrameTrackSample& Sample) c
 	const int32 SampleIndex = (Sample.LargestFrameIndex - FirstFrameIndex) / FramesPerSample;
 	const float X = SampleIndex * SampleW;
 
-	const float Y1 = FMath::RoundToFloat(Viewport.GetViewportYForValue(0.0));
-	float Y2;
+	const FValueAxisViewport& ViewportY = Viewport.GetVerticalAxisViewport();
+
+	const float ViewHeight = FMath::RoundToFloat(Viewport.GetHeight());
+	const float BaselineY = FMath::RoundToFloat(ViewportY.GetOffsetForValue(0.0));
+
+	float ValueY;
 	if (Sample.LargestFrameDuration == std::numeric_limits<double>::infinity())
 	{
-		Y2 = Viewport.Height;
+		ValueY = ViewHeight;
 	}
 	else
 	{
-		Y2 = FMath::RoundToFloat(Viewport.GetViewportYForValue(Sample.LargestFrameDuration));
+		ValueY = FMath::RoundToFloat(ViewportY.GetOffsetForValue(Sample.LargestFrameDuration));
 	}
-	const float Y = FMath::RoundToFloat(Viewport.Height) - Y2;
-	const float H = Y2 - Y1;
+	const float H = ValueY - BaselineY;
+	const float Y = ViewHeight - H;
 
 	const FLinearColor ColorBorder(1.0f, 1.0f, 0.0f, 1.0);
 	DrawContext.DrawBox(X - 1.0f, Y - 1.0f, SampleW + 2.0f, H + 2.0f, BorderBrush, ColorBorder);
@@ -259,18 +269,18 @@ void FFrameTrackDrawHelper::DrawHoveredSample(const FFrameTrackSample& Sample) c
 
 ////////////////////////////////////////////////////////////////////////////////////////////////////
 
-void FFrameTrackDrawHelper::DrawHighlightedInterval(const FFrameTrackTimeline& Timeline, const double StartTime, const double EndTime) const
+void FFrameTrackDrawHelper::DrawHighlightedInterval(const FFrameTrackSeries& Series, const double StartTime, const double EndTime) const
 {
-	const int32 NumSamples = Timeline.Samples.Num();
+	const int32 NumSamples = Series.Samples.Num();
 
 	//TODO: binary search
 	int32 Index1 = 0;
 	int32 Index2 = NumSamples - 1;
-	while (Index1 < NumSamples && Timeline.Samples[Index1].EndTime < StartTime)
+	while (Index1 < NumSamples && Series.Samples[Index1].EndTime < StartTime)
 	{
 		Index1++;
 	}
-	while (Index2 >= Index1 && Timeline.Samples[Index2].StartTime > EndTime)
+	while (Index2 >= Index1 && Series.Samples[Index2].StartTime > EndTime)
 	{
 		Index2--;
 	}
@@ -281,21 +291,21 @@ void FFrameTrackDrawHelper::DrawHighlightedInterval(const FFrameTrackTimeline& T
 		float X1 = Index1 * SampleW;
 		float X2 = Index2 * SampleW;
 
-		constexpr float Y1 = 12.0f; // allows 12px for the horizontal scrollbar (one displayed on top of the track)
-		const float Y2 = Viewport.Height;
+		constexpr float Y1 = 0.0f; // allows 12px for the horizontal scrollbar (one displayed on top of the track)
+		const float Y2 = Viewport.GetHeight();
 		constexpr float D = 2.0f; // line thickness (for both horizontal and vertical lines)
 		constexpr float H = 10.0f; // height of corner lines
 
 		const FLinearColor Color(1.0f, 1.0f, 1.0f, 1.0f);
 
-		if (X1 >= 0.0f && X1 < Viewport.Width - 2.0f)
+		if (X1 >= 0.0f && X1 < Viewport.GetWidth() - 2.0f)
 		{
 			// Draw left side vertical lines.
 			DrawContext.DrawBox(X1 - D, Y1, D, H, WhiteBrush, Color);
 			DrawContext.DrawBox(X1 - D, Y2 - H, D, H, WhiteBrush, Color);
 		}
 
-		if (X2 >= -2.0f && X2 < Viewport.Width)
+		if (X2 >= -2.0f && X2 < Viewport.GetWidth())
 		{
 			// Draw right side vertical lines.
 			DrawContext.DrawBox(X2, Y1, D, H, WhiteBrush, Color);
@@ -306,9 +316,9 @@ void FFrameTrackDrawHelper::DrawHighlightedInterval(const FFrameTrackTimeline& T
 		{
 			X1 = 0.0f;
 		}
-		if (X2 > Viewport.Width)
+		if (X2 > Viewport.GetWidth())
 		{
-			X2 = Viewport.Width;
+			X2 = Viewport.GetWidth();
 		}
 		if (X1 < X2)
 		{
diff --git a/Engine/Source/Developer/TraceInsights/Private/Insights/ViewModels/FrameTrackHelper.h b/Engine/Source/Developer/TraceInsights/Private/Insights/ViewModels/FrameTrackHelper.h
--- a/Engine/Source/Developer/TraceInsights/Private/Insights/ViewModels/FrameTrackHelper.h
+++ b/Engine/Source/Developer/TraceInsights/Private/Insights/ViewModels/FrameTrackHelper.h
@@ -21,23 +21,13 @@ namespace Trace
 
 ////////////////////////////////////////////////////////////////////////////////////////////////////
 
-enum class EFrameType
-{
-	Unknown = -1,
-
-	Game,
-	Render
-};
-
-////////////////////////////////////////////////////////////////////////////////////////////////////
-
 struct FFrameTrackSample
 {
 	int32 NumFrames;
 	double TotalDuration; // sum of durations of all frames in this sample
 	double StartTime; // min start time of all frames in this sample
 	double EndTime; // max end time of all frames in this sample
-	double LargestFrameIndex; // index of the largest frame
+	int32 LargestFrameIndex; // index of the largest frame
 	double LargestFrameStartTime; // start time of the largest frame
 	double LargestFrameDuration; // duration of the largest frame
 
@@ -50,46 +40,73 @@ struct FFrameTrackSample
 		, LargestFrameStartTime(0.0)
 		, LargestFrameDuration(0.0)
 	{}
+
+	FFrameTrackSample(const FFrameTrackSample&) = default;
+	FFrameTrackSample& operator=(const FFrameTrackSample&) = default;
+
+	FFrameTrackSample(FFrameTrackSample&&) = default;
+	FFrameTrackSample& operator=(FFrameTrackSample&&) = default;
+
+	bool Equals(const FFrameTrackSample& Other) const
+	{
+		return NumFrames == Other.NumFrames
+			&& TotalDuration == Other.TotalDuration
+			&& StartTime == Other.StartTime
+			&& EndTime == Other.EndTime
+			&& LargestFrameIndex == Other.LargestFrameIndex
+			&& LargestFrameStartTime == Other.LargestFrameStartTime
+			&& LargestFrameDuration == Other.LargestFrameDuration;
+	}
+
+	static bool AreEquals(const FFrameTrackSample& A, const FFrameTrackSample& B)
+	{
+		return A.Equals(B);
+	}
 };
 
 ////////////////////////////////////////////////////////////////////////////////////////////////////
 
-struct FFrameTrackTimeline
+struct FFrameTrackSeries
 {
-	uint64 Id;
-	//TODO: EFrameType Type;
-
-	//TODO: double MinValue; // min value
-	//TODO: double MaxValue; // max value
-
-	//TODO: float SampleW; // width of a sample, in Slate units
-	//TODO: int32 FramesPerSample; // number of frames in a sample
-	//TODO: int32 FirstFrameIndex; // index of first frame in first sample; can be negative
-	//TODO: int32 NumSamples; // total number of samples
+	int32 FrameType;
+	bool bIsVisible;
 	int32 NumAggregatedFrames; // total number of frames aggregated in samples; ie. sum of all Sample.NumFrames
-
 	TArray<FFrameTrackSample> Samples;
+
+	FFrameTrackSeries(int32 InFrameType)
+		: FrameType(InFrameType)
+		, bIsVisible(true)
+		, NumAggregatedFrames(0)
+		, Samples()
+	{
+	}
+
+	void Reset()
+	{
+		NumAggregatedFrames = 0;
+		Samples.Reset();
+	}
 };
 
 ////////////////////////////////////////////////////////////////////////////////////////////////////
 
-class FFrameTrackTimelineBuilder
+class FFrameTrackSeriesBuilder
 {
 public:
-	FFrameTrackTimelineBuilder(FFrameTrackTimeline& InTimeline, const FFrameTrackViewport& InViewport);
+	FFrameTrackSeriesBuilder(FFrameTrackSeries& InSeries, const FFrameTrackViewport& InViewport);
 
 	/**
 	 * Non-copyable
 	 */
-	FFrameTrackTimelineBuilder(const FFrameTrackTimelineBuilder&) = delete;
-	FFrameTrackTimelineBuilder& operator=(const FFrameTrackTimelineBuilder&) = delete;
+	FFrameTrackSeriesBuilder(const FFrameTrackSeriesBuilder&) = delete;
+	FFrameTrackSeriesBuilder& operator=(const FFrameTrackSeriesBuilder&) = delete;
 
 	void AddFrame(const Trace::FFrame& Frame);
 
 	int32 GetNumAddedFrames() const { return NumAddedFrames; }
 
 private:
-	FFrameTrackTimeline& Timeline;
+	FFrameTrackSeries& Series; // series to update
 	const FFrameTrackViewport& Viewport;
 
 	float SampleW; // width of a sample, in Slate units
@@ -115,11 +132,11 @@ public:
 	FFrameTrackDrawHelper& operator=(const FFrameTrackDrawHelper&) = delete;
 
 	void DrawBackground() const;
-	void DrawCached(const FFrameTrackTimeline& Timeline) const;
+	void DrawCached(const FFrameTrackSeries& Series) const;
 	void DrawHoveredSample(const FFrameTrackSample& Sample) const;
-	void DrawHighlightedInterval(const FFrameTrackTimeline& Timeline, const double StartTime, const double EndTime) const;
+	void DrawHighlightedInterval(const FFrameTrackSeries& Series, const double StartTime, const double EndTime) const;
 
-	FLinearColor GetColorById(int32 Id) const;
+	static FLinearColor GetColorByFrameType(int32 FrameType);
 
 	int32 GetNumFrames() const { return NumFrames; }
 	int32 GetNumDrawSamples() const { return NumDrawSamples; }
diff --git a/Engine/Source/Developer/TraceInsights/Private/Insights/ViewModels/FrameTrackViewport.cpp b/Engine/Source/Developer/TraceInsights/Private/Insights/ViewModels/FrameTrackViewport.cpp
--- a/Engine/Source/Developer/TraceInsights/Private/Insights/ViewModels/FrameTrackViewport.cpp
+++ b/Engine/Source/Developer/TraceInsights/Private/Insights/ViewModels/FrameTrackViewport.cpp
@@ -1,3 +1,133 @@
 // Copyright 1998-2019 Epic Games, Inc. All Rights Reserved.
 
 #include "FrameTrackViewport.h"
+
+#include "Widgets/Layout/SScrollBar.h"
+
+////////////////////////////////////////////////////////////////////////////////////////////////////
+// FIndexAxisViewport
+////////////////////////////////////////////////////////////////////////////////////////////////////
+
+void FIndexAxisViewport::RelativeZoomWithFixedOffset(const float Delta, const float Offset)
+{
+	constexpr float ZoomStep = 0.25f; // as percent
+
+	float NewScale;
+
+	if (Delta > 0)
+	{
+		NewScale = Scale * FMath::Pow(1.0f + ZoomStep, Delta);
+	}
+	else
+	{
+		NewScale = Scale * FMath::Pow(1.0f / (1.0f + ZoomStep), -Delta);
+	}
+
+	// Snap to integer value of either: "number of indices per sample" or "number of pixels per sample".
+	if (NewScale < 1.0f)
+	{
+		// N indices/sample; 1 pixel/sample
+		if (Delta > 0)
+		{
+			NewScale = 1.0f / FMath::FloorToFloat(1.0f / NewScale);
+		}
+		else
+		{
+			NewScale = 1.0f / FMath::CeilToFloat(1.0f / NewScale);
+		}
+	}
+	else
+	{
+		// 1 index/sample; N pixels/sample
+		if (Delta > 0)
+		{
+			NewScale = FMath::CeilToFloat(NewScale);
+		}
+		else
+		{
+			NewScale = FMath::FloorToFloat(NewScale);
+		}
+	}
+
+	ZoomWithFixedOffset(NewScale, Offset);
+}
+
+////////////////////////////////////////////////////////////////////////////////////////////////////
+
+bool FIndexAxisViewport::UpdatePosWithinLimits()
+{
+	float MinPos, MaxPos;
+	if (MaxPosition - MinPosition < Size)
+	{
+		MinPos = MaxPosition - Size;
+		MaxPos = MinPosition;
+	}
+	else
+	{
+		constexpr float ExtraSizeFactor = 0.15f; // allow extra 15% on sides
+		MinPos = MinPosition - ExtraSizeFactor * Size;
+		MaxPos = MaxPosition - (1.0f - ExtraSizeFactor) * Size;
+	}
+
+	constexpr float U = 0.5f; // interpolation factor --> animation speed
+
+	float Pos = Position;
+	if (Pos < MinPos)
+	{
+		Pos = Pos * U + (1.0f - U) * MinPos;
+		if (FMath::IsNearlyEqual(Pos, MinPos, 0.5f))
+		{
+			Pos = MinPos;
+		}
+	}
+	else if (Pos > MaxPos)
+	{
+		Pos = Pos * U + (1.0f - U) * MaxPos;
+		if (FMath::IsNearlyEqual(Pos, MaxPos, 0.5f))
+		{
+			Pos = MaxPos;
+		}
+
+		if (Pos < MinPos)
+		{
+			Pos = MinPos;
+		}
+	}
+
+	if (Position != Pos)
+	{
+		//ScrollAtPos(Pos);
+		ScrollAtIndex(GetIndexAtPos(Pos)); // snap to sample
+		return true;
+	}
+
+	return false;
+}
+
+////////////////////////////////////////////////////////////////////////////////////////////////////
+
+void FIndexAxisViewport::OnUserScrolled(TSharedPtr<SScrollBar> ScrollBar, float ScrollOffset)
+{
+	const float SX = 1.0 / (MaxPosition - MinPosition);
+	const float ThumbSizeFraction = FMath::Clamp<float>(Size * SX, 0.0f, 1.0f);
+	const float OffsetFraction = FMath::Clamp<float>(ScrollOffset, 0.0f, 1.0f - ThumbSizeFraction);
+
+	const float Pos = MinPosition + OffsetFraction * (MaxPosition - MinPosition);
+	ScrollAtPos(Pos);
+
+	ScrollBar->SetState(OffsetFraction, ThumbSizeFraction);
+}
+
+////////////////////////////////////////////////////////////////////////////////////////////////////
+
+void FIndexAxisViewport::UpdateScrollBar(TSharedPtr<SScrollBar> ScrollBar) const
+{
+	const float SX = 1.0 / (MaxPosition - MinPosition);
+	const float ThumbSizeFraction = FMath::Clamp<float>(Size * SX, 0.0f, 1.0f);
+	const float ScrollOffset = (Position - MinPosition) * SX;
+	const float OffsetFraction = FMath::Clamp<float>(ScrollOffset, 0.0f, 1.0f - ThumbSizeFraction);
+
+	ScrollBar->SetState(OffsetFraction, ThumbSizeFraction);
+}
+
+////////////////////////////////////////////////////////////////////////////////////////////////////
diff --git a/Engine/Source/Developer/TraceInsights/Private/Insights/ViewModels/FrameTrackViewport.h b/Engine/Source/Developer/TraceInsights/Private/Insights/ViewModels/FrameTrackViewport.h
--- a/Engine/Source/Developer/TraceInsights/Private/Insights/ViewModels/FrameTrackViewport.h
+++ b/Engine/Source/Developer/TraceInsights/Private/Insights/ViewModels/FrameTrackViewport.h
@@ -4,225 +4,386 @@
 
 #include "CoreMinimal.h"
 
-class FFrameTrackViewport
+class SScrollBar;
+
+////////////////////////////////////////////////////////////////////////////////////////////////////
+
+class FIndexAxisViewport
 {
 private:
 	static constexpr float SLATE_UNITS_TOLERANCE = 0.1f;
 
 public:
-	FFrameTrackViewport()
+	FIndexAxisViewport()
 	{
 		Reset();
 	}
 
 	void Reset()
 	{
-		Width = 0.0f;
-		Height = 0.0f;
+		Size = 0.0f;
 
 		MinIndex = 0;
 		MaxIndex = 0;
 
-		MinX = 0.0f;
-		MaxX = 0.0f;
-		PosX = 0.0f;
-		MinScaleX = 0.0001f;
-		MaxScaleX = 16.0f;
-		ScaleX = MaxScaleX;
+		MinPosition = 0.0f;
+		MaxPosition = 0.0f;
+		Position = 0.0f;
 
-		MinValue = 0.0;
-		MaxValue = 1.0f;
-
-		MinY = 0.0f;
-		MaxY = 1.0f;
-		PosY = 0.0f;
-		MinScaleY = 0.001f;
-		MaxScaleY = 1000000.0f;
-		ScaleY = 1000.0f;
+		MinScale = 0.0001f; // 10000 indices / sample; 1 pixel / sample
+		MaxScale = 16.0f; // 1 index / sample; 16 pixels / sample
+		Scale = MaxScale;
+	}
 
-		LastOnPaintMaxValue = 0.0;
+	FString ToDebugString(const TCHAR* Prefix, const TCHAR* Unit) const
+	{
+		if (Scale >= 1.0f)
+		{
+			return FString::Printf(TEXT("%sScale = %.3f (%d pixels/%s), Pos = %.2f"),
+				Prefix, Scale, FMath::RoundToInt(Scale), Unit, Position);
+		}
+		else
+		{
+			return FString::Printf(TEXT("%sScale = %.3f (%d %ss/pixel), Pos = %.2f"),
+				Prefix, Scale, FMath::RoundToInt(1.0f / Scale), Unit, Position);
+		}
 	}
 
-	void SetSize(const float InWidth, const float InHeight)
+	float GetSize() const { return Size; }
+
+	bool SetSize(const float InSize)
 	{
-		if (!FMath::IsNearlyEqual(Width, InWidth, SLATE_UNITS_TOLERANCE) ||
-			!FMath::IsNearlyEqual(Height, InHeight, SLATE_UNITS_TOLERANCE))
+		if (!FMath::IsNearlyEqual(Size, InSize, SLATE_UNITS_TOLERANCE))
 		{
-			Width = InWidth;
-			Height = InHeight;
+			Size = InSize;
 			OnSizeChanged();
+			return true;
 		}
+		return false;
 	}
 
+	int32 GetMinIndex() const { return MinIndex; }
+	int32 GetMaxIndex() const { return MaxIndex; }
+
+	float GetMinPos() const { return MinPosition; }
+	float GetMaxPos() const { return MaxPosition; }
+	float GetPos() const { return Position; }
+
+	float GetMinScale() const { return MinScale; }
+	float GetMaxScale() const { return MaxScale; }
+	float GetScale() const { return Scale; }
+
+	float GetSampleSize() const { return FMath::Max(1.0f, FMath::RoundToFloat(Scale)); }
+	int32 GetNumIndicesPerSample() const { return FMath::Max(1, FMath::RoundToInt(1.0f / Scale)); }
+
 	void SetMinMaxIndexInterval(const int32 InMinIndex, const int32 InMaxIndex)
 	{
 		MinIndex = InMinIndex;
 		MaxIndex = InMaxIndex;
-		UpdateMinMaxX();
+		UpdateMinMax();
 	}
 
-	int32 GetIndexAtPosX(const float X) const
+	int32 GetIndexAtPos(const float Pos) const
 	{
-		return FMath::RoundToInt(X / ScaleX);
+		return FMath::RoundToInt(Pos / Scale);
 	}
 
-	int32 GetIndexAtViewportX(const float VX) const
+	int32 GetIndexAtOffset(const float Offset) const
 	{
-		return FMath::RoundToInt((PosX + VX) / ScaleX);
+		return FMath::RoundToInt((Position + Offset) / Scale);
 	}
 
-	float GetPosXForIndex(const int32 Index) const
+	float GetPosForIndex(const int32 Index) const
 	{
-		return static_cast<float>(Index) * ScaleX;
+		return static_cast<float>(Index) * Scale;
 	}
 
-	float GetViewportXForIndex(const int32 Index) const
+	float GetOffsetForIndex(const int32 Index) const
 	{
-		return static_cast<float>(Index) * ScaleX - PosX;
+		return static_cast<float>(Index) * Scale - Position;
 	}
 
-	float GetRoundedViewportXForIndex(const int32 Index) const
+	float GetRoundedOffsetForIndex(const int32 Index) const
 	{
-		return FMath::RoundToFloat(static_cast<float>(Index) * ScaleX - PosX);
+		return FMath::RoundToFloat(static_cast<float>(Index) * Scale - Position);
 	}
 
-	void ScrollAtPosX(const float InPosX)
+	void ScrollAtPos(const float Pos)
 	{
-		PosX = InPosX;
-		OnPosXChanged();
+		Position = Pos;
+		OnPositionChanged();
 	}
 
 	void ScrollAtIndex(const int32 Index)
 	{
-		PosX = static_cast<float>(Index) * ScaleX;
-		OnPosXChanged();
+		Position = static_cast<float>(Index) * Scale;
+		OnPositionChanged();
 	}
 
 	void CenterOnIndex(const int32 Index)
 	{
-		float W = FMath::Max(1.0f, ScaleX);
-		if (W > Width)
+		const float SampleSize = FMath::Max(1.0f, Scale);
+		if (SampleSize > Size)
 		{
 			ScrollAtIndex(Index);
 		}
 		else
 		{
-			float VX = (Width - W) / 2.0f;
-			PosX = static_cast<float>(Index) * ScaleX - VX;
-			OnPosXChanged();
+			const float Offset = (Size - SampleSize) / 2.0f;
+			Position = static_cast<float>(Index) * Scale - Offset;
+			OnPositionChanged();
 		}
 	}
 
 	void CenterOnIndexInterval(const int32 IntervalStartIndex, const int32 IntervalEndIndex)
 	{
-		float W = FMath::Max(1.0f, ScaleX * static_cast<float>(IntervalEndIndex - IntervalStartIndex));
-		if (W > Width)
+		const float SampleIntervalSize = FMath::Max(1.0f, Scale * static_cast<float>(IntervalEndIndex - IntervalStartIndex));
+		if (SampleIntervalSize > Size)
 		{
 			ScrollAtIndex(IntervalStartIndex);
 		}
 		else
 		{
-			float VX = (Width - W) / 2.0f;
-			PosX = static_cast<float>(IntervalStartIndex) * ScaleX - VX;
-			OnPosXChanged();
+			const float Offset = (Size - SampleIntervalSize) / 2.0f;
+			Position = static_cast<float>(IntervalStartIndex) * Scale - Offset;
+			OnPositionChanged();
 		}
 	}
 
-	void ZoomWithFixedViewportX(const float InNewScaleX, const float VX)
+	void ZoomWithFixedOffset(const float NewScale, const float Offset)
 	{
-		float LocalNewScaleX = FMath::Clamp(InNewScaleX, MinScaleX, MaxScaleX);
-		if (LocalNewScaleX != ScaleX)
+		const float LocalNewScale = FMath::Clamp(NewScale, MinScale, MaxScale);
+		if (LocalNewScale != Scale)
 		{
-			// Index at viewport X should remain the same. So we resolve equation:
-			//   (PosX + VX) / ScaleX = (NewPosX + VX) / NewScaleX
-			//   ==> NewPosX = (PosX + VX) / ScaleX * NewScaleX - VX
-			PosX = (PosX + VX) * (LocalNewScaleX / ScaleX) - VX;
-			ScaleX = LocalNewScaleX;
-
-			OnPosXChanged();
-			OnScaleXChanged();
+			// Index at Offset should remain the same. So we resolve equation:
+			//   (Position + Offset) / Scale == (NewPosition + Offset) / NewScale
+			//   ==> NewPosition = (Position + Offset) / Scale * NewScale - Offset
+			Position = (Position + Offset) * (LocalNewScale / Scale) - Offset;
+			Scale = LocalNewScale;
+
+			OnPositionChanged();
+			OnScaleChanged();
 		}
 	}
 
-	float GetSampleWidth() const { return FMath::Max(1.0f, FMath::RoundToFloat(ScaleX)); }
-	int32 GetNumFramesPerSample() const { return FMath::Max(1, FMath::RoundToInt(1.0f / ScaleX)); }
+	void RelativeZoomWithFixedOffset(const float Delta, const float Offset);
+
+	bool UpdatePosWithinLimits();
+
+	void OnUserScrolled(TSharedPtr<SScrollBar> ScrollBar, float ScrollOffset);
+	void UpdateScrollBar(TSharedPtr<SScrollBar> ScrollBar) const;
+
+private:
+	void OnSizeChanged()
+	{
+	}
+
+	void OnPositionChanged()
+	{
+	}
+
+	void OnScaleChanged()
+	{
+		UpdateMinMax();
+	}
+
+	void UpdateMinMax()
+	{
+		MinPosition = static_cast<float>(MinIndex) * Scale;
+		MaxPosition = static_cast<float>(MaxIndex) * Scale;
+
+		//float NewPosition = FMath::Clamp(Position, MinPosition, MaxPosition);
+		//if (FMath::IsNearlyEqual(NewPosition, Position, SLATE_UNITS_TOLERANCE))
+		//{
+		//	Position = NewPosition;
+		//	OnPositionChanged();
+		//}
+	}
 
-	int32 GetFirstFrameIndex() const { return GetIndexAtViewportX(0.0f); }
+private:
+	float Size; // size of viewport (ex.: the viewport's width if this is a horizontal axis), in Slate units
+
+	int32 MinIndex; // minimum index
+	int32 MaxIndex; // maximum index (exclusive)
+
+	float MinPosition; // minimum position (corresponding to MinIndex), in Slate units
+	float MaxPosition; // maximum position (corresponding to MaxIndex), in Slate units
+	float Position; // current position of the viewport, in Slate units
+
+	float MinScale; // minimum scale factor
+	float MaxScale; // maximum scale factor
+	float Scale; // current scale factor between viewport index units and Slate units: if > 1, it represents number of pixels (Slate units) for one viewport index, otherwise it represents number of viewport indices in one pixel (Slate unit)
+};
 
-	float GetPosYForValue(const double Value) const
+////////////////////////////////////////////////////////////////////////////////////////////////////
+
+class FValueAxisViewport
+{
+private:
+	static constexpr float SLATE_UNITS_TOLERANCE = 0.1f;
+
+public:
+	FValueAxisViewport()
+	{
+		Reset();
+	}
+
+	void Reset()
 	{
-		return static_cast<float>(Value) * ScaleY;
+		Size = 0.0f;
+
+		MinValue = 0.0;
+		MaxValue = 0.0;
+
+		MinPosition = 0.0f;
+		MaxPosition = 0.0f;
+		Position = 0.0f;
+
+		MinScale = 0.001;
+		MaxScale = 1000000.0;
+		Scale = 1500.0;
 	}
 
-	float GetViewportYForValue(const double Value) const
+	FString ToDebugString(const TCHAR* Prefix) const
 	{
-		return static_cast<float>(Value) * ScaleY - PosY;
+		return FString::Printf(TEXT("%sScale = %g, Pos = %.2f"), Prefix, Scale, Position);
 	}
 
-	void SetScaleY(const float InNewScaleY)
+	float GetSize() const { return Size; }
+
+	bool SetSize(const float InSize)
 	{
-		float LocalNewScaleY = FMath::Clamp(InNewScaleY, MinScaleY, MaxScaleY);
-		if (LocalNewScaleY != ScaleY)
+		if (!FMath::IsNearlyEqual(Size, InSize, SLATE_UNITS_TOLERANCE))
 		{
-			ScaleY = LocalNewScaleY;
-			OnScaleYChanged();
+			Size = InSize;
+			OnSizeChanged();
+			return true;
 		}
+		return false;
 	}
 
-private:
-	void OnSizeChanged()
+	float GetMinPos() const { return MinPosition; }
+	float GetMaxPos() const { return MaxPosition; }
+	float GetPos() const { return Position; }
+
+	double GetMinScale() const { return MinScale; }
+	double GetMaxScale() const { return MaxScale; }
+	double GetScale() const { return Scale; }
+
+	float GetPosForValue(const double Value) const
 	{
+		return static_cast<float>(Value * Scale);
 	}
 
-	void OnPosXChanged()
+	float GetOffsetForValue(const double Value) const
 	{
+		return static_cast<float>(Value * Scale) - Position;
 	}
 
-	void OnScaleXChanged()
+	double GetValueForOffset(const float Offset) const
 	{
-		UpdateMinMaxX();
+		return static_cast<double>(Position + Offset) / Scale;
 	}
 
-	void OnScaleYChanged()
+	void SetScaleLimits(const double InMinScale, const double InMaxScale)
 	{
+		MinScale = InMinScale;
+		MaxScale = InMaxScale;
 	}
 
-	void UpdateMinMaxX()
+	bool SetScale(const double InNewScale)
 	{
-		MinX = static_cast<float>(MinIndex) * ScaleX;
-		MaxX = static_cast<float>(MaxIndex) * ScaleX;
+		float LocalNewScale = FMath::Clamp(InNewScale, MinScale, MaxScale);
+		if (LocalNewScale != Scale)
+		{
+			Scale = LocalNewScale;
+			OnScaleChanged();
+			return true;
+		}
+		return false;
+	}
 
-		//float NewPosX = FMath::Clamp(PosX, MinX, MaxX);
-		//if (FMath::IsNearlyEqual(NewPosX, PosX, SLATE_UNITS_TOLERANCE))
-		//{
-		//	PosX = NewPosX;
-		//	OnPosXChanged();
-		//}
+private:
+	void OnSizeChanged()
+	{
 	}
 
-public:
-	float Width; // width of viewport, in Slate units
-	float Height; // height of viewport, in Slate units
+	void OnPositionChanged()
+	{
+	}
 
-	int32 MinIndex; // minimum index
-	int32 MaxIndex; // maximum index (exclusive)
+	void OnScaleChanged()
+	{
+		//UpdateMinMaxX();
+	}
 
-	float MinX; // minimum horizontal position, in Slate units
-	float MaxX; // maximum horizontal position, in Slate units
-	float PosX;
-	float MinScaleX; // minimum horizontal scale factor
-	float MaxScaleX; // maximum horizontal scale factor
-	float ScaleX; // scale factor between frame index and Slate units: if > 1, it represents number of Slate units for one frame index, otherwise it represents number of frames in one Slate unit
+private:
+	float Size; // size of viewport (ex.: the viewport's height if this is a vertical axis), in Slate units
 
 	double MinValue;
 	double MaxValue;
 
-	float MinY; // minimum vertical scroll position, in Slate units
-	float MaxY; // maximum vertical scroll position, in Slate units
-	float PosY; // current vertical scroll position, in Slate units; origin at bottom
-	float MinScaleY; // minimum vertical scale factor
-	float MaxScaleY; // maximum vertical scale factor
-	float ScaleY; // scale factor between Value units and Slate units
+	float MinPosition; // minimum scroll position (i.e. position corresponding to MinValue), in Slate units
+	float MaxPosition; // maximum scroll position (i.e. position corresponding to MaxValue), in Slate units
+	float Position; // current position of the viewport, in Slate units
+
+	double MinScale; // minimum scale factor
+	double MaxScale; // maximum scale factor
+	double Scale; // scale factor between Value units and Slate units
+};
+
+////////////////////////////////////////////////////////////////////////////////////////////////////
+
+class FFrameTrackViewport
+{
+private:
+	static constexpr float SLATE_UNITS_TOLERANCE = 0.1f;
+
+public:
+	FFrameTrackViewport()
+	{
+		Reset();
+	}
+
+	void Reset()
+	{
+		HorizontalAxisViewport.Reset();
+		VerticalAxisViewport.Reset();
+	}
+
+	const FIndexAxisViewport& GetHorizontalAxisViewport() const { return HorizontalAxisViewport; }
+	FIndexAxisViewport& GetHorizontalAxisViewport() { return HorizontalAxisViewport; }
+
+	const FValueAxisViewport& GetVerticalAxisViewport() const { return VerticalAxisViewport; }
+	FValueAxisViewport& GetVerticalAxisViewport() { return VerticalAxisViewport; }
+
+	float GetWidth() const { return HorizontalAxisViewport.GetSize(); }
+	float GetHeight() const { return VerticalAxisViewport.GetSize(); }
 
-	mutable double LastOnPaintMaxValue;
+	bool SetSize(const float InWidth, const float InHeight)
+	{
+		const bool bWidthChanged = HorizontalAxisViewport.SetSize(InWidth);
+		const bool bHeightChanged = VerticalAxisViewport.SetSize(InHeight);
+		if (bWidthChanged || bHeightChanged)
+		{
+			OnSizeChanged();
+			return true;
+		}
+		return false;
+	}
+
+	float GetSampleWidth() const { return HorizontalAxisViewport.GetSampleSize(); }
+	int32 GetNumFramesPerSample() const { return HorizontalAxisViewport.GetNumIndicesPerSample(); }
+	int32 GetFirstFrameIndex() const { return HorizontalAxisViewport.GetIndexAtOffset(0.0f); }
+
+private:
+	void OnSizeChanged()
+	{
+	}
+
+private:
+	FIndexAxisViewport HorizontalAxisViewport;
+	FValueAxisViewport VerticalAxisViewport;
 };
+
+////////////////////////////////////////////////////////////////////////////////////////////////////
diff --git a/Engine/Source/Developer/TraceInsights/Private/Insights/ViewModels/TimingViewDrawHelper.h b/Engine/Source/Developer/TraceInsights/Private/Insights/ViewModels/TimingViewDrawHelper.h
--- a/Engine/Source/Developer/TraceInsights/Private/Insights/ViewModels/TimingViewDrawHelper.h
+++ b/Engine/Source/Developer/TraceInsights/Private/Insights/ViewModels/TimingViewDrawHelper.h
@@ -75,11 +75,11 @@ struct FTimingViewTooltip
 class FTimingViewDrawHelper
 {
 public:
-	enum class EHighlightMode
+	enum class EHighlightMode : uint32
 	{
-		Hovered,
-		Selected,
-		SelectedAndHovered
+		Hovered = 1,
+		Selected = 2,
+		SelectedAndHovered = 3
 	};
 
 private:
diff --git a/Engine/Source/Developer/TraceInsights/Private/Insights/Widgets/SFrameTrack.cpp b/Engine/Source/Developer/TraceInsights/Private/Insights/Widgets/SFrameTrack.cpp
--- a/Engine/Source/Developer/TraceInsights/Private/Insights/Widgets/SFrameTrack.cpp
+++ b/Engine/Source/Developer/TraceInsights/Private/Insights/Widgets/SFrameTrack.cpp
@@ -8,6 +8,7 @@
 #include "Fonts/FontMeasure.h"
 #include "Fonts/SlateFontInfo.h"
 #include "Framework/Application/SlateApplication.h"
+#include "Framework/MultiBox/MultiBoxBuilder.h"
 #include "HAL/PlatformTime.h"
 #include "Rendering/DrawElements.h"
 #include "TraceServices/AnalysisService.h"
@@ -34,7 +35,6 @@
 ////////////////////////////////////////////////////////////////////////////////////////////////////
 
 SFrameTrack::SFrameTrack()
-	: HoveredSample(nullptr, nullptr)
 {
 	Reset();
 }
@@ -52,10 +52,10 @@ void SFrameTrack::Reset()
 	Viewport.Reset();
 	bIsViewportDirty = true;
 
-	CachedTimelines.Reset();
-	TimelinesOrder.Reset();
-	TimelinesOrder.Add(TraceFrameType_Game);
-	TimelinesOrder.Add(TraceFrameType_Rendering);
+	SeriesMap.Reset();
+	SeriesOrder.Reset();
+	SeriesOrder.Add(TraceFrameType_Game);
+	SeriesOrder.Add(TraceFrameType_Rendering);
 
 	bIsStateDirty = true;
 
@@ -81,12 +81,13 @@ void SFrameTrack::Reset()
 	SelectionStartFrameIndex = 0;
 	SelectionEndFrameIndex = 0;
 
-	HoveredSample.Timeline = nullptr;
-	HoveredSample.Sample = nullptr;
+	HoveredSample.Reset();
+	TooltipDesiredOpacity = 0.9f;
+	TooltipOpacity = 0.0f;
 
 	//ThisGeometry
 
-	CursorType = EFrameTrackCursor::Default;
+	CursorType = ECursorType::Default;
 
 	NumUpdatedFrames = 0;
 	UpdateDurationHistory.Reset();
@@ -138,43 +139,13 @@ void SFrameTrack::Tick(const FGeometry& AllottedGeometry, const double InCurrent
 
 	ThisGeometry = AllottedGeometry;
 
+	FIndexAxisViewport& ViewportX = Viewport.GetHorizontalAxisViewport();
+
 	if (!bIsScrolling)
 	{
-		// Elastic snap to horizontal limits (+15%).
-
-		float MinX, MaxX;
-		if (Viewport.MaxX - Viewport.MinX < Viewport.Width)
-		{
-			MinX = Viewport.MaxX - Viewport.Width;
-			MaxX = Viewport.MinX;
-		}
-		else
-		{
-			MinX = Viewport.MinX - 0.15f * Viewport.Width;
-			MaxX = Viewport.MaxX - 0.85f * Viewport.Width;
-		}
-		const float U = 0.5f;
-
-		float PosX = Viewport.PosX;
-		if (PosX < MinX)
-		{
-			PosX = PosX * U + (1.0f - U) * MinX;
-			if (FMath::IsNearlyEqual(PosX, MinX, 1.0f))
-				PosX = MinX;
-		}
-		else if (PosX > MaxX)
-		{
-			PosX = PosX * U + (1.0f - U) * MaxX;
-			if (FMath::IsNearlyEqual(PosX, MaxX, 1.0f))
-				PosX = MaxX;
-			if (PosX < MinX)
-				PosX = MinX;
-		}
-
-		if (Viewport.PosX != PosX)
+		// Elastic snap to horizontal limits.
+		if (ViewportX.UpdatePosWithinLimits())
 		{
-			//Viewport.ScrollAtPosX(PosX);
-			Viewport.ScrollAtIndex(Viewport.GetIndexAtPosX(PosX));
 			bIsStateDirty = true;
 		}
 	}
@@ -194,23 +165,22 @@ void SFrameTrack::Tick(const FGeometry& AllottedGeometry, const double InCurrent
 
 			for (int32 FrameType = 0; FrameType < TraceFrameType_Count; ++FrameType)
 			{
-				FFrameTrackTimeline& CachedTimeline = CachedTimelines.FindOrAdd(FrameType);
-				CachedTimeline.Id = static_cast<uint64>(FrameType);
+				TSharedPtr<FFrameTrackSeries> SeriesPtr = FindOrAddSeries(FrameType);
 
 				int32 NumFrames = FramesProvider.GetFrameCount(static_cast<ETraceFrameType>(FrameType));
-				if (NumFrames > Viewport.MaxIndex)
+				if (NumFrames > ViewportX.GetMaxIndex())
 				{
-					Viewport.SetMinMaxIndexInterval(0, NumFrames);
+					ViewportX.SetMinMaxIndexInterval(0, NumFrames);
 
 					if (bIsAutoZoomEnabled)
 					{
-						if (Viewport.PosX == 0.0f) // only if the view is unchanged
+						if (ViewportX.GetPos() == 0.0f) // only if the view is unchanged
 						{
 							// Auto zoom out (until entire session time range fits into view).
-							while (Viewport.MaxX > Viewport.Width)
+							while (ViewportX.GetMaxPos() - ViewportX.GetMinPos() > ViewportX.GetSize())
 							{
 								ZoomHorizontally(-0.1f, 0.0f);
-								Viewport.ScrollAtPosX(0.0f);
+								ViewportX.ScrollAtPos(0.0f);
 							}
 						}
 						else
@@ -235,16 +205,50 @@ void SFrameTrack::Tick(const FGeometry& AllottedGeometry, const double InCurrent
 
 ////////////////////////////////////////////////////////////////////////////////////////////////////
 
+TSharedPtr<FFrameTrackSeries> SFrameTrack::FindOrAddSeries(int32 FrameType)
+{
+	TSharedPtr<FFrameTrackSeries>* SeriesPtrPtr = SeriesMap.Find(FrameType);
+	if (SeriesPtrPtr)
+	{
+		ensure((**SeriesPtrPtr).FrameType == FrameType);
+		return *SeriesPtrPtr;
+	}
+	else
+	{
+		TSharedPtr<FFrameTrackSeries> SeriesPtr = MakeShareable(new FFrameTrackSeries(FrameType));
+		SeriesMap.Add(FrameType, SeriesPtr);
+		return SeriesPtr;
+	}
+}
+
+////////////////////////////////////////////////////////////////////////////////////////////////////
+
+TSharedPtr<FFrameTrackSeries> SFrameTrack::FindSeries(int32 FrameType) const
+{
+	const TSharedPtr<FFrameTrackSeries>* SeriesPtrPtr = SeriesMap.Find(FrameType);
+	if (SeriesPtrPtr)
+	{
+		ensure((**SeriesPtrPtr).FrameType == FrameType);
+		return *SeriesPtrPtr;
+	}
+	else
+	{
+		return nullptr;
+	}
+}
+
+////////////////////////////////////////////////////////////////////////////////////////////////////
+
 void SFrameTrack::UpdateState()
 {
 	FStopwatch Stopwatch;
 	Stopwatch.Start();
 
 	// Reset stats.
-	for (TPair<uint64, FFrameTrackTimeline>& KeyValuePair : CachedTimelines)
+	for (TPair<int32, TSharedPtr<FFrameTrackSeries>>& KeyValuePair : SeriesMap)
 	{
-		FFrameTrackTimeline& Timeline = KeyValuePair.Value;
-		Timeline.NumAggregatedFrames = 0;
+		TSharedPtr<FFrameTrackSeries>& SeriesPtr = KeyValuePair.Value;
+		SeriesPtr->NumAggregatedFrames = 0;
 	}
 	NumUpdatedFrames = 0;
 
@@ -255,15 +259,16 @@ void SFrameTrack::UpdateState()
 
 		const Trace::IFrameProvider& FramesProvider = Trace::ReadFrameProvider(*Session.Get());
 
+		const FIndexAxisViewport& ViewportX = Viewport.GetHorizontalAxisViewport();
+
+		const uint64 StartIndex = static_cast<uint64>(FMath::Max(0, ViewportX.GetIndexAtOffset(0.0f)));
+		const uint64 EndIndex = static_cast<uint64>(ViewportX.GetIndexAtOffset(ViewportX.GetSize()));
+
 		for (int32 FrameType = 0; FrameType < TraceFrameType_Count; ++FrameType)
 		{
-			FFrameTrackTimeline& Timeline = CachedTimelines.FindOrAdd(FrameType);
-			Timeline.Id = static_cast<uint64>(FrameType);
-
-			FFrameTrackTimelineBuilder Builder(Timeline, Viewport);
+			TSharedPtr<FFrameTrackSeries> SeriesPtr = FindOrAddSeries(FrameType);
 
-			uint64 StartIndex = static_cast<uint64>(FMath::Max(0, Viewport.GetIndexAtViewportX(0.0f)));
-			uint64 EndIndex = static_cast<uint64>(Viewport.GetIndexAtViewportX(Viewport.Width));
+			FFrameTrackSeriesBuilder Builder(*SeriesPtr, Viewport);
 
 			FramesProvider.EnumerateFrames(static_cast<ETraceFrameType>(FrameType), StartIndex, EndIndex, [&Builder](const Trace::FFrame& Frame)
 			{
@@ -280,63 +285,72 @@ void SFrameTrack::UpdateState()
 
 ////////////////////////////////////////////////////////////////////////////////////////////////////
 
-FSampleRef SFrameTrack::GetSampleAtMousePosition(float X, float Y)
+FFrameTrackSampleRef SFrameTrack::GetSampleAtMousePosition(float X, float Y)
 {
-	float SampleW = Viewport.GetSampleWidth();
-	int32 SampleIndex = FMath::FloorToInt(X / SampleW);
-	if (SampleIndex >= 0)
+	if (!bIsStateDirty)
 	{
-		const float ViewportHeight = FMath::RoundToFloat(Viewport.Height);
-		const float SampleY2 = ViewportHeight - FMath::RoundToFloat(Viewport.GetViewportYForValue(0.0));
-
-		// Search in reverse paint order.
-		for (int32 TimelineIndex = TimelinesOrder.Num() - 1; TimelineIndex >= 0; --TimelineIndex)
+		float SampleW = Viewport.GetSampleWidth();
+		int32 SampleIndex = FMath::FloorToInt(X / SampleW);
+		if (SampleIndex >= 0)
 		{
-			uint64 TimelineId = TimelinesOrder[TimelineIndex];
-
-			if ((TimelineId == TraceFrameType_Rendering && !bShowRenderingFrames) ||
-				(TimelineId == TraceFrameType_Game && !bShowGameFrames))
+			// Search in reverse paint order.
+			for (int32 SeriesIndex = SeriesOrder.Num() - 1; SeriesIndex >= 0; --SeriesIndex)
 			{
-				continue;
-			}
+				int32 FrameType = SeriesOrder[SeriesIndex];
 
-			if (CachedTimelines.Contains(TimelineId))
-			{
-				const FFrameTrackTimeline& Timeline = CachedTimelines[TimelineId];
-				ensure(Timeline.Id == TimelineId);
+				if ((FrameType == TraceFrameType_Rendering && !bShowRenderingFrames) ||
+					(FrameType == TraceFrameType_Game && !bShowGameFrames))
+				{
+					continue;
+				}
 
-				if (Timeline.NumAggregatedFrames > 0 &&
-					SampleIndex < Timeline.Samples.Num())
+				const TSharedPtr<FFrameTrackSeries> SeriesPtr = FindSeries(FrameType);
+				if (SeriesPtr.IsValid())
 				{
-					const FFrameTrackSample& Sample = Timeline.Samples[SampleIndex];
-					if (Sample.NumFrames > 0)
+					if (SeriesPtr->NumAggregatedFrames > 0 &&
+						SampleIndex < SeriesPtr->Samples.Num())
 					{
-						float SampleY1;
-						if (Sample.LargestFrameDuration == std::numeric_limits<double>::infinity())
-						{
-							SampleY1 = 0.0;
-						}
-						else
-						{
-							SampleY1 = ViewportHeight - FMath::RoundToFloat(Viewport.GetViewportYForValue(Sample.LargestFrameDuration));
-						}
-						if (Y >= SampleY1 && Y <= SampleY2)
+						const FFrameTrackSample& Sample = SeriesPtr->Samples[SampleIndex];
+						if (Sample.NumFrames > 0)
 						{
-							return FSampleRef(&Timeline, &Sample);
+							const FValueAxisViewport& ViewportY = Viewport.GetVerticalAxisViewport();
+
+							const float ViewHeight = FMath::RoundToFloat(Viewport.GetHeight());
+							const float BaselineY = FMath::RoundToFloat(ViewportY.GetOffsetForValue(0.0));
+
+							float ValueY;
+							if (Sample.LargestFrameDuration == std::numeric_limits<double>::infinity())
+							{
+								ValueY = ViewHeight;
+							}
+							else
+							{
+								ValueY = FMath::RoundToFloat(ViewportY.GetOffsetForValue(Sample.LargestFrameDuration));
+							}
+
+							constexpr float ToleranceY = 3.0f; // [pixels]
+
+							const float BottomY = FMath::Min(ViewHeight, ViewHeight - BaselineY + ToleranceY);
+							const float TopY = FMath::Max(0.0f, ViewHeight - ValueY - ToleranceY);
+
+							if (Y >= TopY && Y < BottomY)
+							{
+								return FFrameTrackSampleRef(SeriesPtr, MakeShareable(new FFrameTrackSample(Sample)));
+							}
 						}
 					}
 				}
 			}
 		}
 	}
-	return FSampleRef(nullptr, nullptr);
+	return FFrameTrackSampleRef(nullptr, nullptr);
 }
 
 ////////////////////////////////////////////////////////////////////////////////////////////////////
 
 void SFrameTrack::SelectFrameAtMousePosition(float X, float Y)
 {
-	FSampleRef SampleRef = GetSampleAtMousePosition(X, Y);
+	FFrameTrackSampleRef SampleRef = GetSampleAtMousePosition(X, Y);
 	if (!SampleRef.IsValid())
 	{
 		SampleRef = GetSampleAtMousePosition(X - 1.0f, Y);
@@ -366,6 +380,18 @@ void SFrameTrack::SelectFrameAtMousePosition(float X, float Y)
 
 ////////////////////////////////////////////////////////////////////////////////////////////////////
 
+const TCHAR* FrameTypeToString(int32 FrameType)
+{
+	switch (FrameType)
+	{
+		case TraceFrameType_Game:      return TEXT("Game");
+		case TraceFrameType_Rendering: return TEXT("Rendering");
+		default:                       return TEXT("Misc");
+	}
+}
+
+////////////////////////////////////////////////////////////////////////////////////////////////////
+
 int32 SFrameTrack::OnPaint(const FPaintArgs& Args, const FGeometry& AllottedGeometry, const FSlateRect& MyCullingRect, FSlateWindowElementList& OutDrawElements, int32 LayerId, const FWidgetStyle& InWidgetStyle, bool bParentEnabled) const
 {
 	const bool bEnabled = ShouldBeEnabled(bParentEnabled);
@@ -373,6 +399,7 @@ int32 SFrameTrack::OnPaint(const FPaintArgs& Args, const FGeometry& AllottedGeom
 	FDrawContext DrawContext(AllottedGeometry, MyCullingRect, InWidgetStyle, DrawEffects, OutDrawElements, LayerId);
 
 	const TSharedRef<FSlateFontMeasure> FontMeasureService = FSlateApplication::Get().GetRenderer()->GetFontMeasureService();
+	FSlateFontInfo SummaryFont = FCoreStyle::GetDefaultFontStyle("Regular", 8);
 
 	const FSlateBrush* WhiteBrush = FCoreStyle::Get().GetBrush("WhiteBrush");
 
@@ -389,52 +416,63 @@ int32 SFrameTrack::OnPaint(const FPaintArgs& Args, const FGeometry& AllottedGeom
 
 		Helper.DrawBackground();
 
-		// Draw frames, for each visible timeline.
-		for (int32 TimelineIndex = 0; TimelineIndex < TimelinesOrder.Num(); ++TimelineIndex)
+		DrawHorizontalAxisGrid(DrawContext, WhiteBrush, SummaryFont);
+
+		// Draw frames, for each visible Series.
+		for (int32 SeriesIndex = 0; SeriesIndex < SeriesOrder.Num(); ++SeriesIndex)
 		{
-			uint64 TimelineId = TimelinesOrder[TimelineIndex];
+			int32 FrameType = SeriesOrder[SeriesIndex];
 
-			if ((TimelineId == TraceFrameType_Rendering && !bShowRenderingFrames) ||
-				(TimelineId == TraceFrameType_Game && !bShowGameFrames))
+			if ((FrameType == TraceFrameType_Rendering && !bShowRenderingFrames) ||
+				(FrameType == TraceFrameType_Game && !bShowGameFrames))
 			{
 				continue;
 			}
 
-			if (CachedTimelines.Contains(TimelineId))
+			const TSharedPtr<FFrameTrackSeries> SeriesPtr = FindSeries(FrameType);
+			if (SeriesPtr.IsValid())
 			{
-				const FFrameTrackTimeline& Timeline = CachedTimelines[TimelineId];
-				ensure(Timeline.Id == TimelineId);
-
-				Helper.DrawCached(Timeline);
+				Helper.DrawCached(*SeriesPtr);
 			}
 		}
 
 		NumDrawSamples = Helper.GetNumDrawSamples();
 
-		// Highlight the mouse hovered sample.
+		DrawVerticalAxisGrid(DrawContext, WhiteBrush, SummaryFont);
+
+		// Highlight the mouse hovered sample + tooltip.
 		if (HoveredSample.IsValid())
 		{
 			Helper.DrawHoveredSample(*HoveredSample.Sample);
 
-			FString Text = FString::Format(TEXT("{0} frame {1} ({2})"),
+			if (TooltipOpacity < TooltipDesiredOpacity)
+			{
+				TooltipOpacity = TooltipOpacity * 0.9f + TooltipDesiredOpacity * 0.1f;
+			}
+			else
+			{
+				TooltipOpacity = TooltipDesiredOpacity;
+			}
+
+			const FString Text = FString::Format(TEXT("{0} frame {1} ({2})"),
 				{
-					(HoveredSample.Timeline->Id == TraceFrameType_Game) ? TEXT("Game") :
-					(HoveredSample.Timeline->Id == TraceFrameType_Rendering) ? TEXT("Rendering") : TEXT("Misc"),
+					::FrameTypeToString(HoveredSample.Series->FrameType),
 					FText::AsNumber(HoveredSample.Sample->LargestFrameIndex).ToString(),
 					TimeUtils::FormatTimeAuto(HoveredSample.Sample->LargestFrameDuration)
 				});
 
-			FSlateFontInfo SummaryFont = FCoreStyle::GetDefaultFontStyle("Regular", 8);
 			FVector2D TextSize = FontMeasureService->Measure(Text, SummaryFont);
 
 			const float DX = 2.0f;
 			const float W2 = TextSize.X / 2 + DX;
 
-			float X1 = Viewport.GetViewportXForIndex(HoveredSample.Sample->LargestFrameIndex);
+			const FIndexAxisViewport& ViewportX = Viewport.GetHorizontalAxisViewport();
+
+			float X1 = ViewportX.GetOffsetForIndex(HoveredSample.Sample->LargestFrameIndex);
 			float CX = X1 + FMath::RoundToFloat(Viewport.GetSampleWidth() / 2);
-			if (CX + W2 > Viewport.Width)
+			if (CX + W2 > ViewportX.GetSize())
 			{
-				CX = FMath::RoundToFloat(Viewport.Width - W2);
+				CX = FMath::RoundToFloat(ViewportX.GetSize() - W2);
 			}
 			if (CX - W2 < 0)
 			{
@@ -443,8 +481,10 @@ int32 SFrameTrack::OnPaint(const FPaintArgs& Args, const FGeometry& AllottedGeom
 
 			const float Y = 10.0f;
 			const float H = 14.0f;
-			DrawContext.DrawBox(CX - W2, Y, 2 * W2, H, WhiteBrush, FLinearColor(0.7, 0.7, 0.7, 0.9));
-			DrawContext.DrawText(CX - W2 + DX, Y + 1.0f, Text, SummaryFont, FLinearColor(0.0, 0.0, 0.0, 0.9));
+			DrawContext.DrawBox(CX - W2, Y, 2 * W2, H, WhiteBrush, FLinearColor(0.7, 0.7, 0.7, TooltipOpacity));
+			DrawContext.LayerId++;
+			DrawContext.DrawText(CX - W2 + DX, Y + 1.0f, Text, SummaryFont, FLinearColor(0.0, 0.0, 0.0, TooltipOpacity));
+			DrawContext.LayerId++;
 		}
 
 		TSharedPtr<STimingProfilerWindow> Window = FTimingProfilerManager::Get()->GetProfilerWindow();
@@ -453,18 +493,18 @@ int32 SFrameTrack::OnPaint(const FPaintArgs& Args, const FGeometry& AllottedGeom
 			TSharedPtr<STimingView> TimingView = Window->GetTimingView();
 			if (TimingView)
 			{
-				const FFrameTrackTimeline* TimelinePtr = nullptr;
-				for (const TPair<uint64, FFrameTrackTimeline>& KeyValuePair : CachedTimelines)
+				TSharedPtr<FFrameTrackSeries> SeriesPtr = nullptr;
+				for (const TPair<int32, TSharedPtr<FFrameTrackSeries>>& KeyValuePair : SeriesMap)
 				{
-					TimelinePtr = &KeyValuePair.Value;
-					break; // stop at first enumerated timeline
+					SeriesPtr = KeyValuePair.Value;
+					break; // stop at first enumerated Series
 				}
-				if (TimelinePtr)
+				if (SeriesPtr.IsValid())
 				{
 					// Highlight the area corresponding to viewport of Timing View.
 					const double StartTime = TimingView->GetViewport().StartTime;
 					const double EndTime = TimingView->GetViewport().EndTime;
-					Helper.DrawHighlightedInterval(*TimelinePtr, StartTime, EndTime);
+					Helper.DrawHighlightedInterval(*SeriesPtr, StartTime, EndTime);
 				}
 			}
 		}
@@ -478,11 +518,10 @@ int32 SFrameTrack::OnPaint(const FPaintArgs& Args, const FGeometry& AllottedGeom
 	const bool bShouldDisplayDebugInfo = FInsightsManager::Get()->IsDebugInfoEnabled();
 	if (bShouldDisplayDebugInfo)
 	{
-		FSlateFontInfo SummaryFont = FCoreStyle::GetDefaultFontStyle("Regular", 8);
 		const float MaxFontCharHeight = FontMeasureService->Measure(TEXT("!"), SummaryFont).Y;
 		const float DbgDY = MaxFontCharHeight;
 
-		const float DbgW = 240.0f;
+		const float DbgW = 280.0f;
 		const float DbgH = DbgDY * 4 + 3.0f;
 		const float DbgX = ViewWidth - DbgW - 20.0f;
 		float DbgY = 7.0f;
@@ -553,11 +592,7 @@ int32 SFrameTrack::OnPaint(const FPaintArgs& Args, const FGeometry& AllottedGeom
 			OutDrawElements,
 			LayerId,
 			MAKE_PAINT_GEOMETRY_PT(AllottedGeometry, DbgX, DbgY),
-			FString::Printf(TEXT("SX: %.3f (%d %s), X: %.2f"),
-				Viewport.ScaleX,
-				(Viewport.ScaleX >= 1.0f) ? FMath::RoundToInt(Viewport.ScaleX) : FMath::RoundToInt(1.0f / Viewport.ScaleX),
-				(Viewport.ScaleX >= 1.0f) ? TEXT("px/frame") : TEXT("frames/px"),
-				Viewport.PosX),
+			Viewport.GetHorizontalAxisViewport().ToDebugString(TEXT("VX: "), TEXT("frame")),
 			SummaryFont,
 			DrawEffects,
 			DbgTextColor
@@ -570,7 +605,7 @@ int32 SFrameTrack::OnPaint(const FPaintArgs& Args, const FGeometry& AllottedGeom
 			OutDrawElements,
 			LayerId,
 			MAKE_PAINT_GEOMETRY_PT(AllottedGeometry, DbgX, DbgY),
-			FString::Printf(TEXT("SY: %g, Y: %.2f"), Viewport.ScaleY, Viewport.PosY),
+			Viewport.GetVerticalAxisViewport().ToDebugString(TEXT("VY: ")),
 			SummaryFont,
 			DrawEffects,
 			DbgTextColor
@@ -583,18 +618,177 @@ int32 SFrameTrack::OnPaint(const FPaintArgs& Args, const FGeometry& AllottedGeom
 
 ////////////////////////////////////////////////////////////////////////////////////////////////////
 
-FReply SFrameTrack::OnMouseButtonDown(const FGeometry& MyGeometry, const FPointerEvent& MouseEvent)
+void SFrameTrack::DrawVerticalAxisGrid(FDrawContext& DrawContext, const FSlateBrush* Brush, const FSlateFontInfo& Font) const
 {
-	FReply Reply = FReply::Unhandled();
+	const FValueAxisViewport& ViewportY = Viewport.GetVerticalAxisViewport();
+
+	const float RoundedViewHeight = FMath::RoundToFloat(ViewportY.GetSize());
 
-	if (!IsReady())
+	constexpr float MinDY = 24.0f; // min vertical distance between horizontal grid lines
+
+	const double TopValue = ViewportY.GetValueForOffset(RoundedViewHeight);
+	const double GridValue = ViewportY.GetValueForOffset(MinDY);
+	const double BottomValue = ViewportY.GetValueForOffset(0.0f);
+	const double Delta = GridValue - BottomValue;
+
+	if (Delta > 0.0)
 	{
-		return Reply;
+		constexpr double TimeUnits = TimeUtils::Milisecond;
+		const int64 DeltaTicks = static_cast<int64>(Delta / TimeUnits);
+
+		// Compute rounding based on magnitude of visible range of values (Delta).
+		int64 Power10 = 1;
+		int64 Delta10 = DeltaTicks;
+		while (Delta10 > 0)
+		{
+			Delta10 /= 10;
+			Power10 *= 10;
+		}
+		if (Power10 >= 100)
+		{
+			Power10 /= 100;
+		}
+		else
+		{
+			Power10 = 1;
+		}
+
+		const double GridTicks = static_cast<double>(((DeltaTicks + Power10) / Power10) * Power10); // next value divisible with a multiple of 10
+		const double Grid = GridTicks * TimeUnits;
+
+		const double StartValue = FMath::GridSnap(BottomValue, Grid);
+
+		const FLinearColor GridColor(0.0f, 0.0f, 0.0f, 0.1f);
+		const FLinearColor TextBgColor(0.05f, 0.05f, 0.05f, 1.0f);
+		//const FLinearColor TextColor(1.0f, 1.0f, 1.0f, 1.0f);
+
+		const float ViewWidth = Viewport.GetWidth();
+
+		const TSharedRef<FSlateFontMeasure> FontMeasureService = FSlateApplication::Get().GetRenderer()->GetFontMeasureService();
+
+		for (double Value = StartValue; Value < TopValue; Value += Grid)
+		{
+			FLinearColor TextColor;
+			constexpr double Time60fps = 1.0 / 60.0;
+			constexpr double Time30fps = 1.0 / 30.0;
+			if (Value <= Time60fps)
+			{
+				TextColor = FLinearColor(0.5f, 1.0f, 0.5f, 1.0f);
+			}
+			else if (Value <= Time30fps)
+			{
+				//TextColor = FLinearColor(1.0f, 1.0f, 0.5f, 1.0f);
+				const float U = (Value - Time60fps) * 0.5 / (Time30fps - Time60fps);
+				TextColor = FLinearColor(0.5 + U, 1.0f - U, 0.5f, 1.0f);
+			}
+			else
+			{
+				TextColor = FLinearColor(1.0f, 0.5f, 0.5f, 1.0f);
+			}
+
+			const float Y = RoundedViewHeight - FMath::RoundToFloat(ViewportY.GetOffsetForValue(Value));
+
+			// Draw horizontal grid line.
+			DrawContext.DrawBox(0, Y, ViewWidth, 1, Brush, GridColor);
+
+			const FString Text = (Value == 0.0) ? TEXT("0") : FString::Printf(TEXT("%s (%.1f fps)"), *TimeUtils::FormatTimeAuto(Value), FMath::RoundToDouble(1.0 / Value));
+			const FVector2D TextSize = FontMeasureService->Measure(Text, Font);
+			constexpr float TextH = 14.0f;
+
+			// Draw background for value text.
+			DrawContext.DrawBox(ViewWidth - TextSize.X - 4.0f, Y - TextH, TextSize.X + 4.0f, TextH, Brush, TextBgColor);
+
+			// Draw value text.
+			DrawContext.DrawText(ViewWidth - TextSize.X - 2.0f, Y - TextH + 1.0f, Text, Font, TextColor);
+		}
+		DrawContext.LayerId++;
 	}
+}
+
+////////////////////////////////////////////////////////////////////////////////////////////////////
+
+void SFrameTrack::DrawHorizontalAxisGrid(FDrawContext& DrawContext, const FSlateBrush* Brush, const FSlateFontInfo& Font) const
+{
+	const FIndexAxisViewport& ViewportX = Viewport.GetHorizontalAxisViewport();
+
+	const float RoundedViewWidth = FMath::RoundToFloat(ViewportX.GetSize());
+
+	constexpr float MinDX = 100.0f; // min horizontal distance between vertical grid lines
+
+	int32 LeftIndex = ViewportX.GetIndexAtOffset(0.0f);
+	int32 GridIndex = ViewportX.GetIndexAtOffset(MinDX);
+	int32 RightIndex = ViewportX.GetIndexAtOffset(RoundedViewWidth);
+	int32 Delta = GridIndex - LeftIndex;
+
+	if (Delta > 0)
+	{
+		// Compute rounding based on magnitude of visible range of samples (Delta).
+		int32 Power10 = 1;
+		int32 Delta10 = Delta;
+		while (Delta10 > 0)
+		{
+			Delta10 /= 10;
+			Power10 *= 10;
+		}
+		if (Power10 >= 100)
+		{
+			Power10 /= 100;
+		}
+		else
+		{
+			Power10 = 1;
+		}
+
+		const int32 Grid = ((Delta + Power10) / Power10) * Power10; // next value divisible with a multiple of 10
+
+		const FLinearColor GridColor(0.0f, 0.0f, 0.0f, 0.1f);
+		//const FLinearColor TextBgColor(0.05f, 0.05f, 0.05f, 1.0f);
+		//const FLinearColor TextColor(1.0f, 1.0f, 1.0f, 1.0f);
+		const FLinearColor TopTextColor(1.0f, 1.0f, 1.0f, 0.7f);
+
+		// Skip grid lines for negative indices.
+		double StartIndex = ((LeftIndex + Grid - 1) / Grid) * Grid;
+		while (StartIndex < 0)
+		{
+			StartIndex += Grid;
+		}
+
+		const float ViewHeight = Viewport.GetHeight();
+
+		//const TSharedRef<FSlateFontMeasure> FontMeasureService = FSlateApplication::Get().GetRenderer()->GetFontMeasureService();
+
+		for (int32 Index = StartIndex; Index < RightIndex; Index += Grid)
+		{
+			const float X = FMath::RoundToFloat(ViewportX.GetOffsetForIndex(Index));
+
+			// Draw vertical grid line.
+			DrawContext.DrawBox(X, 0, 1, ViewHeight, Brush, GridColor);
+
+			const FString Text = FText::AsNumber(Index).ToString();
+			//const FVector2D TextSize = FontMeasureService->Measure(Text, Font);
+			//constexpr float TextH = 14.0f;
+
+			// Draw background for index text.
+			//DrawContext.DrawBox(X, ViewHeight - TextH, TextSize.X + 4.0f, TextH, Brush, TextBgColor);
+
+			// Draw index text.
+			//DrawContext.DrawText(X + 2.0f, ViewHeight - TextH + 1.0f, Text, Font, TextColor);
+
+			DrawContext.DrawText(X + 2.0f, 10.0f, Text, Font, TopTextColor);
+		}
+		DrawContext.LayerId++;
+	}
+}
+
+////////////////////////////////////////////////////////////////////////////////////////////////////
+
+FReply SFrameTrack::OnMouseButtonDown(const FGeometry& MyGeometry, const FPointerEvent& MouseEvent)
+{
+	FReply Reply = FReply::Unhandled();
 
 	MousePositionOnButtonDown = MyGeometry.AbsoluteToLocal(MouseEvent.GetScreenSpacePosition());
-	ViewportPosXOnButtonDown = Viewport.PosX;
-	ViewportPosYOnButtonDown = Viewport.PosY;
+	ViewportPosXOnButtonDown = Viewport.GetHorizontalAxisViewport().GetPos();
+	ViewportPosYOnButtonDown = Viewport.GetVerticalAxisViewport().GetPos();
 
 	if (MouseEvent.GetEffectingButton() == EKeys::LeftMouseButton)
 	{
@@ -620,11 +814,6 @@ FReply SFrameTrack::OnMouseButtonUp(const FGeometry& MyGeometry, const FPointerE
 {
 	FReply Reply = FReply::Unhandled();
 
-	if (!IsReady())
-	{
-		return Reply;
-	}
-
 	MousePositionOnButtonUp = MyGeometry.AbsoluteToLocal(MouseEvent.GetScreenSpacePosition());
 
 	const bool bIsValidForMouseClick = MousePositionOnButtonUp.Equals(MousePositionOnButtonDown, MOUSE_SNAP_DISTANCE);
@@ -636,7 +825,7 @@ FReply SFrameTrack::OnMouseButtonUp(const FGeometry& MyGeometry, const FPointerE
 			if (bIsScrolling)
 			{
 				bIsScrolling = false;
-				CursorType = EFrameTrackCursor::Default;
+				CursorType = ECursorType::Default;
 			}
 			else if (bIsValidForMouseClick)
 			{
@@ -656,7 +845,7 @@ FReply SFrameTrack::OnMouseButtonUp(const FGeometry& MyGeometry, const FPointerE
 			if (bIsScrolling)
 			{
 				bIsScrolling = false;
-				CursorType = EFrameTrackCursor::Default;
+				CursorType = ECursorType::Default;
 			}
 			else if (bIsValidForMouseClick)
 			{
@@ -679,11 +868,6 @@ FReply SFrameTrack::OnMouseMove(const FGeometry& MyGeometry, const FPointerEvent
 {
 	FReply Reply = FReply::Unhandled();
 
-	if (!IsReady())
-	{
-		return Reply;
-	}
-
 	MousePosition = MyGeometry.AbsoluteToLocal(MouseEvent.GetScreenSpacePosition());
 
 	if (MouseEvent.IsMouseButtonDown(EKeys::LeftMouseButton))
@@ -693,13 +877,14 @@ FReply SFrameTrack::OnMouseMove(const FGeometry& MyGeometry, const FPointerEvent
 			if (!bIsScrolling)
 			{
 				bIsScrolling = true;
-				CursorType = EFrameTrackCursor::Hand;
+				CursorType = ECursorType::Hand;
 
 				HoveredSample.Reset();
 			}
 
-			float PosX = ViewportPosXOnButtonDown + (MousePositionOnButtonDown.X - MousePosition.X);
-			Viewport.ScrollAtIndex(Viewport.GetIndexAtPosX(PosX));
+			FIndexAxisViewport& ViewportX = Viewport.GetHorizontalAxisViewport();
+			const float PosX = ViewportPosXOnButtonDown + (MousePositionOnButtonDown.X - MousePosition.X);
+			ViewportX.ScrollAtIndex(ViewportX.GetIndexAtPos(PosX)); // align pos with at index
 			UpdateHorizontalScrollBar();
 			bIsStateDirty = true;
 
@@ -713,13 +898,14 @@ FReply SFrameTrack::OnMouseMove(const FGeometry& MyGeometry, const FPointerEvent
 			if (!bIsScrolling)
 			{
 				bIsScrolling = true;
-				CursorType = EFrameTrackCursor::Hand;
+				CursorType = ECursorType::Hand;
 
 				HoveredSample.Reset();
 			}
 
-			float PosX = ViewportPosXOnButtonDown + (MousePositionOnButtonDown.X - MousePosition.X);
-			Viewport.ScrollAtIndex(Viewport.GetIndexAtPosX(PosX));
+			FIndexAxisViewport& ViewportX = Viewport.GetHorizontalAxisViewport();
+			const float PosX = ViewportPosXOnButtonDown + (MousePositionOnButtonDown.X - MousePosition.X);
+			ViewportX.ScrollAtIndex(ViewportX.GetIndexAtPos(PosX)); // align pos with at index
 			UpdateHorizontalScrollBar();
 			bIsStateDirty = true;
 
@@ -728,6 +914,10 @@ FReply SFrameTrack::OnMouseMove(const FGeometry& MyGeometry, const FPointerEvent
 	}
 	else
 	{
+		if (!HoveredSample.IsValid())
+		{
+			TooltipOpacity = 0.0f;
+		}
 		HoveredSample = GetSampleAtMousePosition(MousePosition.X, MousePosition.Y);
 		if (!HoveredSample.IsValid())
 		{
@@ -746,7 +936,6 @@ FReply SFrameTrack::OnMouseMove(const FGeometry& MyGeometry, const FPointerEvent
 
 void SFrameTrack::OnMouseEnter(const FGeometry& MyGeometry, const FPointerEvent& MouseEvent)
 {
-
 }
 
 ////////////////////////////////////////////////////////////////////////////////////////////////////
@@ -760,7 +949,7 @@ void SFrameTrack::OnMouseLeave(const FPointerEvent& MouseEvent)
 
 		HoveredSample.Reset();
 
-		CursorType = EFrameTrackCursor::Default;
+		CursorType = ECursorType::Default;
 	}
 }
 
@@ -772,6 +961,8 @@ FReply SFrameTrack::OnMouseWheel(const FGeometry& MyGeometry, const FPointerEven
 
 	if (MouseEvent.GetModifierKeys().IsShiftDown())
 	{
+		FValueAxisViewport& ViewportY = Viewport.GetVerticalAxisViewport();
+
 		// Zoom in/out vertically.
 		const float Delta = MouseEvent.GetWheelDelta();
 		constexpr float ZoomStep = 0.25f; // as percent
@@ -779,14 +970,14 @@ FReply SFrameTrack::OnMouseWheel(const FGeometry& MyGeometry, const FPointerEven
 
 		if (Delta > 0)
 		{
-			ScaleY = Viewport.ScaleY * FMath::Pow(1.0f + ZoomStep, Delta);
+			ScaleY = ViewportY.GetScale() * FMath::Pow(1.0f + ZoomStep, Delta);
 		}
 		else
 		{
-			ScaleY = Viewport.ScaleY * FMath::Pow(1.0f / (1.0f + ZoomStep), -Delta);
+			ScaleY = ViewportY.GetScale() * FMath::Pow(1.0f / (1.0f + ZoomStep), -Delta);
 		}
 
-		Viewport.SetScaleY(ScaleY);
+		ViewportY.SetScale(ScaleY);
 		//UpdateVerticalScrollBar();
 	}
 	else //if (MouseEvent.GetModifierKeys().IsControlDown())
@@ -803,46 +994,9 @@ FReply SFrameTrack::OnMouseWheel(const FGeometry& MyGeometry, const FPointerEven
 
 void SFrameTrack::ZoomHorizontally(const float Delta, const float X)
 {
-	constexpr float ZoomStep = 0.25f; // as percent
-	float ScaleX;
-
-	if (Delta > 0)
-	{
-		ScaleX = Viewport.ScaleX * FMath::Pow(1.0f + ZoomStep, Delta);
-	}
-	else
-	{
-		ScaleX = Viewport.ScaleX * FMath::Pow(1.0f / (1.0f + ZoomStep), -Delta);
-	}
-
-	// Snap to integer value of either: "number of frames per pixel (Slate unit)" or "number of pixels per frame".
-	if (ScaleX < 1.0f)
-	{
-		// frames per Slate unit
-		if (Delta > 0)
-		{
-			ScaleX = 1.0f / FMath::FloorToFloat(1.0f / ScaleX);
-		}
-		else
-		{
-			ScaleX = 1.0f / FMath::CeilToFloat(1.0f / ScaleX);
-		}
-	}
-	else
-	{
-		// Slate units per frame
-		if (Delta > 0)
-		{
-			ScaleX = FMath::CeilToFloat(ScaleX);
-		}
-		else
-		{
-			ScaleX = FMath::FloorToFloat(ScaleX);
-		}
-	}
-
-	//UE_LOG(TimingProfiler, Log, TEXT("%.2f, %.2f, %.2f"), Delta, Viewport.ScaleX, ScaleX);
-	Viewport.ZoomWithFixedViewportX(ScaleX, MousePosition.X);
+	FIndexAxisViewport& ViewportX = Viewport.GetHorizontalAxisViewport();
+	ViewportX.RelativeZoomWithFixedOffset(Delta, X);
+	ViewportX.ScrollAtIndex(ViewportX.GetIndexAtPos(ViewportX.GetPos())); // align pos with at index
 	UpdateHorizontalScrollBar();
 	bIsStateDirty = true;
 }
@@ -860,11 +1014,11 @@ FCursorReply SFrameTrack::OnCursorQuery(const FGeometry& MyGeometry, const FPoin
 {
 	FCursorReply CursorReply = FCursorReply::Unhandled();
 
-	if (CursorType == EFrameTrackCursor::Arrow)
+	if (CursorType == ECursorType::Arrow)
 	{
 		CursorReply = FCursorReply::Cursor(EMouseCursor::ResizeLeftRight);
 	}
-	else if (CursorType == EFrameTrackCursor::Hand)
+	else if (CursorType == ECursorType::Hand)
 	{
 		CursorReply = FCursorReply::Cursor(EMouseCursor::GrabHand);
 	}
@@ -996,23 +1150,25 @@ bool SFrameTrack::ContextMenu_ShowRenderingFrames_IsChecked()
 
 void SFrameTrack::ContextMenu_AutoZoom_Execute()
 {
-	bIsAutoZoomEnabled = !(bIsAutoZoomEnabled && Viewport.PosX == 0.0f);
+	FIndexAxisViewport& ViewportX = Viewport.GetHorizontalAxisViewport();
+
+	bIsAutoZoomEnabled = !(bIsAutoZoomEnabled && ViewportX.GetPos() == 0.0f);
 	if (bIsAutoZoomEnabled)
 	{
-		Viewport.ScrollAtPosX(0.0f);
+		ViewportX.ScrollAtPos(0.0f);
 
 		// Auto zoom in.
-		while (Viewport.MaxX < Viewport.Width)
+		while (ViewportX.GetMaxPos() - ViewportX.GetMinPos() < ViewportX.GetSize())
 		{
 			ZoomHorizontally(+0.1f, 0.0f);
-			Viewport.ScrollAtPosX(0.0f);
+			ViewportX.ScrollAtPos(0.0f);
 		}
 
 		// Auto zoom out (until entire session time range fits into view).
-		while (Viewport.MaxX > Viewport.Width)
+		while (ViewportX.GetMaxPos() - ViewportX.GetMinPos() > ViewportX.GetSize())
 		{
 			ZoomHorizontally(-0.1f, 0.0f);
-			Viewport.ScrollAtPosX(0.0f);
+			ViewportX.ScrollAtPos(0.0f);
 		}
 
 		UpdateHorizontalScrollBar();
@@ -1031,7 +1187,7 @@ bool SFrameTrack::ContextMenu_AutoZoom_CanExecute()
 
 bool SFrameTrack::ContextMenu_AutoZoom_IsChecked()
 {
-	return bIsAutoZoomEnabled && Viewport.PosX == 0.0f;
+	return bIsAutoZoomEnabled && Viewport.GetHorizontalAxisViewport().GetPos() == 0.0f;
 }
 
 ////////////////////////////////////////////////////////////////////////////////////////////////////
@@ -1044,27 +1200,15 @@ void SFrameTrack::BindCommands()
 
 void SFrameTrack::HorizontalScrollBar_OnUserScrolled(float ScrollOffset)
 {
-	const float SX = 1.0 / (Viewport.MaxX - Viewport.MinX);
-	const float ThumbSizeFraction = FMath::Clamp<float>(Viewport.Width * SX, 0.0f, 1.0f);
-	const float OffsetFraction = FMath::Clamp<float>(ScrollOffset, 0.0f, 1.0f - ThumbSizeFraction);
-
-	const float PosX = Viewport.MinX + OffsetFraction * (Viewport.MaxX - Viewport.MinX);
-	Viewport.ScrollAtPosX(PosX);
+	Viewport.GetHorizontalAxisViewport().OnUserScrolled(HorizontalScrollBar, ScrollOffset);
 	bIsStateDirty = true;
-
-	HorizontalScrollBar->SetState(OffsetFraction, ThumbSizeFraction);
 }
 
 ////////////////////////////////////////////////////////////////////////////////////////////////////
 
 void SFrameTrack::UpdateHorizontalScrollBar()
 {
-	const float SX = 1.0 / (Viewport.MaxX - Viewport.MinX);
-	const float ThumbSizeFraction = FMath::Clamp<float>(Viewport.Width * SX, 0.0f, 1.0f);
-	const float ScrollOffset = (Viewport.PosX - Viewport.MinX) * SX;
-	const float OffsetFraction = FMath::Clamp<float>(ScrollOffset, 0.0f, 1.0f - ThumbSizeFraction);
-
-	HorizontalScrollBar->SetState(OffsetFraction, ThumbSizeFraction);
+	Viewport.GetHorizontalAxisViewport().UpdateScrollBar(HorizontalScrollBar);
 }
 
 ////////////////////////////////////////////////////////////////////////////////////////////////////
diff --git a/Engine/Source/Developer/TraceInsights/Private/Insights/Widgets/SFrameTrack.h b/Engine/Source/Developer/TraceInsights/Private/Insights/Widgets/SFrameTrack.h
--- a/Engine/Source/Developer/TraceInsights/Private/Insights/Widgets/SFrameTrack.h
+++ b/Engine/Source/Developer/TraceInsights/Private/Insights/Widgets/SFrameTrack.h
@@ -19,18 +19,54 @@ class SScrollBar;
 
 ////////////////////////////////////////////////////////////////////////////////////////////////////
 
-struct FSampleRef
+struct FFrameTrackSampleRef
 {
-	const FFrameTrackTimeline* Timeline;
-	const FFrameTrackSample* Sample;
+	TSharedPtr<FFrameTrackSeries> Series;
+	TSharedPtr<FFrameTrackSample> Sample;
 
-	FSampleRef(const FFrameTrackTimeline* InTimeline, const FFrameTrackSample* InSample)
-		: Timeline(InTimeline), Sample(InSample)
+	FFrameTrackSampleRef()
+		: Series(), Sample()
 	{
 	}
 
-	void Reset() { Timeline = nullptr; Sample = nullptr; }
-	bool IsValid() const { return Timeline != nullptr && Sample != nullptr; }
+	FFrameTrackSampleRef(TSharedPtr<FFrameTrackSeries> InSeries, TSharedPtr<FFrameTrackSample> InSample)
+		: Series(InSeries), Sample(InSample)
+	{
+	}
+
+	FFrameTrackSampleRef(const FFrameTrackSampleRef& Other)
+		: Series(Other.Series), Sample(Other.Sample)
+	{
+	}
+
+	FFrameTrackSampleRef& operator=(const FFrameTrackSampleRef& Other)
+	{
+		Series = Other.Series;
+		Sample = Other.Sample;
+		return *this;
+	}
+
+	void Reset()
+	{
+		Series.Reset();
+		Sample.Reset();
+	}
+
+	bool IsValid() const
+	{
+		return Series.IsValid() && Sample.IsValid();
+	}
+
+	bool Equals(const FFrameTrackSampleRef& Other) const
+	{
+		return Series == Other.Series
+			&& ((Sample == Other.Sample) || (Sample.IsValid() && Other.Sample.IsValid() && Sample->Equals(*Other.Sample)));
+	}
+
+	static bool AreEquals(const FFrameTrackSampleRef& A, const FFrameTrackSampleRef& B)
+	{
+		return A.Equals(B);
+	}
 };
 
 ////////////////////////////////////////////////////////////////////////////////////////////////////
@@ -43,7 +79,7 @@ public:
 	/** Number of pixels. */
 	static constexpr float MOUSE_SNAP_DISTANCE = 2.0f;
 
-	enum class EFrameTrackCursor
+	enum class ECursorType
 	{
 		Default,
 		Arrow,
@@ -93,11 +129,14 @@ public:
 	virtual FCursorReply OnCursorQuery(const FGeometry& MyGeometry, const FPointerEvent& CursorEvent) const override;
 
 protected:
-	bool IsReady() { return true; }
-
+	TSharedPtr<FFrameTrackSeries> FindOrAddSeries(int32 FrameType);
+	TSharedPtr<FFrameTrackSeries> FindSeries(int32 FrameType) const;
 	void UpdateState();
 
-	FSampleRef GetSampleAtMousePosition(float X, float Y);
+	void DrawHorizontalAxisGrid(FDrawContext& DrawContext, const FSlateBrush* Brush, const FSlateFontInfo& Font) const;
+	void DrawVerticalAxisGrid(FDrawContext& DrawContext, const FSlateBrush* Brush, const FSlateFontInfo& Font) const;
+
+	FFrameTrackSampleRef GetSampleAtMousePosition(float X, float Y);
 	void SelectFrameAtMousePosition(float X, float Y);
 
 	void ShowContextMenu(const FPointerEvent& MouseEvent);
@@ -144,9 +183,9 @@ protected:
 	FFrameTrackViewport Viewport;
 	bool bIsViewportDirty;
 
-	/** Cached info for timelines. */
-	TMap<uint64, FFrameTrackTimeline> CachedTimelines;
-	TArray<uint64> TimelinesOrder;
+	/** Cached info for all frame series. */
+	TMap<int32, TSharedPtr<FFrameTrackSeries>> SeriesMap;
+	TArray<int32> SeriesOrder;
 
 	bool bIsStateDirty;
 
@@ -186,7 +225,10 @@ protected:
 	int32 SelectionStartFrameIndex;
 	int32 SelectionEndFrameIndex;
 
-	FSampleRef HoveredSample;
+	FFrameTrackSampleRef HoveredSample;
+
+	float TooltipDesiredOpacity;
+	mutable float TooltipOpacity;
 
 	//////////////////////////////////////////////////
 	// Misc
@@ -194,7 +236,7 @@ protected:
 	FGeometry ThisGeometry;
 
 	/** Cursor type. */
-	EFrameTrackCursor CursorType;
+	ECursorType CursorType;
 
 	// Debug stats
 	int32 NumUpdatedFrames;
diff --git a/Engine/Source/Developer/TraceInsights/Private/Insights/Widgets/SLogView.h b/Engine/Source/Developer/TraceInsights/Private/Insights/Widgets/SLogView.h
--- a/Engine/Source/Developer/TraceInsights/Private/Insights/Widgets/SLogView.h
+++ b/Engine/Source/Developer/TraceInsights/Private/Insights/Widgets/SLogView.h
@@ -17,6 +17,8 @@
 #include "Insights/ViewModels/LogFilter.h"
 #include "Insights/ViewModels/LogMessage.h"
 
+class FMenuBuilder;
+
 ////////////////////////////////////////////////////////////////////////////////////////////////////
 
 namespace LogViewColumns
diff --git a/Engine/Source/Developer/TraceInsights/Private/Insights/Widgets/STimingProfilerWindow.cpp b/Engine/Source/Developer/TraceInsights/Private/Insights/Widgets/STimingProfilerWindow.cpp
--- a/Engine/Source/Developer/TraceInsights/Private/Insights/Widgets/STimingProfilerWindow.cpp
+++ b/Engine/Source/Developer/TraceInsights/Private/Insights/Widgets/STimingProfilerWindow.cpp
@@ -4,6 +4,7 @@
 
 #include "EditorStyleSet.h"
 #include "Framework/Docking/WorkspaceItem.h"
+#include "Framework/MultiBox/MultiBoxBuilder.h"
 #include "SlateOptMacros.h"
 #include "Widgets/Docking/SDockTab.h"
 #include "Widgets/Images/SImage.h"
diff --git a/Engine/Source/Developer/TraceInsights/Private/Insights/Widgets/STimingProfilerWindow.h b/Engine/Source/Developer/TraceInsights/Private/Insights/Widgets/STimingProfilerWindow.h
--- a/Engine/Source/Developer/TraceInsights/Private/Insights/Widgets/STimingProfilerWindow.h
+++ b/Engine/Source/Developer/TraceInsights/Private/Insights/Widgets/STimingProfilerWindow.h
@@ -18,13 +18,15 @@
 ////////////////////////////////////////////////////////////////////////////////////////////////////
 
 class FActiveTimerHandle;
-class SVerticalBox;
+class FMenuBuilder;
+
 class SFrameTrack;
 class SGraphTrack;
+class SLogView;
 class SStatsView;
 class STimersView;
 class STimingView;
-class SLogView;
+class SVerticalBox;
 
 ////////////////////////////////////////////////////////////////////////////////////////////////////
 
diff --git a/Engine/Source/Developer/TraceInsights/Private/Insights/Widgets/STimingView.h b/Engine/Source/Developer/TraceInsights/Private/Insights/Widgets/STimingView.h
--- a/Engine/Source/Developer/TraceInsights/Private/Insights/Widgets/STimingView.h
+++ b/Engine/Source/Developer/TraceInsights/Private/Insights/Widgets/STimingView.h
@@ -4,7 +4,6 @@
 
 #include "CoreMinimal.h"
 #include "Delegates/Delegate.h"
-#include "Framework/MultiBox/MultiBoxBuilder.h"
 #include "Input/CursorReply.h"
 #include "Input/Reply.h"
 #include "Layout/Geometry.h"
@@ -25,6 +24,7 @@
 
 class FFileActivitySharedState;
 class FLoadingSharedState;
+class FMenuBuilder;
 class FTimingGraphTrack;
 class FTimingViewDrawHelper;
 class SScrollBar;
