VS2019 Fix

diff --git a/Engine/Source/Developer/TraceInsights/Private/Insights/Table/ViewModels/TreeNodeSorting.cpp b/Engine/Source/Developer/TraceInsights/Private/Insights/Table/ViewModels/TreeNodeSorting.cpp
--- a/Engine/Source/Developer/TraceInsights/Private/Insights/Table/ViewModels/TreeNodeSorting.cpp
+++ b/Engine/Source/Developer/TraceInsights/Private/Insights/Table/ViewModels/TreeNodeSorting.cpp
@@ -333,7 +333,7 @@ FTreeNodeSortingByCStringValue::FTreeNodeSortingByCStringValue(TSharedPtr<FTable
 		else
 		{
 			// Sort by value (ascending).
-			return FCString::Strcmp(ValueA, ValueB);
+			return FCString::Strcmp(ValueA, ValueB) < 0;
 		}
 	};
 
@@ -350,7 +350,7 @@ FTreeNodeSortingByCStringValue::FTreeNodeSortingByCStringValue(TSharedPtr<FTable
 		else
 		{
 			// Sort by value (descending).
-			return FCString::Strcmp(ValueB, ValueA);
+			return FCString::Strcmp(ValueB, ValueA) < 0;
 		}
 	};
 }
