[Insights]
  - Frames track: Increased the highlighted interval (white corner lines) to include a frame even if it is only partially visible in Timing view.
  - Frames track: Fixed horizontal grid lines to use only predefined values.
  - Networking Insights: Added prototype for the Packet Breakdown widget (somehow similar with a timing view; panning and zooming; event highlight + tooltip; optimized event reduction + draw state caching). When a packet is selected in Packet Sizes view, the content of the Packet Breakdown widget is updated (mock data).
  - Improved performance when holding mouse over a widget (SFrameTrack, SGraphTrack, STimingView, SPacketSizesView and SPacketBreakdownView). Due to synthetic mouse move events detection of events under mouse were executed even if mouse was not moved.
  - Moved and renamed FIndexAxisViewport and FValueAxisViewport into separate files (FAxisViewportInt32 and FAxisViewportDouble classes).
  - Simplified code re in-place tooltips by reusing the new FTooltipDrawState class (see GraphTrack, FileActivityTimingTrack, LoadingTimingTrack, ThreadTimingTrack and SPacketBreakdownView).
  - Misc code refactoring.

diff --git a/Engine/Source/Developer/TraceInsights/Private/Insights/InsightsStyle.cpp b/Engine/Source/Developer/TraceInsights/Private/Insights/InsightsStyle.cpp
--- a/Engine/Source/Developer/TraceInsights/Private/Insights/InsightsStyle.cpp
+++ b/Engine/Source/Developer/TraceInsights/Private/Insights/InsightsStyle.cpp
@@ -39,6 +39,16 @@ TSharedRef<FSlateStyleSet> FInsightsStyle::Create()
 
 	FSlateStyleSet& Style = StyleRef.Get();
 
+	//////////////////////////////////////////////////
+
+	Style.Set("WhiteBrush", new FSlateColorBrush(FLinearColor::White));
+	Style.Set("SingleBorder", new FSlateBorderBrush(NAME_None, FMargin(1.0f)));
+	Style.Set("DoubleBorder", new FSlateBorderBrush(NAME_None, FMargin(2.0f)));
+
+	Style.Set("EventBorder", new FSlateBorderBrush(NAME_None, FMargin(1.0f)));
+	Style.Set("HoveredEventBorder", new FSlateBorderBrush(NAME_None, FMargin(2.0f)));
+	Style.Set("SelectedEventBorder", new FSlateBorderBrush(NAME_None, FMargin(2.0f)));
+
 	//////////////////////////////////////////////////
 	// Icons for major components
 
diff --git a/Engine/Source/Developer/TraceInsights/Private/Insights/NetworkingProfiler/ViewModels/PacketBreakdownViewHelper.cpp b/Engine/Source/Developer/TraceInsights/Private/Insights/NetworkingProfiler/ViewModels/PacketBreakdownViewHelper.cpp
new file mode 100644
--- /dev/null
+++ b/Engine/Source/Developer/TraceInsights/Private/Insights/NetworkingProfiler/ViewModels/PacketBreakdownViewHelper.cpp
@@ -0,0 +1,411 @@
+// Copyright 1998-2019 Epic Games, Inc. All Rights Reserved.
+
+#include "PacketBreakdownViewHelper.h"
+
+#include "Fonts/FontMeasure.h"
+#include "Framework/Application/SlateApplication.h"
+#include "Rendering/DrawElements.h"
+#include "Styling/CoreStyle.h"
+#include "TraceServices/AnalysisService.h"
+
+// Insights
+#include "Insights/Common/PaintUtils.h"
+#include "Insights/InsightsStyle.h"
+#include "Insights/NetworkingProfiler/ViewModels/PacketBreakdownViewport.h"
+#include "Insights/ViewModels/DrawHelpers.h"
+
+#include <limits>
+
+////////////////////////////////////////////////////////////////////////////////////////////////////
+// FPacketBreakdownViewDrawStateBuilder
+////////////////////////////////////////////////////////////////////////////////////////////////////
+
+FPacketBreakdownViewDrawStateBuilder::FPacketBreakdownViewDrawStateBuilder(FPacketBreakdownViewDrawState& InDrawState, const FPacketBreakdownViewport& InViewport)
+	: DrawState(InDrawState)
+	, Viewport(InViewport)
+	, MaxDepth(-1)
+	, LastEventX2()
+	, LastBox()
+	, EventFont(FCoreStyle::GetDefaultFontStyle("Regular", 8))
+{
+	DrawState.Reset();
+}
+
+////////////////////////////////////////////////////////////////////////////////////////////////////
+
+void FPacketBreakdownViewDrawStateBuilder::AddEvent(int64 EventOffset, int64 EventSize, int32 EventType, int32 Depth)
+{
+	DrawState.Events.AddUninitialized();
+	FPacketBreakdownViewDrawState::FPacketEvent& PacketEvent = DrawState.Events.Last();
+	PacketEvent.Offset = EventOffset;
+	PacketEvent.Size = EventSize;
+	PacketEvent.Type = EventType;
+	PacketEvent.Depth = Depth;
+
+	const FAxisViewportDouble& ViewportX = Viewport.GetHorizontalAxisViewport();
+
+	float EventX1 = ViewportX.GetRoundedOffsetForValue(static_cast<double>(EventOffset));
+	if (EventX1 > Viewport.GetWidth())
+	{
+		return;
+	}
+
+	//float EventX2 = EventSize == 0 ?
+	//	EventX1 + 1.0f :
+	//	ViewportX.GetScale() < 1.0f ?
+	//	ViewportX.GetRoundedOffsetForValue(static_cast<double>(EventOffset + EventSize - 1)) + 1.0f : // x2 = pos of last inclusive bit + 1px
+	//	ViewportX.GetRoundedOffsetForValue(static_cast<double>(EventOffset + EventSize));             // x2 = pos of first exclusive bit
+	float EventX2 = ViewportX.GetRoundedOffsetForValue(static_cast<double>(EventOffset + EventSize));
+	if (EventX2 < 0)
+	{
+		return;
+	}
+
+	// Timing events are displayed with minimum 1px (including empty ones).
+	if (EventX1 == EventX2)
+	{
+		EventX2 = EventX1 + 1.0f;
+	}
+
+	if (Depth > MaxDepth)
+	{
+		MaxDepth = Depth;
+	}
+
+	// Ensure we have enough slots in array. See LastBox[Depth] usage.
+	if (LastBox.Num() <= Depth)
+	{
+		LastBox.AddDefaulted(Depth + 1 - LastBox.Num());
+	}
+
+	constexpr float Y0 = 0.0f;
+	constexpr float EventH = 14.0f;
+	constexpr float EventDY = 2.0f;
+
+	const float EventY = Y0 + (EventH + EventDY) * Depth;
+	if (EventY < -EventH || EventY > Viewport.GetHeight())
+	{
+		return;
+	}
+
+	// Ensure we have enough slots in array. See LastEventX2[Depth] usage.
+	while (LastEventX2.Num() <= Depth)
+	{
+		LastEventX2.Add(-1.0f);
+	}
+
+	// Limit event width on the viewport's left side.
+	// This also makes the text to be displayed in viewport, for very long events.
+	const float MinX = -2.0f; // -2 allows event border to remain outside screen
+	if (EventX1 < MinX)
+	{
+		EventX1 = MinX;
+	}
+
+	// Limit event width on the viewport's right side.
+	const float MaxX = Viewport.GetWidth() + 2.0f; // +2 allows event border to remain outside screen
+	if (EventX2 > MaxX)
+	{
+		EventX2 = MaxX;
+	}
+
+	const float EventW = EventX2 - EventX1;
+
+	// Optimization...
+	if (EventW == 1.0f && EventX1 == LastEventX2[Depth] - 1.0f)
+	{
+		// Do no draw 1 pixel event if the last event was ended on that pixel.
+		return;
+	}
+
+	//////////////////////////////////////////////////
+	// Coloring
+
+	const uint32 Color = (static_cast<uint32>(EventType) + 0x2c2c57ed) * 0x2c2c57ed;
+	const FLinearColor EventColorFill(((Color >> 16) & 0xFF) / 255.0f, ((Color >> 8) & 0xFF) / 255.0f, (Color & 0xFF) / 255.0f, 1.0f);
+	constexpr float BorderColorFactor = 1.25f; // brighter border
+
+	//////////////////////////////////////////////////
+
+	if (EventW > 2.0f && EventH > 2.0f)
+	{
+		// Fill inside of the event box.
+		DrawState.Boxes.AddUninitialized();
+		FPacketBreakdownViewDrawState::FBox& DrawBox = DrawState.Boxes.Last();
+		DrawBox.X = EventX1 + 1.0f;
+		DrawBox.Y = EventY + 1.0f;
+		DrawBox.W = EventW - 2.0f;
+		DrawBox.H = EventH - 2.0f;
+		DrawBox.Color = EventColorFill;
+	}
+
+	// Save X2, for current depth.
+	LastEventX2[Depth] = EventX2;
+
+	// Add border around the event box.
+	if (EventW > 2.0f)
+	{
+		FBoxData& Box = LastBox[Depth];
+		if (Box.X1 < Box.X2)
+		{
+			FlushBox(Box, EventY, EventH);
+			Box.Reset();
+		}
+
+		DrawState.Borders.AddUninitialized();
+		FPacketBreakdownViewDrawState::FBox& DrawBox = DrawState.Borders.Last();
+		DrawBox.X = EventX1;
+		DrawBox.Y = EventY;
+		DrawBox.W = EventW;
+		DrawBox.H = EventH;
+		DrawBox.Color.R = EventColorFill.R * BorderColorFactor;
+		DrawBox.Color.G = EventColorFill.G * BorderColorFactor;
+		DrawBox.Color.B = EventColorFill.B * BorderColorFactor;
+		DrawBox.Color.A = 1.0f;
+	}
+	else // 1px or 2px boxes
+	{
+		FBoxData& Box = LastBox[Depth];
+
+		// Check if we can merge this box with previous one, if any.
+		// Note: We are assuming events are processed in sorted order by X1.
+		if (Color == Box.Color && // same color
+			EventX1 <= Box.X2) // overlapping or adjacent
+		{
+			// Merge it with previous box.
+			Box.X2 = EventX2;
+			DrawState.NumMergedBoxes++;
+		}
+		else
+		{
+			// Flush previous box, if any.
+			if (Box.X1 < Box.X2)
+			{
+				FlushBox(Box, EventY, EventH);
+			}
+
+			// Start new "merge box".
+			Box.X1 = EventX1;
+			Box.X2 = EventX2;
+			Box.Color = Color;
+			Box.LinearColor = FLinearColor(EventColorFill.R * BorderColorFactor, EventColorFill.G * BorderColorFactor, EventColorFill.B * BorderColorFactor, 1.0f);
+		}
+	}
+
+	// Draw the name of the event.
+	if (EventW > 8.0f && EventH > 10.0f)
+	{
+		FString Name = TEXT("NetObject") + FText::AsNumber(EventType).ToString();
+		if (EventW > Name.Len() * 2.0f + 48.0f)
+		{
+			if (EventSize == 1)
+			{
+				Name += TEXT(" (1 bit)");
+			}
+			else
+			{
+				Name += TEXT(" (");
+				Name += FText::AsNumber(EventSize).ToString();
+				Name += TEXT(" bits)");
+			}
+		}
+
+		const TSharedRef<FSlateFontMeasure> FontMeasureService = FSlateApplication::Get().GetRenderer()->GetFontMeasureService();
+		const int32 LastWholeCharacterIndex = FontMeasureService->FindLastWholeCharacterIndexBeforeOffset(Name, EventFont, FMath::RoundToInt(EventW - 2.0f));
+
+		if (LastWholeCharacterIndex >= 0)
+		{
+			// Grey threshold is shifted toward black (0.4 instead of 0.5 in test below) due to "area rule":
+			// a large gray surface (background of an event in this case) is perceived lighter than a smaller area (text pixels).
+			// Ref: https://books.google.ro/books?id=0pVr7dhmdWYC
+			const bool bIsDarkColor = (EventColorFill.ComputeLuminance() < 0.4f);
+
+			DrawState.Texts.AddDefaulted();
+			FPacketBreakdownViewDrawState::FText& DrawText = DrawState.Texts.Last();
+			DrawText.X = EventX1 + 2.0f;
+			DrawText.Y = EventY + 1.0f;
+			DrawText.Text = Name.Left(LastWholeCharacterIndex + 1);
+			DrawText.bWhite = bIsDarkColor;
+		}
+	}
+}
+
+////////////////////////////////////////////////////////////////////////////////////////////////////
+
+void FPacketBreakdownViewDrawStateBuilder::FlushBox(const FBoxData& Box, const float EventY, const float EventH)
+{
+	DrawState.Boxes.AddUninitialized();
+	FPacketBreakdownViewDrawState::FBox& DrawBox = DrawState.Boxes.Last();
+	DrawBox.X = Box.X1;
+	DrawBox.W = Box.X2 - Box.X1;
+	DrawBox.Y = EventY;
+	DrawBox.H = EventH;
+	DrawBox.Color = Box.LinearColor;
+}
+
+////////////////////////////////////////////////////////////////////////////////////////////////////
+
+void FPacketBreakdownViewDrawStateBuilder::Flush()
+{
+	constexpr float Y0 = 0.0f;
+	constexpr float EventH = 14.0f;
+	constexpr float EventDY = 2.0f;
+
+	// Flush merged boxes.
+	for (int32 Depth = 0; Depth <= MaxDepth; ++Depth)
+	{
+		const FBoxData& Box = LastBox[Depth];
+		if (Box.X1 < Box.X2)
+		{
+			const float EventY = Y0 + (EventH + EventDY) * Depth;
+			FlushBox(Box, EventY, EventH);
+		}
+	}
+}
+
+////////////////////////////////////////////////////////////////////////////////////////////////////
+// FPacketBreakdownViewDrawHelper
+////////////////////////////////////////////////////////////////////////////////////////////////////
+
+FPacketBreakdownViewDrawHelper::FPacketBreakdownViewDrawHelper(const FDrawContext& InDrawContext, const FPacketBreakdownViewport& InViewport)
+	: DrawContext(InDrawContext)
+	, Viewport(InViewport)
+	, WhiteBrush(FInsightsStyle::Get().GetBrush("WhiteBrush"))
+	, EventBorderBrush(FInsightsStyle::Get().GetBrush("EventBorder"))
+	, HoveredEventBorderBrush(FInsightsStyle::Get().GetBrush("HoveredEventBorder"))
+	, SelectedEventBorderBrush(FInsightsStyle::Get().GetBrush("SelectedEventBorder"))
+	, EventFont(FCoreStyle::GetDefaultFontStyle("Regular", 8))
+{
+}
+
+////////////////////////////////////////////////////////////////////////////////////////////////////
+
+void FPacketBreakdownViewDrawHelper::DrawBackground() const
+{
+	const FAxisViewportDouble& ViewportX = Viewport.GetHorizontalAxisViewport();
+
+	const float X0 = 0.0f;
+	const float X1 = ViewportX.GetMinPos() - ViewportX.GetPos();
+	const float X2 = ViewportX.GetMaxPos() - ViewportX.GetPos();
+	const float X3 = FMath::CeilToFloat(Viewport.GetWidth());
+
+	const float Y = 0.0f;
+	const float H = FMath::CeilToFloat(Viewport.GetHeight());
+
+	FDrawHelpers::DrawBackground(DrawContext, WhiteBrush, X0, X1, X2, X3, Y, H);
+}
+
+////////////////////////////////////////////////////////////////////////////////////////////////////
+
+FLinearColor FPacketBreakdownViewDrawHelper::GetColorByType(int32 Type)
+{
+	constexpr float Alpha = 1.0f;
+	switch (Type)
+	{
+	case 0:
+		return FLinearColor(0.75f, 0.25f, 0.25f, Alpha);
+
+	case 1:
+		return FLinearColor(0.25f, 0.75f, 0.25f, Alpha);
+
+	case 2:
+		return FLinearColor(0.25f, 0.25f, 0.75f, Alpha);
+
+	case 3:
+		return FLinearColor(0.75f, 0.75f, 0.25f, Alpha);
+
+	case 4:
+		return FLinearColor(0.25f, 0.75f, 0.75f, Alpha);
+
+	case 5:
+		return FLinearColor(0.75f, 0.25f, 0.75f, Alpha);
+
+	//default:
+	//	return FLinearColor(0.75f, 0.75f, 0.75f, Alpha);
+	}
+
+	uint32 Hash = (static_cast<uint32>(Type) + 0x2c2c57ed) * 0x2c2c57ed;
+	return FLinearColor(((Hash >> 16) & 0xFF) / 255.0f,
+						((Hash >>  8) & 0xFF) / 255.0f,
+						((Hash      ) & 0xFF) / 255.0f,
+						1.0f);
+}
+
+////////////////////////////////////////////////////////////////////////////////////////////////////
+
+void FPacketBreakdownViewDrawHelper::Draw(const FPacketBreakdownViewDrawState& DrawState) const
+{
+	// Draw filled boxes.
+	for (const FPacketBreakdownViewDrawState::FBox& Box : DrawState.Boxes)
+	{
+		DrawContext.DrawBox(Box.X, Box.Y, Box.W, Box.H, WhiteBrush, Box.Color);
+	}
+	DrawContext.LayerId++;
+
+	// Draw borders.
+	for (const FPacketBreakdownViewDrawState::FBox& Box : DrawState.Borders)
+	{
+		DrawContext.DrawBox(Box.X, Box.Y, Box.W, Box.H, EventBorderBrush, Box.Color);
+	}
+	DrawContext.LayerId++;
+
+	// Draw texts.
+	for (const FPacketBreakdownViewDrawState::FText& Text : DrawState.Texts)
+	{
+		DrawContext.DrawText(Text.X, Text.Y, Text.Text, EventFont, Text.bWhite ? FLinearColor::White : FLinearColor::Black);
+	}
+	DrawContext.LayerId++;
+}
+
+////////////////////////////////////////////////////////////////////////////////////////////////////
+
+void FPacketBreakdownViewDrawHelper::DrawEventHighlight(const FNetworkPacketEvent& Event) const
+{
+	const FAxisViewportDouble& ViewportX = Viewport.GetHorizontalAxisViewport();
+
+	float EventX1 = ViewportX.GetRoundedOffsetForValue(static_cast<double>(Event.Offset));
+	if (EventX1 > Viewport.GetWidth())
+	{
+		return;
+	}
+
+	float EventX2 = ViewportX.GetRoundedOffsetForValue(static_cast<double>(Event.Offset + Event.Size));
+	if (EventX2 < 0)
+	{
+		return;
+	}
+
+	// Events are displayed with minimum 1px (including empty ones).
+	if (EventX1 == EventX2)
+	{
+		EventX2 = EventX1 + 1.0f;
+	}
+
+	// Limit event width on the viewport's left side.
+	const float MinX = -2.0f; // -2 allows event border to remain outside screen
+	if (EventX1 < MinX)
+	{
+		EventX1 = MinX;
+	}
+
+	// Limit event width on the viewport's right side.
+	const float MaxX = Viewport.GetWidth() + 2.0f; // +2 allows event border to remain outside screen
+	if (EventX2 > MaxX)
+	{
+		EventX2 = MaxX;
+	}
+
+	const float EventW = EventX2 - EventX1;
+
+	constexpr float Y0 = 0.0f;
+	constexpr float EventH = 14.0f;
+	constexpr float EventDY = 2.0f;
+
+	const float EventY = Y0 + (EventH + EventDY) * Event.Depth;
+
+	const FLinearColor Color(1.0f, 1.0f, 0.0f, 1.0f);
+	DrawContext.DrawBox(EventX1 - 2.0f, EventY - 2.0f, EventW + 4.0f, EventH + 4.0f, HoveredEventBorderBrush, Color);
+	DrawContext.LayerId++;
+}
+
+////////////////////////////////////////////////////////////////////////////////////////////////////
diff --git a/Engine/Source/Developer/TraceInsights/Private/Insights/NetworkingProfiler/ViewModels/PacketBreakdownViewHelper.h b/Engine/Source/Developer/TraceInsights/Private/Insights/NetworkingProfiler/ViewModels/PacketBreakdownViewHelper.h
new file mode 100644
--- /dev/null
+++ b/Engine/Source/Developer/TraceInsights/Private/Insights/NetworkingProfiler/ViewModels/PacketBreakdownViewHelper.h
@@ -0,0 +1,195 @@
+// Copyright 1998-2019 Epic Games, Inc. All Rights Reserved.
+
+#pragma once
+
+#include "CoreMinimal.h"
+#include "Fonts/SlateFontInfo.h"
+#include "Styling/WidgetStyle.h"
+
+#include <limits>
+
+enum class ESlateDrawEffect : uint8;
+
+struct FDrawContext;
+struct FGeometry;
+struct FSlateBrush;
+
+class FPacketBreakdownViewport;
+class FSlateWindowElementList;
+
+////////////////////////////////////////////////////////////////////////////////////////////////////
+
+struct FNetworkPacketEvent
+{
+	int64 Offset;
+	int64 Size;
+	int32 Type;
+	int32 Depth;
+
+	FNetworkPacketEvent(int64 InOffset, int64 InSize, int32 InType, int32 InDepth)
+		: Offset(InOffset), Size(InSize), Type(InType), Depth(InDepth)
+	{}
+
+	FNetworkPacketEvent(const FNetworkPacketEvent& Other)
+		: Offset(Other.Offset), Size(Other.Size), Type(Other.Type), Depth(Other.Depth)
+	{
+	}
+
+	FNetworkPacketEvent& operator=(const FNetworkPacketEvent& Other)
+	{
+		Offset = Other.Offset;
+		Size = Other.Size;
+		Type = Other.Type;
+		Depth = Other.Depth;
+		return *this;
+	}
+
+	bool Equals(const FNetworkPacketEvent& Other) const
+	{
+		return Offset == Other.Offset && Size == Other.Size && Type == Other.Type && Depth == Other.Depth;
+	}
+
+	static bool AreEquals(const FNetworkPacketEvent& A, const FNetworkPacketEvent& B)
+	{
+		return A.Equals(B);
+	}
+};
+
+////////////////////////////////////////////////////////////////////////////////////////////////////
+
+struct FPacketBreakdownViewDrawState
+{
+	struct FPacketEvent
+	{
+		int64 Offset; // [bit]
+		int64 Size; // [bit]
+		int32 Type;
+		int32 Depth;
+	};
+
+	struct FBox
+	{
+		float X;
+		float Y;
+		float W;
+		float H;
+		FLinearColor Color;
+	};
+
+	struct FText
+	{
+		float X;
+		float Y;
+		FString Text;
+		bool bWhite;
+	};
+
+	FPacketBreakdownViewDrawState()
+	{
+	}
+
+	void Reset()
+	{
+		Events.Reset();
+		Boxes.Reset();
+		Borders.Reset();
+		Texts.Reset();
+		NumMergedBoxes = 0;
+	}
+
+	int32 GetNumMergedBoxes() const { return NumMergedBoxes; }
+
+	TArray<FPacketEvent> Events;
+	TArray<FBox> Boxes;
+	TArray<FBox> Borders;
+	TArray<FText> Texts;
+
+	// Debug stats.
+	int32 NumMergedBoxes;
+};
+
+////////////////////////////////////////////////////////////////////////////////////////////////////
+
+class FPacketBreakdownViewDrawStateBuilder
+{
+private:
+	struct FBoxData
+	{
+		float X1;
+		float X2;
+		uint32 Color;
+		FLinearColor LinearColor;
+
+		FBoxData() : X1(0.0f), X2(0.0f), Color(0) {}
+		void Reset() { X1 = 0.0f; X2 = 0.0f; Color = 0; }
+	};
+
+public:
+	explicit FPacketBreakdownViewDrawStateBuilder(FPacketBreakdownViewDrawState& InState, const FPacketBreakdownViewport& InViewport);
+
+	/**
+	 * Non-copyable
+	 */
+	FPacketBreakdownViewDrawStateBuilder(const FPacketBreakdownViewDrawStateBuilder&) = delete;
+	FPacketBreakdownViewDrawStateBuilder& operator=(const FPacketBreakdownViewDrawStateBuilder&) = delete;
+
+	void AddEvent(int64 Offset, int64 Size, int32 Type, int32 Depth);
+	void Flush();
+
+private:
+	void FlushBox(const FBoxData& Box, const float EventY, const float EventH);
+
+private:
+	FPacketBreakdownViewDrawState& DrawState; // cached draw state to build
+	const FPacketBreakdownViewport& Viewport;
+
+	int32 MaxDepth;
+
+	TArray<float> LastEventX2; // X2 value for last event on each depth
+	TArray<FBoxData> LastBox;
+
+	const FSlateFontInfo EventFont;
+};
+
+////////////////////////////////////////////////////////////////////////////////////////////////////
+
+class FPacketBreakdownViewDrawHelper
+{
+public:
+	enum class EHighlightMode : uint32
+	{
+		Hovered = 1,
+		Selected = 2,
+		SelectedAndHovered = 3
+	};
+
+public:
+	explicit FPacketBreakdownViewDrawHelper(const FDrawContext& InDrawContext, const FPacketBreakdownViewport& InViewport);
+
+	/**
+	 * Non-copyable
+	 */
+	FPacketBreakdownViewDrawHelper(const FPacketBreakdownViewDrawHelper&) = delete;
+	FPacketBreakdownViewDrawHelper& operator=(const FPacketBreakdownViewDrawHelper&) = delete;
+
+	const FSlateBrush* GetWhiteBrush() const { return WhiteBrush; }
+	const FSlateFontInfo& GetEventFont() const { return EventFont; }
+
+	void DrawBackground() const;
+	void Draw(const FPacketBreakdownViewDrawState& DrawState) const;
+	void DrawEventHighlight(const FNetworkPacketEvent& Event) const;
+
+	static FLinearColor GetColorByType(int32 Type);
+
+private:
+	const FDrawContext& DrawContext;
+	const FPacketBreakdownViewport& Viewport;
+
+	const FSlateBrush* WhiteBrush;
+	const FSlateBrush* EventBorderBrush;
+	const FSlateBrush* HoveredEventBorderBrush;
+	const FSlateBrush* SelectedEventBorderBrush;
+	const FSlateFontInfo EventFont;
+};
+
+////////////////////////////////////////////////////////////////////////////////////////////////////
diff --git a/Engine/Source/Developer/TraceInsights/Private/Insights/NetworkingProfiler/ViewModels/PacketBreakdownViewport.cpp b/Engine/Source/Developer/TraceInsights/Private/Insights/NetworkingProfiler/ViewModels/PacketBreakdownViewport.cpp
new file mode 100644
--- /dev/null
+++ b/Engine/Source/Developer/TraceInsights/Private/Insights/NetworkingProfiler/ViewModels/PacketBreakdownViewport.cpp
@@ -0,0 +1,5 @@
+// Copyright 1998-2019 Epic Games, Inc. All Rights Reserved.
+
+#include "PacketBreakdownViewport.h"
+
+////////////////////////////////////////////////////////////////////////////////////////////////////
diff --git a/Engine/Source/Developer/TraceInsights/Private/Insights/NetworkingProfiler/ViewModels/PacketBreakdownViewport.h b/Engine/Source/Developer/TraceInsights/Private/Insights/NetworkingProfiler/ViewModels/PacketBreakdownViewport.h
new file mode 100644
--- /dev/null
+++ b/Engine/Source/Developer/TraceInsights/Private/Insights/NetworkingProfiler/ViewModels/PacketBreakdownViewport.h
@@ -0,0 +1,54 @@
+// Copyright 1998-2019 Epic Games, Inc. All Rights Reserved.
+
+#pragma once
+
+#include "CoreMinimal.h"
+
+// Insights
+#include "Insights/ViewModels/AxisViewportDouble.h"
+
+////////////////////////////////////////////////////////////////////////////////////////////////////
+
+class FPacketBreakdownViewport
+{
+public:
+	FPacketBreakdownViewport()
+	{
+		Reset();
+	}
+
+	void Reset()
+	{
+		HorizontalAxisViewport.Reset();
+		Height = 0.0f;
+	}
+
+	const FAxisViewportDouble& GetHorizontalAxisViewport() const { return HorizontalAxisViewport; }
+	FAxisViewportDouble& GetHorizontalAxisViewport() { return HorizontalAxisViewport; }
+
+	float GetWidth() const { return HorizontalAxisViewport.GetSize(); }
+	float GetHeight() const { return Height; }
+
+	bool SetSize(const float InWidth, const float InHeight)
+	{
+		const bool bWidthChanged = HorizontalAxisViewport.SetSize(InWidth);
+		if (bWidthChanged || Height != InHeight)
+		{
+			Height = InHeight;
+			OnSizeChanged();
+			return true;
+		}
+		return false;
+	}
+
+private:
+	void OnSizeChanged()
+	{
+	}
+
+private:
+	FAxisViewportDouble HorizontalAxisViewport;
+	float Height;
+};
+
+////////////////////////////////////////////////////////////////////////////////////////////////////
diff --git a/Engine/Source/Developer/TraceInsights/Private/Insights/NetworkingProfiler/ViewModels/PacketSizesViewHelper.cpp b/Engine/Source/Developer/TraceInsights/Private/Insights/NetworkingProfiler/ViewModels/PacketSizesViewHelper.cpp
--- a/Engine/Source/Developer/TraceInsights/Private/Insights/NetworkingProfiler/ViewModels/PacketSizesViewHelper.cpp
+++ b/Engine/Source/Developer/TraceInsights/Private/Insights/NetworkingProfiler/ViewModels/PacketSizesViewHelper.cpp
@@ -2,18 +2,15 @@
 
 #include "PacketSizesViewHelper.h"
 
-#include "Brushes/SlateBorderBrush.h"
-#include "Brushes/SlateColorBrush.h"
-#include "EditorStyleSet.h"
-#include "Fonts/FontMeasure.h"
-#include "Fonts/SlateFontInfo.h"
 #include "Framework/Application/SlateApplication.h"
 #include "Rendering/DrawElements.h"
 #include "TraceServices/AnalysisService.h"
 
 // Insights
 #include "Insights/Common/PaintUtils.h"
+#include "Insights/InsightsStyle.h"
 #include "Insights/NetworkingProfiler/ViewModels/PacketSizesViewport.h"
+#include "Insights/ViewModels/DrawHelpers.h"
 
 #include <limits>
 
@@ -83,8 +80,10 @@ void FNetworkPacketSeriesBuilder::AddPacket(int32 FrameIndex, int64 Size, double
 FPacketSizesViewDrawHelper::FPacketSizesViewDrawHelper(const FDrawContext& InDrawContext, const FPacketSizesViewport& InViewport)
 	: DrawContext(InDrawContext)
 	, Viewport(InViewport)
-	, WhiteBrush(FCoreStyle::Get().GetBrush("WhiteBrush"))
-	, BorderBrush(FEditorStyle::GetBrush("PlainBorder"))
+	, WhiteBrush(FInsightsStyle::Get().GetBrush("WhiteBrush"))
+	//, EventBorderBrush(FInsightsStyle::Get().GetBrush("EventBorder"))
+	, HoveredEventBorderBrush(FInsightsStyle::Get().GetBrush("HoveredEventBorder"))
+	, SelectedEventBorderBrush(FInsightsStyle::Get().GetBrush("SelectedEventBorder"))
 	, NumPackets(0)
 	, NumDrawSamples(0)
 {
@@ -94,48 +93,17 @@ FPacketSizesViewDrawHelper::FPacketSizesViewDrawHelper(const FDrawContext& InDra
 
 void FPacketSizesViewDrawHelper::DrawBackground() const
 {
-	const FSlateBrush* AreaBrush = WhiteBrush;
-	const FLinearColor ValidAreaColor(0.07f, 0.07f, 0.07f, 1.0f);
-	const FLinearColor InvalidAreaColor(0.1f, 0.07f, 0.07f, 1.0f);
+	const FAxisViewportInt32& ViewportX = Viewport.GetHorizontalAxisViewport();
 
-	const FIndexAxisViewport& ViewportX = Viewport.GetHorizontalAxisViewport();
+	const float X0 = 0.0f;
+	const float X1 = ViewportX.GetMinPos() - ViewportX.GetPos();
+	const float X2 = ViewportX.GetMaxPos() - ViewportX.GetPos();
+	const float X3 = FMath::CeilToFloat(Viewport.GetWidth());
 
-	const float X0 = ViewportX.GetMinPos() - ViewportX.GetPos();
-	const float X1 = ViewportX.GetMaxPos() - ViewportX.GetPos();
-
-	const float W = FMath::CeilToFloat(Viewport.GetWidth());
+	const float Y = 0.0f;
 	const float H = FMath::CeilToFloat(Viewport.GetHeight());
 
-	if (X0 >= W || X1 <= 0.0f)
-	{
-		// Draw invalid area (entire view).
-		DrawContext.DrawBox(0.0f, 0.0f, W, H, AreaBrush, InvalidAreaColor);
-	}
-	else // X0 < W && X1 > 0
-	{
-		if (X0 > 0.0f)
-		{
-			// Draw invalid area (left).
-			DrawContext.DrawBox(0.0f, 0.0f, X0, H, AreaBrush, InvalidAreaColor);
-		}
-
-		if (X1 < W)
-		{
-			// Draw invalid area (right).
-			DrawContext.DrawBox(X1, 0.0f, W - X1, H, AreaBrush, InvalidAreaColor);
-		}
-
-		const float ValidX0 = FMath::Max(X0, 0.0f);
-		const float ValidX1 = FMath::Min(X1, W);
-
-		if (ValidX1 > ValidX0)
-		{
-			// Draw valid area.
-			DrawContext.DrawBox(ValidX0, 0.0f, ValidX1 - ValidX0, H, AreaBrush, ValidAreaColor);
-		}
-	}
-
-	DrawContext.LayerId++;
+	FDrawHelpers::DrawBackground(DrawContext, WhiteBrush, X0, X1, X2, X3, Y, H);
 }
 
 ////////////////////////////////////////////////////////////////////////////////////////////////////
@@ -176,7 +144,7 @@ void FPacketSizesViewDrawHelper::DrawCached(const FNetworkPacketSeries& Series)
 	const float SampleW = Viewport.GetSampleWidth();
 	const int32 NumSamples = Series.Samples.Num();
 
-	const FValueAxisViewport& ViewportY = Viewport.GetVerticalAxisViewport();
+	const FAxisViewportDouble& ViewportY = Viewport.GetVerticalAxisViewport();
 
 	const float ViewHeight = FMath::RoundToFloat(Viewport.GetHeight());
 	const float BaselineY = FMath::RoundToFloat(ViewportY.GetOffsetForValue(0.0));
@@ -230,7 +198,7 @@ void FPacketSizesViewDrawHelper::DrawSampleHighlight(const FNetworkPacketAggrega
 	const int32 SampleIndex = (Sample.LargestPacket.FrameIndex - FirstFrameIndex) / PacketsPerSample;
 	const float X = SampleIndex * SampleW;
 
-	const FValueAxisViewport& ViewportY = Viewport.GetVerticalAxisViewport();
+	const FAxisViewportDouble& ViewportY = Viewport.GetVerticalAxisViewport();
 
 	const float ViewHeight = FMath::RoundToFloat(Viewport.GetHeight());
 	const float BaselineY = FMath::RoundToFloat(ViewportY.GetOffsetForValue(0.0));
@@ -241,17 +209,13 @@ void FPacketSizesViewDrawHelper::DrawSampleHighlight(const FNetworkPacketAggrega
 
 	if (Mode == EHighlightMode::Hovered)
 	{
-		const FSlateBrush* HighlightBrush = BorderBrush;
-
 		const FLinearColor Color(1.0f, 1.0f, 0.0f, 1.0f); // yellow
 
 		// Draw border around the timing event box.
-		DrawContext.DrawBox(X - 1.0f, Y - 1.0f, SampleW + 2.0f, H + 2.0f, HighlightBrush, Color);
+		DrawContext.DrawBox(X - 1.0f, Y - 1.0f, SampleW + 2.0f, H + 2.0f, HoveredEventBorderBrush, Color);
 	}
 	else // EHighlightMode::Selected or EHighlightMode::SelectedAndHovered
 	{
-		const FSlateBrush* SelectedBrush = BorderBrush;
-
 		// Animate color from white (if selected and hovered) or yellow (if only selected) to black, using a squared sine function.
 		const double Time = static_cast<double>(FPlatformTime::Cycles64()) * FPlatformTime::GetSecondsPerCycle64();
 		float S = FMath::Sin(2.0 * Time);
@@ -260,7 +224,7 @@ void FPacketSizesViewDrawHelper::DrawSampleHighlight(const FNetworkPacketAggrega
 		const FLinearColor Color(S, S, Blue, 1.0f);
 
 		// Draw border around the timing event box.
-		DrawContext.DrawBox(X - 1.0f, Y - 1.0f, SampleW + 2.0f, H + 2.0f, SelectedBrush, Color);
+		DrawContext.DrawBox(X - 1.0f, Y - 1.0f, SampleW + 2.0f, H + 2.0f, SelectedEventBorderBrush, Color);
 	}
 	DrawContext.LayerId++;
 }
diff --git a/Engine/Source/Developer/TraceInsights/Private/Insights/NetworkingProfiler/ViewModels/PacketSizesViewHelper.h b/Engine/Source/Developer/TraceInsights/Private/Insights/NetworkingProfiler/ViewModels/PacketSizesViewHelper.h
--- a/Engine/Source/Developer/TraceInsights/Private/Insights/NetworkingProfiler/ViewModels/PacketSizesViewHelper.h
+++ b/Engine/Source/Developer/TraceInsights/Private/Insights/NetworkingProfiler/ViewModels/PacketSizesViewHelper.h
@@ -139,7 +139,7 @@ struct FNetworkPacketSeries
 class FNetworkPacketSeriesBuilder
 {
 public:
-	FNetworkPacketSeriesBuilder(FNetworkPacketSeries& InSeries, const FPacketSizesViewport& InViewport);
+	explicit FNetworkPacketSeriesBuilder(FNetworkPacketSeries& InSeries, const FPacketSizesViewport& InViewport);
 
 	/**
 	 * Non-copyable
@@ -177,7 +177,7 @@ public:
 	};
 
 public:
-	FPacketSizesViewDrawHelper(const FDrawContext& InDrawContext, const FPacketSizesViewport& InViewport);
+	explicit FPacketSizesViewDrawHelper(const FDrawContext& InDrawContext, const FPacketSizesViewport& InViewport);
 
 	/**
 	 * Non-copyable
@@ -199,9 +199,13 @@ private:
 	const FPacketSizesViewport& Viewport;
 
 	const FSlateBrush* WhiteBrush;
-	const FSlateBrush* BorderBrush;
+	//const FSlateBrush* EventBorderBrush;
+	const FSlateBrush* HoveredEventBorderBrush;
+	const FSlateBrush* SelectedEventBorderBrush;
 
 	// Debug stats.
 	mutable int32 NumPackets;
 	mutable int32 NumDrawSamples;
 };
+
+////////////////////////////////////////////////////////////////////////////////////////////////////
diff --git a/Engine/Source/Developer/TraceInsights/Private/Insights/NetworkingProfiler/ViewModels/PacketSizesViewport.h b/Engine/Source/Developer/TraceInsights/Private/Insights/NetworkingProfiler/ViewModels/PacketSizesViewport.h
--- a/Engine/Source/Developer/TraceInsights/Private/Insights/NetworkingProfiler/ViewModels/PacketSizesViewport.h
+++ b/Engine/Source/Developer/TraceInsights/Private/Insights/NetworkingProfiler/ViewModels/PacketSizesViewport.h
@@ -5,9 +5,8 @@
 #include "CoreMinimal.h"
 
 // Insights
-#include "Insights/ViewModels/FrameTrackViewport.h"
-//#include "Insights/ViewModels/IndexAxisViewport.h"
-//#include "Insights/ViewModels/ValueAxisViewport.h"
+#include "Insights/ViewModels/AxisViewportInt32.h"
+#include "Insights/ViewModels/AxisViewportDouble.h"
 
 ////////////////////////////////////////////////////////////////////////////////////////////////////
 
@@ -28,11 +27,11 @@ public:
 		VerticalAxisViewport.Reset();
 	}
 
-	const FIndexAxisViewport& GetHorizontalAxisViewport() const { return HorizontalAxisViewport; }
-	FIndexAxisViewport& GetHorizontalAxisViewport() { return HorizontalAxisViewport; }
+	const FAxisViewportInt32& GetHorizontalAxisViewport() const { return HorizontalAxisViewport; }
+	FAxisViewportInt32& GetHorizontalAxisViewport() { return HorizontalAxisViewport; }
 
-	const FValueAxisViewport& GetVerticalAxisViewport() const { return VerticalAxisViewport; }
-	FValueAxisViewport& GetVerticalAxisViewport() { return VerticalAxisViewport; }
+	const FAxisViewportDouble& GetVerticalAxisViewport() const { return VerticalAxisViewport; }
+	FAxisViewportDouble& GetVerticalAxisViewport() { return VerticalAxisViewport; }
 
 	float GetWidth() const { return HorizontalAxisViewport.GetSize(); }
 	float GetHeight() const { return VerticalAxisViewport.GetSize(); }
@@ -50,8 +49,8 @@ public:
 	}
 
 	float GetSampleWidth() const { return HorizontalAxisViewport.GetSampleSize(); }
-	int32 GetNumPacketsPerSample() const { return HorizontalAxisViewport.GetNumIndicesPerSample(); }
-	int32 GetFirstFrameIndex() const { return HorizontalAxisViewport.GetIndexAtOffset(0.0f); }
+	int32 GetNumPacketsPerSample() const { return HorizontalAxisViewport.GetNumSamplesPerPixel(); }
+	int32 GetFirstFrameIndex() const { return HorizontalAxisViewport.GetValueAtOffset(0.0f); }
 
 private:
 	void OnSizeChanged()
@@ -59,8 +58,8 @@ private:
 	}
 
 private:
-	FIndexAxisViewport HorizontalAxisViewport;
-	FValueAxisViewport VerticalAxisViewport;
+	FAxisViewportInt32 HorizontalAxisViewport;
+	FAxisViewportDouble VerticalAxisViewport;
 };
 
 ////////////////////////////////////////////////////////////////////////////////////////////////////
diff --git a/Engine/Source/Developer/TraceInsights/Private/Insights/NetworkingProfiler/Widgets/SNetworkingProfilerToolbar.cpp b/Engine/Source/Developer/TraceInsights/Private/Insights/NetworkingProfiler/Widgets/SNetworkingProfilerToolbar.cpp
--- a/Engine/Source/Developer/TraceInsights/Private/Insights/NetworkingProfiler/Widgets/SNetworkingProfilerToolbar.cpp
+++ b/Engine/Source/Developer/TraceInsights/Private/Insights/NetworkingProfiler/Widgets/SNetworkingProfilerToolbar.cpp
@@ -47,7 +47,33 @@ void SNetworkingProfilerToolbar::Construct(const FArguments& InArgs)
 			{
 				ToolbarBuilder.AddToolBarButton(FNetworkingProfilerCommands::Get().TogglePacketSizesViewVisibility);
 				ToolbarBuilder.AddToolBarButton(FNetworkingProfilerCommands::Get().TogglePacketBreakdownViewVisibility);
-				ToolbarBuilder.AddToolBarButton(FNetworkingProfilerCommands::Get().ToggleDataStreamBreakdownViewVisibility);
+				//ToolbarBuilder.AddToolBarButton(FNetworkingProfilerCommands::Get().ToggleDataStreamBreakdownViewVisibility);
+			}
+			ToolbarBuilder.EndSection();
+			ToolbarBuilder.BeginSection("Connection");
+			{
+				//FUIAction Action_ChooseReplicationSystem
+				//(
+				//	FExecuteAction::CreateSP(this, &STableTreeView::ContextMenu_CopySelectedToClipboard_Execute),
+				//	FCanExecuteAction::CreateSP(this, &STableTreeView::ContextMenu_CopySelectedToClipboard_CanExecute)
+				//);
+				//const FOnGetContent MenuContentGenerator;
+				//ToolbarBuilder.AddComboButton(Action, MenuContentGenerator);
+				TSharedRef<SWidget> ReplicationWidget = SNew(STextBlock)
+					.Text(LOCTEXT("Replication", "Replication System"));
+				ToolbarBuilder.AddWidget(ReplicationWidget);
+
+				ToolbarBuilder.AddSeparator();
+
+				TSharedRef<SWidget> ConnectionWidget = SNew(STextBlock)
+					.Text(LOCTEXT("Connection", "Connection"));
+				ToolbarBuilder.AddWidget(ConnectionWidget);
+
+				ToolbarBuilder.AddSeparator();
+
+				TSharedRef<SWidget> TypeWidget = SNew(STextBlock)
+					.Text(LOCTEXT("Type", "Sent/Received"));
+				ToolbarBuilder.AddWidget(TypeWidget);
 			}
 			ToolbarBuilder.EndSection();
 		}
diff --git a/Engine/Source/Developer/TraceInsights/Private/Insights/NetworkingProfiler/Widgets/SNetworkingProfilerWindow.cpp b/Engine/Source/Developer/TraceInsights/Private/Insights/NetworkingProfiler/Widgets/SNetworkingProfilerWindow.cpp
--- a/Engine/Source/Developer/TraceInsights/Private/Insights/NetworkingProfiler/Widgets/SNetworkingProfilerWindow.cpp
+++ b/Engine/Source/Developer/TraceInsights/Private/Insights/NetworkingProfiler/Widgets/SNetworkingProfilerWindow.cpp
@@ -230,22 +230,23 @@ void SNetworkingProfilerWindow::Construct(const FArguments& InArgs, const TShare
 			->Split
 			(
 				FTabManager::NewStack()
-				->SetSizeCoefficient(0.33f)
+				->SetSizeCoefficient(0.35f)
 				->SetHideTabWell(true)
 				->AddTab(FNetworkingProfilerTabs::PacketSizesViewID, NetworkingProfilerManager->IsPacketSizesViewVisible() ? ETabState::OpenedTab : ETabState::ClosedTab)
 			)
 			->Split
 			(
 				FTabManager::NewStack()
-				->SetSizeCoefficient(0.33f)
+				->SetSizeCoefficient(0.65f)
+				->SetHideTabWell(true)
 				->AddTab(FNetworkingProfilerTabs::PacketBreakdownViewID, NetworkingProfilerManager->IsPacketBreakdownViewVisible() ? ETabState::OpenedTab : ETabState::ClosedTab)
 			)
-			->Split
-			(
-				FTabManager::NewStack()
-				->SetSizeCoefficient(0.33f)
-				->AddTab(FNetworkingProfilerTabs::DataStreamBreakdownViewID, NetworkingProfilerManager->IsDataStreamBreakdownViewVisible() ? ETabState::OpenedTab : ETabState::ClosedTab)
-			)
+			//->Split
+			//(
+			//	FTabManager::NewStack()
+			//	->SetSizeCoefficient(0.33f)
+			//	->AddTab(FNetworkingProfilerTabs::DataStreamBreakdownViewID, NetworkingProfilerManager->IsDataStreamBreakdownViewVisible() ? ETabState::OpenedTab : ETabState::ClosedTab)
+			//)
 		);
 
 	// Create & initialize main menu.
diff --git a/Engine/Source/Developer/TraceInsights/Private/Insights/NetworkingProfiler/Widgets/SPacketBreakdownView.cpp b/Engine/Source/Developer/TraceInsights/Private/Insights/NetworkingProfiler/Widgets/SPacketBreakdownView.cpp
--- a/Engine/Source/Developer/TraceInsights/Private/Insights/NetworkingProfiler/Widgets/SPacketBreakdownView.cpp
+++ b/Engine/Source/Developer/TraceInsights/Private/Insights/NetworkingProfiler/Widgets/SPacketBreakdownView.cpp
@@ -2,7 +2,19 @@
 
 #include "SPacketBreakdownView.h"
 
+#include "Fonts/FontMeasure.h"
+#include "Fonts/SlateFontInfo.h"
+#include "Framework/Application/SlateApplication.h"
+#include "Framework/MultiBox/MultiBoxBuilder.h"
+#include "Rendering/DrawElements.h"
+#include "Styling/CoreStyle.h"
+#include "Widgets/Layout/SScrollBar.h"
+
 // Insights
+#include "Insights/Common/PaintUtils.h"
+#include "Insights/Common/Stopwatch.h"
+#include "Insights/InsightsManager.h"
+#include "Insights/InsightsStyle.h"
 
 ////////////////////////////////////////////////////////////////////////////////////////////////////
 
@@ -11,7 +23,9 @@
 ////////////////////////////////////////////////////////////////////////////////////////////////////
 
 SPacketBreakdownView::SPacketBreakdownView()
+	: DrawState(MakeShareable(new FPacketBreakdownViewDrawState()))
 {
+	Reset();
 }
 
 ////////////////////////////////////////////////////////////////////////////////////////////////////
@@ -22,11 +36,590 @@ SPacketBreakdownView::~SPacketBreakdownView()
 
 ////////////////////////////////////////////////////////////////////////////////////////////////////
 
+void SPacketBreakdownView::Reset()
+{
+	Viewport.Reset();
+	//FAxisViewportDouble& ViewportX = Viewport.GetHorizontalAxisViewport();
+	//ViewportX.SetScaleLimits(0.000001, 100000.0);
+	//ViewportX.SetScale(1.0);
+	bIsViewportDirty = true;
+
+	PacketFrameIndex = -1;
+	PacketSize = 0;
+	DrawState->Reset();
+	bIsStateDirty = true;
+
+	MousePosition = FVector2D::ZeroVector;
+
+	MousePositionOnButtonDown = FVector2D::ZeroVector;
+	ViewportPosXOnButtonDown = 0.0f;
+
+	MousePositionOnButtonUp = FVector2D::ZeroVector;
+
+	bIsLMB_Pressed = false;
+	bIsRMB_Pressed = false;
+
+	bIsScrolling = false;
+
+	HoveredEvent.Reset();
+	Tooltip.Reset();
+
+	//ThisGeometry
+
+	CursorType = ECursorType::Default;
+
+	UpdateDurationHistory.Reset();
+	DrawDurationHistory.Reset();
+	OnPaintDurationHistory.Reset();
+	LastOnPaintTime = FPlatformTime::Cycles64();
+}
+
+////////////////////////////////////////////////////////////////////////////////////////////////////
+
 void SPacketBreakdownView::Construct(const FArguments& InArgs)
 {
-	//ChildSlot
-	//[
-	//];
+	ChildSlot
+	[
+		SNew(SOverlay)
+		.Visibility(EVisibility::SelfHitTestInvisible)
+
+		+ SOverlay::Slot()
+		.VAlign(VAlign_Bottom)
+		.Padding(FMargin(0, 0, 0, 0))
+		[
+			SAssignNew(HorizontalScrollBar, SScrollBar)
+			.Orientation(Orient_Horizontal)
+			.AlwaysShowScrollbar(false)
+			.Visibility(EVisibility::Visible)
+			.Thickness(FVector2D(5.0f, 5.0f))
+			.RenderOpacity(0.75)
+			.OnUserScrolled(this, &SPacketBreakdownView::HorizontalScrollBar_OnUserScrolled)
+		]
+	];
+
+	UpdateHorizontalScrollBar();
+
+	BindCommands();
+}
+
+////////////////////////////////////////////////////////////////////////////////////////////////////
+
+void SPacketBreakdownView::Tick(const FGeometry& AllottedGeometry, const double InCurrentTime, const float InDeltaTime)
+{
+	if (ThisGeometry != AllottedGeometry || bIsViewportDirty)
+	{
+		bIsViewportDirty = false;
+		const float ViewWidth = AllottedGeometry.GetLocalSize().X;
+		const float ViewHeight = AllottedGeometry.GetLocalSize().Y;
+		Viewport.SetSize(ViewWidth, ViewHeight);
+		bIsStateDirty = true;
+	}
+
+	ThisGeometry = AllottedGeometry;
+
+	FAxisViewportDouble& ViewportX = Viewport.GetHorizontalAxisViewport();
+
+	if (!bIsScrolling)
+	{
+		// Elastic snap to horizontal limits.
+		if (ViewportX.UpdatePosWithinLimits())
+		{
+			bIsStateDirty = true;
+		}
+	}
+
+	if (bIsStateDirty)
+	{
+		bIsStateDirty = false;
+		UpdateState();
+	}
+
+	Tooltip.Update();
+	if (!MousePosition.IsZero())
+	{
+		Tooltip.SetPosition(MousePosition, 0.0f, Viewport.GetWidth(), 0.0f, Viewport.GetHeight() - 12.0f);
+	}
+}
+
+////////////////////////////////////////////////////////////////////////////////////////////////////
+
+void SPacketBreakdownView::SetPacket(int32 InFrameIndex, int64 InBitSize)
+{
+	PacketFrameIndex = InFrameIndex;
+	PacketSize = InBitSize;
+
+	FAxisViewportDouble& ViewportX = Viewport.GetHorizontalAxisViewport();
+	ViewportX.SetMinMaxValueInterval(0.0f, static_cast<double>(InBitSize));
+	ViewportX.CenterOnValueInterval(0.0f, static_cast<double>(InBitSize));
+	UpdateHorizontalScrollBar();
+
+	DrawState->Reset();
+	bIsStateDirty = true;
+}
+
+////////////////////////////////////////////////////////////////////////////////////////////////////
+
+void AddMockEventsRec(FPacketBreakdownViewDrawStateBuilder& Builder, FRandomStream& RandomStream , int32 Depth, int64 Offset, int64 Size, int32 Type)
+{
+	Builder.AddEvent(Offset, Size, Type, Depth);
+
+	int64 ChildEventOffset = Offset;
+	while (ChildEventOffset < Offset + Size)
+	{
+		int64 ChildEventSize = RandomStream.RandRange(1, Size);
+		if (ChildEventOffset + ChildEventSize > Offset + Size)
+		{
+			ChildEventSize = Offset + Size - ChildEventOffset;
+		}
+
+		if (ChildEventOffset != Offset || ChildEventSize != Size)
+		{
+			int32 ChildEventType = RandomStream.RandRange(1, 1000);
+			AddMockEventsRec(Builder, RandomStream, Depth + 1, ChildEventOffset, ChildEventSize, ChildEventType);
+		}
+
+		ChildEventOffset += ChildEventSize;
+	}
+}
+
+////////////////////////////////////////////////////////////////////////////////////////////////////
+
+void SPacketBreakdownView::UpdateState()
+{
+	FStopwatch Stopwatch;
+	Stopwatch.Start();
+
+	if (PacketFrameIndex >= 0 && PacketSize > 0)
+	{
+		FPacketBreakdownViewDrawStateBuilder Builder(*DrawState, Viewport);
+
+		// Init packet content with mock data.
+		FRandomStream RandomStream(PacketFrameIndex);
+		AddMockEventsRec(Builder, RandomStream, 0, 0, PacketSize, 0);
+
+		Builder.Flush();
+	}
+
+	Stopwatch.Stop();
+	UpdateDurationHistory.AddValue(Stopwatch.AccumulatedTime);
+}
+
+////////////////////////////////////////////////////////////////////////////////////////////////////
+
+void SPacketBreakdownView::UpdateHoveredEvent()
+{
+	HoveredEvent = GetEventAtMousePosition(MousePosition.X, MousePosition.Y);
+	//if (!HoveredEvent.IsValid())
+	//{
+	//	HoveredEvent = GetEventAtMousePosition(MousePosition.X - 1.0f, MousePosition.Y);
+	//}
+	//if (!HoveredEvent.IsValid())
+	//{
+	//	HoveredEvent = GetEventAtMousePosition(MousePosition.X + 1.0f, MousePosition.Y);
+	//}
+
+	if (HoveredEvent.IsValid())
+	{
+		// Init the tooltip's content.
+		Tooltip.ResetContent();
+		Tooltip.AddTitle(FString::Printf(TEXT("NetObject%d"), HoveredEvent.Event.Type));
+		Tooltip.AddNameValueTextLine(TEXT("Offset:"), FString::Format(TEXT("bit {0}"), { FText::AsNumber(HoveredEvent.Event.Offset).ToString() }));
+		if (HoveredEvent.Event.Size == 1)
+		{
+			Tooltip.AddNameValueTextLine(TEXT("Size:"), TEXT("1 bit"));
+		}
+		else
+		{
+			Tooltip.AddNameValueTextLine(TEXT("Size:"), FString::Format(TEXT("{0} bits"), { FText::AsNumber(HoveredEvent.Event.Size).ToString() }));
+		}
+		Tooltip.AddNameValueTextLine(TEXT("Type:"), FText::AsNumber(HoveredEvent.Event.Type).ToString());
+		Tooltip.AddNameValueTextLine(TEXT("Depth:"), FText::AsNumber(HoveredEvent.Event.Depth).ToString());
+		Tooltip.UpdateLayout();
+
+		Tooltip.SetDesiredOpacity(1.0f);
+	}
+	else
+	{
+		Tooltip.SetDesiredOpacity(0.0f);
+	}
+}
+
+////////////////////////////////////////////////////////////////////////////////////////////////////
+
+FNetworkPacketEventRef SPacketBreakdownView::GetEventAtMousePosition(float X, float Y)
+{
+	if (!bIsStateDirty)
+	{
+		for (const auto& Event : DrawState->Events)
+		{
+			const FAxisViewportDouble& ViewportX = Viewport.GetHorizontalAxisViewport();
+
+			const float EventX1 = ViewportX.GetRoundedOffsetForValue(static_cast<double>(Event.Offset));
+			const float EventX2 = ViewportX.GetRoundedOffsetForValue(static_cast<double>(Event.Offset + Event.Size));
+
+			constexpr float Y0 = 0.0f;
+			constexpr float EventH = 14.0f;
+			constexpr float EventDY = 2.0f;
+			const float EventY = Y0 + (EventH + EventDY) * Event.Depth;
+
+			constexpr float ToleranceX = 1.0f;
+
+			if (X >= EventX1 - ToleranceX && X <= EventX2 &&
+				Y >= EventY - EventDY / 2 && Y < EventY + EventH + EventDY / 2)
+			{
+				return FNetworkPacketEventRef(Event.Offset, Event.Size, Event.Type, Event.Depth);
+			}
+		}
+	}
+	return FNetworkPacketEventRef();
+}
+
+////////////////////////////////////////////////////////////////////////////////////////////////////
+
+int32 SPacketBreakdownView::OnPaint(const FPaintArgs& Args, const FGeometry& AllottedGeometry, const FSlateRect& MyCullingRect, FSlateWindowElementList& OutDrawElements, int32 LayerId, const FWidgetStyle& InWidgetStyle, bool bParentEnabled) const
+{
+	const bool bEnabled = ShouldBeEnabled(bParentEnabled);
+	const ESlateDrawEffect DrawEffects = bEnabled ? ESlateDrawEffect::None : ESlateDrawEffect::DisabledEffect;
+	FDrawContext DrawContext(AllottedGeometry, MyCullingRect, InWidgetStyle, DrawEffects, OutDrawElements, LayerId);
+
+	const float ViewWidth = AllottedGeometry.Size.X;
+	const float ViewHeight = AllottedGeometry.Size.Y;
+
+	//////////////////////////////////////////////////
+	{
+		FStopwatch Stopwatch;
+		Stopwatch.Start();
+
+		FPacketBreakdownViewDrawHelper Helper(DrawContext, Viewport);
+
+		Helper.DrawBackground();
+
+		// Draw the cached state (the network events of the package breakdown).
+		Helper.Draw(*DrawState);
+
+		// Highlight the hovered event.
+		if (HoveredEvent.IsValid())
+		{
+			Helper.DrawEventHighlight(HoveredEvent.Event);
+		}
+
+		// Draw tooltip for hovered Event.
+		Tooltip.Draw(DrawContext);
+
+		Stopwatch.Stop();
+		DrawDurationHistory.AddValue(Stopwatch.AccumulatedTime);
+	}
+	//////////////////////////////////////////////////
+
+	const bool bShouldDisplayDebugInfo = FInsightsManager::Get()->IsDebugInfoEnabled();
+	if (bShouldDisplayDebugInfo)
+	{
+		const FSlateBrush* WhiteBrush = FInsightsStyle::Get().GetBrush("WhiteBrush");
+		FSlateFontInfo SummaryFont = FCoreStyle::GetDefaultFontStyle("Regular", 8);
+
+		const TSharedRef<FSlateFontMeasure> FontMeasureService = FSlateApplication::Get().GetRenderer()->GetFontMeasureService();
+		const float MaxFontCharHeight = FontMeasureService->Measure(TEXT("!"), SummaryFont).Y;
+		const float DbgDY = MaxFontCharHeight;
+
+		const float DbgW = 260.0f;
+		const float DbgH = DbgDY * 4 + 3.0f;
+		const float DbgX = ViewWidth - DbgW - 20.0f;
+		float DbgY = 7.0f;
+
+		DrawContext.LayerId++;
+
+		DrawContext.DrawBox(DbgX - 2.0f, DbgY - 2.0f, DbgW, DbgH, WhiteBrush, FLinearColor(1.0, 1.0, 1.0, 0.9));
+		DrawContext.LayerId++;
+
+		FLinearColor DbgTextColor(0.0, 0.0, 0.0, 0.9);
+
+		// Time interval since last OnPaint call.
+		const uint64 CurrentTime = FPlatformTime::Cycles64();
+		const uint64 OnPaintDuration = CurrentTime - LastOnPaintTime;
+		LastOnPaintTime = CurrentTime;
+		OnPaintDurationHistory.AddValue(OnPaintDuration);
+		const uint64 AvgOnPaintDuration = OnPaintDurationHistory.ComputeAverage();
+		const uint64 AvgOnPaintDurationMs = FStopwatch::Cycles64ToMilliseconds(AvgOnPaintDuration);
+		const double AvgOnPaintFps = AvgOnPaintDurationMs != 0 ? 1.0 / FStopwatch::Cycles64ToSeconds(AvgOnPaintDuration) : 0.0;
+
+		const uint64 AvgUpdateDurationMs = FStopwatch::Cycles64ToMilliseconds(UpdateDurationHistory.ComputeAverage());
+		const uint64 AvgDrawDurationMs = FStopwatch::Cycles64ToMilliseconds(DrawDurationHistory.ComputeAverage());
+
+		// Draw performance info.
+		DrawContext.DrawText
+		(
+			DbgX, DbgY,
+			FString::Printf(TEXT("U: %llu ms    D: %llu ms + %llu ms = %llu ms (%d fps)"),
+				AvgUpdateDurationMs, // average duration of UpdateState calls
+				AvgDrawDurationMs, // drawing time
+				AvgOnPaintDurationMs - AvgDrawDurationMs, // other overhead to OnPaint calls
+				AvgOnPaintDurationMs, // average time between two OnPaint calls
+				FMath::RoundToInt(AvgOnPaintFps)), // framerate of OnPaint calls
+			SummaryFont, DbgTextColor
+		);
+		DbgY += DbgDY;
+
+		// Draw "the update stats".
+		DrawContext.DrawText
+		(
+			DbgX, DbgY,
+			FString::Printf(TEXT("U: %s events"),
+				*FText::AsNumber(DrawState->Events.Num()).ToString()),
+			SummaryFont, DbgTextColor
+		);
+		DbgY += DbgDY;
+
+		// Draw "the draw stats".
+		DrawContext.DrawText
+		(
+			DbgX, DbgY,
+			FString::Printf(TEXT("D: %s boxes, %s borders (%s merged), %s texts"),
+				*FText::AsNumber(DrawState->Boxes.Num()).ToString(),
+				*FText::AsNumber(DrawState->Borders.Num()).ToString(),
+				*FText::AsNumber(DrawState->GetNumMergedBoxes()).ToString(),
+				*FText::AsNumber(DrawState->Texts.Num()).ToString()),
+			SummaryFont, DbgTextColor
+		);
+		DbgY += DbgDY;
+
+		// Draw viewport's horizontal info.
+		DrawContext.DrawText
+		(
+			DbgX, DbgY,
+			Viewport.GetHorizontalAxisViewport().ToDebugString(TEXT("X")),
+			SummaryFont, DbgTextColor
+		);
+		DbgY += DbgDY;
+	}
+
+	return SCompoundWidget::OnPaint(Args, AllottedGeometry, MyCullingRect, OutDrawElements, LayerId, InWidgetStyle, bParentEnabled && IsEnabled());
+}
+
+////////////////////////////////////////////////////////////////////////////////////////////////////
+
+FReply SPacketBreakdownView::OnMouseButtonDown(const FGeometry& MyGeometry, const FPointerEvent& MouseEvent)
+{
+	FReply Reply = FReply::Unhandled();
+
+	MousePositionOnButtonDown = MyGeometry.AbsoluteToLocal(MouseEvent.GetScreenSpacePosition());
+	ViewportPosXOnButtonDown = Viewport.GetHorizontalAxisViewport().GetPos();
+
+	if (MouseEvent.GetEffectingButton() == EKeys::LeftMouseButton)
+	{
+		bIsLMB_Pressed = true;
+
+		// Capture mouse.
+		Reply = FReply::Handled().CaptureMouse(SharedThis(this));
+	}
+	else if (MouseEvent.GetEffectingButton() == EKeys::RightMouseButton)
+	{
+		bIsRMB_Pressed = true;
+
+		// Capture mouse, so we can scroll outside this widget.
+		Reply = FReply::Handled().CaptureMouse(SharedThis(this));
+	}
+
+	return Reply;
+}
+
+////////////////////////////////////////////////////////////////////////////////////////////////////
+
+FReply SPacketBreakdownView::OnMouseButtonUp(const FGeometry& MyGeometry, const FPointerEvent& MouseEvent)
+{
+	FReply Reply = FReply::Unhandled();
+
+	MousePositionOnButtonUp = MyGeometry.AbsoluteToLocal(MouseEvent.GetScreenSpacePosition());
+
+	const bool bIsValidForMouseClick = MousePositionOnButtonUp.Equals(MousePositionOnButtonDown, MOUSE_SNAP_DISTANCE);
+
+	if (MouseEvent.GetEffectingButton() == EKeys::LeftMouseButton)
+	{
+		if (bIsLMB_Pressed)
+		{
+			if (bIsScrolling)
+			{
+				bIsScrolling = false;
+				CursorType = ECursorType::Default;
+			}
+			else if (bIsValidForMouseClick)
+			{
+				//SelectEventAtMousePosition(MousePositionOnButtonUp.X, MousePositionOnButtonUp.Y);
+			}
+
+			bIsLMB_Pressed = false;
+
+			// Release the mouse.
+			Reply = FReply::Handled().ReleaseMouseCapture();
+		}
+	}
+	else if (MouseEvent.GetEffectingButton() == EKeys::RightMouseButton)
+	{
+		if (bIsRMB_Pressed)
+		{
+			if (bIsScrolling)
+			{
+				bIsScrolling = false;
+				CursorType = ECursorType::Default;
+			}
+			else if (bIsValidForMouseClick)
+			{
+				//ShowContextMenu(MouseEvent);
+			}
+
+			bIsRMB_Pressed = false;
+
+			// Release mouse as we no longer scroll.
+			Reply = FReply::Handled().ReleaseMouseCapture();
+		}
+	}
+
+	return Reply;
+}
+
+////////////////////////////////////////////////////////////////////////////////////////////////////
+
+FReply SPacketBreakdownView::OnMouseMove(const FGeometry& MyGeometry, const FPointerEvent& MouseEvent)
+{
+	FReply Reply = FReply::Unhandled();
+
+	MousePosition = MyGeometry.AbsoluteToLocal(MouseEvent.GetScreenSpacePosition());
+
+	if (!MouseEvent.GetCursorDelta().IsZero())
+	{
+		if (MouseEvent.IsMouseButtonDown(EKeys::LeftMouseButton) ||
+			MouseEvent.IsMouseButtonDown(EKeys::RightMouseButton))
+		{
+			if (HasMouseCapture())
+			{
+				if (!bIsScrolling)
+				{
+					bIsScrolling = true;
+					CursorType = ECursorType::Hand;
+
+					HoveredEvent.Reset();
+					Tooltip.SetDesiredOpacity(0.0f);
+				}
+
+				FAxisViewportDouble& ViewportX = Viewport.GetHorizontalAxisViewport();
+				const float PosX = ViewportPosXOnButtonDown + (MousePositionOnButtonDown.X - MousePosition.X);
+				ViewportX.ScrollAtPos(PosX);
+				UpdateHorizontalScrollBar();
+				bIsStateDirty = true;
+			}
+		}
+		else
+		{
+			UpdateHoveredEvent();
+		}
+
+		Reply = FReply::Handled();
+	}
+
+	return Reply;
+}
+
+////////////////////////////////////////////////////////////////////////////////////////////////////
+
+void SPacketBreakdownView::OnMouseEnter(const FGeometry& MyGeometry, const FPointerEvent& MouseEvent)
+{
+}
+
+////////////////////////////////////////////////////////////////////////////////////////////////////
+
+void SPacketBreakdownView::OnMouseLeave(const FPointerEvent& MouseEvent)
+{
+	if (!HasMouseCapture())
+	{
+		// No longer scrolling (unless we have mouse capture).
+		bIsScrolling = false;
+
+		bIsLMB_Pressed = false;
+		bIsRMB_Pressed = false;
+
+		MousePosition = FVector2D::ZeroVector;
+
+		HoveredEvent.Reset();
+		Tooltip.SetDesiredOpacity(0.0f);
+
+		CursorType = ECursorType::Default;
+	}
+}
+
+////////////////////////////////////////////////////////////////////////////////////////////////////
+
+FReply SPacketBreakdownView::OnMouseWheel(const FGeometry& MyGeometry, const FPointerEvent& MouseEvent)
+{
+	MousePosition = MyGeometry.AbsoluteToLocal(MouseEvent.GetScreenSpacePosition());
+
+	//if (MouseEvent.GetModifierKeys().IsShiftDown())
+	//{
+	//}
+	//else //if (MouseEvent.GetModifierKeys().IsControlDown())
+	{
+		// Zoom in/out horizontally.
+		const float Delta = MouseEvent.GetWheelDelta();
+		ZoomHorizontally(Delta, MousePosition.X);
+	}
+
+	return FReply::Handled();
+}
+
+////////////////////////////////////////////////////////////////////////////////////////////////////
+
+FReply SPacketBreakdownView::OnMouseButtonDoubleClick(const FGeometry& MyGeometry, const FPointerEvent& MouseEvent)
+{
+	return FReply::Unhandled();
+}
+
+////////////////////////////////////////////////////////////////////////////////////////////////////
+
+FCursorReply SPacketBreakdownView::OnCursorQuery(const FGeometry& MyGeometry, const FPointerEvent& CursorEvent) const
+{
+	FCursorReply CursorReply = FCursorReply::Unhandled();
+
+	if (CursorType == ECursorType::Arrow)
+	{
+		CursorReply = FCursorReply::Cursor(EMouseCursor::ResizeLeftRight);
+	}
+	else if (CursorType == ECursorType::Hand)
+	{
+		CursorReply = FCursorReply::Cursor(EMouseCursor::GrabHand);
+	}
+
+	return CursorReply;
+}
+
+////////////////////////////////////////////////////////////////////////////////////////////////////
+
+void SPacketBreakdownView::BindCommands()
+{
+}
+
+////////////////////////////////////////////////////////////////////////////////////////////////////
+
+void SPacketBreakdownView::HorizontalScrollBar_OnUserScrolled(float ScrollOffset)
+{
+	FAxisViewportDouble& ViewportX = Viewport.GetHorizontalAxisViewport();
+	ViewportX.OnUserScrolled(HorizontalScrollBar, ScrollOffset);
+	bIsStateDirty = true;
+}
+
+////////////////////////////////////////////////////////////////////////////////////////////////////
+
+void SPacketBreakdownView::UpdateHorizontalScrollBar()
+{
+	FAxisViewportDouble& ViewportX = Viewport.GetHorizontalAxisViewport();
+	ViewportX.UpdateScrollBar(HorizontalScrollBar);
+}
+
+////////////////////////////////////////////////////////////////////////////////////////////////////
+
+void SPacketBreakdownView::ZoomHorizontally(const float Delta, const float X)
+{
+	FAxisViewportDouble& ViewportX = Viewport.GetHorizontalAxisViewport();
+	ViewportX.RelativeZoomWithFixedOffset(Delta, X);
+	UpdateHorizontalScrollBar();
+	bIsStateDirty = true;
 }
 
 ////////////////////////////////////////////////////////////////////////////////////////////////////
diff --git a/Engine/Source/Developer/TraceInsights/Private/Insights/NetworkingProfiler/Widgets/SPacketBreakdownView.h b/Engine/Source/Developer/TraceInsights/Private/Insights/NetworkingProfiler/Widgets/SPacketBreakdownView.h
--- a/Engine/Source/Developer/TraceInsights/Private/Insights/NetworkingProfiler/Widgets/SPacketBreakdownView.h
+++ b/Engine/Source/Developer/TraceInsights/Private/Insights/NetworkingProfiler/Widgets/SPacketBreakdownView.h
@@ -3,12 +3,86 @@
 #pragma once
 
 #include "CoreMinimal.h"
-#include "Styling/SlateTypes.h"
+#include "Input/CursorReply.h"
+#include "Input/Reply.h"
+#include "Layout/Geometry.h"
+#include "Rendering/RenderingCommon.h"
 #include "Widgets/DeclarativeSyntaxSupport.h"
 #include "Widgets/SCompoundWidget.h"
 
+// Insights
+#include "Insights/Common/FixedCircularBuffer.h"
+#include "Insights/NetworkingProfiler/ViewModels/PacketBreakdownViewHelper.h"
+#include "Insights/NetworkingProfiler/ViewModels/PacketBreakdownViewport.h"
+#include "Insights/ViewModels/TooltipDrawState.h"
+
+class SScrollBar;
+
+////////////////////////////////////////////////////////////////////////////////////////////////////
+
+struct FNetworkPacketEventRef
+{
+	FNetworkPacketEvent Event;
+
+	FNetworkPacketEventRef()
+	: Event(0, 0, -1, -1)
+	{
+	}
+
+	FNetworkPacketEventRef(int64 InOffset, int64 InSize, int32 InType, int32 InDepth)
+		: Event(InOffset, InSize, InType, InDepth)
+	{
+	}
+
+	FNetworkPacketEventRef(const FNetworkPacketEventRef& Other)
+		: Event(Other.Event)
+	{
+	}
+
+	FNetworkPacketEventRef& operator=(const FNetworkPacketEventRef& Other)
+	{
+		Event = Other.Event;
+		return *this;
+	}
+
+	void Reset()
+	{
+		Event.Type = -1;
+	}
+
+	bool IsValid() const
+	{
+		return Event.Type >= 0;
+	}
+
+	bool Equals(const FNetworkPacketEventRef& Other) const
+	{
+		return Event.Equals(Other.Event);
+	}
+
+	static bool AreEquals(const FNetworkPacketEventRef& A, const FNetworkPacketEventRef& B)
+	{
+		return A.Equals(B);
+	}
+};
+
+////////////////////////////////////////////////////////////////////////////////////////////////////
+/**
+ * Widget used to present content of a network packet.
+ */
 class SPacketBreakdownView : public SCompoundWidget
 {
+public:
+	/** Number of pixels. */
+	static constexpr float MOUSE_SNAP_DISTANCE = 2.0f;
+
+	enum class ECursorType
+	{
+		Default,
+		Arrow,
+		Hand,
+	};
+
 public:
 	/** Default constructor. */
 	SPacketBreakdownView();
@@ -16,13 +90,110 @@ public:
 	/** Virtual destructor. */
 	virtual ~SPacketBreakdownView();
 
+	/** Resets internal widget's data to the default one. */
+	void Reset();
+
 	SLATE_BEGIN_ARGS(SPacketBreakdownView) {}
 	SLATE_END_ARGS()
 
 	/**
 	 * Construct this widget
-	 *
-	 * @param	InArgs	The declaration data for this widget
+	 * @param InArgs The declaration data for this widget
 	 */
 	void Construct(const FArguments& InArgs);
+
+	/**
+	 * Ticks this widget. Override in derived classes, but always call the parent implementation.
+	 * @param  AllottedGeometry The space allotted for this widget
+	 * @param  InCurrentTime  Current absolute real time
+	 * @param  InDeltaTime  Real time passed since last tick
+	 */
+	virtual void Tick(const FGeometry& AllottedGeometry, const double InCurrentTime, const float InDeltaTime) override;
+
+	virtual int32 OnPaint(const FPaintArgs& Args, const FGeometry& AllottedGeometry, const FSlateRect& MyCullingRect, FSlateWindowElementList& OutDrawElements, int32 LayerId, const FWidgetStyle& InWidgetStyle, bool bParentEnabled) const override;
+
+	virtual FReply OnMouseButtonDown(const FGeometry& MyGeometry, const FPointerEvent& MouseEvent) override;
+	virtual FReply OnMouseButtonUp(const FGeometry& MyGeometry, const FPointerEvent& MouseEvent) override;
+	virtual FReply OnMouseMove(const FGeometry& MyGeometry, const FPointerEvent& MouseEvent) override;
+	virtual void OnMouseEnter(const FGeometry& MyGeometry, const FPointerEvent& MouseEvent) override;
+	virtual void OnMouseLeave(const FPointerEvent& MouseEvent) override;
+	virtual FReply OnMouseWheel(const FGeometry& MyGeometry, const FPointerEvent& MouseEvent) override;
+	virtual FReply OnMouseButtonDoubleClick(const FGeometry& MyGeometry, const FPointerEvent& MouseEvent) override;
+
+	virtual FCursorReply OnCursorQuery(const FGeometry& MyGeometry, const FPointerEvent& CursorEvent) const override;
+
+	void SetPacket(int32 InPacketFrameIndex, int64 InPacketBitSize);
+
+private:
+	void UpdateState();
+	void UpdateHoveredEvent();
+	FNetworkPacketEventRef GetEventAtMousePosition(float X, float Y);
+
+	//void ShowContextMenu(const FPointerEvent& MouseEvent);
+
+	/** Binds our UI commands to delegates. */
+	void BindCommands();
+
+	/**
+	 * Called when the user scrolls the horizontal scrollbar.
+	 * @param ScrollOffset Scroll offset as a fraction between 0 and 1.
+	 */
+	void HorizontalScrollBar_OnUserScrolled(float ScrollOffset);
+	void UpdateHorizontalScrollBar();
+
+	void ZoomHorizontally(const float Delta, const float X);
+
+private:
+	/** The track's viewport. Encapsulates info about position and scale. */
+	FPacketBreakdownViewport Viewport;
+	bool bIsViewportDirty;
+
+	int32 PacketFrameIndex;
+	int64 PacketSize; // total number of bits; [bit]
+
+	/** Cached draw state of the packet content (i.e. all it needs to render). */
+	TSharedPtr<FPacketBreakdownViewDrawState> DrawState;
+	bool bIsStateDirty;
+	//////////////////////////////////////////////////
+
+	TSharedPtr<SScrollBar> HorizontalScrollBar;
+
+	//////////////////////////////////////////////////
+	// Panning and Zooming behaviors
+
+	/** The current mouse position. */
+	FVector2D MousePosition;
+
+	/** Mouse position during the call on mouse button down. */
+	FVector2D MousePositionOnButtonDown;
+	float ViewportPosXOnButtonDown;
+
+	/** Mouse position during the call on mouse button up. */
+	FVector2D MousePositionOnButtonUp;
+
+	bool bIsLMB_Pressed;
+	bool bIsRMB_Pressed;
+
+	/** True, if the user is currently interactively scrolling the view (ex.: by holding the left mouse button and dragging). */
+	bool bIsScrolling;
+
+	//////////////////////////////////////////////////
+	// Selection
+
+	FNetworkPacketEventRef HoveredEvent;
+	FTooltipDrawState Tooltip;
+
+	//////////////////////////////////////////////////
+	// Misc
+
+	FGeometry ThisGeometry;
+
+	/** Cursor type. */
+	ECursorType CursorType;
+
+	// Debug stats
+	TFixedCircularBuffer<uint64, 32> UpdateDurationHistory;
+	mutable TFixedCircularBuffer<uint64, 32> DrawDurationHistory;
+	mutable TFixedCircularBuffer<uint64, 32> OnPaintDurationHistory;
+	mutable uint64 LastOnPaintTime;
 };
diff --git a/Engine/Source/Developer/TraceInsights/Private/Insights/NetworkingProfiler/Widgets/SPacketSizesView.cpp b/Engine/Source/Developer/TraceInsights/Private/Insights/NetworkingProfiler/Widgets/SPacketSizesView.cpp
--- a/Engine/Source/Developer/TraceInsights/Private/Insights/NetworkingProfiler/Widgets/SPacketSizesView.cpp
+++ b/Engine/Source/Developer/TraceInsights/Private/Insights/NetworkingProfiler/Widgets/SPacketSizesView.cpp
@@ -2,15 +2,13 @@
 
 #include "SPacketSizesView.h"
 
-#include "Brushes/SlateBorderBrush.h"
-#include "Brushes/SlateColorBrush.h"
-#include "EditorStyleSet.h"
 #include "Fonts/FontMeasure.h"
 #include "Fonts/SlateFontInfo.h"
 #include "Framework/Application/SlateApplication.h"
 #include "Framework/MultiBox/MultiBoxBuilder.h"
 #include "HAL/PlatformTime.h"
 #include "Rendering/DrawElements.h"
+#include "Styling/CoreStyle.h"
 #include "Widgets/Layout/SScrollBar.h"
 
 // Insights
@@ -18,6 +16,10 @@
 #include "Insights/Common/Stopwatch.h"
 #include "Insights/Common/TimeUtils.h"
 #include "Insights/InsightsManager.h"
+#include "Insights/NetworkingProfiler/NetworkingProfilerManager.h"
+#include "Insights/NetworkingProfiler/Widgets/SNetworkingProfilerWindow.h"
+#include "Insights/NetworkingProfiler/Widgets/SPacketBreakdownView.h"
+#include "Insights/InsightsStyle.h"
 
 #include <limits>
 
@@ -44,8 +46,7 @@ SPacketSizesView::~SPacketSizesView()
 void SPacketSizesView::Reset()
 {
 	Viewport.Reset();
-
-	FValueAxisViewport& ViewportY = Viewport.GetVerticalAxisViewport();
+	FAxisViewportDouble& ViewportY = Viewport.GetVerticalAxisViewport();
 	ViewportY.SetScaleLimits(0.0001, 50.0);
 	ViewportY.SetScale(0.08);
 	bIsViewportDirty = true;
@@ -62,7 +63,6 @@ void SPacketSizesView::Reset()
 
 	MousePositionOnButtonDown = FVector2D::ZeroVector;
 	ViewportPosXOnButtonDown = 0.0f;
-	ViewportPosYOnButtonDown = 0.0f;
 
 	MousePositionOnButtonUp = FVector2D::ZeroVector;
 
@@ -133,7 +133,7 @@ void SPacketSizesView::Tick(const FGeometry& AllottedGeometry, const double InCu
 
 	ThisGeometry = AllottedGeometry;
 
-	FIndexAxisViewport& ViewportX = Viewport.GetHorizontalAxisViewport();
+	FAxisViewportInt32& ViewportX = Viewport.GetHorizontalAxisViewport();
 
 	if (!bIsScrolling)
 	{
@@ -173,16 +173,17 @@ void SPacketSizesView::UpdateState()
 
 	// Init series with mock data.
 	{
-		FIndexAxisViewport& ViewportX = Viewport.GetHorizontalAxisViewport();
+		FAxisViewportInt32& ViewportX = Viewport.GetHorizontalAxisViewport();
 
 		constexpr int32 NumPackets = 100000;
 
-		ViewportX.SetMinMaxIndexInterval(0, NumPackets);
-
-		const int32 StartIndex = FMath::Max(0, ViewportX.GetIndexAtOffset(0.0f));
-		const int32 EndIndex = FMath::Min(NumPackets, ViewportX.GetIndexAtOffset(ViewportX.GetSize()));
+		ViewportX.SetMinMaxInterval(0, NumPackets);
 
 		FNetworkPacketSeriesBuilder Builder(*PacketSeries, Viewport);
+
+		const int32 StartIndex = FMath::Max(0, ViewportX.GetValueAtOffset(0.0f));
+		const int32 EndIndex = FMath::Min(NumPackets, ViewportX.GetValueAtOffset(ViewportX.GetSize()));
+
 		for (int32 FrameIndex = StartIndex; FrameIndex < EndIndex; ++FrameIndex)
 		{
 			FRandomStream RandomStream((FrameIndex * FrameIndex * FrameIndex) ^ 0x2c2c57ed);
@@ -207,6 +208,7 @@ void SPacketSizesView::UpdateState()
 
 			Builder.AddPacket(FrameIndex, Size, TimeSent, TimeAck, Status);
 		}
+
 		NumUpdatedPackets += Builder.GetNumAddedPackets();
 	}
 
@@ -230,7 +232,7 @@ FNetworkPacketSampleRef SPacketSizesView::GetSampleAtMousePosition(float X, floa
 				const FNetworkPacketAggregatedSample& Sample = PacketSeries->Samples[SampleIndex];
 				if (Sample.NumPackets > 0)
 				{
-					const FValueAxisViewport& ViewportY = Viewport.GetVerticalAxisViewport();
+					const FAxisViewportDouble& ViewportY = Viewport.GetVerticalAxisViewport();
 
 					const float ViewHeight = FMath::RoundToFloat(Viewport.GetHeight());
 					const float BaselineY = FMath::RoundToFloat(ViewportY.GetOffsetForValue(0.0));
@@ -269,6 +271,16 @@ void SPacketSizesView::SelectSampleAtMousePosition(float X, float Y)
 	if (SampleRef.IsValid())
 	{
 		SelectedSample = SampleRef;
+
+		TSharedPtr<SNetworkingProfilerWindow> Wnd = FNetworkingProfilerManager::Get()->GetProfilerWindow();
+		if (Wnd.IsValid())
+		{
+			TSharedPtr<SPacketBreakdownView> PacketBreakdownView = Wnd->GetPacketBreakdownView();
+			if (PacketBreakdownView.IsValid())
+			{
+				PacketBreakdownView->SetPacket(SampleRef.Sample->LargestPacket.FrameIndex, SampleRef.Sample->LargestPacket.Size);
+			}
+		}
 	}
 }
 
@@ -309,13 +321,14 @@ int32 SPacketSizesView::OnPaint(const FPaintArgs& Args, const FGeometry& Allotte
 	const TSharedRef<FSlateFontMeasure> FontMeasureService = FSlateApplication::Get().GetRenderer()->GetFontMeasureService();
 	FSlateFontInfo SummaryFont = FCoreStyle::GetDefaultFontStyle("Regular", 8);
 
-	const FSlateBrush* WhiteBrush = FCoreStyle::Get().GetBrush("WhiteBrush");
+	const FSlateBrush* WhiteBrush = FInsightsStyle::Get().GetBrush("WhiteBrush");
 
 	const float ViewWidth = AllottedGeometry.Size.X;
 	const float ViewHeight = AllottedGeometry.Size.Y;
 
 	int32 NumDrawSamples = 0;
 
+	//////////////////////////////////////////////////
 	{
 		FStopwatch Stopwatch;
 		Stopwatch.Start();
@@ -324,6 +337,7 @@ int32 SPacketSizesView::OnPaint(const FPaintArgs& Args, const FGeometry& Allotte
 
 		Helper.DrawBackground();
 
+		// Draw the horizontal axis grid.
 		DrawHorizontalAxisGrid(DrawContext, WhiteBrush, SummaryFont);
 
 		Helper.DrawCached(*PacketSeries);
@@ -341,6 +355,7 @@ int32 SPacketSizesView::OnPaint(const FPaintArgs& Args, const FGeometry& Allotte
 			Helper.DrawSampleHighlight(*HoveredSample.Sample, FPacketSizesViewDrawHelper::EHighlightMode::Hovered);
 		}
 
+		// Draw the vertical axis grid.
 		DrawVerticalAxisGrid(DrawContext, WhiteBrush, SummaryFont);
 
 		// Draw tooltip for hovered sample.
@@ -401,9 +416,9 @@ int32 SPacketSizesView::OnPaint(const FPaintArgs& Args, const FGeometry& Allotte
 			const float DX = 2.0f;
 			const float W2 = TextSize.X / 2 + DX;
 
-			const FIndexAxisViewport& ViewportX = Viewport.GetHorizontalAxisViewport();
+			const FAxisViewportInt32& ViewportX = Viewport.GetHorizontalAxisViewport();
 
-			float X1 = ViewportX.GetOffsetForIndex(HoveredSample.Sample->LargestPacket.FrameIndex);
+			float X1 = ViewportX.GetOffsetForValue(HoveredSample.Sample->LargestPacket.FrameIndex);
 			float CX = X1 + FMath::RoundToFloat(Viewport.GetSampleWidth() / 2);
 			if (CX + W2 > ViewportX.GetSize())
 			{
@@ -425,7 +440,6 @@ int32 SPacketSizesView::OnPaint(const FPaintArgs& Args, const FGeometry& Allotte
 		Stopwatch.Stop();
 		DrawDurationHistory.AddValue(Stopwatch.AccumulatedTime);
 	}
-
 	//////////////////////////////////////////////////
 
 	const bool bShouldDisplayDebugInfo = FInsightsManager::Get()->IsDebugInfoEnabled();
@@ -439,17 +453,11 @@ int32 SPacketSizesView::OnPaint(const FPaintArgs& Args, const FGeometry& Allotte
 		const float DbgX = ViewWidth - DbgW - 20.0f;
 		float DbgY = 7.0f;
 
-		++LayerId;
-		FSlateDrawElement::MakeBox(
-			OutDrawElements,
-			LayerId,
-			MAKE_PAINT_GEOMETRY_RC(AllottedGeometry, DbgX - 2.0f, DbgY - 2.0f, DbgW, DbgH),
-			WhiteBrush,
-			DrawEffects,
-			FLinearColor(1.0, 1.0, 1.0, 0.9)
-		);
+		DrawContext.LayerId++;
+
+		DrawContext.DrawBox(DbgX - 2.0f, DbgY - 2.0f, DbgW, DbgH, WhiteBrush, FLinearColor(1.0, 1.0, 1.0, 0.9));
+		DrawContext.LayerId++;
 
-		++LayerId;
 		FLinearColor DbgTextColor(0.0, 0.0, 0.0, 0.9);
 
 		// Time interval since last OnPaint call.
@@ -465,63 +473,45 @@ int32 SPacketSizesView::OnPaint(const FPaintArgs& Args, const FGeometry& Allotte
 		const uint64 AvgDrawDurationMs = FStopwatch::Cycles64ToMilliseconds(DrawDurationHistory.ComputeAverage());
 
 		// Draw performance info.
-		FSlateDrawElement::MakeText
+		DrawContext.DrawText
 		(
-			OutDrawElements,
-			LayerId,
-			MAKE_PAINT_GEOMETRY_PT(AllottedGeometry, DbgX, DbgY),
+			DbgX, DbgY,
 			FString::Printf(TEXT("U: %llu ms, D: %llu ms + %llu ms = %llu ms (%d fps)"),
 				AvgUpdateDurationMs, // caching time
 				AvgDrawDurationMs, // drawing time
 				AvgOnPaintDurationMs - AvgDrawDurationMs, // other overhead to OnPaint calls
 				AvgOnPaintDurationMs, // average time between two OnPaint calls
 				FMath::RoundToInt(AvgOnPaintFps)), // framerate of OnPaint calls
-			SummaryFont,
-			DrawEffects,
-			DbgTextColor
+			SummaryFont, DbgTextColor
 		);
 		DbgY += DbgDY;
 
 		// Draw number of draw calls.
-		FSlateDrawElement::MakeText
+		DrawContext.DrawText
 		(
-			OutDrawElements,
-			LayerId,
-			MAKE_PAINT_GEOMETRY_PT(AllottedGeometry, DbgX, DbgY),
-			FString::Format(TEXT("U: {0} packets, D: {1} samples"),
-				{
-					FText::AsNumber(NumUpdatedPackets).ToString(),
-					FText::AsNumber(NumDrawSamples).ToString()
-				}),
-			SummaryFont,
-			DrawEffects,
-			DbgTextColor
+			DbgX, DbgY,
+			FString::Printf(TEXT("U: %s packets, D: %s samples"),
+				*FText::AsNumber(NumUpdatedPackets).ToString(),
+				*FText::AsNumber(NumDrawSamples).ToString()),
+			SummaryFont, DbgTextColor
 		);
 		DbgY += DbgDY;
 
 		// Draw viewport's horizontal info.
-		FSlateDrawElement::MakeText
+		DrawContext.DrawText
 		(
-			OutDrawElements,
-			LayerId,
-			MAKE_PAINT_GEOMETRY_PT(AllottedGeometry, DbgX, DbgY),
-			Viewport.GetHorizontalAxisViewport().ToDebugString(TEXT("VX: "), TEXT("packet")),
-			SummaryFont,
-			DrawEffects,
-			DbgTextColor
+			DbgX, DbgY,
+			Viewport.GetHorizontalAxisViewport().ToDebugString(TEXT("X"), TEXT("packet")),
+			SummaryFont, DbgTextColor
 		);
 		DbgY += DbgDY;
 
 		// Draw viewport's vertical info.
-		FSlateDrawElement::MakeText
+		DrawContext.DrawText
 		(
-			OutDrawElements,
-			LayerId,
-			MAKE_PAINT_GEOMETRY_PT(AllottedGeometry, DbgX, DbgY),
-			Viewport.GetVerticalAxisViewport().ToDebugString(TEXT("VY: ")),
-			SummaryFont,
-			DrawEffects,
-			DbgTextColor
+			DbgX, DbgY,
+			Viewport.GetVerticalAxisViewport().ToDebugString(TEXT("Y")),
+			SummaryFont, DbgTextColor
 		);
 		DbgY += DbgDY;
 	}
@@ -533,15 +523,15 @@ int32 SPacketSizesView::OnPaint(const FPaintArgs& Args, const FGeometry& Allotte
 
 void SPacketSizesView::DrawVerticalAxisGrid(FDrawContext& DrawContext, const FSlateBrush* Brush, const FSlateFontInfo& Font) const
 {
-	const FValueAxisViewport& ViewportY = Viewport.GetVerticalAxisViewport();
+	const FAxisViewportDouble& ViewportY = Viewport.GetVerticalAxisViewport();
 
 	const float RoundedViewHeight = FMath::RoundToFloat(ViewportY.GetSize());
 
 	constexpr float MinDY = 32.0f; // min vertical distance between horizontal grid lines
 
-	const double TopValue = ViewportY.GetValueForOffset(RoundedViewHeight);
-	const double GridValue = ViewportY.GetValueForOffset(MinDY);
-	const double BottomValue = ViewportY.GetValueForOffset(0.0f);
+	const double TopValue = ViewportY.GetValueAtOffset(RoundedViewHeight);
+	const double GridValue = ViewportY.GetValueAtOffset(MinDY);
+	const double BottomValue = ViewportY.GetValueAtOffset(0.0f);
 	const double Delta = GridValue - BottomValue;
 
 	if (Delta > 0.0)
@@ -602,15 +592,15 @@ void SPacketSizesView::DrawVerticalAxisGrid(FDrawContext& DrawContext, const FSl
 
 void SPacketSizesView::DrawHorizontalAxisGrid(FDrawContext& DrawContext, const FSlateBrush* Brush, const FSlateFontInfo& Font) const
 {
-	const FIndexAxisViewport& ViewportX = Viewport.GetHorizontalAxisViewport();
+	const FAxisViewportInt32& ViewportX = Viewport.GetHorizontalAxisViewport();
 
 	const float RoundedViewWidth = FMath::RoundToFloat(ViewportX.GetSize());
 
 	constexpr float MinDX = 100.0f; // min horizontal distance between vertical grid lines
 
-	int32 LeftIndex = ViewportX.GetIndexAtOffset(0.0f);
-	int32 GridIndex = ViewportX.GetIndexAtOffset(MinDX);
-	int32 RightIndex = ViewportX.GetIndexAtOffset(RoundedViewWidth);
+	int32 LeftIndex = ViewportX.GetValueAtOffset(0.0f);
+	int32 GridIndex = ViewportX.GetValueAtOffset(MinDX);
+	int32 RightIndex = ViewportX.GetValueAtOffset(RoundedViewWidth);
 	int32 Delta = GridIndex - LeftIndex;
 
 	if (Delta > 0)
@@ -652,7 +642,7 @@ void SPacketSizesView::DrawHorizontalAxisGrid(FDrawContext& DrawContext, const F
 
 		for (int32 Index = StartIndex; Index < RightIndex; Index += Grid)
 		{
-			const float X = FMath::RoundToFloat(ViewportX.GetOffsetForIndex(Index));
+			const float X = FMath::RoundToFloat(ViewportX.GetOffsetForValue(Index));
 
 			// Draw vertical grid line.
 			DrawContext.DrawBox(X, 0, 1, ViewHeight, Brush, GridColor);
@@ -681,7 +671,6 @@ FReply SPacketSizesView::OnMouseButtonDown(const FGeometry& MyGeometry, const FP
 
 	MousePositionOnButtonDown = MyGeometry.AbsoluteToLocal(MouseEvent.GetScreenSpacePosition());
 	ViewportPosXOnButtonDown = Viewport.GetHorizontalAxisViewport().GetPos();
-	ViewportPosYOnButtonDown = Viewport.GetVerticalAxisViewport().GetPos();
 
 	if (MouseEvent.GetEffectingButton() == EKeys::LeftMouseButton)
 	{
@@ -763,63 +752,46 @@ FReply SPacketSizesView::OnMouseMove(const FGeometry& MyGeometry, const FPointer
 
 	MousePosition = MyGeometry.AbsoluteToLocal(MouseEvent.GetScreenSpacePosition());
 
-	if (MouseEvent.IsMouseButtonDown(EKeys::LeftMouseButton))
+	if (!MouseEvent.GetCursorDelta().IsZero())
 	{
-		if (HasMouseCapture() && !MouseEvent.GetCursorDelta().IsZero())
+		if (MouseEvent.IsMouseButtonDown(EKeys::LeftMouseButton) ||
+			MouseEvent.IsMouseButtonDown(EKeys::RightMouseButton))
 		{
-			if (!bIsScrolling)
+			if (HasMouseCapture())
 			{
-				bIsScrolling = true;
-				CursorType = ECursorType::Hand;
-
-				HoveredSample.Reset();
-			}
+				if (!bIsScrolling)
+				{
+					bIsScrolling = true;
+					CursorType = ECursorType::Hand;
 
-			FIndexAxisViewport& ViewportX = Viewport.GetHorizontalAxisViewport();
-			const float PosX = ViewportPosXOnButtonDown + (MousePositionOnButtonDown.X - MousePosition.X);
-			ViewportX.ScrollAtIndex(ViewportX.GetIndexAtPos(PosX)); // align pos with at index
-			UpdateHorizontalScrollBar();
-			bIsStateDirty = true;
+					HoveredSample.Reset();
+				}
 
-			Reply = FReply::Handled();
+				FAxisViewportInt32& ViewportX = Viewport.GetHorizontalAxisViewport();
+				const float PosX = ViewportPosXOnButtonDown + (MousePositionOnButtonDown.X - MousePosition.X);
+				ViewportX.ScrollAtValue(ViewportX.GetValueAtPos(PosX)); // align viewport position with sample
+				UpdateHorizontalScrollBar();
+				bIsStateDirty = true;
+			}
 		}
-	}
-	else if (MouseEvent.IsMouseButtonDown(EKeys::RightMouseButton))
-	{
-		if (HasMouseCapture() && !MouseEvent.GetCursorDelta().IsZero())
+		else
 		{
-			if (!bIsScrolling)
+			if (!HoveredSample.IsValid())
 			{
-				bIsScrolling = true;
-				CursorType = ECursorType::Hand;
-
-				HoveredSample.Reset();
+				TooltipOpacity = 0.0f;
+			}
+			HoveredSample = GetSampleAtMousePosition(MousePosition.X, MousePosition.Y);
+			if (!HoveredSample.IsValid())
+			{
+				HoveredSample = GetSampleAtMousePosition(MousePosition.X - 1.0f, MousePosition.Y);
+			}
+			if (!HoveredSample.IsValid())
+			{
+				HoveredSample = GetSampleAtMousePosition(MousePosition.X + 1.0f, MousePosition.Y);
 			}
-
-			FIndexAxisViewport& ViewportX = Viewport.GetHorizontalAxisViewport();
-			const float PosX = ViewportPosXOnButtonDown + (MousePositionOnButtonDown.X - MousePosition.X);
-			ViewportX.ScrollAtIndex(ViewportX.GetIndexAtPos(PosX)); // align pos with at index
-			UpdateHorizontalScrollBar();
-			bIsStateDirty = true;
-
-			Reply = FReply::Handled();
-		}
-	}
-	else
-	{
-		if (!HoveredSample.IsValid())
-		{
-			TooltipOpacity = 0.0f;
-		}
-		HoveredSample = GetSampleAtMousePosition(MousePosition.X, MousePosition.Y);
-		if (!HoveredSample.IsValid())
-		{
-			HoveredSample = GetSampleAtMousePosition(MousePosition.X - 1.0f, MousePosition.Y);
-		}
-		if (!HoveredSample.IsValid())
-		{
-			HoveredSample = GetSampleAtMousePosition(MousePosition.X + 1.0f, MousePosition.Y);
 		}
+
+		Reply = FReply::Handled();
 	}
 
 	return Reply;
@@ -854,7 +826,7 @@ FReply SPacketSizesView::OnMouseWheel(const FGeometry& MyGeometry, const FPointe
 
 	if (MouseEvent.GetModifierKeys().IsShiftDown())
 	{
-		FValueAxisViewport& ViewportY = Viewport.GetVerticalAxisViewport();
+		FAxisViewportDouble& ViewportY = Viewport.GetVerticalAxisViewport();
 
 		// Zoom in/out vertically.
 		const float Delta = MouseEvent.GetWheelDelta();
@@ -885,17 +857,6 @@ FReply SPacketSizesView::OnMouseWheel(const FGeometry& MyGeometry, const FPointe
 
 ////////////////////////////////////////////////////////////////////////////////////////////////////
 
-void SPacketSizesView::ZoomHorizontally(const float Delta, const float X)
-{
-	FIndexAxisViewport& ViewportX = Viewport.GetHorizontalAxisViewport();
-	ViewportX.RelativeZoomWithFixedOffset(Delta, X);
-	ViewportX.ScrollAtIndex(ViewportX.GetIndexAtPos(ViewportX.GetPos())); // align pos with at index
-	UpdateHorizontalScrollBar();
-	bIsStateDirty = true;
-}
-
-////////////////////////////////////////////////////////////////////////////////////////////////////
-
 FReply SPacketSizesView::OnMouseButtonDoubleClick(const FGeometry& MyGeometry, const FPointerEvent& MouseEvent)
 {
 	return FReply::Unhandled();
@@ -957,7 +918,7 @@ void SPacketSizesView::ShowContextMenu(const FPointerEvent& MouseEvent)
 
 void SPacketSizesView::ContextMenu_AutoZoom_Execute()
 {
-	FIndexAxisViewport& ViewportX = Viewport.GetHorizontalAxisViewport();
+	FAxisViewportInt32& ViewportX = Viewport.GetHorizontalAxisViewport();
 
 	bIsAutoZoomEnabled = !(bIsAutoZoomEnabled && ViewportX.GetPos() == 0.0f);
 	if (bIsAutoZoomEnabled)
@@ -1007,7 +968,8 @@ void SPacketSizesView::BindCommands()
 
 void SPacketSizesView::HorizontalScrollBar_OnUserScrolled(float ScrollOffset)
 {
-	Viewport.GetHorizontalAxisViewport().OnUserScrolled(HorizontalScrollBar, ScrollOffset);
+	FAxisViewportInt32& ViewportX = Viewport.GetHorizontalAxisViewport();
+	ViewportX.OnUserScrolled(HorizontalScrollBar, ScrollOffset);
 	bIsStateDirty = true;
 }
 
@@ -1015,7 +977,19 @@ void SPacketSizesView::HorizontalScrollBar_OnUserScrolled(float ScrollOffset)
 
 void SPacketSizesView::UpdateHorizontalScrollBar()
 {
-	Viewport.GetHorizontalAxisViewport().UpdateScrollBar(HorizontalScrollBar);
+	FAxisViewportInt32& ViewportX = Viewport.GetHorizontalAxisViewport();
+	ViewportX.UpdateScrollBar(HorizontalScrollBar);
+}
+
+////////////////////////////////////////////////////////////////////////////////////////////////////
+
+void SPacketSizesView::ZoomHorizontally(const float Delta, const float X)
+{
+	FAxisViewportInt32& ViewportX = Viewport.GetHorizontalAxisViewport();
+	ViewportX.RelativeZoomWithFixedOffset(Delta, X);
+	ViewportX.ScrollAtValue(ViewportX.GetValueAtPos(ViewportX.GetPos())); // align viewport position with sample
+	UpdateHorizontalScrollBar();
+	bIsStateDirty = true;
 }
 
 ////////////////////////////////////////////////////////////////////////////////////////////////////
diff --git a/Engine/Source/Developer/TraceInsights/Private/Insights/NetworkingProfiler/Widgets/SPacketSizesView.h b/Engine/Source/Developer/TraceInsights/Private/Insights/NetworkingProfiler/Widgets/SPacketSizesView.h
--- a/Engine/Source/Developer/TraceInsights/Private/Insights/NetworkingProfiler/Widgets/SPacketSizesView.h
+++ b/Engine/Source/Developer/TraceInsights/Private/Insights/NetworkingProfiler/Widgets/SPacketSizesView.h
@@ -12,8 +12,8 @@
 
 // Insights
 #include "Insights/Common/FixedCircularBuffer.h"
-#include "Insights/NetworkingProfiler/ViewModels/PacketSizesViewport.h"
 #include "Insights/NetworkingProfiler/ViewModels/PacketSizesViewHelper.h"
+#include "Insights/NetworkingProfiler/ViewModels/PacketSizesViewport.h"
 
 class SScrollBar;
 
@@ -71,7 +71,7 @@ struct FNetworkPacketSampleRef
 
 ////////////////////////////////////////////////////////////////////////////////////////////////////
 /**
- * Widget used to present frames data in a track.
+ * Widget used to present the network packets as a bar track.
  */
 class SPacketSizesView : public SCompoundWidget
 {
@@ -128,7 +128,7 @@ public:
 
 	virtual FCursorReply OnCursorQuery(const FGeometry& MyGeometry, const FPointerEvent& CursorEvent) const override;
 
-protected:
+private:
 	void UpdateState();
 
 	void DrawHorizontalAxisGrid(FDrawContext& DrawContext, const FSlateBrush* Brush, const FSlateFontInfo& Font) const;
@@ -155,7 +155,7 @@ protected:
 
 	void ZoomHorizontally(const float Delta, const float X);
 
-protected:
+private:
 	/** The track's viewport. Encapsulates info about position and scale. */
 	FPacketSizesViewport Viewport;
 	bool bIsViewportDirty;
@@ -181,7 +181,6 @@ protected:
 	/** Mouse position during the call on mouse button down. */
 	FVector2D MousePositionOnButtonDown;
 	float ViewportPosXOnButtonDown;
-	float ViewportPosYOnButtonDown;
 
 	/** Mouse position during the call on mouse button up. */
 	FVector2D MousePositionOnButtonUp;
@@ -189,7 +188,7 @@ protected:
 	bool bIsLMB_Pressed;
 	bool bIsRMB_Pressed;
 
-	/** True, if the user is currently interactively scrolling the view by holding the right mouse button and dragging. */
+	/** True, if the user is currently interactively scrolling the view (ex.: by holding the left mouse button and dragging). */
 	bool bIsScrolling;
 
 	//////////////////////////////////////////////////
diff --git a/Engine/Source/Developer/TraceInsights/Private/Insights/Table/ViewModels/TreeNodeGrouping.cpp b/Engine/Source/Developer/TraceInsights/Private/Insights/Table/ViewModels/TreeNodeGrouping.cpp
--- a/Engine/Source/Developer/TraceInsights/Private/Insights/Table/ViewModels/TreeNodeGrouping.cpp
+++ b/Engine/Source/Developer/TraceInsights/Private/Insights/Table/ViewModels/TreeNodeGrouping.cpp
@@ -2,8 +2,6 @@
 
 #include "TreeNodeGrouping.h"
 
-#include "EditorStyleSet.h"
-
 #define LOCTEXT_NAMESPACE "Insights_TreeNode"
 
 namespace Insights
diff --git a/Engine/Source/Developer/TraceInsights/Private/Insights/Table/Widgets/STableTreeViewCell.cpp b/Engine/Source/Developer/TraceInsights/Private/Insights/Table/Widgets/STableTreeViewCell.cpp
--- a/Engine/Source/Developer/TraceInsights/Private/Insights/Table/Widgets/STableTreeViewCell.cpp
+++ b/Engine/Source/Developer/TraceInsights/Private/Insights/Table/Widgets/STableTreeViewCell.cpp
@@ -7,6 +7,7 @@
 #include "Widgets/Images/SImage.h"
 #include "Widgets/Input/SButton.h"
 #include "Widgets/SBoxPanel.h"
+#include "Widgets/SToolTip.h"
 #include "Widgets/Text/STextBlock.h"
 #include "Widgets/Views/SExpanderArrow.h"
 
diff --git a/Engine/Source/Developer/TraceInsights/Private/Insights/Table/Widgets/STableTreeViewTooltip.h b/Engine/Source/Developer/TraceInsights/Private/Insights/Table/Widgets/STableTreeViewTooltip.h
--- a/Engine/Source/Developer/TraceInsights/Private/Insights/Table/Widgets/STableTreeViewTooltip.h
+++ b/Engine/Source/Developer/TraceInsights/Private/Insights/Table/Widgets/STableTreeViewTooltip.h
@@ -3,14 +3,9 @@
 #pragma once
 
 #include "CoreMinimal.h"
-#include "EditorStyleSet.h"
-#include "SlateOptMacros.h"
-#include "TraceServices/AnalysisService.h"
-#include "Widgets/Images/SImage.h"
-#include "Widgets/Layout/SGridPanel.h"
-#include "Widgets/Layout/SSeparator.h"
-#include "Widgets/SToolTip.h"
-#include "Widgets/Text/STextBlock.h"
+
+class SToolTip;
+class SGridPanel;
 
 namespace Insights
 {
diff --git a/Engine/Source/Developer/TraceInsights/Private/Insights/TraceInsightsModule.cpp b/Engine/Source/Developer/TraceInsights/Private/Insights/TraceInsightsModule.cpp
--- a/Engine/Source/Developer/TraceInsights/Private/Insights/TraceInsightsModule.cpp
+++ b/Engine/Source/Developer/TraceInsights/Private/Insights/TraceInsightsModule.cpp
@@ -1,7 +1,6 @@
 // Copyright 1998-2019 Epic Games, Inc. All Rights Reserved.
 
 #include "CoreMinimal.h"
-#include "EditorStyleSet.h"
 #include "HAL/PlatformApplicationMisc.h"
 #include "Modules/ModuleManager.h"
 #include "TraceServices/ITraceServicesModule.h"
diff --git a/Engine/Source/Developer/TraceInsights/Private/Insights/ViewModels/BaseTimingTrack.h b/Engine/Source/Developer/TraceInsights/Private/Insights/ViewModels/BaseTimingTrack.h
--- a/Engine/Source/Developer/TraceInsights/Private/Insights/ViewModels/BaseTimingTrack.h
+++ b/Engine/Source/Developer/TraceInsights/Private/Insights/ViewModels/BaseTimingTrack.h
@@ -31,7 +31,7 @@ class FBaseTimingTrack : public TSharedFromThis<FBaseTimingTrack>
 	friend class FTimingViewDrawHelper;
 
 protected:
-	FBaseTimingTrack(uint64 InId)
+	explicit FBaseTimingTrack(uint64 InId)
 		: Id(InId)
 		, PosY(0.0f)
 		, Height(0.0f)
