[Insights]

diff --git a/Engine/Source/Developer/TraceInsights/Private/Insights/ViewModels/DrawHelpers.cpp b/Engine/Source/Developer/TraceInsights/Private/Insights/ViewModels/DrawHelpers.cpp
--- a/Engine/Source/Developer/TraceInsights/Private/Insights/ViewModels/DrawHelpers.cpp
+++ b/Engine/Source/Developer/TraceInsights/Private/Insights/ViewModels/DrawHelpers.cpp
@@ -2,8 +2,6 @@
 
 #include "DrawHelpers.h"
 
-#include "Brushes/SlateBorderBrush.h"
-#include "Brushes/SlateColorBrush.h"
 #include "Fonts/FontMeasure.h"
 #include "Fonts/SlateFontInfo.h"
 
@@ -16,14 +14,15 @@
 
 void FDrawHelpers::DrawBackground(
 	const FDrawContext& DrawContext,
-	const FTimingTrackViewport& Viewport,
 	const FSlateBrush* BackgroundAreaBrush,
 	const FLinearColor& ValidAreaColor,
 	const FLinearColor& InvalidAreaColor,
 	const FLinearColor& EdgeColor,
 	const float X0,
+	const float X1,
+	const float X2,
+	const float X3,
 	const float Y,
-	const float W,
 	const float H,
 	float& OutValidAreaX,
 	float& OutValidAreaW)
@@ -34,17 +33,13 @@ void FDrawHelpers::DrawBackground(
 	//  ++++++|*****|++++++
 	//  ++++++|*****|++++++
 
-	const float X1 = Viewport.TimeToSlateUnitsRounded(Viewport.MinValidTime);
-	const float X2 = Viewport.TimeToSlateUnitsRounded(Viewport.MaxValidTime);
-	const float X3 = X0 + W;
-
 	if (X1 >= X3 || X2 <= X0)
 	{
 		OutValidAreaX = X0;
-		OutValidAreaW = W;
+		OutValidAreaW = X3 - X0;
 
 		// Draw invalid area (entire view).
-		DrawContext.DrawBox(X0, Y, W, H, BackgroundAreaBrush, InvalidAreaColor);
+		DrawContext.DrawBox(X0, Y, X3 - X0, H, BackgroundAreaBrush, InvalidAreaColor);
 	}
 	else // X1 < X3 && X2 > X0
 	{
@@ -91,6 +86,46 @@ void FDrawHelpers::DrawBackground(
 
 ////////////////////////////////////////////////////////////////////////////////////////////////////
 
+void FDrawHelpers::DrawBackground(const FDrawContext& DrawContext,
+								  const FSlateBrush* BackgroundAreaBrush,
+								  const float X0,
+								  const float X1,
+								  const float X2,
+								  const float X3,
+								  const float Y,
+								  const float H)
+{
+	const FLinearColor ValidAreaColor(0.07f, 0.07f, 0.07f, 1.0f);
+	const FLinearColor InvalidAreaColor(0.1f, 0.07f, 0.07f, 1.0f);
+	const FLinearColor EdgeColor(0.05f, 0.05f, 0.05f, 1.0f);
+
+	float ValidAreaX, ValidAreaW;
+	FDrawHelpers::DrawBackground(DrawContext, BackgroundAreaBrush, ValidAreaColor, InvalidAreaColor, EdgeColor, X0, X1, X2, X3, Y, H, ValidAreaX, ValidAreaW);
+}
+
+////////////////////////////////////////////////////////////////////////////////////////////////////
+
+void FDrawHelpers::DrawBackground(const FDrawContext& DrawContext,
+								  const FSlateBrush* BackgroundAreaBrush,
+								  const FTimingTrackViewport& Viewport,
+								  const float Y,
+								  const float H)
+{
+	const FLinearColor ValidAreaColor(0.07f, 0.07f, 0.07f, 1.0f);
+	const FLinearColor InvalidAreaColor(0.1f, 0.07f, 0.07f, 1.0f);
+	const FLinearColor EdgeColor(0.05f, 0.05f, 0.05f, 1.0f);
+
+	const float X0 = 0.0f;
+	const float X1 = Viewport.TimeToSlateUnitsRounded(Viewport.GetMinValidTime());
+	const float X2 = Viewport.TimeToSlateUnitsRounded(Viewport.GetMaxValidTime());
+	const float X3 = FMath::CeilToFloat(Viewport.GetWidth());
+
+	float ValidAreaX, ValidAreaW;
+	FDrawHelpers::DrawBackground(DrawContext, BackgroundAreaBrush, ValidAreaColor, InvalidAreaColor, EdgeColor, X0, X1, X2, X3, Y, H, ValidAreaX, ValidAreaW);
+}
+
+////////////////////////////////////////////////////////////////////////////////////////////////////
+
 void FDrawHelpers::DrawTimeRangeSelection(
 	const FDrawContext& DrawContext,
 	const FTimingTrackViewport& Viewport,
@@ -101,155 +136,182 @@ void FDrawHelpers::DrawTimeRangeSelection(
 {
 	if (EndTime > StartTime)
 	{
+		const float MinX = 0.0f;
+		const float MaxX = Viewport.GetWidth();
+
 		float SelectionX1 = Viewport.TimeToSlateUnitsRounded(StartTime);
 		float SelectionX2 = Viewport.TimeToSlateUnitsRounded(EndTime);
 
-		if (SelectionX1 <= Viewport.Width &&
-			SelectionX2 >= 0)
+		if (SelectionX1 <= MaxX && SelectionX2 >= MinX)
 		{
-			float ClipLeft = 0.0f;
-			float ClipRight = 0.0f;
-			if (SelectionX1 < 0.0f)
-			{
-				ClipLeft = -SelectionX1;
-				SelectionX1 = 0.0f;
-			}
-			if (SelectionX2 > Viewport.Width)
-			{
-				ClipRight = SelectionX2 - Viewport.Width;
-				SelectionX2 = Viewport.Width;
-			}
+			const float Y = 0.0f;
+			const float H = Viewport.GetHeight();
+
+			const double Duration = EndTime - StartTime;
+			const FString Text = TimeUtils::FormatTimeAuto(Duration);
+
+			DrawSelection(DrawContext, MinX, MaxX, SelectionX1, SelectionX2, Y, H, Text, Brush, Font);
+		}
+	}
+}
+
+////////////////////////////////////////////////////////////////////////////////////////////////////
+
+void FDrawHelpers::DrawSelection(
+	const FDrawContext& DrawContext,
+	const float MinX,
+	const float MaxX,
+	float SelectionX1,
+	float SelectionX2,
+	const float Y,
+	const float H,
+	const FString Text,
+	const FSlateBrush* Brush,
+	const FSlateFontInfo& Font)
+{
+	float ClipLeft = 0.0f;
+	float ClipRight = 0.0f;
+	if (SelectionX1 < MinX)
+	{
+		ClipLeft = MinX - SelectionX1;
+		SelectionX1 = MinX;
+	}
+	if (SelectionX2 > MaxX)
+	{
+		ClipRight = SelectionX2 - MaxX;
+		SelectionX2 = MaxX;
+	}
 
-			// Draw selection area.
-			DrawContext.DrawBox(SelectionX1, 0.0f, SelectionX2 - SelectionX1, Viewport.Height, Brush, FLinearColor(0.25f, 0.5f, 1.0f, 0.25f));
-			DrawContext.LayerId++;
+	const FLinearColor BackgroundColor(0.25f, 0.5f, 1.0f, 0.25f);
 
-			const FColor ArrowFillColor(32, 64, 128, 255);
-			const FLinearColor ArrowColor(ArrowFillColor);
+	// Draw selection area.
+	DrawContext.DrawBox(SelectionX1, Y, SelectionX2 - SelectionX1, H, Brush, BackgroundColor);
+	DrawContext.LayerId++;
+
+	const FColor ArrowFillColor(32, 64, 128, 255);
+	const FLinearColor ArrowColor(ArrowFillColor);
+
+	if (SelectionX1 > MinX)
+	{
+		// Draw left side (vertical line).
+		DrawContext.DrawBox(SelectionX1 - 1.0f, Y, 1.0f, H, Brush, ArrowColor);
+	}
+
+	if (SelectionX2 < MaxX)
+	{
+		// Draw right side (vertical line).
+		DrawContext.DrawBox(SelectionX2, Y, 1.0f, H, Brush, ArrowColor);
+	}
+
+	DrawContext.LayerId++;
 
-			if (SelectionX1 > 0.0f)
+	const float ArrowSize = 6.0f;
+	const float ArrowY = Y + 6.0f;
+
+	if (SelectionX2 - SelectionX1 > 2 * ArrowSize)
+	{
+		// Draw horizontal line.
+		float HorizLineX1 = SelectionX1;
+		if (ClipLeft == 0.0f)
+		{
+			HorizLineX1 += 1.0f;
+		}
+		float HorizLineX2 = SelectionX2;
+		if (ClipRight == 0.0f)
+		{
+			HorizLineX2 -= 1.0f;
+		}
+		DrawContext.DrawBox(HorizLineX1, ArrowY - 1.0f, HorizLineX2 - HorizLineX1, 3.0f, Brush, ArrowColor);
+
+		if (ClipLeft < ArrowSize)
+		{
+			// Draw left arrow.
+			for (float AH = 0.0f; AH < ArrowSize; AH += 1.0f)
 			{
-				// Draw left side (vertical line).
-				DrawContext.DrawBox(SelectionX1 - 1.0f, 0.0f, 1.0f, Viewport.Height, Brush, ArrowColor);
+				DrawContext.DrawBox(SelectionX1 - ClipLeft + AH, ArrowY - AH, 1.0f, 2.0f * AH + 1.0f, Brush, ArrowColor);
 			}
+		}
 
-			if (SelectionX2 < Viewport.Width)
+		if (ClipRight < ArrowSize)
+		{
+			// Draw right arrow.
+			for (float AH = 0.0f; AH < ArrowSize; AH += 1.0f)
 			{
-				// Draw right side (vertical line).
-				DrawContext.DrawBox(SelectionX2, 0.0f, 1.0f, Viewport.Height, Brush, ArrowColor);
+				DrawContext.DrawBox(SelectionX2 + ClipRight - AH - 1.0f, ArrowY - AH, 1.0f, 2.0f * AH + 1.0f, Brush, ArrowColor);
 			}
+		}
 
-			DrawContext.LayerId++;
-
-			const float ArrowSize = 6.0f;
-			const float ArrowY = 6.0f;
-
-			if (SelectionX2 - SelectionX1 > 2 * ArrowSize)
-			{
-				// Draw horizontal line.
-				float HorizLineX1 = SelectionX1;
-				if (ClipLeft == 0.0f)
-				{
-					HorizLineX1 += 1.0f;
-				}
-				float HorizLineX2 = SelectionX2;
-				if (ClipRight == 0.0f)
-				{
-					HorizLineX2 -= 1.0f;
-				}
-				DrawContext.DrawBox(HorizLineX1, ArrowY - 1.0f, HorizLineX2 - HorizLineX1, 3.0f, Brush, ArrowColor);
-
-				if (ClipLeft < ArrowSize)
-				{
-					// Draw left arrow.
-					for (float AH = 0.0f; AH < ArrowSize; AH += 1.0f)
-					{
-						DrawContext.DrawBox(SelectionX1 - ClipLeft + AH, ArrowY - AH, 1.0f, 2.0f * AH + 1.0f, Brush, ArrowColor);
-					}
-				}
-
-				if (ClipRight < ArrowSize)
-				{
-					// Draw right arrow.
-					for (float AH = 0.0f; AH < ArrowSize; AH += 1.0f)
-					{
-						DrawContext.DrawBox(SelectionX2 + ClipRight - AH - 1.0f, ArrowY - AH, 1.0f, 2.0f * AH + 1.0f, Brush, ArrowColor);
-					}
-				}
-
-				DrawContext.LayerId++;
+		DrawContext.LayerId++;
 
 #if 0
-				//im: This should be a more efficeint way top draw the arrows, but it renders them with artifacts (missing vertical lines; shader bug?)!
-
-				const FSlateBrush* MyBrush = WhiteBrush;
-				FSlateResourceHandle ResourceHandle = FSlateApplication::Get().GetRenderer()->GetResourceHandle(*MyBrush);
-				const FSlateShaderResourceProxy* ResourceProxy = ResourceHandle.GetResourceProxy();
-
-				FVector2D AtlasOffset = ResourceProxy ? ResourceProxy->StartUV : FVector2D(0.0f, 0.0f);
-				FVector2D AtlasUVSize = ResourceProxy ? ResourceProxy->SizeUV : FVector2D(1.0f, 1.0f);
-
-				const FVector2D Pos = AllottedGeometry.GetAbsolutePosition() + FVector2D(0.0f, 40.0f);
-				const float Scale = AllottedGeometry.Scale;
-
-				FSlateRenderTransform RenderTransform;
-
-				TArray<FSlateVertex> Verts;
-				Verts.Reserve(6);
-				Verts.Add(FSlateVertex::Make<ESlateVertexRounding::Disabled>(RenderTransform, Pos + FVector2D(0.5f + SelectionX1 + ArrowSize, 0.5f + ArrowY + ArrowSize) * Scale, AtlasOffset + FVector2D(0.0f, 1.0f) * AtlasUVSize, ArrowFillColor));
-				Verts.Add(FSlateVertex::Make<ESlateVertexRounding::Disabled>(RenderTransform, Pos + FVector2D(0.5f + SelectionX1, 0.5f + ArrowY) * Scale, AtlasOffset + FVector2D(1.0f, 0.5f) * AtlasUVSize, ArrowFillColor));
-				Verts.Add(FSlateVertex::Make<ESlateVertexRounding::Disabled>(RenderTransform, Pos + FVector2D(0.5f + SelectionX1 + ArrowSize, 0.5f + ArrowY - ArrowSize) * Scale, AtlasOffset + FVector2D(0.0f, 0.0f) * AtlasUVSize, ArrowFillColor));
-				Verts.Add(FSlateVertex::Make<ESlateVertexRounding::Disabled>(RenderTransform, Pos + FVector2D(0.5f + SelectionX2 - ArrowSize, 0.5f + ArrowY - ArrowSize) * Scale, AtlasOffset + FVector2D(0.0f, 0.0f) * AtlasUVSize, ArrowFillColor));
-				Verts.Add(FSlateVertex::Make<ESlateVertexRounding::Disabled>(RenderTransform, Pos + FVector2D(0.5f + SelectionX2, 0.5f + ArrowY) * Scale, AtlasOffset + FVector2D(1.0f, 0.5f) * AtlasUVSize, ArrowFillColor));
-				Verts.Add(FSlateVertex::Make<ESlateVertexRounding::Disabled>(RenderTransform, Pos + FVector2D(0.5f + SelectionX2 - ArrowSize, 0.5f + ArrowY + ArrowSize) * Scale, AtlasOffset + FVector2D(0.0f, 1.0f) * AtlasUVSize, ArrowFillColor));
-
-				TArray<SlateIndex> Indices;
-				Indices.Reserve(6);
-				if (ClipLeft < ArrowSize)
-				{
-					Indices.Add(0);
-					Indices.Add(1);
-					Indices.Add(2);
-				}
-				if (ClipRight < ArrowSize)
-				{
-					Indices.Add(3);
-					Indices.Add(4);
-					Indices.Add(5);
-				}
-
-				FSlateDrawElement::MakeCustomVerts(
-					OutDrawElements,
-					LayerId,
-					ResourceHandle,
-					Verts,
-					Indices,
-					nullptr,
-					0,
-					0,
-					ESlateDrawEffect::PreMultipliedAlpha);
-
-				DrawContext.LayerId++;
-#endif
-			}
+		//im: This should be a more efficeint way top draw the arrows, but it renders them with artifacts (missing vertical lines; shader bug?)!
 
-			//////////////////////////////////////////////////
-			// Draw duration for selected time interval.
+		const FSlateBrush* MyBrush = WhiteBrush;
+		FSlateResourceHandle ResourceHandle = FSlateApplication::Get().GetRenderer()->GetResourceHandle(*MyBrush);
+		const FSlateShaderResourceProxy* ResourceProxy = ResourceHandle.GetResourceProxy();
 
-			const double Duration = EndTime - StartTime;
-			const FString Text = TimeUtils::FormatTimeAuto(Duration);
+		FVector2D AtlasOffset = ResourceProxy ? ResourceProxy->StartUV : FVector2D(0.0f, 0.0f);
+		FVector2D AtlasUVSize = ResourceProxy ? ResourceProxy->SizeUV : FVector2D(1.0f, 1.0f);
 
-			const TSharedRef<FSlateFontMeasure> FontMeasureService = FSlateApplication::Get().GetRenderer()->GetFontMeasureService();
-			const float TextWidth = FontMeasureService->Measure(Text, Font).X;
+		const FVector2D Pos = AllottedGeometry.GetAbsolutePosition() + FVector2D(0.0f, 40.0f);
+		const float Scale = AllottedGeometry.Scale;
 
-			const float CenterX = (SelectionX1 + SelectionX2) / 2.0f;
+		FSlateRenderTransform RenderTransform;
 
-			DrawContext.DrawBox(CenterX - TextWidth / 2 - 2.0, ArrowY - 6.0f, TextWidth + 4.0f, 13.0f, Brush, ArrowColor);
-			DrawContext.LayerId++;
+		TArray<FSlateVertex> Verts;
+		Verts.Reserve(6);
+		Verts.Add(FSlateVertex::Make<ESlateVertexRounding::Disabled>(RenderTransform, Pos + FVector2D(0.5f + SelectionX1 + ArrowSize, 0.5f + ArrowY + ArrowSize) * Scale, AtlasOffset + FVector2D(0.0f, 1.0f) * AtlasUVSize, ArrowFillColor));
+		Verts.Add(FSlateVertex::Make<ESlateVertexRounding::Disabled>(RenderTransform, Pos + FVector2D(0.5f + SelectionX1, 0.5f + ArrowY) * Scale, AtlasOffset + FVector2D(1.0f, 0.5f) * AtlasUVSize, ArrowFillColor));
+		Verts.Add(FSlateVertex::Make<ESlateVertexRounding::Disabled>(RenderTransform, Pos + FVector2D(0.5f + SelectionX1 + ArrowSize, 0.5f + ArrowY - ArrowSize) * Scale, AtlasOffset + FVector2D(0.0f, 0.0f) * AtlasUVSize, ArrowFillColor));
+		Verts.Add(FSlateVertex::Make<ESlateVertexRounding::Disabled>(RenderTransform, Pos + FVector2D(0.5f + SelectionX2 - ArrowSize, 0.5f + ArrowY - ArrowSize) * Scale, AtlasOffset + FVector2D(0.0f, 0.0f) * AtlasUVSize, ArrowFillColor));
+		Verts.Add(FSlateVertex::Make<ESlateVertexRounding::Disabled>(RenderTransform, Pos + FVector2D(0.5f + SelectionX2, 0.5f + ArrowY) * Scale, AtlasOffset + FVector2D(1.0f, 0.5f) * AtlasUVSize, ArrowFillColor));
+		Verts.Add(FSlateVertex::Make<ESlateVertexRounding::Disabled>(RenderTransform, Pos + FVector2D(0.5f + SelectionX2 - ArrowSize, 0.5f + ArrowY + ArrowSize) * Scale, AtlasOffset + FVector2D(0.0f, 1.0f) * AtlasUVSize, ArrowFillColor));
 
-			DrawContext.DrawText(CenterX - TextWidth / 2, ArrowY - 6.0f, Text, Font, FLinearColor::White);
-			DrawContext.LayerId++;
+		TArray<SlateIndex> Indices;
+		Indices.Reserve(6);
+		if (ClipLeft < ArrowSize)
+		{
+			Indices.Add(0);
+			Indices.Add(1);
+			Indices.Add(2);
+		}
+		if (ClipRight < ArrowSize)
+		{
+			Indices.Add(3);
+			Indices.Add(4);
+			Indices.Add(5);
 		}
+
+		FSlateDrawElement::MakeCustomVerts(
+			OutDrawElements,
+			LayerId,
+			ResourceHandle,
+			Verts,
+			Indices,
+			nullptr,
+			0,
+			0,
+			ESlateDrawEffect::PreMultipliedAlpha);
+
+		DrawContext.LayerId++;
+#endif
+	}
+
+	//////////////////////////////////////////////////
+	// Draw text info for selected interval.
+
+	if (!Text.IsEmpty())
+	{
+		const TSharedRef<FSlateFontMeasure> FontMeasureService = FSlateApplication::Get().GetRenderer()->GetFontMeasureService();
+		const float TextWidth = FontMeasureService->Measure(Text, Font).X;
+
+		const float CenterX = (SelectionX1 + SelectionX2) / 2.0f;
+
+		DrawContext.DrawBox(CenterX - TextWidth / 2 - 2.0, ArrowY - 6.0f, TextWidth + 4.0f, 13.0f, Brush, ArrowColor);
+		DrawContext.LayerId++;
+
+		DrawContext.DrawText(CenterX - TextWidth / 2, ArrowY - 6.0f, Text, Font, FLinearColor::White);
+		DrawContext.LayerId++;
 	}
 }
 
diff --git a/Engine/Source/Developer/TraceInsights/Private/Insights/ViewModels/DrawHelpers.h b/Engine/Source/Developer/TraceInsights/Private/Insights/ViewModels/DrawHelpers.h
--- a/Engine/Source/Developer/TraceInsights/Private/Insights/ViewModels/DrawHelpers.h
+++ b/Engine/Source/Developer/TraceInsights/Private/Insights/ViewModels/DrawHelpers.h
@@ -20,20 +20,42 @@ class FDrawHelpers
 public:
 	/**
 	 * Draw background.
+	 * X1-X2 specifies the valid area. X0-X1 and X2-X3 specifies the invalid areas.
+	 *
+	 *  X0    X1    X2     X3
+	 *  ++++++|*****|++++++
+	 *  ++++++|*****|++++++
+	 *  ++++++|*****|++++++
 	 */
 	static void DrawBackground(const FDrawContext& DrawContext,
-							   const FTimingTrackViewport& Viewport,
 							   const FSlateBrush* BackgroundAreaBrush,
 							   const FLinearColor& ValidAreaColor,
 							   const FLinearColor& InvalidAreaColor,
 							   const FLinearColor& EdgeColor,
-							   const float X,
+							   const float X0,
+							   const float X1,
+							   const float X2,
+							   const float X3,
 							   const float Y,
-							   const float W,
 							   const float H,
 							   float& OutValidAreaX,
 							   float& OutValidAreaW);
 
+	static void DrawBackground(const FDrawContext& DrawContext,
+							   const FSlateBrush* BackgroundAreaBrush,
+							   const float X0,
+							   const float X1,
+							   const float X2,
+							   const float X3,
+							   const float Y,
+							   const float H);
+
+	static void DrawBackground(const FDrawContext& DrawContext,
+							   const FSlateBrush* BackgroundAreaBrush,
+							   const FTimingTrackViewport& Viewport,
+							   const float Y,
+							   const float H);
+
 	/**
 	 * Draw time range selection.
 	 */
@@ -43,6 +65,17 @@ public:
 									   const double EndTime,
 									   const FSlateBrush* Brush,
 									   const FSlateFontInfo& Font);
+
+	static void DrawSelection(const FDrawContext& DrawContext,
+							  const float MinX,
+							  const float MaxX,
+							  float SelectionX1,
+							  float SelectionX2,
+							  const float Y,
+							  const float H,
+							  const FString Text,
+							  const FSlateBrush* Brush,
+							  const FSlateFontInfo& Font);
 };
 
 ////////////////////////////////////////////////////////////////////////////////////////////////////
diff --git a/Engine/Source/Developer/TraceInsights/Private/Insights/ViewModels/FileActivityTimingTrack.cpp b/Engine/Source/Developer/TraceInsights/Private/Insights/ViewModels/FileActivityTimingTrack.cpp
--- a/Engine/Source/Developer/TraceInsights/Private/Insights/ViewModels/FileActivityTimingTrack.cpp
+++ b/Engine/Source/Developer/TraceInsights/Private/Insights/ViewModels/FileActivityTimingTrack.cpp
@@ -16,6 +16,7 @@
 #include "Insights/ViewModels/TimingEvent.h"
 #include "Insights/ViewModels/TimingTrackViewport.h"
 #include "Insights/ViewModels/TimingViewDrawHelper.h"
+#include "Insights/ViewModels/TooltipDrawState.h"
 
 #include <limits>
 
@@ -393,43 +394,13 @@ void FFileActivitySharedState::Update()
 // FFileActivityTimingTrack
 ////////////////////////////////////////////////////////////////////////////////////////////////////
 
-void FFileActivityTimingTrack::DrawTooltip(FTimingViewTooltip& Tooltip, const FVector2D& MousePosition, const FTimingEvent& HoveredTimingEvent, const FTimingTrackViewport& Viewport, const FDrawContext& DrawContext, const FSlateBrush* WhiteBrush, const FSlateFontInfo& Font) const
+void FFileActivityTimingTrack::InitTooltip(FTooltipDrawState& Tooltip, const FTimingEvent& HoveredTimingEvent) const
 {
+	Tooltip.ResetContent();
+
 	const Trace::EFileActivityType ActivityType = static_cast<Trace::EFileActivityType>(HoveredTimingEvent.TypeId & 0x0F);
 	const bool bHasFailed = ((HoveredTimingEvent.TypeId & 0xF0) != 0);
-
-	FString FilePath(HoveredTimingEvent.Path);
-
-	constexpr float ValueOffsetX = 50.0f;
-	constexpr float MinValueTextWidth = 220.0f;
-
-	// Compute desired tooltip width.
-	const TSharedRef<FSlateFontMeasure> FontMeasureService = FSlateApplication::Get().GetRenderer()->GetFontMeasureService();
-	const float MinWidth1 = FontMeasureService->Measure(FilePath, Font).X + 2 * FTimingViewTooltip::BorderX;
-	const float MinWidth2 = ValueOffsetX + MinValueTextWidth + 2 * FTimingViewTooltip::BorderX;
-	const float DesiredTooltipWidth = FMath::Max3(MinWidth1, MinWidth2, FTimingViewTooltip::MinWidth);
-
-	constexpr float LineH = 14.0f;
-
-	// Compute desired tooltip height.
-	const int32 LineCount = (ActivityType == Trace::FileActivityType_Read || ActivityType == Trace::FileActivityType_Write) ? 6 : 4;
-	const float MinHeight = LineCount * LineH + 2 * FTimingViewTooltip::BorderY;
-	const float DesiredTooltipHeight = FMath::Max(MinHeight, FTimingViewTooltip::MinHeight);
-
-	Tooltip.Update(MousePosition, DesiredTooltipWidth, DesiredTooltipHeight, Viewport.Width, Viewport.Height);
-
-	const FLinearColor BackgroundColor(0.05f, 0.05f, 0.05f, Tooltip.Opacity);
-	const FLinearColor NameColor(0.9f, 0.9f, 0.5f, Tooltip.Opacity);
-	const FLinearColor TextColor(0.6f, 0.6f, 0.6f, Tooltip.Opacity);
-	const FLinearColor ValueColor(1.0f, 1.0f, 1.0f, Tooltip.Opacity);
-
-	DrawContext.DrawBox(Tooltip.PosX, Tooltip.PosY, Tooltip.Width, Tooltip.Height, WhiteBrush, BackgroundColor);
-	DrawContext.LayerId++;
-
-	const float X = Tooltip.PosX + FTimingViewTooltip::BorderX;
-	const float ValueX = X + ValueOffsetX;
-	float Y = Tooltip.PosY + FTimingViewTooltip::BorderY;
-
+	
 	FString TypeStr;
 	uint32 TypeColor;
 	if (bHasFailed)
@@ -447,37 +418,20 @@ void FFileActivityTimingTrack::DrawTooltip(FTimingViewTooltip& Tooltip, const FV
 	TypeLinearColor.R *= 2.0f;
 	TypeLinearColor.G *= 2.0f;
 	TypeLinearColor.B *= 2.0f;
-	DrawContext.DrawText(X, Y, TypeStr, Font, TypeLinearColor);
-	Y += LineH;
-
-	DrawContext.DrawText(X, Y, FilePath, Font, NameColor);
-	Y += LineH;
+	Tooltip.AddTitle(TypeStr, TypeLinearColor);
 
-	DrawContext.DrawText(X, Y, TEXT("Duration:"), Font, TextColor);
-	FString DurationStr = TimeUtils::FormatTimeAuto(HoveredTimingEvent.Duration());
-	DrawContext.DrawText(ValueX, Y, DurationStr, Font, ValueColor);
-	Y += LineH;
+	Tooltip.AddTitle(HoveredTimingEvent.Path);
 
-	DrawContext.DrawText(X, Y, TEXT("Depth:"), Font, TextColor);
-	FString DepthStr = FString::Printf(TEXT("%d"), HoveredTimingEvent.Depth);
-	DrawContext.DrawText(ValueX, Y, DepthStr, Font, ValueColor);
-	Y += LineH;
+	Tooltip.AddNameValueTextLine(TEXT("Duration:"), TimeUtils::FormatTimeAuto(HoveredTimingEvent.Duration()));
+	Tooltip.AddNameValueTextLine(TEXT("Depth:"), FString::Printf(TEXT("%d"), HoveredTimingEvent.Depth));
 
 	if (ActivityType == Trace::FileActivityType_Read || ActivityType == Trace::FileActivityType_Write)
 	{
-		DrawContext.DrawText(X, Y, TEXT("Offset:"), Font, TextColor);
-		FString OffsetStr = FText::AsNumber(HoveredTimingEvent.Offset).ToString();
-		DrawContext.DrawText(ValueX, Y, OffsetStr, Font, ValueColor);
-		Y += LineH;
-
-		DrawContext.DrawText(X, Y, TEXT("Size:"), Font, TextColor);
-		FString SizeStr = FText::AsNumber(HoveredTimingEvent.Size).ToString();
-		SizeStr += TEXT(" bytes");
-		DrawContext.DrawText(ValueX, Y, SizeStr, Font, ValueColor);
-		Y += LineH;
+		Tooltip.AddNameValueTextLine(TEXT("Offset:"), FText::AsNumber(HoveredTimingEvent.Offset).ToString() + TEXT(" bytes"));
+		Tooltip.AddNameValueTextLine(TEXT("Size:"), FText::AsNumber(HoveredTimingEvent.Size).ToString() + TEXT(" bytes"));
 	}
 
-	DrawContext.LayerId++;
+	Tooltip.UpdateLayout();
 }
 
 ////////////////////////////////////////////////////////////////////////////////////////////////////
@@ -590,11 +544,11 @@ void FOverviewFileActivityTimingTrack::Draw(FTimingViewDrawHelper& Helper) const
 			//const double EventEndTime = Event.EndTime; // keep duration of events
 			const double EventEndTime = Event.StartTime; // make all 0 duration events
 
-			if (EventEndTime <= Helper.GetViewport().StartTime)
+			if (EventEndTime <= Helper.GetViewport().GetStartTime())
 			{
 				continue;
 			}
-			if (Event.StartTime >= Helper.GetViewport().EndTime)
+			if (Event.StartTime >= Helper.GetViewport().GetEndTime())
 			{
 				break;
 			}
@@ -636,11 +590,11 @@ void FDetailedFileActivityTimingTrack::Draw(FTimingViewDrawHelper& Helper) const
 		{
 			for (const TSharedPtr<FFileActivitySharedState::FIoFileActivity> Activity : State->FileActivities)
 			{
-				if (Activity->EndTime <= Helper.GetViewport().StartTime)
+				if (Activity->EndTime <= Helper.GetViewport().GetStartTime())
 				{
 					continue;
 				}
-				if (Activity->StartTime >= Helper.GetViewport().EndTime)
+				if (Activity->StartTime >= Helper.GetViewport().GetEndTime())
 				{
 					break;
 				}
@@ -654,11 +608,11 @@ void FDetailedFileActivityTimingTrack::Draw(FTimingViewDrawHelper& Helper) const
 		// Draw IO file activity foreground events.
 		for (const FFileActivitySharedState::FIoTimingEvent& Event : State->AllIoEvents)
 		{
-			if (Event.EndTime <= Helper.GetViewport().StartTime)
+			if (Event.EndTime <= Helper.GetViewport().GetStartTime())
 			{
 				continue;
 			}
-			if (Event.StartTime >= Helper.GetViewport().EndTime)
+			if (Event.StartTime >= Helper.GetViewport().GetEndTime())
 			{
 				break;
 			}
diff --git a/Engine/Source/Developer/TraceInsights/Private/Insights/ViewModels/FileActivityTimingTrack.h b/Engine/Source/Developer/TraceInsights/Private/Insights/ViewModels/FileActivityTimingTrack.h
--- a/Engine/Source/Developer/TraceInsights/Private/Insights/ViewModels/FileActivityTimingTrack.h
+++ b/Engine/Source/Developer/TraceInsights/Private/Insights/ViewModels/FileActivityTimingTrack.h
@@ -68,14 +68,14 @@ private:
 class FFileActivityTimingTrack : public FTimingEventsTrack
 {
 public:
-	FFileActivityTimingTrack(uint64 InTrackId, const FName SubType, const FString& InName, TSharedPtr<FFileActivitySharedState> InState)
+	explicit FFileActivityTimingTrack(uint64 InTrackId, const FName SubType, const FString& InName, TSharedPtr<FFileActivitySharedState> InState)
 		: FTimingEventsTrack(InTrackId, FName(TEXT("FileActivity")), SubType, InName)
 		, State(InState)
 	{
 	}
 	virtual ~FFileActivityTimingTrack() {}
 
-	virtual void DrawTooltip(FTimingViewTooltip& Tooltip, const FVector2D& MousePosition, const FTimingEvent& HoveredTimingEvent, const FTimingTrackViewport& Viewport, const FDrawContext& DrawContext, const FSlateBrush* WhiteBrush, const FSlateFontInfo& Font) const override;
+	virtual void InitTooltip(FTooltipDrawState& Tooltip, const FTimingEvent& HoveredTimingEvent) const override;
 
 protected:
 	bool SearchEvent(const double InStartTime, const double InEndTime, TFunctionRef<bool(double, double, uint32)> InPredicate, FTimingEvent& InOutTimingEvent, bool bInStopAtFirstMatch, bool bInSearchForLargestEvent, bool bIgnoreEventDepth) const;
@@ -89,7 +89,7 @@ protected:
 class FOverviewFileActivityTimingTrack : public FFileActivityTimingTrack
 {
 public:
-	FOverviewFileActivityTimingTrack(uint64 InTrackId, TSharedPtr<FFileActivitySharedState> InState)
+	explicit FOverviewFileActivityTimingTrack(uint64 InTrackId, TSharedPtr<FFileActivitySharedState> InState)
 		: FFileActivityTimingTrack(InTrackId, FName(TEXT("Overview")), TEXT("I/O Overview"), InState)
 	{
 	}
@@ -107,7 +107,7 @@ public:
 class FDetailedFileActivityTimingTrack : public FFileActivityTimingTrack
 {
 public:
-	FDetailedFileActivityTimingTrack(uint64 InTrackId, TSharedPtr<FFileActivitySharedState> InState)
+	explicit FDetailedFileActivityTimingTrack(uint64 InTrackId, TSharedPtr<FFileActivitySharedState> InState)
 		: FFileActivityTimingTrack(InTrackId, FName(TEXT("Detailed")), TEXT("I/O Activity"), InState)
 	{
 	}
diff --git a/Engine/Source/Developer/TraceInsights/Private/Insights/ViewModels/FrameTrackHelper.cpp b/Engine/Source/Developer/TraceInsights/Private/Insights/ViewModels/FrameTrackHelper.cpp
--- a/Engine/Source/Developer/TraceInsights/Private/Insights/ViewModels/FrameTrackHelper.cpp
+++ b/Engine/Source/Developer/TraceInsights/Private/Insights/ViewModels/FrameTrackHelper.cpp
@@ -2,9 +2,6 @@
 
 #include "FrameTrackHelper.h"
 
-#include "Brushes/SlateBorderBrush.h"
-#include "Brushes/SlateColorBrush.h"
-#include "EditorStyleSet.h"
 #include "Fonts/FontMeasure.h"
 #include "Fonts/SlateFontInfo.h"
 #include "Framework/Application/SlateApplication.h"
@@ -13,6 +10,8 @@
 
 // Insights
 #include "Insights/Common/PaintUtils.h"
+#include "Insights/InsightsStyle.h"
+#include "Insights/ViewModels/DrawHelpers.h"
 #include "Insights/ViewModels/FrameTrackViewport.h"
 
 #include <limits>
@@ -78,8 +77,8 @@ void FFrameTrackSeriesBuilder::AddFrame(const Trace::FFrame& Frame)
 FFrameTrackDrawHelper::FFrameTrackDrawHelper(const FDrawContext& InDrawContext, const FFrameTrackViewport& InViewport)
 	: DrawContext(InDrawContext)
 	, Viewport(InViewport)
-	, WhiteBrush(FCoreStyle::Get().GetBrush("WhiteBrush"))
-	, BorderBrush(FEditorStyle::GetBrush("PlainBorder"))
+	, WhiteBrush(FInsightsStyle::Get().GetBrush("WhiteBrush"))
+	, HoveredFrameBorderBrush(FInsightsStyle::Get().GetBrush("HoveredEventBorder"))
 	, NumFrames(0)
 	, NumDrawSamples(0)
 {
@@ -89,48 +88,17 @@ FFrameTrackDrawHelper::FFrameTrackDrawHelper(const FDrawContext& InDrawContext,
 
 void FFrameTrackDrawHelper::DrawBackground() const
 {
-	const FSlateBrush* AreaBrush = WhiteBrush;
-	const FLinearColor ValidAreaColor(0.07f, 0.07f, 0.07f, 1.0f);
-	const FLinearColor InvalidAreaColor(0.1f, 0.07f, 0.07f, 1.0f);
+	const FAxisViewportInt32& ViewportX = Viewport.GetHorizontalAxisViewport();
 
-	const FIndexAxisViewport& ViewportX = Viewport.GetHorizontalAxisViewport();
+	const float X0 = 0.0f;
+	const float X1 = ViewportX.GetMinPos() - ViewportX.GetPos();
+	const float X2 = ViewportX.GetMaxPos() - ViewportX.GetPos();
+	const float X3 = FMath::CeilToFloat(Viewport.GetWidth());
 
-	const float X0 = ViewportX.GetMinPos() - ViewportX.GetPos();
-	const float X1 = ViewportX.GetMaxPos() - ViewportX.GetPos();
-
-	const float W = FMath::CeilToFloat(Viewport.GetWidth());
+	const float Y = 0.0f;
 	const float H = FMath::CeilToFloat(Viewport.GetHeight());
 
-	if (X0 >= W || X1 <= 0.0f)
-	{
-		// Draw invalid area (entire view).
-		DrawContext.DrawBox(0.0f, 0.0f, W, H, AreaBrush, InvalidAreaColor);
-	}
-	else // X0 < W && X1 > 0
-	{
-		if (X0 > 0.0f)
-		{
-			// Draw invalid area (left).
-			DrawContext.DrawBox(0.0f, 0.0f, X0, H, AreaBrush, InvalidAreaColor);
-		}
-
-		if (X1 < W)
-		{
-			// Draw invalid area (right).
-			DrawContext.DrawBox(X1, 0.0f, W - X1, H, AreaBrush, InvalidAreaColor);
-		}
-
-		const float ValidX0 = FMath::Max(X0, 0.0f);
-		const float ValidX1 = FMath::Min(X1, W);
-
-		if (ValidX1 > ValidX0)
-		{
-			// Draw valid area.
-			DrawContext.DrawBox(ValidX0, 0.0f, ValidX1 - ValidX0, H, AreaBrush, ValidAreaColor);
-		}
-	}
-
-	DrawContext.LayerId++;
+	FDrawHelpers::DrawBackground(DrawContext, WhiteBrush, X0, X1, X2, X3, Y, H);
 }
 
 ////////////////////////////////////////////////////////////////////////////////////////////////////
@@ -170,7 +138,7 @@ void FFrameTrackDrawHelper::DrawCached(const FFrameTrackSeries& Series) const
 	const float SampleW = Viewport.GetSampleWidth();
 	const int32 NumSamples = Series.Samples.Num();
 
-	const FValueAxisViewport& ViewportY = Viewport.GetVerticalAxisViewport();
+	const FAxisViewportDouble& ViewportY = Viewport.GetVerticalAxisViewport();
 
 	const float ViewHeight = FMath::RoundToFloat(Viewport.GetHeight());
 	const float BaselineY = FMath::RoundToFloat(ViewportY.GetOffsetForValue(0.0));
@@ -245,7 +213,7 @@ void FFrameTrackDrawHelper::DrawHoveredSample(const FFrameTrackSample& Sample) c
 	const int32 SampleIndex = (Sample.LargestFrameIndex - FirstFrameIndex) / FramesPerSample;
 	const float X = SampleIndex * SampleW;
 
-	const FValueAxisViewport& ViewportY = Viewport.GetVerticalAxisViewport();
+	const FAxisViewportDouble& ViewportY = Viewport.GetVerticalAxisViewport();
 
 	const float ViewHeight = FMath::RoundToFloat(Viewport.GetHeight());
 	const float BaselineY = FMath::RoundToFloat(ViewportY.GetOffsetForValue(0.0));
@@ -263,7 +231,7 @@ void FFrameTrackDrawHelper::DrawHoveredSample(const FFrameTrackSample& Sample) c
 	const float Y = ViewHeight - H;
 
 	const FLinearColor ColorBorder(1.0f, 1.0f, 0.0f, 1.0);
-	DrawContext.DrawBox(X - 1.0f, Y - 1.0f, SampleW + 2.0f, H + 2.0f, BorderBrush, ColorBorder);
+	DrawContext.DrawBox(X - 1.0f, Y - 1.0f, SampleW + 2.0f, H + 2.0f, HoveredFrameBorderBrush, ColorBorder);
 	DrawContext.LayerId++;
 }
 
@@ -289,7 +257,7 @@ void FFrameTrackDrawHelper::DrawHighlightedInterval(const FFrameTrackSeries& Ser
 	{
 		const float SampleW = Viewport.GetSampleWidth();
 		float X1 = Index1 * SampleW;
-		float X2 = Index2 * SampleW;
+		float X2 = (Index2 + 1) * SampleW;
 
 		constexpr float Y1 = 0.0f; // allows 12px for the horizontal scrollbar (one displayed on top of the track)
 		const float Y2 = Viewport.GetHeight();
diff --git a/Engine/Source/Developer/TraceInsights/Private/Insights/ViewModels/FrameTrackHelper.h b/Engine/Source/Developer/TraceInsights/Private/Insights/ViewModels/FrameTrackHelper.h
--- a/Engine/Source/Developer/TraceInsights/Private/Insights/ViewModels/FrameTrackHelper.h
+++ b/Engine/Source/Developer/TraceInsights/Private/Insights/ViewModels/FrameTrackHelper.h
@@ -70,10 +70,10 @@ struct FFrameTrackSeries
 {
 	int32 FrameType;
 	bool bIsVisible;
-	int32 NumAggregatedFrames; // total number of frames aggregated in samples; ie. sum of all Sample.NumFrames
-	TArray<FFrameTrackSample> Samples;
+	int32 NumAggregatedFrames; // total number of frames aggregated in samples; i.e. sum of all Sample.NumFrames
+	TArray<FFrameTrackSample> Samples; // the aggregated samples
 
-	FFrameTrackSeries(int32 InFrameType)
+	explicit FFrameTrackSeries(int32 InFrameType)
 		: FrameType(InFrameType)
 		, bIsVisible(true)
 		, NumAggregatedFrames(0)
@@ -93,7 +93,7 @@ struct FFrameTrackSeries
 class FFrameTrackSeriesBuilder
 {
 public:
-	FFrameTrackSeriesBuilder(FFrameTrackSeries& InSeries, const FFrameTrackViewport& InViewport);
+	explicit FFrameTrackSeriesBuilder(FFrameTrackSeries& InSeries, const FFrameTrackViewport& InViewport);
 
 	/**
 	 * Non-copyable
@@ -123,7 +123,7 @@ private:
 class FFrameTrackDrawHelper
 {
 public:
-	FFrameTrackDrawHelper(const FDrawContext& InDrawContext, const FFrameTrackViewport& InViewport);
+	explicit FFrameTrackDrawHelper(const FDrawContext& InDrawContext, const FFrameTrackViewport& InViewport);
 
 	/**
 	 * Non-copyable
@@ -146,7 +146,7 @@ private:
 	const FFrameTrackViewport& Viewport;
 
 	const FSlateBrush* WhiteBrush;
-	const FSlateBrush* BorderBrush;
+	const FSlateBrush* HoveredFrameBorderBrush;
 
 	// Debug stats.
 	mutable int32 NumFrames;
diff --git a/Engine/Source/Developer/TraceInsights/Private/Insights/ViewModels/FrameTrackViewport.cpp b/Engine/Source/Developer/TraceInsights/Private/Insights/ViewModels/FrameTrackViewport.cpp
--- a/Engine/Source/Developer/TraceInsights/Private/Insights/ViewModels/FrameTrackViewport.cpp
+++ b/Engine/Source/Developer/TraceInsights/Private/Insights/ViewModels/FrameTrackViewport.cpp
@@ -2,132 +2,4 @@
 
 #include "FrameTrackViewport.h"
 
-#include "Widgets/Layout/SScrollBar.h"
-
-////////////////////////////////////////////////////////////////////////////////////////////////////
-// FIndexAxisViewport
-////////////////////////////////////////////////////////////////////////////////////////////////////
-
-void FIndexAxisViewport::RelativeZoomWithFixedOffset(const float Delta, const float Offset)
-{
-	constexpr float ZoomStep = 0.25f; // as percent
-
-	float NewScale;
-
-	if (Delta > 0)
-	{
-		NewScale = Scale * FMath::Pow(1.0f + ZoomStep, Delta);
-	}
-	else
-	{
-		NewScale = Scale * FMath::Pow(1.0f / (1.0f + ZoomStep), -Delta);
-	}
-
-	// Snap to integer value of either: "number of indices per sample" or "number of pixels per sample".
-	if (NewScale < 1.0f)
-	{
-		// N indices/sample; 1 pixel/sample
-		if (Delta > 0)
-		{
-			NewScale = 1.0f / FMath::FloorToFloat(1.0f / NewScale);
-		}
-		else
-		{
-			NewScale = 1.0f / FMath::CeilToFloat(1.0f / NewScale);
-		}
-	}
-	else
-	{
-		// 1 index/sample; N pixels/sample
-		if (Delta > 0)
-		{
-			NewScale = FMath::CeilToFloat(NewScale);
-		}
-		else
-		{
-			NewScale = FMath::FloorToFloat(NewScale);
-		}
-	}
-
-	ZoomWithFixedOffset(NewScale, Offset);
-}
-
-////////////////////////////////////////////////////////////////////////////////////////////////////
-
-bool FIndexAxisViewport::UpdatePosWithinLimits()
-{
-	float MinPos, MaxPos;
-	if (MaxPosition - MinPosition < Size)
-	{
-		MinPos = MaxPosition - Size;
-		MaxPos = MinPosition;
-	}
-	else
-	{
-		constexpr float ExtraSizeFactor = 0.15f; // allow extra 15% on sides
-		MinPos = MinPosition - ExtraSizeFactor * Size;
-		MaxPos = MaxPosition - (1.0f - ExtraSizeFactor) * Size;
-	}
-
-	constexpr float U = 0.5f; // interpolation factor --> animation speed
-
-	float Pos = Position;
-	if (Pos < MinPos)
-	{
-		Pos = Pos * U + (1.0f - U) * MinPos;
-		if (FMath::IsNearlyEqual(Pos, MinPos, 0.5f))
-		{
-			Pos = MinPos;
-		}
-	}
-	else if (Pos > MaxPos)
-	{
-		Pos = Pos * U + (1.0f - U) * MaxPos;
-		if (FMath::IsNearlyEqual(Pos, MaxPos, 0.5f))
-		{
-			Pos = MaxPos;
-		}
-
-		if (Pos < MinPos)
-		{
-			Pos = MinPos;
-		}
-	}
-
-	if (Position != Pos)
-	{
-		//ScrollAtPos(Pos);
-		ScrollAtIndex(GetIndexAtPos(Pos)); // snap to sample
-		return true;
-	}
-
-	return false;
-}
-
-////////////////////////////////////////////////////////////////////////////////////////////////////
-
-void FIndexAxisViewport::OnUserScrolled(TSharedPtr<SScrollBar> ScrollBar, float ScrollOffset)
-{
-	const float SX = 1.0 / (MaxPosition - MinPosition);
-	const float ThumbSizeFraction = FMath::Clamp<float>(Size * SX, 0.0f, 1.0f);
-	const float OffsetFraction = FMath::Clamp<float>(ScrollOffset, 0.0f, 1.0f - ThumbSizeFraction);
-
-	const float Pos = MinPosition + OffsetFraction * (MaxPosition - MinPosition);
-	ScrollAtPos(Pos);
-
-	ScrollBar->SetState(OffsetFraction, ThumbSizeFraction);
-}
-
-////////////////////////////////////////////////////////////////////////////////////////////////////
-
-void FIndexAxisViewport::UpdateScrollBar(TSharedPtr<SScrollBar> ScrollBar) const
-{
-	const float SX = 1.0 / (MaxPosition - MinPosition);
-	const float ThumbSizeFraction = FMath::Clamp<float>(Size * SX, 0.0f, 1.0f);
-	const float ScrollOffset = (Position - MinPosition) * SX;
-	const float OffsetFraction = FMath::Clamp<float>(ScrollOffset, 0.0f, 1.0f - ThumbSizeFraction);
-
-	ScrollBar->SetState(OffsetFraction, ThumbSizeFraction);
-}
-
 ////////////////////////////////////////////////////////////////////////////////////////////////////
diff --git a/Engine/Source/Developer/TraceInsights/Private/Insights/ViewModels/FrameTrackViewport.h b/Engine/Source/Developer/TraceInsights/Private/Insights/ViewModels/FrameTrackViewport.h
--- a/Engine/Source/Developer/TraceInsights/Private/Insights/ViewModels/FrameTrackViewport.h
+++ b/Engine/Source/Developer/TraceInsights/Private/Insights/ViewModels/FrameTrackViewport.h
@@ -4,333 +4,8 @@
 
 #include "CoreMinimal.h"
 
-class SScrollBar;
-
-////////////////////////////////////////////////////////////////////////////////////////////////////
-
-class FIndexAxisViewport
-{
-private:
-	static constexpr float SLATE_UNITS_TOLERANCE = 0.1f;
-
-public:
-	FIndexAxisViewport()
-	{
-		Reset();
-	}
-
-	void Reset()
-	{
-		Size = 0.0f;
-
-		MinIndex = 0;
-		MaxIndex = 0;
-
-		MinPosition = 0.0f;
-		MaxPosition = 0.0f;
-		Position = 0.0f;
-
-		MinScale = 0.0001f; // 10000 indices / sample; 1 pixel / sample
-		MaxScale = 16.0f; // 1 index / sample; 16 pixels / sample
-		Scale = MaxScale;
-	}
-
-	FString ToDebugString(const TCHAR* Prefix, const TCHAR* Unit) const
-	{
-		if (Scale >= 1.0f)
-		{
-			return FString::Printf(TEXT("%sScale = %.3f (%d pixels/%s), Pos = %.2f"),
-				Prefix, Scale, FMath::RoundToInt(Scale), Unit, Position);
-		}
-		else
-		{
-			return FString::Printf(TEXT("%sScale = %.3f (%d %ss/pixel), Pos = %.2f"),
-				Prefix, Scale, FMath::RoundToInt(1.0f / Scale), Unit, Position);
-		}
-	}
-
-	float GetSize() const { return Size; }
-
-	bool SetSize(const float InSize)
-	{
-		if (!FMath::IsNearlyEqual(Size, InSize, SLATE_UNITS_TOLERANCE))
-		{
-			Size = InSize;
-			OnSizeChanged();
-			return true;
-		}
-		return false;
-	}
-
-	int32 GetMinIndex() const { return MinIndex; }
-	int32 GetMaxIndex() const { return MaxIndex; }
-
-	float GetMinPos() const { return MinPosition; }
-	float GetMaxPos() const { return MaxPosition; }
-	float GetPos() const { return Position; }
-
-	float GetMinScale() const { return MinScale; }
-	float GetMaxScale() const { return MaxScale; }
-	float GetScale() const { return Scale; }
-
-	float GetSampleSize() const { return FMath::Max(1.0f, FMath::RoundToFloat(Scale)); }
-	int32 GetNumIndicesPerSample() const { return FMath::Max(1, FMath::RoundToInt(1.0f / Scale)); }
-
-	void SetMinMaxIndexInterval(const int32 InMinIndex, const int32 InMaxIndex)
-	{
-		MinIndex = InMinIndex;
-		MaxIndex = InMaxIndex;
-		UpdateMinMax();
-	}
-
-	int32 GetIndexAtPos(const float Pos) const
-	{
-		return FMath::RoundToInt(Pos / Scale);
-	}
-
-	int32 GetIndexAtOffset(const float Offset) const
-	{
-		return FMath::RoundToInt((Position + Offset) / Scale);
-	}
-
-	float GetPosForIndex(const int32 Index) const
-	{
-		return static_cast<float>(Index) * Scale;
-	}
-
-	float GetOffsetForIndex(const int32 Index) const
-	{
-		return static_cast<float>(Index) * Scale - Position;
-	}
-
-	float GetRoundedOffsetForIndex(const int32 Index) const
-	{
-		return FMath::RoundToFloat(static_cast<float>(Index) * Scale - Position);
-	}
-
-	void ScrollAtPos(const float Pos)
-	{
-		Position = Pos;
-		OnPositionChanged();
-	}
-
-	void ScrollAtIndex(const int32 Index)
-	{
-		Position = static_cast<float>(Index) * Scale;
-		OnPositionChanged();
-	}
-
-	void CenterOnIndex(const int32 Index)
-	{
-		const float SampleSize = FMath::Max(1.0f, Scale);
-		if (SampleSize > Size)
-		{
-			ScrollAtIndex(Index);
-		}
-		else
-		{
-			const float Offset = (Size - SampleSize) / 2.0f;
-			Position = static_cast<float>(Index) * Scale - Offset;
-			OnPositionChanged();
-		}
-	}
-
-	void CenterOnIndexInterval(const int32 IntervalStartIndex, const int32 IntervalEndIndex)
-	{
-		const float SampleIntervalSize = FMath::Max(1.0f, Scale * static_cast<float>(IntervalEndIndex - IntervalStartIndex));
-		if (SampleIntervalSize > Size)
-		{
-			ScrollAtIndex(IntervalStartIndex);
-		}
-		else
-		{
-			const float Offset = (Size - SampleIntervalSize) / 2.0f;
-			Position = static_cast<float>(IntervalStartIndex) * Scale - Offset;
-			OnPositionChanged();
-		}
-	}
-
-	void ZoomWithFixedOffset(const float NewScale, const float Offset)
-	{
-		const float LocalNewScale = FMath::Clamp(NewScale, MinScale, MaxScale);
-		if (LocalNewScale != Scale)
-		{
-			// Index at Offset should remain the same. So we resolve equation:
-			//   (Position + Offset) / Scale == (NewPosition + Offset) / NewScale
-			//   ==> NewPosition = (Position + Offset) / Scale * NewScale - Offset
-			Position = (Position + Offset) * (LocalNewScale / Scale) - Offset;
-			Scale = LocalNewScale;
-
-			OnPositionChanged();
-			OnScaleChanged();
-		}
-	}
-
-	void RelativeZoomWithFixedOffset(const float Delta, const float Offset);
-
-	bool UpdatePosWithinLimits();
-
-	void OnUserScrolled(TSharedPtr<SScrollBar> ScrollBar, float ScrollOffset);
-	void UpdateScrollBar(TSharedPtr<SScrollBar> ScrollBar) const;
-
-private:
-	void OnSizeChanged()
-	{
-	}
-
-	void OnPositionChanged()
-	{
-	}
-
-	void OnScaleChanged()
-	{
-		UpdateMinMax();
-	}
-
-	void UpdateMinMax()
-	{
-		MinPosition = static_cast<float>(MinIndex) * Scale;
-		MaxPosition = static_cast<float>(MaxIndex) * Scale;
-
-		//float NewPosition = FMath::Clamp(Position, MinPosition, MaxPosition);
-		//if (FMath::IsNearlyEqual(NewPosition, Position, SLATE_UNITS_TOLERANCE))
-		//{
-		//	Position = NewPosition;
-		//	OnPositionChanged();
-		//}
-	}
-
-private:
-	float Size; // size of viewport (ex.: the viewport's width if this is a horizontal axis), in Slate units
-
-	int32 MinIndex; // minimum index
-	int32 MaxIndex; // maximum index (exclusive)
-
-	float MinPosition; // minimum position (corresponding to MinIndex), in Slate units
-	float MaxPosition; // maximum position (corresponding to MaxIndex), in Slate units
-	float Position; // current position of the viewport, in Slate units
-
-	float MinScale; // minimum scale factor
-	float MaxScale; // maximum scale factor
-	float Scale; // current scale factor between viewport index units and Slate units: if > 1, it represents number of pixels (Slate units) for one viewport index, otherwise it represents number of viewport indices in one pixel (Slate unit)
-};
-
-////////////////////////////////////////////////////////////////////////////////////////////////////
-
-class FValueAxisViewport
-{
-private:
-	static constexpr float SLATE_UNITS_TOLERANCE = 0.1f;
-
-public:
-	FValueAxisViewport()
-	{
-		Reset();
-	}
-
-	void Reset()
-	{
-		Size = 0.0f;
-
-		MinValue = 0.0;
-		MaxValue = 0.0;
-
-		MinPosition = 0.0f;
-		MaxPosition = 0.0f;
-		Position = 0.0f;
-
-		MinScale = 0.001;
-		MaxScale = 1000000.0;
-		Scale = 1500.0;
-	}
-
-	FString ToDebugString(const TCHAR* Prefix) const
-	{
-		return FString::Printf(TEXT("%sScale = %g, Pos = %.2f"), Prefix, Scale, Position);
-	}
-
-	float GetSize() const { return Size; }
-
-	bool SetSize(const float InSize)
-	{
-		if (!FMath::IsNearlyEqual(Size, InSize, SLATE_UNITS_TOLERANCE))
-		{
-			Size = InSize;
-			OnSizeChanged();
-			return true;
-		}
-		return false;
-	}
-
-	float GetMinPos() const { return MinPosition; }
-	float GetMaxPos() const { return MaxPosition; }
-	float GetPos() const { return Position; }
-
-	double GetMinScale() const { return MinScale; }
-	double GetMaxScale() const { return MaxScale; }
-	double GetScale() const { return Scale; }
-
-	float GetPosForValue(const double Value) const
-	{
-		return static_cast<float>(Value * Scale);
-	}
-
-	float GetOffsetForValue(const double Value) const
-	{
-		return static_cast<float>(Value * Scale) - Position;
-	}
-
-	double GetValueForOffset(const float Offset) const
-	{
-		return static_cast<double>(Position + Offset) / Scale;
-	}
-
-	void SetScaleLimits(const double InMinScale, const double InMaxScale)
-	{
-		MinScale = InMinScale;
-		MaxScale = InMaxScale;
-	}
-
-	bool SetScale(const double InNewScale)
-	{
-		float LocalNewScale = FMath::Clamp(InNewScale, MinScale, MaxScale);
-		if (LocalNewScale != Scale)
-		{
-			Scale = LocalNewScale;
-			OnScaleChanged();
-			return true;
-		}
-		return false;
-	}
-
-private:
-	void OnSizeChanged()
-	{
-	}
-
-	void OnPositionChanged()
-	{
-	}
-
-	void OnScaleChanged()
-	{
-		//UpdateMinMaxX();
-	}
-
-private:
-	float Size; // size of viewport (ex.: the viewport's height if this is a vertical axis), in Slate units
-
-	double MinValue;
-	double MaxValue;
-
-	float MinPosition; // minimum scroll position (i.e. position corresponding to MinValue), in Slate units
-	float MaxPosition; // maximum scroll position (i.e. position corresponding to MaxValue), in Slate units
-	float Position; // current position of the viewport, in Slate units
-
-	double MinScale; // minimum scale factor
-	double MaxScale; // maximum scale factor
-	double Scale; // scale factor between Value units and Slate units
-};
+#include "Insights/ViewModels/AxisViewportInt32.h"
+#include "Insights/ViewModels/AxisViewportDouble.h"
 
 ////////////////////////////////////////////////////////////////////////////////////////////////////
 
@@ -351,11 +26,11 @@ public:
 		VerticalAxisViewport.Reset();
 	}
 
-	const FIndexAxisViewport& GetHorizontalAxisViewport() const { return HorizontalAxisViewport; }
-	FIndexAxisViewport& GetHorizontalAxisViewport() { return HorizontalAxisViewport; }
+	const FAxisViewportInt32& GetHorizontalAxisViewport() const { return HorizontalAxisViewport; }
+	FAxisViewportInt32& GetHorizontalAxisViewport() { return HorizontalAxisViewport; }
 
-	const FValueAxisViewport& GetVerticalAxisViewport() const { return VerticalAxisViewport; }
-	FValueAxisViewport& GetVerticalAxisViewport() { return VerticalAxisViewport; }
+	const FAxisViewportDouble& GetVerticalAxisViewport() const { return VerticalAxisViewport; }
+	FAxisViewportDouble& GetVerticalAxisViewport() { return VerticalAxisViewport; }
 
 	float GetWidth() const { return HorizontalAxisViewport.GetSize(); }
 	float GetHeight() const { return VerticalAxisViewport.GetSize(); }
@@ -373,8 +48,8 @@ public:
 	}
 
 	float GetSampleWidth() const { return HorizontalAxisViewport.GetSampleSize(); }
-	int32 GetNumFramesPerSample() const { return HorizontalAxisViewport.GetNumIndicesPerSample(); }
-	int32 GetFirstFrameIndex() const { return HorizontalAxisViewport.GetIndexAtOffset(0.0f); }
+	int32 GetNumFramesPerSample() const { return HorizontalAxisViewport.GetNumSamplesPerPixel(); }
+	int32 GetFirstFrameIndex() const { return HorizontalAxisViewport.GetValueAtOffset(0.0f); }
 
 private:
 	void OnSizeChanged()
@@ -382,8 +57,8 @@ private:
 	}
 
 private:
-	FIndexAxisViewport HorizontalAxisViewport;
-	FValueAxisViewport VerticalAxisViewport;
+	FAxisViewportInt32 HorizontalAxisViewport;
+	FAxisViewportDouble VerticalAxisViewport;
 };
 
 ////////////////////////////////////////////////////////////////////////////////////////////////////
diff --git a/Engine/Source/Developer/TraceInsights/Private/Insights/ViewModels/GraphTrack.cpp b/Engine/Source/Developer/TraceInsights/Private/Insights/ViewModels/GraphTrack.cpp
--- a/Engine/Source/Developer/TraceInsights/Private/Insights/ViewModels/GraphTrack.cpp
+++ b/Engine/Source/Developer/TraceInsights/Private/Insights/ViewModels/GraphTrack.cpp
@@ -2,8 +2,6 @@
 
 #include "GraphTrack.h"
 
-#include "Brushes/SlateColorBrush.h"
-#include "Brushes/SlateBorderBrush.h"
 #include "EditorStyleSet.h"
 #include "Framework/MultiBox/MultiBoxBuilder.h"
 #include "Rendering/DrawElements.h"
@@ -12,6 +10,7 @@
 // Insights
 #include "Insights/Common/PaintUtils.h"
 #include "Insights/Common/TimeUtils.h"
+#include "Insights/InsightsStyle.h"
 #include "Insights/ViewModels/DrawHelpers.h"
 #include "Insights/ViewModels/TimingTrackViewport.h"
 
@@ -121,10 +120,9 @@ FString FGraphSeries::FormatValue(double Value) const
 FGraphTrack::FGraphTrack(uint64 InTrackId)
 	: FBaseTimingTrack(InTrackId)
 	//, AllSeries()
-	, WhiteBrush(FCoreStyle::Get().GetBrush("WhiteBrush"))
+	, WhiteBrush(FInsightsStyle::Get().GetBrush("WhiteBrush"))
 	, PointBrush(FEditorStyle::GetBrush("Graph.ExecutionBubble"))
-	, BorderBrush(new FSlateBorderBrush(NAME_None, FMargin(1.0f)))
-	//, BorderBrush(FEditorStyle::GetBrush("PlainBorder"))
+	, BorderBrush(FInsightsStyle::Get().GetBrush("SingleBorder"))
 	, Font(FCoreStyle::GetDefaultFontStyle("Regular", 8))
 	, bDrawPoints(true)
 	, bDrawPointsWithBorder(true)
@@ -132,6 +130,9 @@ FGraphTrack::FGraphTrack(uint64 InTrackId)
 	, bDrawPolygon(true)
 	, bUseEventDuration(true)
 	, bDrawBoxes(true)
+	, HoveredGraphEvent()
+	, Tooltip()
+	, TooltipMousePosition(0.0f, 0.0f)
 	, NumAddedEvents(0)
 	, NumDrawPoints(0)
 	, NumDrawLines(0)
@@ -143,13 +144,16 @@ FGraphTrack::FGraphTrack(uint64 InTrackId)
 
 FGraphTrack::~FGraphTrack()
 {
-	delete BorderBrush;
 }
 
 ////////////////////////////////////////////////////////////////////////////////////////////////////
 
 void FGraphTrack::Reset()
 {
+	HoveredGraphEvent.Reset();
+	Tooltip.Reset();
+	TooltipMousePosition.Set(0.0f, 0.0f);
+
 	AllSeries.Reset();
 }
 
@@ -160,7 +164,7 @@ void FGraphTrack::UpdateHoveredState(float MouseX, float MouseY, const FTimingTr
 	constexpr float HeaderWidth = 100.0f;
 	constexpr float HeaderHeight = 14.0f;
 
-	HoveredGraphEvent.Series = nullptr;
+	HoveredGraphEvent.Reset();
 
 	if (MouseY >= GetPosY() && MouseY < GetPosY() + GetHeight())
 	{
@@ -173,6 +177,29 @@ void FGraphTrack::UpdateHoveredState(float MouseX, float MouseY, const FTimingTr
 	{
 		SetHoveredState(false);
 	}
+
+	if (HoveredGraphEvent.IsValid())
+	{
+		InitTooltip();
+		Tooltip.SetDesiredOpacity(1.0f);
+	}
+	else
+	{
+		Tooltip.SetDesiredOpacity(0.0f);
+	}
+	TooltipMousePosition.Set(MouseX, MouseY);
+}
+
+////////////////////////////////////////////////////////////////////////////////////////////////////
+
+void FGraphTrack::Update(const FTimingTrackViewport& Viewport)
+{
+	Tooltip.Update();
+
+	if (!TooltipMousePosition.IsZero())
+	{
+		Tooltip.SetPosition(TooltipMousePosition, 0.0f, Viewport.GetWidth(), 0.0f, Viewport.GetHeight());
+	}
 }
 
 ////////////////////////////////////////////////////////////////////////////////////////////////////
@@ -198,18 +225,7 @@ void FGraphTrack::UpdateStats()
 
 void FGraphTrack::DrawBackground(FDrawContext& DrawContext, const FTimingTrackViewport& Viewport) const
 {
-	const FLinearColor ValidAreaColor(0.07f, 0.07f, 0.07f, 1.0f);
-	const FLinearColor InvalidAreaColor(0.1f, 0.07f, 0.07f, 1.0f);
-	const FLinearColor EdgeColor(0.05f, 0.05f, 0.05f, 1.0f);
-
-	const float X = 0.0f;
-	const float W = FMath::CeilToFloat(Viewport.Width);
-	const float Y = GetPosY();
-	const float H = GetHeight();
-
-	float ValidAreaX, ValidAreaW;
-
-	FDrawHelpers::DrawBackground(DrawContext, Viewport, WhiteBrush, ValidAreaColor, InvalidAreaColor, EdgeColor, X, Y, W, H, ValidAreaX, ValidAreaW);
+	FDrawHelpers::DrawBackground(DrawContext, WhiteBrush, Viewport, GetPosY(), GetHeight());
 }
 
 ////////////////////////////////////////////////////////////////////////////////////////////////////
@@ -221,7 +237,7 @@ void FGraphTrack::Draw(FDrawContext& DrawContext, const FTimingTrackViewport& Vi
 	{
 		FVector2D AbsPos = DrawContext.Geometry.GetAbsolutePosition();
 		const float L = AbsPos.X;
-		const float R = AbsPos.X + Viewport.Width;
+		const float R = AbsPos.X + Viewport.GetWidth();
 		const float T = AbsPos.Y + GetPosY();
 		const float B = AbsPos.Y + GetPosY() + GetHeight();
 		FSlateClippingZone ClipZone(FVector2D(L, T), FVector2D(R, T), FVector2D(L, B), FVector2D(R, B));
@@ -238,7 +254,7 @@ void FGraphTrack::Draw(FDrawContext& DrawContext, const FTimingTrackViewport& Vi
 		const float BaselineY = FMath::RoundToFloat(static_cast<float>(Series->GetBaselineY()));
 		if (BaselineY >= 0.0f && BaselineY < GetHeight())
 		{
-			DrawContext.DrawBox(0.0f, GetPosY() + BaselineY, Viewport.Width, 1.0f, WhiteBrush, FLinearColor(0.05f, 0.05f, 0.05f, 1.0f));
+			DrawContext.DrawBox(0.0f, GetPosY() + BaselineY, Viewport.GetWidth(), 1.0f, WhiteBrush, FLinearColor(0.05f, 0.05f, 0.05f, 1.0f));
 			DrawContext.LayerId++;
 		}
 	}
@@ -263,14 +279,7 @@ void FGraphTrack::Draw(FDrawContext& DrawContext, const FTimingTrackViewport& Vi
 
 void FGraphTrack::PostDraw(FDrawContext& DrawContext, const FTimingTrackViewport& Viewport, const FVector2D& MousePosition) const
 {
-	if (HoveredGraphEvent.IsValid() && !MousePosition.IsZero())
-	{
-		DrawTooltip(DrawContext, Viewport, MousePosition, HoveredGraphEvent);
-	}
-	else
-	{
-		Tooltip.Reset();
-	}
+	Tooltip.Draw(DrawContext);
 }
 
 ////////////////////////////////////////////////////////////////////////////////////////////////////
@@ -562,62 +571,14 @@ void FGraphTrack::DrawHighlightedEvent(FDrawContext& DrawContext, const FTimingT
 
 ////////////////////////////////////////////////////////////////////////////////////////////////////
 
-void FGraphTrack::DrawTooltip(FDrawContext& DrawContext, const FTimingTrackViewport& Viewport, const FVector2D& MousePosition, const FGraphTrack::FEvent& GraphEvent) const
+void FGraphTrack::InitTooltip()
 {
-	const TSharedRef<FSlateFontMeasure> FontMeasureService = FSlateApplication::Get().GetRenderer()->GetFontMeasureService();
-
-	constexpr float ValueOffsetX = 48.0f;
-
-	float DesiredTooltipWidth = FTimingViewTooltip::MinWidth;
-
-	const float DesiredTooltipWidthByName = FontMeasureService->Measure(GraphEvent.Series->GetName(), Font).X + 2 * FTimingViewTooltip::BorderX;
-	if (DesiredTooltipWidth < DesiredTooltipWidthByName)
-	{
-		DesiredTooltipWidth = DesiredTooltipWidthByName;
-	}
-
-	FString ValueStr = GraphEvent.Series->FormatValue(GraphEvent.SeriesEvent.Value);
-	const float DesiredTooltipWidthByValue = ValueOffsetX + FontMeasureService->Measure(ValueStr, Font).X + 2 * FTimingViewTooltip::BorderX;
-	if (DesiredTooltipWidth < DesiredTooltipWidthByValue)
-	{
-		DesiredTooltipWidth = DesiredTooltipWidthByValue;
-	}
-
-	constexpr float LineH = 14.0f;
-	const float DesiredTooltipHeight = 4 * LineH + 2 * FTimingViewTooltip::BorderY;
-
-	Tooltip.Update(MousePosition, DesiredTooltipWidth, DesiredTooltipHeight, Viewport.Width, Viewport.Height);
-
-	const FLinearColor BackgroundColor(0.05f, 0.05f, 0.05f, Tooltip.Opacity);
-	const FLinearColor TextColor(0.6f, 0.6f, 0.6f, Tooltip.Opacity);
-	const FLinearColor ValueColor(1.0f, 1.0f, 1.0f, Tooltip.Opacity);
-
-	DrawContext.DrawBox(Tooltip.PosX, Tooltip.PosY, Tooltip.Width, Tooltip.Height, WhiteBrush, BackgroundColor);
-	DrawContext.LayerId++;
-
-	float X = Tooltip.PosX + FTimingViewTooltip::BorderX;
-	float Y = Tooltip.PosY + FTimingViewTooltip::BorderY;
-
-	DrawContext.DrawText(X, Y, GraphEvent.Series->GetName().ToString(), Font, GraphEvent.Series->GetColor());
-	Y += LineH;
-
-	const float ValueX = X + ValueOffsetX;
-
-	DrawContext.DrawTextAligned(HAlign_Right, ValueX - 4.0f, Y, TEXT("Time:"), Font, TextColor);
-	FString TimeStr = TimeUtils::FormatTimeAuto(GraphEvent.SeriesEvent.Time);
-	DrawContext.DrawText(ValueX, Y, TimeStr, Font, ValueColor);
-	Y += LineH;
-
-	DrawContext.DrawTextAligned(HAlign_Right, ValueX - 4.0f, Y, TEXT("Duration:"), Font, TextColor);
-	FString DurationStr = TimeUtils::FormatTimeAuto(GraphEvent.SeriesEvent.Duration);
-	DrawContext.DrawText(ValueX, Y, DurationStr, Font, ValueColor);
-	Y += LineH;
-
-	DrawContext.DrawTextAligned(HAlign_Right, ValueX - 4.0f, Y, TEXT("Value:"), Font, TextColor);
-	DrawContext.DrawText(ValueX, Y, ValueStr, Font, ValueColor);
-	Y += LineH;
-
-	DrawContext.LayerId++;
+	Tooltip.ResetContent();
+	Tooltip.AddTitle(HoveredGraphEvent.Series->GetName().ToString(), HoveredGraphEvent.Series->GetColor());
+	Tooltip.AddNameValueTextLine(TEXT("Time:"), TimeUtils::FormatTimeAuto(HoveredGraphEvent.SeriesEvent.Time));
+	Tooltip.AddNameValueTextLine(TEXT("Duration:"), TimeUtils::FormatTimeAuto(HoveredGraphEvent.SeriesEvent.Duration));
+	Tooltip.AddNameValueTextLine(TEXT("Value:"), HoveredGraphEvent.Series->FormatValue(HoveredGraphEvent.SeriesEvent.Value));
+	Tooltip.UpdateLayout();
 }
 
 ////////////////////////////////////////////////////////////////////////////////////////////////////
@@ -981,6 +942,8 @@ FRandomGraphTrack::~FRandomGraphTrack()
 
 void FRandomGraphTrack::Update(const FTimingTrackViewport& Viewport)
 {
+	FGraphTrack::Update(Viewport);
+
 	if (IsDirty())
 	{
 		ClearDirtyFlag();
@@ -1041,7 +1004,7 @@ void FRandomGraphTrack::GenerateSeries(FGraphSeries& Series, const FTimingTrackV
 		FGraphTrackBuilder Builder(*this, Series, Viewport);
 
 		int32 Index = 0;
-		while (Index < EventCount && Events[Index].Time < Viewport.StartTime)
+		while (Index < EventCount && Events[Index].Time < Viewport.GetStartTime())
 		{
 			++Index;
 		}
@@ -1054,7 +1017,7 @@ void FRandomGraphTrack::GenerateSeries(FGraphSeries& Series, const FTimingTrackV
 			const FGraphEvent& Ev = Events[Index];
 			Builder.AddEvent(Ev.Time, Ev.Duration, Ev.Value);
 
-			if (Ev.Time > Viewport.EndTime)
+			if (Ev.Time > Viewport.GetEndTime())
 			{
 				// one point outside screen (right side)
 				break;
@@ -1146,13 +1109,13 @@ void FGraphTrackBuilder::BeginPoints()
 bool FGraphTrackBuilder::AddPoint(double Time, double Value)
 {
 	const float X = Viewport.TimeToSlateUnitsRounded(Time);
-	if (X < -FGraphTrack::PointVisualSize / 2.0f || X >= Viewport.Width + FGraphTrack::PointVisualSize / 2.0f)
+	if (X < -FGraphTrack::PointVisualSize / 2.0f || X >= Viewport.GetWidth() + FGraphTrack::PointVisualSize / 2.0f)
 	{
 		return false;
 	}
 
 	// Align the X with a grid of GraphTrackPointDX pixels in size, in the global space (i.e. scroll independent).
-	const double AlignedX = FMath::RoundToDouble(Time * Viewport.ScaleX / FGraphTrack::PointSizeX) * FGraphTrack::PointSizeX;
+	const double AlignedX = FMath::RoundToDouble(Time * Viewport.GetScaleX() / FGraphTrack::PointSizeX) * FGraphTrack::PointSizeX;
 
 	if (AlignedX > PointsCurrentX + FGraphTrack::PointSizeX - 0.5)
 	{
@@ -1236,7 +1199,7 @@ void FGraphTrackBuilder::AddConnectedLine(double Time, double Value)
 		return;
 	}
 
-	if (X >= Viewport.Width)
+	if (X >= Viewport.GetWidth())
 	{
 		if (!bIsLastLineAdded)
 		{
@@ -1321,7 +1284,7 @@ void FGraphTrackBuilder::BeginBoxes()
 void FGraphTrackBuilder::AddBox(double Time, double Duration, double Value)
 {
 	float X1 = Viewport.TimeToSlateUnitsRounded(Time);
-	if (X1 > Viewport.Width)
+	if (X1 > Viewport.GetWidth())
 	{
 		return;
 	}
diff --git a/Engine/Source/Developer/TraceInsights/Private/Insights/ViewModels/GraphTrack.h b/Engine/Source/Developer/TraceInsights/Private/Insights/ViewModels/GraphTrack.h
--- a/Engine/Source/Developer/TraceInsights/Private/Insights/ViewModels/GraphTrack.h
+++ b/Engine/Source/Developer/TraceInsights/Private/Insights/ViewModels/GraphTrack.h
@@ -7,7 +7,7 @@
 
 // Insights
 #include "Insights/ViewModels/BaseTimingTrack.h"
-#include "Insights/ViewModels/TimingViewDrawHelper.h"
+#include "Insights/ViewModels/TooltipDrawState.h"
 
 struct FDrawContext;
 struct FSlateBrush;
@@ -183,6 +183,7 @@ public:
 		TSharedPtr<FGraphSeries> Series;
 		FGraphSeries::FEvent SeriesEvent;
 
+		void Reset() { Series.Reset(); }
 		bool IsValid() const { return Series.IsValid(); }
 	};
 
@@ -195,14 +196,14 @@ private:
 	static constexpr float PointSizeY = 3.0f;
 
 public:
-	FGraphTrack(uint64 InTrackId);
+	explicit FGraphTrack(uint64 InTrackId);
 	virtual ~FGraphTrack();
 
 	virtual void Reset() override;
 
 	virtual void UpdateHoveredState(float MouseX, float MouseY, const FTimingTrackViewport& Viewport) override;
 
-	virtual void Update(const FTimingTrackViewport& Viewport) override = 0;
+	virtual void Update(const FTimingTrackViewport& Viewport) override;
 
 	void Draw(FDrawContext& DrawContext, const FTimingTrackViewport& Viewport, const FVector2D& MousePosition) const;
 	void PostDraw(FDrawContext& DrawContext, const FTimingTrackViewport& Viewport, const FVector2D& MousePosition) const;
@@ -234,7 +235,8 @@ protected:
 	void DrawBackground(FDrawContext& DrawContext, const FTimingTrackViewport& Viewport) const;
 	void DrawSeries(const FGraphSeries& Series, FDrawContext& DrawContext, const FTimingTrackViewport& Viewport) const;
 	void DrawHighlightedEvent(FDrawContext& DrawContext, const FTimingTrackViewport& Viewport, const FGraphTrack::FEvent& GraphEvent) const;
-	void DrawTooltip(FDrawContext& DrawContext, const FTimingTrackViewport& Viewport, const FVector2D& MousePosition, const FGraphTrack::FEvent& GraphEvent) const;
+
+	void InitTooltip();
 
 private:
 	void ContextMenu_ShowPoints_Execute();
@@ -282,7 +284,8 @@ protected:
 	bool bDrawBoxes;
 
 	mutable FGraphTrack::FEvent HoveredGraphEvent;
-	mutable FTimingViewTooltip Tooltip;
+	FTooltipDrawState Tooltip;
+	FVector2D TooltipMousePosition;
 
 	// Stats
 	int32 NumAddedEvents; // total event count
@@ -296,7 +299,7 @@ protected:
 class FRandomGraphTrack : public FGraphTrack
 {
 public:
-	FRandomGraphTrack(uint64 InTrackId);
+	explicit FRandomGraphTrack(uint64 InTrackId);
 	virtual ~FRandomGraphTrack();
 
 	virtual void Update(const FTimingTrackViewport& Viewport) override;
@@ -322,7 +325,7 @@ private:
 	};
 
 public:
-	FGraphTrackBuilder(FGraphTrack& InTrack, FGraphSeries& InSeries, const FTimingTrackViewport& InViewport);
+	explicit FGraphTrackBuilder(FGraphTrack& InTrack, FGraphSeries& InSeries, const FTimingTrackViewport& InViewport);
 	~FGraphTrackBuilder();
 
 	/**
diff --git a/Engine/Source/Developer/TraceInsights/Private/Insights/ViewModels/LoadingTimingTrack.cpp b/Engine/Source/Developer/TraceInsights/Private/Insights/ViewModels/LoadingTimingTrack.cpp
--- a/Engine/Source/Developer/TraceInsights/Private/Insights/ViewModels/LoadingTimingTrack.cpp
+++ b/Engine/Source/Developer/TraceInsights/Private/Insights/ViewModels/LoadingTimingTrack.cpp
@@ -14,6 +14,7 @@
 #include "Insights/ViewModels/TimingEvent.h"
 #include "Insights/ViewModels/TimingTrackViewport.h"
 #include "Insights/ViewModels/TimingViewDrawHelper.h"
+#include "Insights/ViewModels/TooltipDrawState.h"
 
 #define LOCTEXT_NAMESPACE "LoadingTimingTrack"
 
@@ -136,110 +137,36 @@ void FLoadingSharedState::SetColorSchema(int32 Schema)
 // FLoadingTimingTrack
 ////////////////////////////////////////////////////////////////////////////////////////////////////
 
-void FLoadingTimingTrack::DrawTooltip(FTimingViewTooltip& Tooltip, const FVector2D& MousePosition, const FTimingEvent& HoveredTimingEvent, const FTimingTrackViewport& Viewport, const FDrawContext& DrawContext, const FSlateBrush* WhiteBrush, const FSlateFontInfo& Font) const
+void FLoadingTimingTrack::InitTooltip(FTooltipDrawState& Tooltip, const FTimingEvent& HoveredTimingEvent) const
 {
-	FString EventName(State->GetEventNameEx(HoveredTimingEvent.Depth, HoveredTimingEvent.LoadingInfo));
+	Tooltip.ResetContent();
+
+	Tooltip.AddTitle(State->GetEventNameEx(HoveredTimingEvent.Depth, HoveredTimingEvent.LoadingInfo));
 
 	const Trace::FPackageInfo* Package = HoveredTimingEvent.LoadingInfo.Package;
 	const Trace::FPackageExportInfo* Export = HoveredTimingEvent.LoadingInfo.Export;
 
-	FString PackageName = Package ? Package->Name : TEXT("N/A");
-
-	constexpr float ValueOffsetX = 100.0f;
-	constexpr float MinValueTextWidth = 220.0f;
-
-	// Compute desired tooltip width.
-	const TSharedRef<FSlateFontMeasure> FontMeasureService = FSlateApplication::Get().GetRenderer()->GetFontMeasureService();
-	const float NameWidth = FontMeasureService->Measure(EventName, Font).X;
-	const float PackageNameWidth = FontMeasureService->Measure(PackageName, Font).X;
-	const float MinWidth1 = NameWidth + 2 * FTimingViewTooltip::BorderX;
-	const float MinWidth2 = ValueOffsetX + FMath::Max(PackageNameWidth, MinValueTextWidth) + 2 * FTimingViewTooltip::BorderX;
-	const float DesiredTooltipWidth = FMath::Max3(MinWidth1, MinWidth2, FTimingViewTooltip::MinWidth);
+	Tooltip.AddNameValueTextLine(TEXT("Duration:"), TimeUtils::FormatTimeAuto(HoveredTimingEvent.Duration()));
+	Tooltip.AddNameValueTextLine(TEXT("Depth:"), FString::Printf(TEXT("%d"), HoveredTimingEvent.Depth));
+	Tooltip.AddNameValueTextLine(TEXT("Package Event:"), ::GetName(HoveredTimingEvent.LoadingInfo.PackageEventType));
 
-	constexpr float LineH = 14.0f;
-
-	// Compute desired tooltip height.
-	int32 LineCount = 5;
 	if (Package)
 	{
-		LineCount += 3;
+		Tooltip.AddNameValueTextLine(TEXT("Package Name:"), Package->Name);
+		Tooltip.AddNameValueTextLine(TEXT("Header Size:"), FString::Printf(TEXT("%s bytes"), *FText::AsNumber(Package->Summary.TotalHeaderSize).ToString()));
+		Tooltip.AddNameValueTextLine(TEXT("Package Summary:"), FString::Printf(TEXT("%d names, %d imports, %d exports"), Package->Summary.NameCount, Package->Summary.ImportCount, Package->Summary.ExportCount));
 	}
-	if (Export)
-	{
-		LineCount += 2;
-	}
-	const float MinHeight = LineCount * LineH + 2 * FTimingViewTooltip::BorderY;
-	const float DesiredTooltipHeight = FMath::Max(MinHeight, FTimingViewTooltip::MinHeight);
-
-	Tooltip.Update(MousePosition, DesiredTooltipWidth, DesiredTooltipHeight, Viewport.Width, Viewport.Height);
-
-	const FLinearColor BackgroundColor(0.05f, 0.05f, 0.05f, Tooltip.Opacity);
-	const FLinearColor NameColor(0.9f, 0.9f, 0.5f, Tooltip.Opacity);
-	const FLinearColor TextColor(0.6f, 0.6f, 0.6f, Tooltip.Opacity);
-	const FLinearColor ValueColor(1.0f, 1.0f, 1.0f, Tooltip.Opacity);
-
-	DrawContext.DrawBox(Tooltip.PosX, Tooltip.PosY, Tooltip.Width, Tooltip.Height, WhiteBrush, BackgroundColor);
-	DrawContext.LayerId++;
-
-	const float X = Tooltip.PosX + FTimingViewTooltip::BorderX;
-	const float ValueX = X + ValueOffsetX;
-	float Y = Tooltip.PosY + FTimingViewTooltip::BorderY;
 
-	DrawContext.DrawText(X, Y, EventName, Font, NameColor);
-	Y += LineH;
-
-	DrawContext.DrawText(X, Y, TEXT("Duration:"), Font, TextColor);
-	FString InclStr = TimeUtils::FormatTimeAuto(HoveredTimingEvent.Duration());
-	DrawContext.DrawText(ValueX, Y, InclStr, Font, ValueColor);
-	Y += LineH;
-
-	DrawContext.DrawText(X, Y, TEXT("Depth:"), Font, TextColor);
-	FString DepthStr = FString::Printf(TEXT("%d"), HoveredTimingEvent.Depth);
-	DrawContext.DrawText(ValueX, Y, DepthStr, Font, ValueColor);
-	Y += LineH;
-
-	DrawContext.DrawText(X, Y, TEXT("Package Event:"), Font, TextColor);
-	DrawContext.DrawText(ValueX, Y, ::GetName(HoveredTimingEvent.LoadingInfo.PackageEventType), Font, ValueColor);
-	Y += LineH;
-
-	if (Package)
-	{
-		DrawContext.DrawText(X, Y, TEXT("Package Name:"), Font, TextColor);
-		DrawContext.DrawText(ValueX, Y, PackageName, Font, ValueColor);
-		Y += LineH;
-
-		DrawContext.DrawText(X, Y, TEXT("Header Size:"), Font, TextColor);
-		FString HeaderSizeStr = FString::Printf(TEXT("%d bytes"), Package->Summary.TotalHeaderSize);
-		DrawContext.DrawText(ValueX, Y, HeaderSizeStr, Font, ValueColor);
-		Y += LineH;
-
-		DrawContext.DrawText(X, Y, TEXT("Package Summary:"), Font, TextColor);
-		FString SummaryStr = FString::Printf(TEXT("%d names, %d imports, %d exports"), Package->Summary.NameCount, Package->Summary.ImportCount, Package->Summary.ExportCount);
-		DrawContext.DrawText(ValueX, Y, SummaryStr, Font, ValueColor);
-		Y += LineH;
-	}
-
-	{
-		DrawContext.DrawText(X, Y, TEXT("Export Event:"), Font, TextColor);
-		FString ExportTypeStr = FString::Printf(TEXT("%s%s"), ::GetName(HoveredTimingEvent.LoadingInfo.ExportEventType), Export && Export->IsAsset ? TEXT(" [asset]") : TEXT(""));
-		DrawContext.DrawText(ValueX, Y, ExportTypeStr, Font, ValueColor);
-		Y += LineH;
-	}
+	Tooltip.AddNameValueTextLine(TEXT("Export Event:"), FString::Printf(TEXT("%s%s"), ::GetName(HoveredTimingEvent.LoadingInfo.ExportEventType), Export && Export->IsAsset ? TEXT(" [asset]") : TEXT("")));
 
 	if (Export)
 	{
-		DrawContext.DrawText(X, Y, TEXT("Export Class:"), Font, TextColor);
-		FString ClassStr = FString::Printf(TEXT("%s"), Export->Class ? Export->Class->Name : TEXT("N/A"));
-		DrawContext.DrawText(ValueX, Y, ClassStr, Font, ValueColor);
-		Y += LineH;
-
-		DrawContext.DrawText(X, Y, TEXT("Serial:"), Font, TextColor);
-		FString SerialStr = FString::Printf(TEXT("Offset: %llu, Size: %llu%s"), Export->SerialOffset, Export->SerialSize);
-		DrawContext.DrawText(ValueX, Y, SerialStr, Font, ValueColor);
-		Y += LineH;
+		Tooltip.AddNameValueTextLine(TEXT("Export Class:"), Export->Class ? Export->Class->Name : TEXT("N/A"));
+		Tooltip.AddNameValueTextLine(TEXT("Serial Offset:"), FString::Printf(TEXT("%s bytes"), *FText::AsNumber(Export->SerialOffset).ToString()));
+		Tooltip.AddNameValueTextLine(TEXT("Serial Size:"), FString::Printf(TEXT("%s bytes"), *FText::AsNumber(Export->SerialSize).ToString()));
 	}
 
-	DrawContext.LayerId++;
+	Tooltip.UpdateLayout();
 }
 
 ////////////////////////////////////////////////////////////////////////////////////////////////////
@@ -314,8 +241,8 @@ void FLoadingMainThreadTimingTrack::Draw(FTimingViewDrawHelper& Helper) const
 			{
 				if (FTimingEventsTrack::bUseDownSampling)
 				{
-					const double SecondsPerPixel = 1.0 / Helper.GetViewport().ScaleX;
-					Timeline.EnumerateEventsDownSampled(Helper.GetViewport().StartTime, Helper.GetViewport().EndTime, SecondsPerPixel, [this, &Helper](double StartTime, double EndTime, uint32 Depth, const Trace::FLoadTimeProfilerCpuEvent& Event)
+					const double SecondsPerPixel = 1.0 / Helper.GetViewport().GetScaleX();
+					Timeline.EnumerateEventsDownSampled(Helper.GetViewport().GetStartTime(), Helper.GetViewport().GetEndTime(), SecondsPerPixel, [this, &Helper](double StartTime, double EndTime, uint32 Depth, const Trace::FLoadTimeProfilerCpuEvent& Event)
 					{
 						const TCHAR* Name = State->GetEventNameEx(Depth, Event);
 						Helper.AddEvent(StartTime, EndTime, Depth, Name);
@@ -323,7 +250,7 @@ void FLoadingMainThreadTimingTrack::Draw(FTimingViewDrawHelper& Helper) const
 				}
 				else
 				{
-					Timeline.EnumerateEvents(Helper.GetViewport().StartTime, Helper.GetViewport().EndTime, [this, &Helper](double StartTime, double EndTime, uint32 Depth, const Trace::FLoadTimeProfilerCpuEvent& Event)
+					Timeline.EnumerateEvents(Helper.GetViewport().GetStartTime(), Helper.GetViewport().GetEndTime(), [this, &Helper](double StartTime, double EndTime, uint32 Depth, const Trace::FLoadTimeProfilerCpuEvent& Event)
 					{
 						const TCHAR* Name = State->GetEventNameEx(Depth, Event);
 						Helper.AddEvent(StartTime, EndTime, Depth, Name);
@@ -392,8 +319,8 @@ void FLoadingAsyncThreadTimingTrack::Draw(FTimingViewDrawHelper& Helper) const
 			{
 				if (FTimingEventsTrack::bUseDownSampling)
 				{
-					const double SecondsPerPixel = 1.0 / Helper.GetViewport().ScaleX;
-					Timeline.EnumerateEventsDownSampled(Helper.GetViewport().StartTime, Helper.GetViewport().EndTime, SecondsPerPixel, [this, &Helper](double StartTime, double EndTime, uint32 Depth, const Trace::FLoadTimeProfilerCpuEvent& Event)
+					const double SecondsPerPixel = 1.0 / Helper.GetViewport().GetScaleX();
+					Timeline.EnumerateEventsDownSampled(Helper.GetViewport().GetStartTime(), Helper.GetViewport().GetEndTime(), SecondsPerPixel, [this, &Helper](double StartTime, double EndTime, uint32 Depth, const Trace::FLoadTimeProfilerCpuEvent& Event)
 					{
 						const TCHAR* Name = State->GetEventNameEx(Depth, Event);
 						Helper.AddEvent(StartTime, EndTime, Depth, Name);
@@ -401,7 +328,7 @@ void FLoadingAsyncThreadTimingTrack::Draw(FTimingViewDrawHelper& Helper) const
 				}
 				else
 				{
-					Timeline.EnumerateEvents(Helper.GetViewport().StartTime, Helper.GetViewport().EndTime, [this, &Helper](double StartTime, double EndTime, uint32 Depth, const Trace::FLoadTimeProfilerCpuEvent& Event)
+					Timeline.EnumerateEvents(Helper.GetViewport().GetStartTime(), Helper.GetViewport().GetEndTime(), [this, &Helper](double StartTime, double EndTime, uint32 Depth, const Trace::FLoadTimeProfilerCpuEvent& Event)
 					{
 						const TCHAR* Name = State->GetEventNameEx(Depth, Event);
 						Helper.AddEvent(StartTime, EndTime, Depth, Name);
diff --git a/Engine/Source/Developer/TraceInsights/Private/Insights/ViewModels/LoadingTimingTrack.h b/Engine/Source/Developer/TraceInsights/Private/Insights/ViewModels/LoadingTimingTrack.h
--- a/Engine/Source/Developer/TraceInsights/Private/Insights/ViewModels/LoadingTimingTrack.h
+++ b/Engine/Source/Developer/TraceInsights/Private/Insights/ViewModels/LoadingTimingTrack.h
@@ -37,14 +37,14 @@ private:
 class FLoadingTimingTrack : public FTimingEventsTrack
 {
 public:
-	FLoadingTimingTrack(uint64 InTrackId, const FName& InSubType, const FString& InName, TSharedPtr<FLoadingSharedState> InState)
+	explicit FLoadingTimingTrack(uint64 InTrackId, const FName& InSubType, const FString& InName, TSharedPtr<FLoadingSharedState> InState)
 		: FTimingEventsTrack(InTrackId, FName(TEXT("Loading")), InSubType, InName)
 		, State(InState)
 	{
 	}
 	virtual ~FLoadingTimingTrack() {}
 
-	virtual void DrawTooltip(FTimingViewTooltip& Tooltip, const FVector2D& MousePosition, const FTimingEvent& HoveredTimingEvent, const FTimingTrackViewport& Viewport, const FDrawContext& DrawContext, const FSlateBrush* WhiteBrush, const FSlateFontInfo& Font) const override;
+	virtual void InitTooltip(FTooltipDrawState& Tooltip, const FTimingEvent& HoveredTimingEvent) const override;
 
 protected:
 	TSharedPtr<FLoadingSharedState> State;
@@ -55,7 +55,7 @@ protected:
 class FLoadingMainThreadTimingTrack : public FLoadingTimingTrack
 {
 public:
-	FLoadingMainThreadTimingTrack(uint64 InTrackId, TSharedPtr<FLoadingSharedState> InState)
+	explicit FLoadingMainThreadTimingTrack(uint64 InTrackId, TSharedPtr<FLoadingSharedState> InState)
 		: FLoadingTimingTrack(InTrackId, FName(TEXT("MainThread")), TEXT("Loading - Main Thread"), InState)
 	{
 	}
@@ -69,7 +69,7 @@ public:
 class FLoadingAsyncThreadTimingTrack : public FLoadingTimingTrack
 {
 public:
-	FLoadingAsyncThreadTimingTrack(uint64 InTrackId, TSharedPtr<FLoadingSharedState> InState)
+	explicit FLoadingAsyncThreadTimingTrack(uint64 InTrackId, TSharedPtr<FLoadingSharedState> InState)
 		: FLoadingTimingTrack(InTrackId, FName(TEXT("AsyncThread")), TEXT("Loading - Async Thread"), InState)
 	{
 	}
diff --git a/Engine/Source/Developer/TraceInsights/Private/Insights/ViewModels/MarkersTimingTrack.cpp b/Engine/Source/Developer/TraceInsights/Private/Insights/ViewModels/MarkersTimingTrack.cpp
--- a/Engine/Source/Developer/TraceInsights/Private/Insights/ViewModels/MarkersTimingTrack.cpp
+++ b/Engine/Source/Developer/TraceInsights/Private/Insights/ViewModels/MarkersTimingTrack.cpp
@@ -2,7 +2,6 @@
 
 #include "MarkersTimingTrack.h"
 
-#include "Brushes/SlateColorBrush.h"
 #include "Fonts/FontMeasure.h"
 #include "Framework/Application/SlateApplication.h"
 #include "Styling/CoreStyle.h"
@@ -11,6 +10,7 @@
 // Insights
 #include "Insights/Common/PaintUtils.h"
 #include "Insights/InsightsManager.h"
+#include "Insights/InsightsStyle.h"
 #include "Insights/TimingProfilerManager.h"
 #include "Insights/ViewModels/TimingTrackViewport.h"
 
@@ -33,7 +33,7 @@ FMarkersTimingTrack::FMarkersTimingTrack(uint64 InTrackId)
 	, NumLogMessages(0)
 	, NumDrawBoxes(0)
 	, NumDrawTexts(0)
-	, WhiteBrush(FCoreStyle::Get().GetBrush("WhiteBrush"))
+	, WhiteBrush(FInsightsStyle::Get().GetBrush("WhiteBrush"))
 	, Font(FCoreStyle::GetDefaultFontStyle("Regular", 8))
 {
 }
@@ -145,8 +145,8 @@ void FMarkersTimingTrack::Update(const FTimingTrackViewport& InViewport)
 		Builder.BeginLog(LogProvider);
 
 		LogProvider.EnumerateMessages(
-			Builder.GetViewport().StartTime,
-			Builder.GetViewport().EndTime,
+			Builder.GetViewport().GetStartTime(),
+			Builder.GetViewport().GetEndTime(),
 			[&Builder](const Trace::FLogMessage& Message) { Builder.AddLogMessage(Message); });
 
 		Builder.EndLog();
@@ -157,8 +157,10 @@ void FMarkersTimingTrack::Update(const FTimingTrackViewport& InViewport)
 
 void FMarkersTimingTrack::Draw(FDrawContext& DrawContext, const FTimingTrackViewport& Viewport) const
 {
+	const FLinearColor BackgroundColor(0.04f, 0.04f, 0.04f, 1.0f);
+
 	// Draw background.
-	DrawContext.DrawBox(0.0f, GetPosY(), Viewport.Width, GetHeight(), WhiteBrush, FLinearColor(0.04f, 0.04f, 0.04f, 1.0f));
+	DrawContext.DrawBox(0.0f, GetPosY(), Viewport.GetWidth(), GetHeight(), WhiteBrush, BackgroundColor);
 	DrawContext.LayerId++;
 
 	// Draw the track's header, in background.
@@ -180,10 +182,10 @@ void FMarkersTimingTrack::Draw(FDrawContext& DrawContext, const FTimingTrackView
 	else
 	{
 		BoxY = 0.0f;
-		BoxH = Viewport.Height;
+		BoxH = Viewport.GetHeight();
 	}
 
-	int32 NumBoxes = TimeMarkerBoxes.Num();
+	const int32 NumBoxes = TimeMarkerBoxes.Num();
 	for (int32 BoxIndex = 0; BoxIndex < NumBoxes; BoxIndex++)
 	{
 		const FTimeMarkerBoxInfo& Box = TimeMarkerBoxes[BoxIndex];
@@ -198,7 +200,7 @@ void FMarkersTimingTrack::Draw(FDrawContext& DrawContext, const FTimingTrackView
 	const float CategoryY = GetPosY() + 2.0f;
 	const float MessageY = GetPosY() + (IsBookmarksTrack() ? 1.0f : 14.0f);
 
-	int32 NumTexts = TimeMarkerTexts.Num();
+	const int32 NumTexts = TimeMarkerTexts.Num();
 	for (int32 TextIndex = 0; TextIndex < NumTexts; TextIndex++)
 	{
 		const FTimeMarkerTextInfo& TextInfo = TimeMarkerTexts[TextIndex];
@@ -452,7 +454,7 @@ void FTimeMarkerTrackBuilder::Flush(float AvailableTextW)
 
 void FTimeMarkerTrackBuilder::AddTimeMarker(const float X, const uint64 LogIndex, const ELogVerbosity::Type Verbosity, const TCHAR* const Category, const TCHAR* Message)
 {
-	float W = X - LastX2;
+	const float W = X - LastX2;
 
 	if (W > 0.0f) // There is at least 1px from previous box?
 	{
@@ -516,7 +518,7 @@ void FTimeMarkerTrackBuilder::AddTimeMarker(const float X, const uint64 LogIndex
 
 void FTimeMarkerTrackBuilder::EndLog()
 {
-	Flush(Viewport.Width - LastX2);
+	Flush(Viewport.GetWidth() - LastX2);
 }
 
 ////////////////////////////////////////////////////////////////////////////////////////////////////
diff --git a/Engine/Source/Developer/TraceInsights/Private/Insights/ViewModels/MarkersTimingTrack.h b/Engine/Source/Developer/TraceInsights/Private/Insights/ViewModels/MarkersTimingTrack.h
--- a/Engine/Source/Developer/TraceInsights/Private/Insights/ViewModels/MarkersTimingTrack.h
+++ b/Engine/Source/Developer/TraceInsights/Private/Insights/ViewModels/MarkersTimingTrack.h
@@ -44,7 +44,7 @@ class FMarkersTimingTrack : public FBaseTimingTrack
 	friend class FTimeMarkerTrackBuilder;
 
 public:
-	FMarkersTimingTrack(uint64 InTrackId);
+	explicit FMarkersTimingTrack(uint64 InTrackId);
 	virtual ~FMarkersTimingTrack();
 
 	virtual void Reset() override;
@@ -106,7 +106,7 @@ private:
 class FTimeMarkerTrackBuilder
 {
 public:
-	FTimeMarkerTrackBuilder(FMarkersTimingTrack& InTrack, const FTimingTrackViewport& InViewport);
+	explicit FTimeMarkerTrackBuilder(FMarkersTimingTrack& InTrack, const FTimingTrackViewport& InViewport);
 
 	/**
 	 * Non-copyable
diff --git a/Engine/Source/Developer/TraceInsights/Private/Insights/ViewModels/ThreadTimingTrack.cpp b/Engine/Source/Developer/TraceInsights/Private/Insights/ViewModels/ThreadTimingTrack.cpp
--- a/Engine/Source/Developer/TraceInsights/Private/Insights/ViewModels/ThreadTimingTrack.cpp
+++ b/Engine/Source/Developer/TraceInsights/Private/Insights/ViewModels/ThreadTimingTrack.cpp
@@ -14,6 +14,7 @@
 #include "Insights/ViewModels/TimingEvent.h"
 #include "Insights/ViewModels/TimingTrackViewport.h"
 #include "Insights/ViewModels/TimingViewDrawHelper.h"
+#include "Insights/ViewModels/TooltipDrawState.h"
 
 #define LOCTEXT_NAMESPACE "ThreadTimingTrack"
 
@@ -38,15 +39,15 @@ void FThreadTimingTrack::Draw(FTimingViewDrawHelper& Helper) const
 				{
 					if (FTimingEventsTrack::bUseDownSampling)
 					{
-						const double SecondsPerPixel = 1.0 / Helper.GetViewport().ScaleX;
-						Timeline.EnumerateEventsDownSampled(Helper.GetViewport().StartTime, Helper.GetViewport().EndTime, SecondsPerPixel, [this, &Helper, Timers](double StartTime, double EndTime, uint32 Depth, const Trace::FTimingProfilerEvent& Event)
+						const double SecondsPerPixel = 1.0 / Helper.GetViewport().GetScaleX();
+						Timeline.EnumerateEventsDownSampled(Helper.GetViewport().GetStartTime(), Helper.GetViewport().GetEndTime(), SecondsPerPixel, [this, &Helper, Timers](double StartTime, double EndTime, uint32 Depth, const Trace::FTimingProfilerEvent& Event)
 						{
 							Helper.AddEvent(StartTime, EndTime, Depth, Timers[Event.TimerIndex].Name);
 						});
 					}
 					else
 					{
-						Timeline.EnumerateEvents(Helper.GetViewport().StartTime, Helper.GetViewport().EndTime, [this, &Helper, Timers](double StartTime, double EndTime, uint32 Depth, const Trace::FTimingProfilerEvent& Event)
+						Timeline.EnumerateEvents(Helper.GetViewport().GetStartTime(), Helper.GetViewport().GetEndTime(), [this, &Helper, Timers](double StartTime, double EndTime, uint32 Depth, const Trace::FTimingProfilerEvent& Event)
 						{
 							Helper.AddEvent(StartTime, EndTime, Depth, Timers[Event.TimerIndex].Name);
 						});
@@ -72,8 +73,8 @@ void FThreadTimingTrack::DrawSelectedEventInfo(const FTimingEvent& SelectedTimin
 
 	const TSharedRef<FSlateFontMeasure> FontMeasureService = FSlateApplication::Get().GetRenderer()->GetFontMeasureService();
 	const FVector2D Size = FontMeasureService->Measure(Str, Font);
-	const float X = Viewport.Width - Size.X - 23.0f;
-	const float Y = Viewport.Height - Size.Y - 18.0f;
+	const float X = Viewport.GetWidth() - Size.X - 23.0f;
+	const float Y = Viewport.GetHeight() - Size.Y - 18.0f;
 
 	const FLinearColor BackgroundColor(0.05f, 0.05f, 0.05f, 1.0f);
 	const FLinearColor TextColor(0.7f, 0.7f, 0.7f, 1.0f);
@@ -87,61 +88,19 @@ void FThreadTimingTrack::DrawSelectedEventInfo(const FTimingEvent& SelectedTimin
 
 ////////////////////////////////////////////////////////////////////////////////////////////////////
 
-void FThreadTimingTrack::DrawTooltip(FTimingViewTooltip& Tooltip,
-									 const FVector2D& MousePosition,
-									 const FTimingEvent& HoveredTimingEvent,
-									 const FTimingTrackViewport& Viewport,
-									 const FDrawContext& DrawContext,
-									 const FSlateBrush* WhiteBrush,
-									 const FSlateFontInfo& Font) const
+void FThreadTimingTrack::InitTooltip(FTooltipDrawState& Tooltip, const FTimingEvent& HoveredTimingEvent) const
 {
+	Tooltip.ResetContent();
+
 	const FTimerNodePtr TimerNodePtr = FTimingProfilerManager::Get()->GetTimerNode(HoveredTimingEvent.TypeId);
 	FString TimerName = TimerNodePtr ? TimerNodePtr->GetName().ToString() : TEXT("N/A");
+	Tooltip.AddTitle(TimerName);
 
-	const TSharedRef<FSlateFontMeasure> FontMeasureService = FSlateApplication::Get().GetRenderer()->GetFontMeasureService();
-	float DesiredTooltipWidth = FontMeasureService->Measure(TimerName, Font).X + 2 * FTimingViewTooltip::BorderX;
-	if (DesiredTooltipWidth < FTimingViewTooltip::MinWidth)
-	{
-		DesiredTooltipWidth = FTimingViewTooltip::MinWidth;
-	}
-
-	constexpr float LineH = 14.0f;
-	const float DesiredTooltipHeight = 4 * LineH + 2 * FTimingViewTooltip::BorderY;
-
-	Tooltip.Update(MousePosition, DesiredTooltipWidth, DesiredTooltipHeight, Viewport.Width, Viewport.Height);
-
-	const FLinearColor BackgroundColor(0.05f, 0.05f, 0.05f, Tooltip.Opacity);
-	const FLinearColor NameColor(0.9f, 0.9f, 0.5f, Tooltip.Opacity);
-	const FLinearColor TextColor(0.6f, 0.6f, 0.6f, Tooltip.Opacity);
-	const FLinearColor ValueColor(1.0f, 1.0f, 1.0f, Tooltip.Opacity);
-
-	DrawContext.DrawBox(Tooltip.PosX, Tooltip.PosY, Tooltip.Width, Tooltip.Height, WhiteBrush, BackgroundColor);
-	DrawContext.LayerId++;
-
-	float X = Tooltip.PosX + FTimingViewTooltip::BorderX;
-	float Y = Tooltip.PosY + FTimingViewTooltip::BorderY;
-
-	DrawContext.DrawText(X, Y, TimerName, Font, NameColor);
-	Y += LineH;
-
-	const float ValueX = X + 58.0f;
+	Tooltip.AddNameValueTextLine(TEXT("Incl. Time:"), TimeUtils::FormatTimeAuto(HoveredTimingEvent.Duration()));
+	Tooltip.AddNameValueTextLine(TEXT("Excl. Time:"), TimeUtils::FormatTimeAuto(HoveredTimingEvent.ExclusiveTime));
+	Tooltip.AddNameValueTextLine(TEXT("Depth:"), FString::Printf(TEXT("%d"), HoveredTimingEvent.Depth));
 
-	DrawContext.DrawText(X + 3.0f, Y, TEXT("Incl. Time:"), Font, TextColor);
-	FString InclStr = TimeUtils::FormatTimeAuto(HoveredTimingEvent.Duration());
-	DrawContext.DrawText(ValueX, Y, InclStr, Font, ValueColor);
-	Y += LineH;
-
-	DrawContext.DrawText(X, Y, TEXT("Excl. Time:"), Font, TextColor);
-	FString ExclStr = TimeUtils::FormatTimeAuto(HoveredTimingEvent.ExclusiveTime);
-	DrawContext.DrawText(ValueX, Y, ExclStr, Font, ValueColor);
-	Y += LineH;
-
-	DrawContext.DrawText(X + 24.0f, Y, TEXT("Depth:"), Font, TextColor);
-	FString DepthStr = FString::Printf(TEXT("%d"), HoveredTimingEvent.Depth);
-	DrawContext.DrawText(ValueX, Y, DepthStr, Font, ValueColor);
-	Y += LineH;
-
-	DrawContext.LayerId++;
+	Tooltip.UpdateLayout();
 }
 
 ////////////////////////////////////////////////////////////////////////////////////////////////////
diff --git a/Engine/Source/Developer/TraceInsights/Private/Insights/ViewModels/ThreadTimingTrack.h b/Engine/Source/Developer/TraceInsights/Private/Insights/ViewModels/ThreadTimingTrack.h
--- a/Engine/Source/Developer/TraceInsights/Private/Insights/ViewModels/ThreadTimingTrack.h
+++ b/Engine/Source/Developer/TraceInsights/Private/Insights/ViewModels/ThreadTimingTrack.h
@@ -12,7 +12,7 @@
 class FThreadTimingTrack : public FTimingEventsTrack
 {
 public:
-	FThreadTimingTrack(uint64 InTrackId, const FName& InSubType, const FString& InName, const TCHAR* InGroupName)
+	explicit FThreadTimingTrack(uint64 InTrackId, const FName& InSubType, const FString& InName, const TCHAR* InGroupName)
 		: FTimingEventsTrack(InTrackId, FName(TEXT("Thread")), InSubType, InName)
 		, GroupName(InGroupName)
 		, ThreadId(0)
@@ -28,7 +28,7 @@ public:
 
 	virtual void Draw(FTimingViewDrawHelper& Helper) const override;
 	virtual void DrawSelectedEventInfo(const FTimingEvent& SelectedTimingEvent, const FTimingTrackViewport& Viewport, const FDrawContext& DrawContext, const FSlateBrush* WhiteBrush, const FSlateFontInfo& Font) const override;
-	virtual void DrawTooltip(FTimingViewTooltip& Tooltip, const FVector2D& MousePosition, const FTimingEvent& HoveredTimingEvent, const FTimingTrackViewport& Viewport, const FDrawContext& DrawContext, const FSlateBrush* WhiteBrush, const FSlateFontInfo& Font) const override;
+	virtual void InitTooltip(FTooltipDrawState& Tooltip, const FTimingEvent& HoveredTimingEvent) const override;
 
 	virtual bool SearchTimingEvent(const double InStartTime, const double InEndTime, TFunctionRef<bool(double, double, uint32)> InPredicate, FTimingEvent& InOutTimingEvent, bool bInStopAtFirstMatch, bool bInSearchForLargestEvent) const override;
 
@@ -42,7 +42,7 @@ private:
 class FCpuTimingTrack : public FThreadTimingTrack
 {
 public:
-	FCpuTimingTrack(uint64 InTrackId, const FString& InName, const TCHAR* InGroupName)
+	explicit FCpuTimingTrack(uint64 InTrackId, const FString& InName, const TCHAR* InGroupName)
 		: FThreadTimingTrack(InTrackId, FName(TEXT("CPU")), InName, InGroupName)
 	{
 	}
@@ -53,7 +53,7 @@ public:
 class FGpuTimingTrack : public FThreadTimingTrack
 {
 public:
-	FGpuTimingTrack(uint64 InTrackId, const FString& InName, const TCHAR* InGroupName)
+	explicit FGpuTimingTrack(uint64 InTrackId, const FString& InName, const TCHAR* InGroupName)
 		: FThreadTimingTrack(InTrackId, FName(TEXT("GPU")), InName, InGroupName)
 	{
 	}
diff --git a/Engine/Source/Developer/TraceInsights/Private/Insights/ViewModels/TimeRulerTrack.cpp b/Engine/Source/Developer/TraceInsights/Private/Insights/ViewModels/TimeRulerTrack.cpp
--- a/Engine/Source/Developer/TraceInsights/Private/Insights/ViewModels/TimeRulerTrack.cpp
+++ b/Engine/Source/Developer/TraceInsights/Private/Insights/ViewModels/TimeRulerTrack.cpp
@@ -2,15 +2,17 @@
 
 #include "TimeRulerTrack.h"
 
-#include "Brushes/SlateColorBrush.h"
 #include "Fonts/FontMeasure.h"
 #include "Fonts/SlateFontInfo.h"
 #include "Framework/Application/SlateApplication.h"
 #include "Rendering/DrawElements.h"
+#include "Styling/CoreStyle.h"
 
 // Insights
 #include "Insights/Common/PaintUtils.h"
 #include "Insights/Common/TimeUtils.h"
+#include "Insights/InsightsStyle.h"
+#include "Insights/ViewModels/DrawHelpers.h"
 #include "Insights/ViewModels/TimingTrackViewport.h"
 
 #define LOCTEXT_NAMESPACE "TimeRulerTrack"
@@ -19,7 +21,7 @@
 
 FTimeRulerTrack::FTimeRulerTrack(uint64 InTrackId)
 	: FBaseTimingTrack(InTrackId)
-	, WhiteBrush(FCoreStyle::Get().GetBrush("WhiteBrush"))
+	, WhiteBrush(FInsightsStyle::Get().GetBrush("WhiteBrush"))
 	, Font(FCoreStyle::GetDefaultFontStyle("Regular", 8))
 {
 }
@@ -62,58 +64,7 @@ void FTimeRulerTrack::UpdateHoveredState(float MouseX, float MouseY, const FTimi
 
 void FTimeRulerTrack::DrawBackground(FDrawContext& DrawContext, const FTimingTrackViewport& Viewport) const
 {
-	const FLinearColor ValidAreaColor(0.07f, 0.07f, 0.07f, 1.0f);
-	const FLinearColor InvalidAreaColor(0.1f, 0.07f, 0.07f, 1.0f);
-	const FLinearColor StartEndMarkerColor(0.05f, 0.05f, 0.05f, 1.0f);
-
-	const float X0 = Viewport.TimeToSlateUnitsRounded(0.0);
-	const float X1 = Viewport.TimeToSlateUnitsRounded(Viewport.MaxValidTime);
-	const float W = FMath::CeilToFloat(Viewport.Width);
-
-	const float Y = GetPosY();
-	const float H = GetHeight();
-
-	if (X0 >= W || X1 <= 0.0f)
-	{
-		// Draw invalid area (entire view).
-		DrawContext.DrawBox(0.0f, Y, W, H, WhiteBrush, InvalidAreaColor);
-	}
-	else // X0 < W && X1 > 0
-	{
-		if (X0 > 0.0f)
-		{
-			// Draw invalid area (left).
-			DrawContext.DrawBox(0.0f, Y, X0, H, WhiteBrush, InvalidAreaColor);
-		}
-
-		if (X1 < W)
-		{
-			// Draw invalid area (right).
-			DrawContext.DrawBox(X1 + 1.0f, Y, W - X1 - 1.0f, H, WhiteBrush, InvalidAreaColor);
-
-			// Draw the end time marker.
-			DrawContext.DrawBox(X1, Y, 1.0f, H, WhiteBrush, StartEndMarkerColor);
-		}
-
-		float ValidX0 = FMath::Max(X0, 0.0f);
-		float ValidX1 = FMath::Min(X1, W);
-
-		if (X0 >= 0.0f)
-		{
-			// Draw the start time marker.
-			DrawContext.DrawBox(X0, Y, 1.0f, H, WhiteBrush, StartEndMarkerColor);
-
-			ValidX0 += 1.0f; // to not overlap the start time marker
-		}
-
-		if (ValidX1 > ValidX0)
-		{
-			// Draw valid area.
-			DrawContext.DrawBox(ValidX0, Y, ValidX1 - ValidX0, H, WhiteBrush, ValidAreaColor);
-		}
-	}
-
-	DrawContext.LayerId++;
+	FDrawHelpers::DrawBackground(DrawContext, WhiteBrush, Viewport, GetPosY(), GetHeight());
 }
 
 ////////////////////////////////////////////////////////////////////////////////////////////////////
@@ -131,7 +82,7 @@ void FTimeRulerTrack::Draw(FDrawContext& DrawContext, const FTimingTrackViewport
 	double MinorTickMarkTime = Viewport.GetDurationForViewportDX(MinorTickMark);
 	double MajorTickMarkTime = Viewport.GetDurationForViewportDX(MajorTickMark);
 
-	double VX = Viewport.StartTime * Viewport.ScaleX;
+	double VX = Viewport.GetStartTime() * Viewport.GetScaleX();
 	double MinorN = FMath::FloorToDouble(VX / static_cast<double>(MinorTickMark));
 	double MajorN = FMath::FloorToDouble(VX / static_cast<double>(MajorTickMark));
 	float MinorOX = static_cast<float>(FMath::RoundToDouble(MinorN * static_cast<double>(MinorTickMark) - VX));
@@ -141,7 +92,7 @@ void FTimeRulerTrack::Draw(FDrawContext& DrawContext, const FTimingTrackViewport
 	DrawBackground(DrawContext, Viewport);
 
 	// Draw the minor tick marks.
-	for (float X = MinorOX; X < Viewport.Width; X += MinorTickMark)
+	for (float X = MinorOX; X < Viewport.GetWidth(); X += MinorTickMark)
 	{
 		const bool bIsTenth = ((int32)(((X - MajorOX) / MinorTickMark) + 0.4f) % 2 == 0);
 		const float MinorTickH = bIsTenth ? MinorTickMarkHeight : MinorTickMarkHeight - 1.0f;
@@ -149,25 +100,25 @@ void FTimeRulerTrack::Draw(FDrawContext& DrawContext, const FTimingTrackViewport
 			bIsTenth ? FLinearColor(0.3f, 0.3f, 0.3f, 1.0f) : FLinearColor(0.25f, 0.25f, 0.25f, 1.0f));
 	}
 	// Draw the major tick marks.
-	for (float X = MajorOX; X < Viewport.Width; X += MajorTickMark)
+	for (float X = MajorOX; X < Viewport.GetWidth(); X += MajorTickMark)
 	{
 		DrawContext.DrawBox(X, GetPosY(), 1.0f, MajorTickMarkHeight, WhiteBrush, FLinearColor(0.4f, 0.4f, 0.4f, 1.0f));
 	}
 	DrawContext.LayerId++;
 
-	const double DT = static_cast<double>(MajorTickMark) / Viewport.ScaleX;
+	const double DT = static_cast<double>(MajorTickMark) / Viewport.GetScaleX();
 	const double Precision = FMath::Max(DT / 10.0, TimeUtils::Nanosecond);
 
 	const TSharedRef<FSlateFontMeasure> FontMeasureService = FSlateApplication::Get().GetRenderer()->GetFontMeasureService();
 
 	// Draw the time at major tick marks.
-	for (float X = MajorOX; X < Viewport.Width + MajorTickMark; X += MajorTickMark)
+	for (float X = MajorOX; X < Viewport.GetWidth() + MajorTickMark; X += MajorTickMark)
 	{
 		const double T = Viewport.SlateUnitsToTime(X);
 		FString Text = TimeUtils::FormatTime(T, Precision);
 		const float TextWidth = FontMeasureService->Measure(Text, Font).X;
 		DrawContext.DrawText(X - TextWidth / 2, TextY, Text, Font,
-			(T < 0 || T >= Viewport.MaxValidTime) ? FLinearColor(0.7f, 0.5f, 0.5f, 1.0f) : FLinearColor(0.8f, 0.8f, 0.8f, 1.0f));
+			(T < Viewport.GetMinValidTime() || T >= Viewport.GetMaxValidTime()) ? FLinearColor(0.7f, 0.5f, 0.5f, 1.0f) : FLinearColor(0.8f, 0.8f, 0.8f, 1.0f));
 	}
 	DrawContext.LayerId++;
 
@@ -226,7 +177,7 @@ void FTimeRulerTrack::Draw(FDrawContext& DrawContext, const FTimingTrackViewport
 			//DrawContext.DrawBox(0.0f, MousePosition.Y, Viewport.Width, 1.0f, WhiteBrush, MousePosLineColor);
 
 			// Draw vertical line at mouse position.
-			DrawContext.DrawBox(MousePosition.X, 0.0f, 1.0f, Viewport.Height, WhiteBrush, MousePosLineColor);
+			DrawContext.DrawBox(MousePosition.X, 0.0f, 1.0f, Viewport.GetHeight(), WhiteBrush, MousePosLineColor);
 
 			// Stroke the vertical line above current time box.
 			DrawContext.DrawBox(MousePosition.X, 0.0f, 1.0f, TextY, WhiteBrush, MousePosTextBackgroundColor);
diff --git a/Engine/Source/Developer/TraceInsights/Private/Insights/ViewModels/TimeRulerTrack.h b/Engine/Source/Developer/TraceInsights/Private/Insights/ViewModels/TimeRulerTrack.h
--- a/Engine/Source/Developer/TraceInsights/Private/Insights/ViewModels/TimeRulerTrack.h
+++ b/Engine/Source/Developer/TraceInsights/Private/Insights/ViewModels/TimeRulerTrack.h
@@ -17,7 +17,7 @@ class FTimingTrackViewport;
 class FTimeRulerTrack : public FBaseTimingTrack
 {
 public:
-	FTimeRulerTrack(uint64 InTrackId);
+	explicit FTimeRulerTrack(uint64 InTrackId);
 	virtual ~FTimeRulerTrack();
 
 	virtual void Reset() override;
diff --git a/Engine/Source/Developer/TraceInsights/Private/Insights/ViewModels/TimingEventsTrack.h b/Engine/Source/Developer/TraceInsights/Private/Insights/ViewModels/TimingEventsTrack.h
--- a/Engine/Source/Developer/TraceInsights/Private/Insights/ViewModels/TimingEventsTrack.h
+++ b/Engine/Source/Developer/TraceInsights/Private/Insights/ViewModels/TimingEventsTrack.h
@@ -14,6 +14,7 @@ struct FTimingEvent;
 class FTimingTrackViewport;
 struct FTimingViewTooltip;
 class FTimingViewDrawHelper;
+class FTooltipDrawState;
 
 ////////////////////////////////////////////////////////////////////////////////////////////////////
 
@@ -55,7 +56,7 @@ struct FTimingEventsTrackLayout
 class FTimingEventsTrack : public FBaseTimingTrack
 {
 public:
-	FTimingEventsTrack(uint64 InTrackId, const FName& InType, const FName& InSubType, const FString& InName);
+	explicit FTimingEventsTrack(uint64 InTrackId, const FName& InType, const FName& InSubType, const FString& InName);
 	virtual ~FTimingEventsTrack();
 
 	const FName& GetType() const { return Type; }
@@ -75,7 +76,7 @@ public:
 
 	virtual void Draw(FTimingViewDrawHelper& Helper) const = 0;
 	virtual void DrawSelectedEventInfo(const FTimingEvent& SelectedTimingEvent, const FTimingTrackViewport& Viewport, const FDrawContext& DrawContext, const FSlateBrush* WhiteBrush, const FSlateFontInfo& Font) const {}
-	virtual void DrawTooltip(FTimingViewTooltip& Tooltip, const FVector2D& MousePosition, const FTimingEvent& HoveredTimingEvent, const FTimingTrackViewport& Viewport, const FDrawContext& DrawContext, const FSlateBrush* WhiteBrush, const FSlateFontInfo& Font) const {}
+	virtual void InitTooltip(FTooltipDrawState& Tooltip, const FTimingEvent& HoveredTimingEvent) const {}
 
 	virtual bool SearchTimingEvent(const double InStartTime, const double InEndTime, TFunctionRef<bool(double, double, uint32)> InPredicate, FTimingEvent& InOutTimingEvent, bool bInStopAtFirstMatch, bool bInSearchForLargestEvent) const { return false; }
 
