[Insights]

diff --git a/Engine/Source/Developer/TraceInsights/Private/Insights/ViewModels/TimingGraphTrack.cpp b/Engine/Source/Developer/TraceInsights/Private/Insights/ViewModels/TimingGraphTrack.cpp
--- a/Engine/Source/Developer/TraceInsights/Private/Insights/ViewModels/TimingGraphTrack.cpp
+++ b/Engine/Source/Developer/TraceInsights/Private/Insights/ViewModels/TimingGraphTrack.cpp
@@ -174,6 +174,8 @@ FTimingGraphTrack::~FTimingGraphTrack()
 
 void FTimingGraphTrack::Update(const FTimingTrackViewport& Viewport)
 {
+	FGraphTrack::Update(Viewport);
+
 	const bool bIsEntireGraphTrackDirty = IsDirty();
 	bool bNeedsUpdate = bIsEntireGraphTrackDirty;
 
@@ -272,7 +274,7 @@ void FTimingGraphTrack::UpdateStatsCounterSeries(FTimingGraphSeries& Series, con
 
 				if (Counter.IsFloatingPoint())
 				{
-					Counter.EnumerateFloatValues(Viewport.StartTime, Viewport.EndTime, [&Builder, &MinValue, &MaxValue](double Time, double Value)
+					Counter.EnumerateFloatValues(Viewport.GetStartTime(), Viewport.GetEndTime(), [&Builder, &MinValue, &MaxValue](double Time, double Value)
 					{
 						if (Value < MinValue)
 						{
@@ -286,7 +288,7 @@ void FTimingGraphTrack::UpdateStatsCounterSeries(FTimingGraphSeries& Series, con
 				}
 				else
 				{
-					Counter.EnumerateValues(Viewport.StartTime, Viewport.EndTime, [&Builder, &MinValue, &MaxValue](double Time, int64 IntValue)
+					Counter.EnumerateValues(Viewport.GetStartTime(), Viewport.GetEndTime(), [&Builder, &MinValue, &MaxValue](double Time, int64 IntValue)
 					{
 						const double Value = static_cast<double>(IntValue);
 						if (Value < MinValue)
@@ -350,7 +352,7 @@ void FTimingGraphTrack::UpdateStatsCounterSeries(FTimingGraphSeries& Series, con
 
 			if (Counter.IsFloatingPoint())
 			{
-				Counter.EnumerateFloatValues(Viewport.StartTime, Viewport.EndTime, [&Builder](double Time, double Value)
+				Counter.EnumerateFloatValues(Viewport.GetStartTime(), Viewport.GetEndTime(), [&Builder](double Time, double Value)
 				{
 					//TODO: add a "value unit converter"
 					Builder.AddEvent(Time, 0.0, Value);
@@ -358,7 +360,7 @@ void FTimingGraphTrack::UpdateStatsCounterSeries(FTimingGraphSeries& Series, con
 			}
 			else
 			{
-				Counter.EnumerateValues(Viewport.StartTime, Viewport.EndTime, [&Builder](double Time, int64 IntValue)
+				Counter.EnumerateValues(Viewport.GetStartTime(), Viewport.GetEndTime(), [&Builder](double Time, int64 IntValue)
 				{
 					//TODO: add a "value unit converter"
 					Builder.AddEvent(Time, 0.0, static_cast<double>(IntValue));
diff --git a/Engine/Source/Developer/TraceInsights/Private/Insights/ViewModels/TimingGraphTrack.h b/Engine/Source/Developer/TraceInsights/Private/Insights/ViewModels/TimingGraphTrack.h
--- a/Engine/Source/Developer/TraceInsights/Private/Insights/ViewModels/TimingGraphTrack.h
+++ b/Engine/Source/Developer/TraceInsights/Private/Insights/ViewModels/TimingGraphTrack.h
@@ -34,7 +34,7 @@ public:
 class FTimingGraphTrack : public FGraphTrack
 {
 public:
-	FTimingGraphTrack(uint64 InTrackId);
+	explicit FTimingGraphTrack(uint64 InTrackId);
 	virtual ~FTimingGraphTrack();
 
 	virtual void Update(const FTimingTrackViewport& Viewport) override;
diff --git a/Engine/Source/Developer/TraceInsights/Private/Insights/ViewModels/TimingTrackViewport.cpp b/Engine/Source/Developer/TraceInsights/Private/Insights/ViewModels/TimingTrackViewport.cpp
--- a/Engine/Source/Developer/TraceInsights/Private/Insights/ViewModels/TimingTrackViewport.cpp
+++ b/Engine/Source/Developer/TraceInsights/Private/Insights/ViewModels/TimingTrackViewport.cpp
@@ -2,10 +2,26 @@
 
 #include "TimingTrackViewport.h"
 
+#include "Widgets/Layout/SScrollBar.h"
+
 #include <limits>
 
 ////////////////////////////////////////////////////////////////////////////////////////////////////
 
+bool FTimingTrackViewport::UpdateSize(const float InWidth, const float InHeight)
+{
+	if (Width != InWidth || Height != InHeight)
+	{
+		Width = InWidth;
+		Height = InHeight;
+		EndTime = SlateUnitsToTime(Width);
+		return true;
+	}
+	return false;
+}
+
+////////////////////////////////////////////////////////////////////////////////////////////////////
+
 bool FTimingTrackViewport::ScrollAtTime(const double Time)
 {
 	const double NewStartTime = AlignTimeToPixel(Time);
@@ -55,6 +71,25 @@ bool FTimingTrackViewport::ZoomOnTimeInterval(const double Time, const double Du
 
 ////////////////////////////////////////////////////////////////////////////////////////////////////
 
+bool FTimingTrackViewport::RelativeZoomWithFixedX(const float Delta, const float X)
+{
+	constexpr double ZoomStep = 0.25; // as percent
+	double NewScaleX;
+
+	if (Delta > 0)
+	{
+		NewScaleX = ScaleX * FMath::Pow(1.0 + ZoomStep, Delta);
+	}
+	else
+	{
+		NewScaleX = ScaleX * FMath::Pow(1.0 / (1.0 + ZoomStep), -Delta);
+	}
+
+	return ZoomWithFixedX(NewScaleX, X);
+}
+
+////////////////////////////////////////////////////////////////////////////////////////////////////
+
 bool FTimingTrackViewport::ZoomWithFixedX(const double NewScaleX, const float X)
 {
 	const double LocalNewScaleX = FMath::Clamp(NewScaleX, MinScaleX, MaxScaleX);
@@ -129,18 +164,92 @@ bool FTimingTrackViewport::EnforceHorizontalScrollLimits(const double U)
 	{
 		NewStartTime = (1.0 - U) * NewStartTime + U * MinT;
 		if (FMath::IsNearlyEqual(NewStartTime, MinT, 1.0 / ScaleX))
+		{
 			NewStartTime = MinT;
+		}
 	}
 	else if (NewStartTime > MaxT)
 	{
 		NewStartTime = (1.0 - U) * NewStartTime + U * MaxT;
 		if (FMath::IsNearlyEqual(NewStartTime, MaxT, 1.0 / ScaleX))
+		{
 			NewStartTime = MaxT;
+		}
 		if (NewStartTime < MinT)
+		{
 			NewStartTime = MinT;
+		}
 	}
 
 	return ScrollAtTime(NewStartTime);
 }
 
 ////////////////////////////////////////////////////////////////////////////////////////////////////
+
+bool FTimingTrackViewport::OnUserScrolled(TSharedPtr<SScrollBar> ScrollBar, float ScrollOffset)
+{
+	const double S = 1.0 / (MaxValidTime - MinValidTime);
+	const double D = EndTime - StartTime;
+	const float ThumbSizeFraction = FMath::Clamp<float>(D * S, 0.0f, 1.0f);
+	const float OffsetFraction = FMath::Clamp<float>(ScrollOffset, 0.0f, 1.0f - ThumbSizeFraction);
+
+	ScrollBar->SetState(OffsetFraction, ThumbSizeFraction);
+
+	const double NewStartTime = MinValidTime + static_cast<double>(OffsetFraction) * (MaxValidTime - MinValidTime);
+	if (NewStartTime != StartTime)
+	{
+		ScrollAtTime(NewStartTime);
+		return true;
+	}
+
+	return false;
+}
+
+////////////////////////////////////////////////////////////////////////////////////////////////////
+
+void FTimingTrackViewport::UpdateScrollBar(TSharedPtr<SScrollBar> ScrollBar) const
+{
+	const double S = 1.0 / (MaxValidTime - MinValidTime);
+	const double D = EndTime - StartTime;
+	const float ThumbSizeFraction = FMath::Clamp<float>(D * S, 0.0f, 1.0f);
+	const float ScrollOffset = static_cast<float>((StartTime - MinValidTime) * S);
+	const float OffsetFraction = FMath::Clamp<float>(ScrollOffset, 0.0f, 1.0f - ThumbSizeFraction);
+
+	ScrollBar->SetState(OffsetFraction, ThumbSizeFraction);
+}
+
+////////////////////////////////////////////////////////////////////////////////////////////////////
+
+bool FTimingTrackViewport::OnUserScrolledY(TSharedPtr<SScrollBar> ScrollBar, float ScrollOffset)
+{
+	const float S = 1.0 / ScrollHeight;
+	const float D = Height - TopOffset;
+	const float ThumbSizeFraction = FMath::Clamp<float>(D * S, 0.0f, 1.0f);
+	const float OffsetFraction = FMath::Clamp<float>(ScrollOffset, 0.0f, 1.0f - ThumbSizeFraction);
+
+	ScrollBar->SetState(OffsetFraction, ThumbSizeFraction);
+
+	const float NewScrollPosY = OffsetFraction * ScrollHeight;
+	if (NewScrollPosY != ScrollPosY)
+	{
+		ScrollPosY = NewScrollPosY;
+		return true;
+	}
+
+	return false;
+}
+
+////////////////////////////////////////////////////////////////////////////////////////////////////
+
+void FTimingTrackViewport::UpdateScrollBarY(TSharedPtr<SScrollBar> ScrollBar) const
+{
+	const float S = 1.0 / ScrollHeight;
+	const float D = Height - TopOffset;
+	const float ThumbSizeFraction = FMath::Clamp<float>(D * S, 0.0f, 1.0f);
+	const float ScrollOffset = ScrollPosY * S;
+	const float OffsetFraction = FMath::Clamp<float>(ScrollOffset, 0.0f, 1.0f - ThumbSizeFraction);
+
+	ScrollBar->SetState(OffsetFraction, ThumbSizeFraction);
+}
+
+////////////////////////////////////////////////////////////////////////////////////////////////////
diff --git a/Engine/Source/Developer/TraceInsights/Private/Insights/ViewModels/TimingTrackViewport.h b/Engine/Source/Developer/TraceInsights/Private/Insights/ViewModels/TimingTrackViewport.h
--- a/Engine/Source/Developer/TraceInsights/Private/Insights/ViewModels/TimingTrackViewport.h
+++ b/Engine/Source/Developer/TraceInsights/Private/Insights/ViewModels/TimingTrackViewport.h
@@ -4,24 +4,12 @@
 
 #include "CoreMinimal.h"
 
-class FTimingTrackViewport
-{
-public:
-	float Width; // width of viewport, in Slate units
-	float Height; // height of viewport, in Slate units
+class SScrollBar;
 
-	double MinValidTime; // min session time, in seconds
-	double MaxValidTime; // max session time, in seconds
-	double StartTime; // time of viewport's left side, in seconds
-	double EndTime; // time of viewport's right side, in seconds; computed when StartTime or Width changes
-	double MinScaleX;
-	double MaxScaleX;
-	double ScaleX; // scale factor between seconds and pixels (Slate units), in pixels (Slate units) per second
-
-	float TopOffset; // top offset (to allow the time ruller to be visible), in pixels (Slate units)
-	float ScrollHeight; // height of the vertical scrollable area, in pixels (Slate units)
-	float ScrollPosY; // current vertical scroll position, in pixels (Slate units)
+////////////////////////////////////////////////////////////////////////////////////////////////////
 
+class FTimingTrackViewport
+{
 public:
 	FTimingTrackViewport()
 	{
@@ -35,8 +23,10 @@ public:
 
 		MinValidTime = 0.0;
 		MaxValidTime = 0.0;
+
 		StartTime = 0.0;
 		EndTime = 0.0;
+
 		MinScaleX = (5 * 20) / 3600.0; // 1h between major tick marks
 		MaxScaleX = 1.0E10; // 10ns between major tick marks
 		ScaleX = (5 * 20) / 5.0; // 5s between major tick marks
@@ -46,6 +36,40 @@ public:
 		ScrollPosY = 0.0f;
 	}
 
+	/** Width of viewport, in pixels. [px] */
+	float GetWidth() const { return Width; }
+
+	/** Height of viewport, in pixels. [px] */
+	float GetHeight() const { return Height; }
+
+	bool UpdateSize(const float InWidth, const float InHeight);
+
+	/** Minimum session time, in seconds. [s] */
+	double GetMinValidTime() const { return MinValidTime; }
+	/** Maximum session time, in seconds. [s] */
+	double GetMaxValidTime() const { return MaxValidTime; }
+	void SetMaxValidTime(const double InMaxValidTime) { MaxValidTime = InMaxValidTime; }
+
+	/**
+	  * Time of viewport's left side, in seconds. [s]
+	  */
+	double GetStartTime() const { return StartTime; }
+	/**
+	  * Time of viewport's right side, in seconds. [s]
+	  * Computed when StartTime, ScaleX or Width changes.
+	  */
+	double GetEndTime() const { return EndTime; }
+	double GetDuration() const { return EndTime - StartTime; }
+
+	double GetMinScaleX() const { return MinScaleX; }
+	double GetMaxScaleX() const { return MaxScaleX; }
+
+	/** Current scale factor between seconds and pixels (Slate units). [px/s] */
+	double GetScaleX() const { return ScaleX; }
+
+	//////////////////////////////////////////////////
+	// Conversions
+
 	float TimeToSlateUnits(const double Time) const
 	{
 		return static_cast<float>((Time - StartTime) * ScaleX);
@@ -71,18 +95,6 @@ public:
 		return StartTime + static_cast<double>(X) / ScaleX;
 	}
 
-	bool UpdateSize(const float InWidth, const float InHeight)
-	{
-		if (Width != InWidth || Height != InHeight)
-		{
-			Width = InWidth;
-			Height = InHeight;
-			EndTime = SlateUnitsToTime(Width);
-			return true;
-		}
-		return false;
-	}
-
 	double AlignTimeToPixel(const double InTime, const double InScaleX) const
 	{
 		return FMath::RoundToDouble(InTime * InScaleX) / InScaleX;
@@ -93,17 +105,70 @@ public:
 		return AlignTimeToPixel(Time, ScaleX);
 	}
 
+	double RestrictEndTime(const double InEndTime) const;
+	double RestrictDuration(const double InStartTime, const double InEndTime) const;
+
+	//////////////////////////////////////////////////
+	// Scrolling
+
 	bool ScrollAtTime(const double Time);
 	bool CenterOnTimeInterval(const double Time, const double Duration);
 
+	void GetHorizontalScrollLimits(double& OutMinT, double& OutMaxT);
+	bool EnforceHorizontalScrollLimits(const double U);
+
+	//////////////////////////////////////////////////
+	// Zooming
+
 	bool ZoomOnTimeInterval(const double Time, const double Duration);
+	bool RelativeZoomWithFixedX(const float Delta, const float X);
 	bool ZoomWithFixedX(const double NewScaleX, const float X);
 
-	double RestrictEndTime(const double InEndTime) const;
-	double RestrictDuration(const double InStartTime, const double InEndTime) const;
+	//////////////////////////////////////////////////
+	// Vertical axis
 
-	void GetHorizontalScrollLimits(double& OutMinT, double& OutMaxT);
-	bool EnforceHorizontalScrollLimits(const double U);
+	/** Top offset (to allow the time ruller to be visible), in pixels. [px] */
+	float GetTopOffset() const { return TopOffset; }
+
+	void SetTopOffset(float InTopOffset) { TopOffset = InTopOffset; }
+
+	/** Height of the vertical scrollable area, in pixels. [px] */
+	float GetScrollHeight() const { return ScrollHeight; }
+	void SetScrollHeight(const float InScrollHeight) { ScrollHeight = InScrollHeight; }
+
+	/** Current vertical scroll position, in pixels. [px] */
+	float GetScrollPosY() const { return ScrollPosY; }
+	void SetScrollPosY(const float InScrollPosY) { ScrollPosY = InScrollPosY; }
 
 	float GetViewportY(const float Y) const { return TopOffset + Y - ScrollPosY; }
+
+	//////////////////////////////////////////////////
+
+	bool OnUserScrolled(TSharedPtr<SScrollBar> ScrollBar, float ScrollOffset);
+	void UpdateScrollBar(TSharedPtr<SScrollBar> ScrollBar) const;
+
+	bool OnUserScrolledY(TSharedPtr<SScrollBar> ScrollBar, float ScrollOffset);
+	void UpdateScrollBarY(TSharedPtr<SScrollBar> ScrollBar) const;
+
+	//////////////////////////////////////////////////
+
+private:
+	float Width; // width of viewport, in pixels; [px]
+	float Height; // height of viewport, in pixels; [px]
+
+	double MinValidTime; // min session time, in seconds; [s]
+	double MaxValidTime; // max session time, in seconds; [s]
+
+	double StartTime; // time of viewport's left side, in seconds; [s]
+	double EndTime; // time of viewport's right side, in seconds; [s]; computed when StartTime, Scale or Width changes
+
+	double MinScaleX; // min scale factor; [px/s]
+	double MaxScaleX; // max scale factor; [px/s]
+	double ScaleX; // scale factor between seconds and pixels; [px/s]
+
+	float TopOffset; // top offset (to allow the time ruller to be visible), in pixels; [px]
+	float ScrollHeight; // height of the vertical scrollable area, in pixels; [px]
+	float ScrollPosY; // current vertical scroll position, in pixels; [px]
 };
+
+////////////////////////////////////////////////////////////////////////////////////////////////////
diff --git a/Engine/Source/Developer/TraceInsights/Private/Insights/ViewModels/TimingViewDrawHelper.cpp b/Engine/Source/Developer/TraceInsights/Private/Insights/ViewModels/TimingViewDrawHelper.cpp
--- a/Engine/Source/Developer/TraceInsights/Private/Insights/ViewModels/TimingViewDrawHelper.cpp
+++ b/Engine/Source/Developer/TraceInsights/Private/Insights/ViewModels/TimingViewDrawHelper.cpp
@@ -2,9 +2,6 @@
 
 #include "TimingViewDrawHelper.h"
 
-#include "Brushes/SlateBorderBrush.h"
-#include "Brushes/SlateColorBrush.h"
-#include "EditorStyleSet.h"
 #include "Fonts/FontMeasure.h"
 #include "Fonts/SlateFontInfo.h"
 #include "Framework/Application/SlateApplication.h"
@@ -15,6 +12,7 @@
 // Insights
 #include "Insights/Common/PaintUtils.h"
 #include "Insights/Common/TimeUtils.h"
+#include "Insights/InsightsStyle.h"
 #include "Insights/ViewModels/DrawHelpers.h"
 #include "Insights/ViewModels/MarkersTimingTrack.h"
 #include "Insights/ViewModels/TimingEventsTrack.h"
@@ -26,9 +24,10 @@ FTimingViewDrawHelper::FTimingViewDrawHelper(const FDrawContext& InDC, const FTi
 	: DrawContext(InDC)
 	, Viewport(InViewport)
 	, Layout(InLayout)
-	, WhiteBrush(FCoreStyle::Get().GetBrush("WhiteBrush"))
-	, BorderBrush(FEditorStyle::GetBrush("PlainBorder"))
-	, EventsBorderBrush(new FSlateBorderBrush(NAME_None, FMargin(1.0f)))
+	, WhiteBrush(FInsightsStyle::Get().GetBrush("WhiteBrush"))
+	, EventBorderBrush(FInsightsStyle::Get().GetBrush("EventBorder"))
+	, HoveredEventBorderBrush(FInsightsStyle::Get().GetBrush("HoveredEventBorder"))
+	, SelectedEventBorderBrush(FInsightsStyle::Get().GetBrush("SelectedEventBorder"))
 	, BackgroundAreaBrush(WhiteBrush)
 	, ValidAreaColor(0.07f, 0.07f, 0.07f, 1.0f)
 	, InvalidAreaColor(0.1f, 0.07f, 0.07f, 1.0f)
@@ -42,18 +41,15 @@ FTimingViewDrawHelper::FTimingViewDrawHelper(const FDrawContext& InDC, const FTi
 
 FTimingViewDrawHelper::~FTimingViewDrawHelper()
 {
-	delete EventsBorderBrush; //im: is it safe to delete brush (i.e in OnPaint, imediately after use)?
 }
 
 ////////////////////////////////////////////////////////////////////////////////////////////////////
 
 void FTimingViewDrawHelper::DrawBackground() const
 {
-	const float X = 0.0f;
-	const float W = FMath::CeilToFloat(Viewport.Width);
 	const float Y = 0.0f;
-	const float H = FMath::CeilToFloat(Viewport.Height);
-	FDrawHelpers::DrawBackground(DrawContext, Viewport, BackgroundAreaBrush, ValidAreaColor, InvalidAreaColor, EdgeColor, X, Y, W, H, ValidAreaX, ValidAreaW);
+	const float H = FMath::CeilToFloat(Viewport.GetHeight());
+	FDrawHelpers::DrawBackground(DrawContext, BackgroundAreaBrush, Viewport, Y, H);
 }
 
 ////////////////////////////////////////////////////////////////////////////////////////////////////
@@ -61,13 +57,13 @@ void FTimingViewDrawHelper::DrawBackground() const
 void FTimingViewDrawHelper::BeginTimelines()
 {
 	TimelineIndex = -1;
-	TimelineTopY = -Viewport.ScrollPosY;
+	TimelineTopY = -Viewport.GetScrollPosY();
 
-	float Y = TimelineTopY + Viewport.TopOffset;
+	float Y = TimelineTopY + Viewport.GetTopOffset();
 
 	if (Y > 0.0f && ValidAreaW > 0.0f)
 	{
-		Y = FMath::Min(Y, Viewport.Height);
+		Y = FMath::Min(Y, Viewport.GetHeight());
 
 		// Draw invalid area (top).
 		DrawContext.DrawBox(DrawContext.LayerId + ToInt32(EDrawLayer::TimelineHeader), ValidAreaX, 0.0f, ValidAreaW, Y, BackgroundAreaBrush, InvalidAreaColor);
@@ -80,7 +76,7 @@ bool FTimingViewDrawHelper::BeginTimeline(FTimingEventsTrack& Track)
 {
 	TimelineIndex++;
 
-	Track.SetPosY(TimelineTopY + Viewport.ScrollPosY);
+	Track.SetPosY(TimelineTopY + Viewport.GetScrollPosY());
 
 	if (Track.GetHeight() < Layout.MinTimelineH)
 	{
@@ -93,7 +89,7 @@ bool FTimingViewDrawHelper::BeginTimeline(FTimingEventsTrack& Track)
 		return false;
 	}
 
-	if (TimelineTopY > Viewport.Height - Viewport.TopOffset)
+	if (TimelineTopY > Viewport.GetHeight() - Viewport.GetTopOffset())
 	{
 		return false;
 	}
@@ -101,7 +97,7 @@ bool FTimingViewDrawHelper::BeginTimeline(FTimingEventsTrack& Track)
 	MaxDepth = -1;
 
 	// +1.0f is for horizontal line between timelines
-	TimelineY = Viewport.TopOffset + TimelineTopY + 1.0f + Layout.TimelineDY;
+	TimelineY = Viewport.GetTopOffset() + TimelineTopY + 1.0f + Layout.TimelineDY;
 
 	// Reset "last event X2" for all depths.
 	LastEventX2.Reset();
@@ -174,7 +170,7 @@ void FTimingViewDrawHelper::EndTimeline(FTimingEventsTrack& Track)
 
 	if (Track.GetHeight() > 0)
 	{
-		const float Y = Viewport.TopOffset + TimelineTopY;
+		const float Y = Viewport.GetTopOffset() + TimelineTopY;
 
 		FLinearColor TextColor(1.0f, 1.0f, 1.0f, 1.0f);
 
@@ -195,7 +191,7 @@ void FTimingViewDrawHelper::EndTimeline(FTimingEventsTrack& Track)
 		}
 
 		// Draw a horizontal line between timelines.
-		DrawContext.DrawBox(DrawContext.LayerId + ToInt32(EDrawLayer::TimelineHeader), 0, Y, Viewport.Width, 1.0f, WhiteBrush, EdgeColor);
+		DrawContext.DrawBox(DrawContext.LayerId + ToInt32(EDrawLayer::TimelineHeader), 0, Y, Viewport.GetWidth(), 1.0f, WhiteBrush, EdgeColor);
 
 		// Draw name of timeline.
 		//const FString Name = FString::Printf(TEXT("%d Y: %g, H: %g, VO: %g, VY: %g"), TimelineIndex, Track.Y, Track.H, Viewport.TopOffset, Viewport.ScrollPosY); //debug
@@ -215,12 +211,12 @@ void FTimingViewDrawHelper::AddEvent(double EventStartTime, double EventEndTime,
 	Stats.NumEvents++;
 
 	float EventX1 = Viewport.TimeToSlateUnitsRounded(EventStartTime);
-	if (EventX1 > Viewport.Width)
+	if (EventX1 > Viewport.GetWidth())
 	{
 		return;
 	}
 
-	double RestrictedEndTime = Viewport.RestrictEndTime(EventEndTime);
+	const double RestrictedEndTime = Viewport.RestrictEndTime(EventEndTime);
 	float EventX2 = Viewport.TimeToSlateUnitsRounded(RestrictedEndTime);
 	if (EventX2 < 0)
 	{
@@ -246,7 +242,7 @@ void FTimingViewDrawHelper::AddEvent(double EventStartTime, double EventEndTime,
 	}
 
 	float EventY = TimelineY + (Layout.EventH + Layout.EventDY) * Depth;
-	if (EventY < -Layout.EventH || EventY > Viewport.Height)
+	if (EventY < -Layout.EventH || EventY > Viewport.GetHeight())
 	{
 		return;
 	}
@@ -266,13 +262,13 @@ void FTimingViewDrawHelper::AddEvent(double EventStartTime, double EventEndTime,
 	}
 
 	// Limit event width on the viewport's right side.
-	const float MaxX = Viewport.Width + 2.0f; // +2 allows event border to remain outside screen
+	const float MaxX = Viewport.GetWidth() + 2.0f; // +2 allows event border to remain outside screen
 	if (EventX2 > MaxX)
 	{
 		EventX2 = MaxX;
 	}
 
-	float EventW = EventX2 - EventX1;
+	const float EventW = EventX2 - EventX1;
 
 	// Optimization...
 	if (EventW == 1.0f && EventX1 == LastEventX2[Depth] - 1.0f)
@@ -385,7 +381,7 @@ void FTimingViewDrawHelper::AddEvent(double EventStartTime, double EventEndTime,
 		}
 
 		const FLinearColor EventColorBorder(EventColorFill.R * BorderColorFactor, EventColorFill.G * BorderColorFactor, EventColorFill.B * BorderColorFactor, 1.0f);
-		DrawContext.DrawBox(DrawContext.LayerId + ToInt32(EDrawLayer::EventBorder), EventX1, EventY, EventW, Layout.EventH, EventsBorderBrush, EventColorBorder);
+		DrawContext.DrawBox(DrawContext.LayerId + ToInt32(EDrawLayer::EventBorder), EventX1, EventY, EventW, Layout.EventH, EventBorderBrush, EventColorBorder);
 		Stats.NumDrawBorders++;
 	}
 	else // 1px or 2px boxes
@@ -457,18 +453,18 @@ void FTimingViewDrawHelper::DrawBox(const FBoxData& Box, const float EventY, con
 
 void FTimingViewDrawHelper::EndTimelines()
 {
-	float Y = Viewport.TopOffset + TimelineTopY;
+	float Y = Viewport.GetTopOffset() + TimelineTopY;
 
 	// Draw a last horizontal line.
-	DrawContext.DrawBox(DrawContext.LayerId + ToInt32(EDrawLayer::TimelineHeader), 0, Y, Viewport.Width, 1.0f, WhiteBrush, EdgeColor);
+	DrawContext.DrawBox(DrawContext.LayerId + ToInt32(EDrawLayer::TimelineHeader), 0, Y, Viewport.GetWidth(), 1.0f, WhiteBrush, EdgeColor);
 	Y += 1.0f;
 
-	if (Y < Viewport.Height && ValidAreaW > 0.0f)
+	if (Y < Viewport.GetHeight() && ValidAreaW > 0.0f)
 	{
 		Y = FMath::Max(Y, 0.0f);
 
 		// Draw invalid area (bottom).
-		DrawContext.DrawBox(DrawContext.LayerId + ToInt32(EDrawLayer::TimelineHeader), ValidAreaX, Y, ValidAreaW, Viewport.Height - Y, BackgroundAreaBrush, InvalidAreaColor);
+		DrawContext.DrawBox(DrawContext.LayerId + ToInt32(EDrawLayer::TimelineHeader), ValidAreaX, Y, ValidAreaW, Viewport.GetHeight() - Y, BackgroundAreaBrush, InvalidAreaColor);
 	}
 
 	DrawContext.LayerId += ToInt32(EDrawLayer::Count);
@@ -479,7 +475,7 @@ void FTimingViewDrawHelper::EndTimelines()
 void FTimingViewDrawHelper::DrawTimingEventHighlight(double StartTime, double EndTime, float Y, EHighlightMode Mode)
 {
 	float EventX1 = Viewport.TimeToSlateUnitsRounded(StartTime);
-	if (EventX1 > Viewport.Width)
+	if (EventX1 > Viewport.GetWidth())
 	{
 		return;
 	}
@@ -491,6 +487,7 @@ void FTimingViewDrawHelper::DrawTimingEventHighlight(double StartTime, double En
 		return;
 	}
 
+	// Timing events are displayed with minimum 1px (including empty ones).
 	if (EventX1 == EventX2)
 	{
 		EventX2 = EventX1 + 1.0f;
@@ -504,27 +501,23 @@ void FTimingViewDrawHelper::DrawTimingEventHighlight(double StartTime, double En
 	}
 
 	// Limit event width on the viewport's right side.
-	const float MaxX = Viewport.Width + 2.0f; // +2 allows event border to remain outside screen
+	const float MaxX = Viewport.GetWidth() + 2.0f; // +2 allows event border to remain outside screen
 	if (EventX2 > MaxX)
 	{
 		EventX2 = MaxX;
 	}
 
-	float EventW = EventX2 - EventX1;
+	const float EventW = EventX2 - EventX1;
 
 	if (Mode == EHighlightMode::Hovered)
 	{
-		const FSlateBrush* HighlightBrush = BorderBrush;
-
 		const FLinearColor Color(1.0f, 1.0f, 0.0f, 1.0f); // yellow
 
 		// Draw border around the timing event box.
-		DrawContext.DrawBox(EventX1 - 2.0f, Y - 2.0f, EventW + 4.0f, Layout.EventH + 4.0f, HighlightBrush, Color);
+		DrawContext.DrawBox(EventX1 - 2.0f, Y - 2.0f, EventW + 4.0f, Layout.EventH + 4.0f, HoveredEventBorderBrush, Color);
 	}
 	else // EHighlightMode::Selected or EHighlightMode::SelectedAndHovered
 	{
-		const FSlateBrush* SelectedBrush = BorderBrush;
-
 		// Animate color from white (if selected and hovered) or yellow (if only selected) to black, using a squared sine function.
 		const double Time = static_cast<double>(FPlatformTime::Cycles64()) * FPlatformTime::GetSecondsPerCycle64();
 		float S = FMath::Sin(2.0 * Time);
@@ -533,7 +526,7 @@ void FTimingViewDrawHelper::DrawTimingEventHighlight(double StartTime, double En
 		const FLinearColor Color(S, S, Blue, 1.0f);
 
 		// Draw border around the timing event box.
-		DrawContext.DrawBox(EventX1 - 2.0f, Y - 2.0f, EventW + 4.0f, Layout.EventH + 4.0f, SelectedBrush, Color);
+		DrawContext.DrawBox(EventX1 - 2.0f, Y - 2.0f, EventW + 4.0f, Layout.EventH + 4.0f, SelectedEventBorderBrush, Color);
 	}
 	DrawContext.LayerId++;
 }
diff --git a/Engine/Source/Developer/TraceInsights/Private/Insights/ViewModels/TimingViewDrawHelper.h b/Engine/Source/Developer/TraceInsights/Private/Insights/ViewModels/TimingViewDrawHelper.h
--- a/Engine/Source/Developer/TraceInsights/Private/Insights/ViewModels/TimingViewDrawHelper.h
+++ b/Engine/Source/Developer/TraceInsights/Private/Insights/ViewModels/TimingViewDrawHelper.h
@@ -17,61 +17,6 @@ class FTimingTrackViewport;
 
 ////////////////////////////////////////////////////////////////////////////////////////////////////
 
-struct FTimingViewTooltip
-{
-	static constexpr float BorderX = 6.0f;
-	static constexpr float BorderY = 3.0f;
-	static constexpr float MinWidth = 128.0f;
-	static constexpr float MinHeight = 0.0f;
-
-	FTimingViewTooltip() : PosX(0.0f), PosY(0.0f), Width(FTimingViewTooltip::MinWidth), Height(FTimingViewTooltip::MinHeight), Opacity(0.0f) {}
-
-	void Reset()
-	{
-		PosX = 0.0f;
-		PosY = 0.0f;
-		Width = MinWidth;
-		Height = MinHeight;
-		Opacity = 0.0f;
-	}
-
-	void Update(const FVector2D& MousePosition, const float DesiredWidth, const float DesiredHeight, const float ViewportWidth, const float ViewportHeight)
-	{
-		if (Width != DesiredWidth)
-		{
-			Width = Width * 0.75f + DesiredWidth * 0.25f;
-		}
-		
-		Height = DesiredHeight;
-
-		const float MaxX = FMath::Max(0.0f, ViewportWidth - Width - 12.0f); // -12.0f is to avoid overlapping the vertical scrollbar (one on the right side of the view)
-		const float X = FMath::Clamp<float>(MousePosition.X + 12.0f, 0.0f, MaxX);
-		PosX = X;
-
-		const float MaxY = FMath::Max(0.0f, ViewportHeight - Height - 12.0f); // -12.0f is to avoid overlapping the horizontal scrollbar (one on the bottom of the view)
-		float Y = FMath::Clamp<float>(MousePosition.Y + 15.0f, 0.0f, MaxY);
-		PosY = Y;
-
-		const float DesiredOpacity = 1.0f - FMath::Abs(Width - DesiredWidth) / DesiredWidth;
-		if (Opacity < DesiredOpacity)
-		{
-			Opacity = Opacity * 0.9f + DesiredOpacity * 0.1f;
-		}
-		else
-		{
-			Opacity = DesiredOpacity;
-		}
-	}
-
-	float PosX;
-	float PosY;
-	float Width;
-	float Height;
-	float Opacity;
-};
-
-////////////////////////////////////////////////////////////////////////////////////////////////////
-
 class FTimingViewDrawHelper
 {
 public:
@@ -129,8 +74,6 @@ private:
 		int32 NumMergedBoxes;
 		int32 NumDrawBorders;
 		int32 NumDrawTexts;
-		int32 NumDrawTimeMarkerBoxes;
-		int32 NumDrawTimeMarkerTexts;
 
 		FStats()
 			: NumEvents(0)
@@ -138,13 +81,11 @@ private:
 			, NumMergedBoxes(0)
 			, NumDrawBorders(0)
 			, NumDrawTexts(0)
-			, NumDrawTimeMarkerBoxes(0)
-			, NumDrawTimeMarkerTexts(0)
 		{}
 	};
 
 public:
-	FTimingViewDrawHelper(const FDrawContext& InDrawContext, const FTimingTrackViewport& InViewport, const FTimingEventsTrackLayout& InLayout);
+	explicit FTimingViewDrawHelper(const FDrawContext& InDrawContext, const FTimingTrackViewport& InViewport, const FTimingEventsTrackLayout& InLayout);
 	~FTimingViewDrawHelper();
 
 	/**
@@ -158,17 +99,13 @@ public:
 	const FTimingEventsTrackLayout& GetLayout() const { return Layout; }
 
 	const FSlateBrush* GetWhiteBrush() const { return WhiteBrush; }
-	const FSlateBrush* GetBorderBrush() const { return BorderBrush; }
-	const FSlateBrush* GetEventsBorderBrush() const { return EventsBorderBrush; }
 	const FSlateFontInfo& GetEventFont() const { return EventFont; }
 
-	int32 GetNumEvents() const              { return Stats.NumEvents; }
-	int32 GetNumDrawBoxes() const           { return Stats.NumDrawBoxes; }
-	int32 GetNumMergedBoxes() const         { return Stats.NumMergedBoxes; }
-	int32 GetNumDrawBorders() const         { return Stats.NumDrawBorders; }
-	int32 GetNumDrawTexts() const           { return Stats.NumDrawTexts; }
-	int32 GetNumDrawTimeMarkerBoxes() const { return Stats.NumDrawTimeMarkerBoxes; }
-	int32 GetNumDrawTimeMarkerTexts() const { return Stats.NumDrawTimeMarkerTexts; }
+	int32 GetNumEvents() const      { return Stats.NumEvents; }
+	int32 GetNumDrawBoxes() const   { return Stats.NumDrawBoxes; }
+	int32 GetNumMergedBoxes() const { return Stats.NumMergedBoxes; }
+	int32 GetNumDrawBorders() const { return Stats.NumDrawBorders; }
+	int32 GetNumDrawTexts() const   { return Stats.NumDrawTexts; }
 
 	void DrawBackground() const;
 	void DrawTimingEventHighlight(double StartTime, double EndTime, float Y, EHighlightMode Mode);
@@ -189,8 +126,9 @@ private:
 	const FTimingEventsTrackLayout& Layout;
 
 	const FSlateBrush* WhiteBrush;
-	const FSlateBrush* BorderBrush;
-	const FSlateBrush* EventsBorderBrush;
+	const FSlateBrush* EventBorderBrush;
+	const FSlateBrush* HoveredEventBorderBrush;
+	const FSlateBrush* SelectedEventBorderBrush;
 	const FSlateBrush* BackgroundAreaBrush;
 	const FLinearColor ValidAreaColor;
 	const FLinearColor InvalidAreaColor;
diff --git a/Engine/Source/Developer/TraceInsights/Private/Insights/ViewModels/TooltipDrawState.cpp b/Engine/Source/Developer/TraceInsights/Private/Insights/ViewModels/TooltipDrawState.cpp
new file mode 100644
--- /dev/null
+++ b/Engine/Source/Developer/TraceInsights/Private/Insights/ViewModels/TooltipDrawState.cpp
@@ -0,0 +1,259 @@
+// Copyright 1998-2019 Epic Games, Inc. All Rights Reserved.
+
+#include "TooltipDrawState.h"
+
+#include "Fonts/FontMeasure.h"
+#include "Fonts/SlateFontInfo.h"
+#include "Styling/CoreStyle.h"
+
+// Insights
+#include "Insights/Common/PaintUtils.h"
+#include "Insights/InsightsStyle.h"
+
+////////////////////////////////////////////////////////////////////////////////////////////////////
+
+FLinearColor FTooltipDrawState::DefaultTitleColor(0.9f, 0.9f, 0.5f, 1.0f);
+FLinearColor FTooltipDrawState::DefaultNameColor(0.6f, 0.6f, 0.6f, 1.0f);
+FLinearColor FTooltipDrawState::DefaultValueColor(1.0f, 1.0f, 1.0f, 1.0f);
+
+////////////////////////////////////////////////////////////////////////////////////////////////////
+
+FTooltipDrawState::FTooltipDrawState()
+	: WhiteBrush(FInsightsStyle::Get().GetBrush("WhiteBrush"))
+	, Font(FCoreStyle::GetDefaultFontStyle("Regular", 8))
+	, BackgroundColor(0.05f, 0.05f, 0.05f, 1.0f)
+	, Size(0.0f, 0.0f)
+	, DesiredSize(0.0f, 0.0f)
+	, Position(0.0f, 0.0f)
+	, ValueOffsetX(0.0f)
+	, NewLineY(0.0f)
+	, Opacity(0.0f)
+	, DesiredOpacity(0.0f)
+	, Texts()
+{
+}
+
+////////////////////////////////////////////////////////////////////////////////////////////////////
+
+FTooltipDrawState::~FTooltipDrawState()
+{
+}
+
+////////////////////////////////////////////////////////////////////////////////////////////////////
+
+void FTooltipDrawState::Reset()
+{
+	Opacity = 0.0f;
+	DesiredOpacity = 0.0f;
+
+	ResetContent();
+}
+
+////////////////////////////////////////////////////////////////////////////////////////////////////
+
+void FTooltipDrawState::ResetContent()
+{
+	Texts.Reset();
+
+	ValueOffsetX = 0.0f;
+	NewLineY = BorderY;
+}
+
+////////////////////////////////////////////////////////////////////////////////////////////////////
+
+void FTooltipDrawState::AddTitle(const FString& Title)
+{
+	const TSharedRef<FSlateFontMeasure> FontMeasureService = FSlateApplication::Get().GetRenderer()->GetFontMeasureService();
+
+	const FVector2D TextSize = FontMeasureService->Measure(Title, Font);
+	Texts.Add({ BorderX, NewLineY, TextSize, Title, DefaultTitleColor, FDrawTextType::Title });
+
+	NewLineY += DefaultTitleHeight;
+}
+
+////////////////////////////////////////////////////////////////////////////////////////////////////
+
+void FTooltipDrawState::AddTitle(const FString& Title, const FLinearColor& Color)
+{
+	const TSharedRef<FSlateFontMeasure> FontMeasureService = FSlateApplication::Get().GetRenderer()->GetFontMeasureService();
+
+	const FVector2D TextSize = FontMeasureService->Measure(Title, Font);
+	Texts.Add({ BorderX, NewLineY, TextSize, Title, Color, FDrawTextType::Misc });
+
+	NewLineY += DefaultTitleHeight;
+}
+
+////////////////////////////////////////////////////////////////////////////////////////////////////
+
+void FTooltipDrawState::AddNameValueTextLine(const FString& Name, const FString& Value)
+{
+	const TSharedRef<FSlateFontMeasure> FontMeasureService = FSlateApplication::Get().GetRenderer()->GetFontMeasureService();
+
+	const FVector2D NameTextSize = FontMeasureService->Measure(Name, Font);
+	Texts.Add({ 0.0f, NewLineY, NameTextSize, Name, DefaultNameColor, FDrawTextType::Name });
+
+	const FVector2D ValueTextSize = FontMeasureService->Measure(Value, Font);
+	Texts.Add({ 0.0f, NewLineY, ValueTextSize, Value, DefaultValueColor, FDrawTextType::Value });
+
+	NewLineY += DefaultLineHeight;
+}
+
+////////////////////////////////////////////////////////////////////////////////////////////////////
+
+void FTooltipDrawState::AddTextLine(const FString& Text, const FLinearColor& Color)
+{
+	const TSharedRef<FSlateFontMeasure> FontMeasureService = FSlateApplication::Get().GetRenderer()->GetFontMeasureService();
+
+	const FVector2D TextSize = FontMeasureService->Measure(Text, Font);
+	Texts.Add({ BorderX, NewLineY, TextSize, Text, Color, FDrawTextType::Misc });
+
+	NewLineY += DefaultLineHeight;
+}
+
+////////////////////////////////////////////////////////////////////////////////////////////////////
+
+void FTooltipDrawState::AddTextLine(const float X, const float Y, const FString& Text, const FLinearColor& Color)
+{
+	const TSharedRef<FSlateFontMeasure> FontMeasureService = FSlateApplication::Get().GetRenderer()->GetFontMeasureService();
+
+	const FVector2D TextSize = FontMeasureService->Measure(Text, Font);
+	Texts.Add({ X, Y, TextSize, Text, Color, FDrawTextType::Misc });
+}
+
+////////////////////////////////////////////////////////////////////////////////////////////////////
+
+void FTooltipDrawState::UpdateLayout()
+{
+	ValueOffsetX = 0.0f;
+	for (const FDrawTextInfo& TextInfo : Texts)
+	{
+		if (TextInfo.Type == FDrawTextType::Name)
+		{
+			if (ValueOffsetX < TextInfo.TextSize.X)
+			{
+				ValueOffsetX = TextInfo.TextSize.X;
+			}
+		}
+	}
+	ValueOffsetX += BorderX;
+
+	DesiredSize.Set(BorderX, NewLineY + BorderY);
+	for (const FDrawTextInfo& TextInfo : Texts)
+	{
+		float RightX;
+		switch (TextInfo.Type)
+		{
+		case FDrawTextType::Name:
+			RightX = ValueOffsetX;
+			break;
+
+		case FDrawTextType::Value:
+			RightX = ValueOffsetX + NameValueDX + TextInfo.TextSize.X;
+			break;
+
+		default:
+			RightX = TextInfo.X + TextInfo.TextSize.X;
+		}
+		if (DesiredSize.X < RightX)
+		{
+			DesiredSize.X = RightX;
+		}
+	}
+	DesiredSize.X += BorderX;
+}
+
+////////////////////////////////////////////////////////////////////////////////////////////////////
+
+void FTooltipDrawState::Update()
+{
+	if (Size.X != DesiredSize.X)
+	{
+		Size.X = Size.X * 0.75f + DesiredSize.X * 0.25f;
+
+		if (FMath::IsNearlyEqual(Size.X, DesiredSize.X))
+		{
+			Size.X = DesiredSize.X;
+		}
+	}
+
+	if (Size.Y != DesiredSize.Y)
+	{
+		Size.Y = Size.Y * 0.5f + DesiredSize.Y * 0.5f;
+
+		if (FMath::IsNearlyEqual(Size.Y, DesiredSize.Y))
+		{
+			Size.Y = DesiredSize.Y;
+		}
+	}
+
+	const float RealDesiredOpacity = FMath::Min(DesiredOpacity, 1.0f - FMath::Abs(Size.X - DesiredSize.X) / DesiredSize.X);
+	if (Opacity != RealDesiredOpacity)
+	{
+		if (Opacity < RealDesiredOpacity)
+		{
+			// slow fade in
+			Opacity = Opacity * 0.9f + RealDesiredOpacity * 0.1f;
+		}
+		else
+		{
+			// fast fade out
+			Opacity = Opacity * 0.75f + RealDesiredOpacity * 0.25f;
+		}
+
+		if (FMath::IsNearlyEqual(Opacity, RealDesiredOpacity))
+		{
+			Opacity = RealDesiredOpacity;
+		}
+	}
+}
+
+////////////////////////////////////////////////////////////////////////////////////////////////////
+
+void FTooltipDrawState::SetPosition(const FVector2D& MousePosition, const float MinX, const float MaxX, const float MinY, const float MaxY)
+{
+	Position.X = FMath::Max(MinX, FMath::Min(MousePosition.X + 12.0f, MaxX - Size.X));
+	Position.Y = FMath::Max(MinY, FMath::Min(MousePosition.Y + 15.0f, MaxY - Size.Y));
+}
+
+////////////////////////////////////////////////////////////////////////////////////////////////////
+
+void FTooltipDrawState::Draw(const FDrawContext& DrawContext) const
+{
+	if (Opacity > 0.0f && Size.X > 0.0f && Size.Y > 0.0f)
+	{
+		// Draw background.
+		DrawContext.DrawBox(Position.X, Position.Y, Size.X, Size.Y, WhiteBrush, BackgroundColor.CopyWithNewOpacity(Opacity));
+		if (Size.X < DesiredSize.X)
+		{
+			DrawContext.DrawBox(Position.X + Size.X, Position.Y, DesiredSize.X - Size.X, Size.Y, WhiteBrush, BackgroundColor.CopyWithNewOpacity(Opacity * 0.5f));
+		}
+		DrawContext.LayerId++;
+
+		// Draw border.
+		//DrawContext.DrawBox(Position.X, Position.Y, Size.X, Size.Y, BorderBrush, BorderColor.CopyWithNewOpacity(Opacity));
+		//DrawContext.LayerId++;
+
+		// Draw cached texts.
+		for (const FDrawTextInfo& TextInfo : Texts)
+		{
+			float X = Position.X;
+			switch (TextInfo.Type)
+			{
+				case FDrawTextType::Name:
+					X += ValueOffsetX - TextInfo.TextSize.X;
+					break;
+
+				case FDrawTextType::Value:
+					X += ValueOffsetX + NameValueDX;
+					break;
+
+				default:
+					X += TextInfo.X;
+			}
+			DrawContext.DrawText(X, Position.Y + TextInfo.Y, TextInfo.Text, Font, TextInfo.Color.CopyWithNewOpacity(Opacity));
+		}
+		DrawContext.LayerId++;
+	}
+}
+
+////////////////////////////////////////////////////////////////////////////////////////////////////
diff --git a/Engine/Source/Developer/TraceInsights/Private/Insights/Widgets/SFrameTrack.cpp b/Engine/Source/Developer/TraceInsights/Private/Insights/Widgets/SFrameTrack.cpp
--- a/Engine/Source/Developer/TraceInsights/Private/Insights/Widgets/SFrameTrack.cpp
+++ b/Engine/Source/Developer/TraceInsights/Private/Insights/Widgets/SFrameTrack.cpp
@@ -2,15 +2,13 @@
 
 #include "SFrameTrack.h"
 
-#include "Brushes/SlateBorderBrush.h"
-#include "Brushes/SlateColorBrush.h"
-#include "EditorStyleSet.h"
 #include "Fonts/FontMeasure.h"
 #include "Fonts/SlateFontInfo.h"
 #include "Framework/Application/SlateApplication.h"
 #include "Framework/MultiBox/MultiBoxBuilder.h"
 #include "HAL/PlatformTime.h"
 #include "Rendering/DrawElements.h"
+#include "Styling/CoreStyle.h"
 #include "TraceServices/AnalysisService.h"
 #include "TraceServices/SessionService.h"
 #include "Widgets/Layout/SScrollBar.h"
@@ -20,6 +18,7 @@
 #include "Insights/Common/Stopwatch.h"
 #include "Insights/Common/TimeUtils.h"
 #include "Insights/InsightsManager.h"
+#include "Insights/InsightsStyle.h"
 #include "Insights/TimingProfilerCommon.h"
 #include "Insights/TimingProfilerManager.h"
 #include "Insights/ViewModels/FrameTrackHelper.h"
@@ -50,6 +49,9 @@ SFrameTrack::~SFrameTrack()
 void SFrameTrack::Reset()
 {
 	Viewport.Reset();
+	FAxisViewportDouble& ViewportY = Viewport.GetVerticalAxisViewport();
+	ViewportY.SetScaleLimits(0.01, 1000000.0);
+	ViewportY.SetScale(1500.0);
 	bIsViewportDirty = true;
 
 	SeriesMap.Reset();
@@ -69,7 +71,6 @@ void SFrameTrack::Reset()
 
 	MousePositionOnButtonDown = FVector2D::ZeroVector;
 	ViewportPosXOnButtonDown = 0.0f;
-	ViewportPosYOnButtonDown = 0.0f;
 
 	MousePositionOnButtonUp = FVector2D::ZeroVector;
 
@@ -78,9 +79,6 @@ void SFrameTrack::Reset()
 
 	bIsScrolling = false;
 
-	SelectionStartFrameIndex = 0;
-	SelectionEndFrameIndex = 0;
-
 	HoveredSample.Reset();
 	TooltipDesiredOpacity = 0.9f;
 	TooltipOpacity = 0.0f;
@@ -139,7 +137,7 @@ void SFrameTrack::Tick(const FGeometry& AllottedGeometry, const double InCurrent
 
 	ThisGeometry = AllottedGeometry;
 
-	FIndexAxisViewport& ViewportX = Viewport.GetHorizontalAxisViewport();
+	FAxisViewportInt32& ViewportX = Viewport.GetHorizontalAxisViewport();
 
 	if (!bIsScrolling)
 	{
@@ -168,9 +166,9 @@ void SFrameTrack::Tick(const FGeometry& AllottedGeometry, const double InCurrent
 				TSharedPtr<FFrameTrackSeries> SeriesPtr = FindOrAddSeries(FrameType);
 
 				int32 NumFrames = FramesProvider.GetFrameCount(static_cast<ETraceFrameType>(FrameType));
-				if (NumFrames > ViewportX.GetMaxIndex())
+				if (NumFrames > ViewportX.GetMaxValue())
 				{
-					ViewportX.SetMinMaxIndexInterval(0, NumFrames);
+					ViewportX.SetMinMaxInterval(0, NumFrames);
 
 					if (bIsAutoZoomEnabled)
 					{
@@ -259,10 +257,10 @@ void SFrameTrack::UpdateState()
 
 		const Trace::IFrameProvider& FramesProvider = Trace::ReadFrameProvider(*Session.Get());
 
-		const FIndexAxisViewport& ViewportX = Viewport.GetHorizontalAxisViewport();
+		const FAxisViewportInt32& ViewportX = Viewport.GetHorizontalAxisViewport();
 
-		const uint64 StartIndex = static_cast<uint64>(FMath::Max(0, ViewportX.GetIndexAtOffset(0.0f)));
-		const uint64 EndIndex = static_cast<uint64>(ViewportX.GetIndexAtOffset(ViewportX.GetSize()));
+		const uint64 StartIndex = static_cast<uint64>(FMath::Max(0, ViewportX.GetValueAtOffset(0.0f)));
+		const uint64 EndIndex = static_cast<uint64>(ViewportX.GetValueAtOffset(ViewportX.GetSize()));
 
 		for (int32 FrameType = 0; FrameType < TraceFrameType_Count; ++FrameType)
 		{
@@ -313,7 +311,7 @@ FFrameTrackSampleRef SFrameTrack::GetSampleAtMousePosition(float X, float Y)
 						const FFrameTrackSample& Sample = SeriesPtr->Samples[SampleIndex];
 						if (Sample.NumFrames > 0)
 						{
-							const FValueAxisViewport& ViewportY = Viewport.GetVerticalAxisViewport();
+							const FAxisViewportDouble& ViewportY = Viewport.GetVerticalAxisViewport();
 
 							const float ViewHeight = FMath::RoundToFloat(Viewport.GetHeight());
 							const float BaselineY = FMath::RoundToFloat(ViewportY.GetOffsetForValue(0.0));
@@ -401,13 +399,14 @@ int32 SFrameTrack::OnPaint(const FPaintArgs& Args, const FGeometry& AllottedGeom
 	const TSharedRef<FSlateFontMeasure> FontMeasureService = FSlateApplication::Get().GetRenderer()->GetFontMeasureService();
 	FSlateFontInfo SummaryFont = FCoreStyle::GetDefaultFontStyle("Regular", 8);
 
-	const FSlateBrush* WhiteBrush = FCoreStyle::Get().GetBrush("WhiteBrush");
+	const FSlateBrush* WhiteBrush = FInsightsStyle::Get().GetBrush("WhiteBrush");
 
 	const float ViewWidth = AllottedGeometry.Size.X;
 	const float ViewHeight = AllottedGeometry.Size.Y;
 
 	int32 NumDrawSamples = 0;
 
+	//////////////////////////////////////////////////
 	{
 		FStopwatch Stopwatch;
 		Stopwatch.Start();
@@ -416,6 +415,7 @@ int32 SFrameTrack::OnPaint(const FPaintArgs& Args, const FGeometry& AllottedGeom
 
 		Helper.DrawBackground();
 
+		// Draw the horizontal axis grid.
 		DrawHorizontalAxisGrid(DrawContext, WhiteBrush, SummaryFont);
 
 		// Draw frames, for each visible Series.
@@ -438,13 +438,18 @@ int32 SFrameTrack::OnPaint(const FPaintArgs& Args, const FGeometry& AllottedGeom
 
 		NumDrawSamples = Helper.GetNumDrawSamples();
 
-		DrawVerticalAxisGrid(DrawContext, WhiteBrush, SummaryFont);
-
-		// Highlight the mouse hovered sample + tooltip.
+		// Highlight the mouse hovered sample (frame).
 		if (HoveredSample.IsValid())
 		{
 			Helper.DrawHoveredSample(*HoveredSample.Sample);
+		}
 
+		// Draw the vertical axis grid.
+		DrawVerticalAxisGrid(DrawContext, WhiteBrush, SummaryFont);
+
+		// Draw tooltip for hovered sample (frame).
+		if (HoveredSample.IsValid())
+		{
 			if (TooltipOpacity < TooltipDesiredOpacity)
 			{
 				TooltipOpacity = TooltipOpacity * 0.9f + TooltipDesiredOpacity * 0.1f;
@@ -466,9 +471,9 @@ int32 SFrameTrack::OnPaint(const FPaintArgs& Args, const FGeometry& AllottedGeom
 			const float DX = 2.0f;
 			const float W2 = TextSize.X / 2 + DX;
 
-			const FIndexAxisViewport& ViewportX = Viewport.GetHorizontalAxisViewport();
+			const FAxisViewportInt32& ViewportX = Viewport.GetHorizontalAxisViewport();
 
-			float X1 = ViewportX.GetOffsetForIndex(HoveredSample.Sample->LargestFrameIndex);
+			float X1 = ViewportX.GetOffsetForValue(HoveredSample.Sample->LargestFrameIndex);
 			float CX = X1 + FMath::RoundToFloat(Viewport.GetSampleWidth() / 2);
 			if (CX + W2 > ViewportX.GetSize())
 			{
@@ -502,8 +507,8 @@ int32 SFrameTrack::OnPaint(const FPaintArgs& Args, const FGeometry& AllottedGeom
 				if (SeriesPtr.IsValid())
 				{
 					// Highlight the area corresponding to viewport of Timing View.
-					const double StartTime = TimingView->GetViewport().StartTime;
-					const double EndTime = TimingView->GetViewport().EndTime;
+					const double StartTime = TimingView->GetViewport().GetStartTime();
+					const double EndTime = TimingView->GetViewport().GetEndTime();
 					Helper.DrawHighlightedInterval(*SeriesPtr, StartTime, EndTime);
 				}
 			}
@@ -512,7 +517,6 @@ int32 SFrameTrack::OnPaint(const FPaintArgs& Args, const FGeometry& AllottedGeom
 		Stopwatch.Stop();
 		DrawDurationHistory.AddValue(Stopwatch.AccumulatedTime);
 	}
-
 	//////////////////////////////////////////////////
 
 	const bool bShouldDisplayDebugInfo = FInsightsManager::Get()->IsDebugInfoEnabled();
@@ -526,17 +530,11 @@ int32 SFrameTrack::OnPaint(const FPaintArgs& Args, const FGeometry& AllottedGeom
 		const float DbgX = ViewWidth - DbgW - 20.0f;
 		float DbgY = 7.0f;
 
-		++LayerId;
-		FSlateDrawElement::MakeBox(
-			OutDrawElements,
-			LayerId,
-			MAKE_PAINT_GEOMETRY_RC(AllottedGeometry, DbgX - 2.0f, DbgY - 2.0f, DbgW, DbgH),
-			WhiteBrush,
-			DrawEffects,
-			FLinearColor(1.0, 1.0, 1.0, 0.9)
-		);
+		DrawContext.LayerId++;
+
+		DrawContext.DrawBox(DbgX - 2.0f, DbgY - 2.0f, DbgW, DbgH, WhiteBrush, FLinearColor(1.0, 1.0, 1.0, 0.9));
+		DrawContext.LayerId++;
 
-		++LayerId;
 		FLinearColor DbgTextColor(0.0, 0.0, 0.0, 0.9);
 
 		// Time interval since last OnPaint call.
@@ -552,63 +550,45 @@ int32 SFrameTrack::OnPaint(const FPaintArgs& Args, const FGeometry& AllottedGeom
 		const uint64 AvgDrawDurationMs = FStopwatch::Cycles64ToMilliseconds(DrawDurationHistory.ComputeAverage());
 
 		// Draw performance info.
-		FSlateDrawElement::MakeText
+		DrawContext.DrawText
 		(
-			OutDrawElements,
-			LayerId,
-			MAKE_PAINT_GEOMETRY_PT(AllottedGeometry, DbgX, DbgY),
+			DbgX, DbgY,
 			FString::Printf(TEXT("U: %llu ms, D: %llu ms + %llu ms = %llu ms (%d fps)"),
 				AvgUpdateDurationMs, // caching time
 				AvgDrawDurationMs, // drawing time
 				AvgOnPaintDurationMs - AvgDrawDurationMs, // other overhead to OnPaint calls
 				AvgOnPaintDurationMs, // average time between two OnPaint calls
 				FMath::RoundToInt(AvgOnPaintFps)), // framerate of OnPaint calls
-			SummaryFont,
-			DrawEffects,
-			DbgTextColor
+			SummaryFont, DbgTextColor
 		);
 		DbgY += DbgDY;
 
 		// Draw number of draw calls.
-		FSlateDrawElement::MakeText
+		DrawContext.DrawText
 		(
-			OutDrawElements,
-			LayerId,
-			MAKE_PAINT_GEOMETRY_PT(AllottedGeometry, DbgX, DbgY),
-			FString::Format(TEXT("U: {0} frames, D: {1} samples"),
-				{
-					FText::AsNumber(NumUpdatedFrames).ToString(),
-					FText::AsNumber(NumDrawSamples).ToString()
-				}),
-			SummaryFont,
-			DrawEffects,
-			DbgTextColor
+			DbgX, DbgY,
+			FString::Printf(TEXT("U: %s frames, D: %s samples"),
+				*FText::AsNumber(NumUpdatedFrames).ToString(),
+				*FText::AsNumber(NumDrawSamples).ToString()),
+			SummaryFont, DbgTextColor
 		);
 		DbgY += DbgDY;
 
 		// Draw viewport's horizontal info.
-		FSlateDrawElement::MakeText
+		DrawContext.DrawText
 		(
-			OutDrawElements,
-			LayerId,
-			MAKE_PAINT_GEOMETRY_PT(AllottedGeometry, DbgX, DbgY),
-			Viewport.GetHorizontalAxisViewport().ToDebugString(TEXT("VX: "), TEXT("frame")),
-			SummaryFont,
-			DrawEffects,
-			DbgTextColor
+			DbgX, DbgY,
+			Viewport.GetHorizontalAxisViewport().ToDebugString(TEXT("X"), TEXT("frame")),
+			SummaryFont, DbgTextColor
 		);
 		DbgY += DbgDY;
 
 		// Draw viewport's vertical info.
-		FSlateDrawElement::MakeText
+		DrawContext.DrawText
 		(
-			OutDrawElements,
-			LayerId,
-			MAKE_PAINT_GEOMETRY_PT(AllottedGeometry, DbgX, DbgY),
-			Viewport.GetVerticalAxisViewport().ToDebugString(TEXT("VY: ")),
-			SummaryFont,
-			DrawEffects,
-			DbgTextColor
+			DbgX, DbgY,
+			Viewport.GetVerticalAxisViewport().ToDebugString(TEXT("Y")),
+			SummaryFont, DbgTextColor
 		);
 		DbgY += DbgDY;
 	}
@@ -620,104 +600,104 @@ int32 SFrameTrack::OnPaint(const FPaintArgs& Args, const FGeometry& AllottedGeom
 
 void SFrameTrack::DrawVerticalAxisGrid(FDrawContext& DrawContext, const FSlateBrush* Brush, const FSlateFontInfo& Font) const
 {
-	const FValueAxisViewport& ViewportY = Viewport.GetVerticalAxisViewport();
+	const float ViewWidth = Viewport.GetWidth();
 
+	const FAxisViewportDouble& ViewportY = Viewport.GetVerticalAxisViewport();
 	const float RoundedViewHeight = FMath::RoundToFloat(ViewportY.GetSize());
 
-	constexpr float MinDY = 24.0f; // min vertical distance between horizontal grid lines
+	constexpr float TextH = 14.0f;
+	constexpr float MinDY = 12.0f; // min vertical distance between horizontal grid lines
+
+	const FLinearColor GridColor(0.0f, 0.0f, 0.0f, 0.1f);
+	const FLinearColor TextBgColor(0.05f, 0.05f, 0.05f, 1.0f);
+	//const FLinearColor TextColor(1.0f, 1.0f, 1.0f, 1.0f);
 
-	const double TopValue = ViewportY.GetValueForOffset(RoundedViewHeight);
-	const double GridValue = ViewportY.GetValueForOffset(MinDY);
-	const double BottomValue = ViewportY.GetValueForOffset(0.0f);
-	const double Delta = GridValue - BottomValue;
+	const TSharedRef<FSlateFontMeasure> FontMeasureService = FSlateApplication::Get().GetRenderer()->GetFontMeasureService();
 
-	if (Delta > 0.0)
+	const double GridValues[] =
+	{
+		0.0,
+		1.0 / 200.0, //    5 ms (200 fps)
+		1.0 / 60.0,  // 16.7 ms (60 fps)
+		1.0 / 30.0,  // 33.3 ms (30 fps)
+		1.0 / 20.0,  //   50 ms (20 fps)
+		1.0 / 15.0,  // 66.7 ms (15 fps)
+		1.0 / 10.0,  //  100 ms (10 fps)
+		1.0 / 5.0,   //  200 ms (5 fps)
+		1.0,   // 1s
+		10.0,  // 10s
+		60.0,  // 1m
+		600.0, // 10m
+		3600.0 // 1h
+	};
+	constexpr int32 NumGridValues = sizeof(GridValues) / sizeof(double);
+
+	double PreviousY = -MinDY;
+
+	for (int32 Index = 0; Index < NumGridValues; ++Index)
 	{
-		constexpr double TimeUnits = TimeUtils::Milisecond;
-		const int64 DeltaTicks = static_cast<int64>(Delta / TimeUnits);
+		const double Value = GridValues[Index];
 
-		// Compute rounding based on magnitude of visible range of values (Delta).
-		int64 Power10 = 1;
-		int64 Delta10 = DeltaTicks;
-		while (Delta10 > 0)
+		constexpr double Time60fps = 1.0 / 60.0;
+		constexpr double Time30fps = 1.0 / 30.0;
+
+		FLinearColor TextColor;
+		if (Value <= Time60fps)
 		{
-			Delta10 /= 10;
-			Power10 *= 10;
+			TextColor = FLinearColor(0.5f, 1.0f, 0.5f, 1.0f);
 		}
-		if (Power10 >= 100)
+		else if (Value <= Time30fps)
 		{
-			Power10 /= 100;
+			TextColor = FLinearColor(1.0f, 1.0f, 0.5f, 1.0f);
+			//const float U = (Value - Time60fps) * 0.5 / (Time30fps - Time60fps);
+			//TextColor = FLinearColor(0.5f + U, 1.0f - U, 0.5f, 1.0f);
 		}
 		else
 		{
-			Power10 = 1;
+			TextColor = FLinearColor(1.0f, 0.5f, 0.5f, 1.0f);
 		}
 
-		const double GridTicks = static_cast<double>(((DeltaTicks + Power10) / Power10) * Power10); // next value divisible with a multiple of 10
-		const double Grid = GridTicks * TimeUnits;
-
-		const double StartValue = FMath::GridSnap(BottomValue, Grid);
-
-		const FLinearColor GridColor(0.0f, 0.0f, 0.0f, 0.1f);
-		const FLinearColor TextBgColor(0.05f, 0.05f, 0.05f, 1.0f);
-		//const FLinearColor TextColor(1.0f, 1.0f, 1.0f, 1.0f);
-
-		const float ViewWidth = Viewport.GetWidth();
-
-		const TSharedRef<FSlateFontMeasure> FontMeasureService = FSlateApplication::Get().GetRenderer()->GetFontMeasureService();
-
-		for (double Value = StartValue; Value < TopValue; Value += Grid)
+		const float Y = RoundedViewHeight - FMath::RoundToFloat(ViewportY.GetOffsetForValue(Value));
+		if (Y < -TextH)
 		{
-			FLinearColor TextColor;
-			constexpr double Time60fps = 1.0 / 60.0;
-			constexpr double Time30fps = 1.0 / 30.0;
-			if (Value <= Time60fps)
-			{
-				TextColor = FLinearColor(0.5f, 1.0f, 0.5f, 1.0f);
-			}
-			else if (Value <= Time30fps)
-			{
-				//TextColor = FLinearColor(1.0f, 1.0f, 0.5f, 1.0f);
-				const float U = (Value - Time60fps) * 0.5 / (Time30fps - Time60fps);
-				TextColor = FLinearColor(0.5 + U, 1.0f - U, 0.5f, 1.0f);
-			}
-			else
-			{
-				TextColor = FLinearColor(1.0f, 0.5f, 0.5f, 1.0f);
-			}
-
-			const float Y = RoundedViewHeight - FMath::RoundToFloat(ViewportY.GetOffsetForValue(Value));
+			break;
+		}
+		if (Y > RoundedViewHeight || FMath::Abs(PreviousY - Y) < MinDY)
+		{
+			continue;
+		}
+		PreviousY = Y;
 
-			// Draw horizontal grid line.
-			DrawContext.DrawBox(0, Y, ViewWidth, 1, Brush, GridColor);
+		// Draw horizontal grid line.
+		DrawContext.DrawBox(0, Y, ViewWidth, 1, Brush, GridColor);
 
-			const FString Text = (Value == 0.0) ? TEXT("0") : FString::Printf(TEXT("%s (%.1f fps)"), *TimeUtils::FormatTimeAuto(Value), FMath::RoundToDouble(1.0 / Value));
-			const FVector2D TextSize = FontMeasureService->Measure(Text, Font);
-			constexpr float TextH = 14.0f;
+		const FString Text = (Value == 0.0) ? TEXT("0") :
+							 (Value <= 1.0) ? FString::Printf(TEXT("%s (%.0f fps)"), *TimeUtils::FormatTimeAuto(Value), 1.0 / Value) :
+											  TimeUtils::FormatTimeAuto(Value);
+		const FVector2D TextSize = FontMeasureService->Measure(Text, Font);
 
-			// Draw background for value text.
-			DrawContext.DrawBox(ViewWidth - TextSize.X - 4.0f, Y - TextH, TextSize.X + 4.0f, TextH, Brush, TextBgColor);
+		// Draw background for value text.
+		DrawContext.DrawBox(ViewWidth - TextSize.X - 4.0f, Y - TextH, TextSize.X + 4.0f, TextH, Brush, TextBgColor);
 
-			// Draw value text.
-			DrawContext.DrawText(ViewWidth - TextSize.X - 2.0f, Y - TextH + 1.0f, Text, Font, TextColor);
-		}
-		DrawContext.LayerId++;
+		// Draw value text.
+		DrawContext.DrawText(ViewWidth - TextSize.X - 2.0f, Y - TextH + 1.0f, Text, Font, TextColor);
 	}
+	DrawContext.LayerId++;
 }
 
 ////////////////////////////////////////////////////////////////////////////////////////////////////
 
 void SFrameTrack::DrawHorizontalAxisGrid(FDrawContext& DrawContext, const FSlateBrush* Brush, const FSlateFontInfo& Font) const
 {
-	const FIndexAxisViewport& ViewportX = Viewport.GetHorizontalAxisViewport();
+	const FAxisViewportInt32& ViewportX = Viewport.GetHorizontalAxisViewport();
 
 	const float RoundedViewWidth = FMath::RoundToFloat(ViewportX.GetSize());
 
 	constexpr float MinDX = 100.0f; // min horizontal distance between vertical grid lines
 
-	int32 LeftIndex = ViewportX.GetIndexAtOffset(0.0f);
-	int32 GridIndex = ViewportX.GetIndexAtOffset(MinDX);
-	int32 RightIndex = ViewportX.GetIndexAtOffset(RoundedViewWidth);
+	int32 LeftIndex = ViewportX.GetValueAtOffset(0.0f);
+	int32 GridIndex = ViewportX.GetValueAtOffset(MinDX);
+	int32 RightIndex = ViewportX.GetValueAtOffset(RoundedViewWidth);
 	int32 Delta = GridIndex - LeftIndex;
 
 	if (Delta > 0)
@@ -759,7 +739,7 @@ void SFrameTrack::DrawHorizontalAxisGrid(FDrawContext& DrawContext, const FSlate
 
 		for (int32 Index = StartIndex; Index < RightIndex; Index += Grid)
 		{
-			const float X = FMath::RoundToFloat(ViewportX.GetOffsetForIndex(Index));
+			const float X = FMath::RoundToFloat(ViewportX.GetOffsetForValue(Index));
 
 			// Draw vertical grid line.
 			DrawContext.DrawBox(X, 0, 1, ViewHeight, Brush, GridColor);
@@ -788,7 +768,6 @@ FReply SFrameTrack::OnMouseButtonDown(const FGeometry& MyGeometry, const FPointe
 
 	MousePositionOnButtonDown = MyGeometry.AbsoluteToLocal(MouseEvent.GetScreenSpacePosition());
 	ViewportPosXOnButtonDown = Viewport.GetHorizontalAxisViewport().GetPos();
-	ViewportPosYOnButtonDown = Viewport.GetVerticalAxisViewport().GetPos();
 
 	if (MouseEvent.GetEffectingButton() == EKeys::LeftMouseButton)
 	{
@@ -870,63 +849,46 @@ FReply SFrameTrack::OnMouseMove(const FGeometry& MyGeometry, const FPointerEvent
 
 	MousePosition = MyGeometry.AbsoluteToLocal(MouseEvent.GetScreenSpacePosition());
 
-	if (MouseEvent.IsMouseButtonDown(EKeys::LeftMouseButton))
+	if (!MouseEvent.GetCursorDelta().IsZero())
 	{
-		if (HasMouseCapture() && !MouseEvent.GetCursorDelta().IsZero())
+		if (MouseEvent.IsMouseButtonDown(EKeys::LeftMouseButton) ||
+			MouseEvent.IsMouseButtonDown(EKeys::RightMouseButton))
 		{
-			if (!bIsScrolling)
+			if (HasMouseCapture())
 			{
-				bIsScrolling = true;
-				CursorType = ECursorType::Hand;
-
-				HoveredSample.Reset();
-			}
+				if (!bIsScrolling)
+				{
+					bIsScrolling = true;
+					CursorType = ECursorType::Hand;
 
-			FIndexAxisViewport& ViewportX = Viewport.GetHorizontalAxisViewport();
-			const float PosX = ViewportPosXOnButtonDown + (MousePositionOnButtonDown.X - MousePosition.X);
-			ViewportX.ScrollAtIndex(ViewportX.GetIndexAtPos(PosX)); // align pos with at index
-			UpdateHorizontalScrollBar();
-			bIsStateDirty = true;
+					HoveredSample.Reset();
+				}
 
-			Reply = FReply::Handled();
+				FAxisViewportInt32& ViewportX = Viewport.GetHorizontalAxisViewport();
+				const float PosX = ViewportPosXOnButtonDown + (MousePositionOnButtonDown.X - MousePosition.X);
+				ViewportX.ScrollAtValue(ViewportX.GetValueAtPos(PosX)); // align viewport position with sample (frame index)
+				UpdateHorizontalScrollBar();
+				bIsStateDirty = true;
+			}
 		}
-	}
-	else if (MouseEvent.IsMouseButtonDown(EKeys::RightMouseButton))
-	{
-		if (HasMouseCapture() && !MouseEvent.GetCursorDelta().IsZero())
+		else
 		{
-			if (!bIsScrolling)
+			if (!HoveredSample.IsValid())
 			{
-				bIsScrolling = true;
-				CursorType = ECursorType::Hand;
-
-				HoveredSample.Reset();
+				TooltipOpacity = 0.0f;
+			}
+			HoveredSample = GetSampleAtMousePosition(MousePosition.X, MousePosition.Y);
+			if (!HoveredSample.IsValid())
+			{
+				HoveredSample = GetSampleAtMousePosition(MousePosition.X - 1.0f, MousePosition.Y);
+			}
+			if (!HoveredSample.IsValid())
+			{
+				HoveredSample = GetSampleAtMousePosition(MousePosition.X + 1.0f, MousePosition.Y);
 			}
-
-			FIndexAxisViewport& ViewportX = Viewport.GetHorizontalAxisViewport();
-			const float PosX = ViewportPosXOnButtonDown + (MousePositionOnButtonDown.X - MousePosition.X);
-			ViewportX.ScrollAtIndex(ViewportX.GetIndexAtPos(PosX)); // align pos with at index
-			UpdateHorizontalScrollBar();
-			bIsStateDirty = true;
-
-			Reply = FReply::Handled();
-		}
-	}
-	else
-	{
-		if (!HoveredSample.IsValid())
-		{
-			TooltipOpacity = 0.0f;
-		}
-		HoveredSample = GetSampleAtMousePosition(MousePosition.X, MousePosition.Y);
-		if (!HoveredSample.IsValid())
-		{
-			HoveredSample = GetSampleAtMousePosition(MousePosition.X - 1.0f, MousePosition.Y);
-		}
-		if (!HoveredSample.IsValid())
-		{
-			HoveredSample = GetSampleAtMousePosition(MousePosition.X + 1.0f, MousePosition.Y);
 		}
+
+		Reply = FReply::Handled();
 	}
 
 	return Reply;
@@ -961,7 +923,7 @@ FReply SFrameTrack::OnMouseWheel(const FGeometry& MyGeometry, const FPointerEven
 
 	if (MouseEvent.GetModifierKeys().IsShiftDown())
 	{
-		FValueAxisViewport& ViewportY = Viewport.GetVerticalAxisViewport();
+		FAxisViewportDouble& ViewportY = Viewport.GetVerticalAxisViewport();
 
 		// Zoom in/out vertically.
 		const float Delta = MouseEvent.GetWheelDelta();
@@ -994,9 +956,9 @@ FReply SFrameTrack::OnMouseWheel(const FGeometry& MyGeometry, const FPointerEven
 
 void SFrameTrack::ZoomHorizontally(const float Delta, const float X)
 {
-	FIndexAxisViewport& ViewportX = Viewport.GetHorizontalAxisViewport();
+	FAxisViewportInt32& ViewportX = Viewport.GetHorizontalAxisViewport();
 	ViewportX.RelativeZoomWithFixedOffset(Delta, X);
-	ViewportX.ScrollAtIndex(ViewportX.GetIndexAtPos(ViewportX.GetPos())); // align pos with at index
+	ViewportX.ScrollAtValue(ViewportX.GetValueAtPos(ViewportX.GetPos())); // align viewport position with sample (frame index)
 	UpdateHorizontalScrollBar();
 	bIsStateDirty = true;
 }
@@ -1150,7 +1112,7 @@ bool SFrameTrack::ContextMenu_ShowRenderingFrames_IsChecked()
 
 void SFrameTrack::ContextMenu_AutoZoom_Execute()
 {
-	FIndexAxisViewport& ViewportX = Viewport.GetHorizontalAxisViewport();
+	FAxisViewportInt32& ViewportX = Viewport.GetHorizontalAxisViewport();
 
 	bIsAutoZoomEnabled = !(bIsAutoZoomEnabled && ViewportX.GetPos() == 0.0f);
 	if (bIsAutoZoomEnabled)
diff --git a/Engine/Source/Developer/TraceInsights/Private/Insights/Widgets/SFrameTrack.h b/Engine/Source/Developer/TraceInsights/Private/Insights/Widgets/SFrameTrack.h
--- a/Engine/Source/Developer/TraceInsights/Private/Insights/Widgets/SFrameTrack.h
+++ b/Engine/Source/Developer/TraceInsights/Private/Insights/Widgets/SFrameTrack.h
@@ -71,7 +71,7 @@ struct FFrameTrackSampleRef
 
 ////////////////////////////////////////////////////////////////////////////////////////////////////
 /**
- * Widget used to present frames data in a track.
+ * Widget used to present frames data in a bar track.
  */
 class SFrameTrack : public SCompoundWidget
 {
@@ -165,19 +165,6 @@ protected:
 
 	void ZoomHorizontally(const float Delta, const float X);
 
-	//////////////////////////////////////////////////
-	// SelectionBoxChanged Event
-
-public:
-	/** The event to execute when the selection box has been changed. */
-	DECLARE_EVENT_TwoParams(SFrameTrack, FSelectionBoxChangedEvent, int32 /*FrameStart*/, int32 /*FrameEnd*/);
-	FSelectionBoxChangedEvent& OnSelectionBoxChanged() { return SelectionBoxChangedEvent; }
-protected:
-	/** The event to execute when the selection box has been changed. */
-	FSelectionBoxChangedEvent SelectionBoxChangedEvent;
-
-	//////////////////////////////////////////////////
-
 protected:
 	/** The track's viewport. Encapsulates info about position and scale. */
 	FFrameTrackViewport Viewport;
@@ -208,7 +195,6 @@ protected:
 	/** Mouse position during the call on mouse button down. */
 	FVector2D MousePositionOnButtonDown;
 	float ViewportPosXOnButtonDown;
-	float ViewportPosYOnButtonDown;
 
 	/** Mouse position during the call on mouse button up. */
 	FVector2D MousePositionOnButtonUp;
@@ -216,15 +202,12 @@ protected:
 	bool bIsLMB_Pressed;
 	bool bIsRMB_Pressed;
 
-	/** True, if the user is currently interactively scrolling the view by holding the right mouse button and dragging. */
+	/** True, if the user is currently interactively scrolling the view (ex.: by holding the left mouse button and dragging). */
 	bool bIsScrolling;
 
 	//////////////////////////////////////////////////
 	// Selection
 
-	int32 SelectionStartFrameIndex;
-	int32 SelectionEndFrameIndex;
-
 	FFrameTrackSampleRef HoveredSample;
 
 	float TooltipDesiredOpacity;
diff --git a/Engine/Source/Developer/TraceInsights/Private/Insights/Widgets/SGraphTrack.cpp b/Engine/Source/Developer/TraceInsights/Private/Insights/Widgets/SGraphTrack.cpp
--- a/Engine/Source/Developer/TraceInsights/Private/Insights/Widgets/SGraphTrack.cpp
+++ b/Engine/Source/Developer/TraceInsights/Private/Insights/Widgets/SGraphTrack.cpp
@@ -2,12 +2,8 @@
 
 #include "SGraphTrack.h"
 
-#include "Brushes/SlateBorderBrush.h"
-#include "Brushes/SlateBoxBrush.h"
-#include "Brushes/SlateColorBrush.h"
 #include "Containers/ArrayBuilder.h"
 #include "Containers/MapBuilder.h"
-#include "EditorStyleSet.h"
 #include "Fonts/FontMeasure.h"
 #include "Fonts/SlateFontInfo.h"
 #include "Framework/Application/MenuStack.h"
@@ -27,6 +23,7 @@
 // Insights
 #include "Insights/Common/PaintUtils.h"
 #include "Insights/Common/TimeUtils.h"
+#include "Insights/InsightsStyle.h"
 #include "Insights/TimingProfilerManager.h"
 #include "Insights/ViewModels/DrawHelpers.h"
 
@@ -41,7 +38,7 @@
 SGraphTrack::SGraphTrack()
 	: TimeRulerTrack(MakeShareable(new FTimeRulerTrack(0)))
 	, GraphTrack(MakeShareable(new FRandomGraphTrack(1)))
-	, WhiteBrush(FCoreStyle::Get().GetBrush("WhiteBrush"))
+	, WhiteBrush(FInsightsStyle::Get().GetBrush("WhiteBrush"))
 	, MainFont(FCoreStyle::GetDefaultFontStyle("Regular", 8))
 {
 	Reset();
@@ -62,8 +59,8 @@ void SGraphTrack::Reset()
 	StaticCastSharedPtr<FRandomGraphTrack>(GraphTrack)->AddDefaultSeries();
 
 	Viewport.Reset();
-	Viewport.MaxValidTime = 84.0 * 60.0;
-	//Viewport.ScaleX = (5 * 20) / 0.1; // 100ms between major tick marks
+	Viewport.SetMaxValidTime(84.0 * 60.0);
+	//Viewport.SetScaleX((5 * 20) / 0.1); // 100ms between major tick marks
 
 	bIsViewportDirty = true;
 	bIsVerticalViewportDirty = true;
@@ -112,7 +109,7 @@ void SGraphTrack::Tick(const FGeometry& AllottedGeometry, const double InCurrent
 		GraphTrack->GetHeight() != TrackHeight)
 	{
 		GraphTrack->SetPosY(TimeRulerTrack->GetHeight());
-		GraphTrack->SetHeight(Viewport.Height - TimeRulerTrack->GetHeight());
+		GraphTrack->SetHeight(Viewport.GetHeight() - TimeRulerTrack->GetHeight());
 		GraphTrack->SetDirtyFlag();
 	}
 
@@ -130,7 +127,7 @@ int32 SGraphTrack::OnPaint(const FPaintArgs& Args, const FGeometry& AllottedGeom
 	GraphTrack->Draw(DrawContext, Viewport, MousePosition);
 
 	TimeRulerTrack->Draw(DrawContext, Viewport, MousePosition, bIsSelecting, SelectionStartTime, SelectionEndTime);
-	DrawContext.DrawBox(0.0f, TimeRulerTrack->GetHeight(), Viewport.Width, 1.0f, WhiteBrush, FLinearColor(0.05f, 0.05f, 0.05f, 1.0f));
+	DrawContext.DrawBox(0.0f, TimeRulerTrack->GetHeight(), Viewport.GetWidth(), 1.0f, WhiteBrush, FLinearColor(0.05f, 0.05f, 0.05f, 1.0f));
 
 	FDrawHelpers::DrawTimeRangeSelection(DrawContext, Viewport, SelectionStartTime, SelectionEndTime, WhiteBrush, MainFont);
 
@@ -147,8 +144,8 @@ int32 SGraphTrack::OnPaint(const FPaintArgs& Args, const FGeometry& AllottedGeom
 
 		const float DbgW = 320.0f;
 		const float DbgH = DbgDY * 2 + 3.0f;
-		const float DbgX = Viewport.Width - DbgW - 20.0f;
-		float DbgY = Viewport.TopOffset + 10.0f;
+		const float DbgX = Viewport.GetWidth() - DbgW - 20.0f;
+		float DbgY = Viewport.GetTopOffset() + 10.0f;
 
 		DrawContext.LayerId++;
 
@@ -163,7 +160,10 @@ int32 SGraphTrack::OnPaint(const FPaintArgs& Args, const FGeometry& AllottedGeom
 		DrawContext.DrawText
 		(
 			DbgX, DbgY,
-			FString::Printf(TEXT("SX: %g, ST: %g, ET: %s"), Viewport.ScaleX, Viewport.StartTime, *TimeUtils::FormatTimeAuto(Viewport.MaxValidTime)),
+			FString::Printf(TEXT("SX: %g, ST: %g, ET: %s"),
+				Viewport.GetScaleX(),
+				Viewport.GetStartTime(),
+				*TimeUtils::FormatTimeAuto(Viewport.GetMaxValidTime())),
 			SummaryFont, DbgTextColor
 		);
 		DbgY += DbgDY;
@@ -174,7 +174,10 @@ int32 SGraphTrack::OnPaint(const FPaintArgs& Args, const FGeometry& AllottedGeom
 		DrawContext.DrawText
 		(
 			DbgX, DbgY,
-			FString::Printf(TEXT("Y: %.2f, H: %g, VH: %g"), Viewport.ScrollPosY, Viewport.ScrollHeight, Viewport.Height),
+			FString::Printf(TEXT("Y: %.2f, H: %g, VH: %g"),
+				Viewport.GetScrollPosY(),
+				Viewport.GetScrollHeight(),
+				Viewport.GetHeight()),
 			SummaryFont, DbgTextColor
 		);
 		DbgY += DbgDY;
@@ -241,8 +244,8 @@ FReply SGraphTrack::OnMouseButtonDown(const FGeometry& MyGeometry, const FPointe
 		bIsPanning = true;
 		bIsDragging = false;
 
-		ViewportStartTimeOnButtonDown = Viewport.StartTime;
-		ViewportScrollPosYOnButtonDown = Viewport.ScrollPosY;
+		ViewportStartTimeOnButtonDown = Viewport.GetStartTime();
+		ViewportScrollPosYOnButtonDown = Viewport.GetScrollPosY();
 
 		if (MouseEvent.GetModifierKeys().IsControlDown())
 		{
@@ -364,50 +367,51 @@ FReply SGraphTrack::OnMouseMove(const FGeometry& MyGeometry, const FPointerEvent
 
 	MousePosition = MyGeometry.AbsoluteToLocal(MouseEvent.GetScreenSpacePosition());
 
-	if (bIsPanning)
+	if (!MouseEvent.GetCursorDelta().IsZero())
 	{
-		if (HasMouseCapture() && !MouseEvent.GetCursorDelta().IsZero())
+		if (bIsPanning)
 		{
-			bIsDragging = true;
-
-			if ((int32)PanningMode & (int32)EPanningMode::Horizontal)
+			if (HasMouseCapture())
 			{
-				const double StartTime = ViewportStartTimeOnButtonDown + static_cast<double>(MousePositionOnButtonDown.X - MousePosition.X) / Viewport.ScaleX;
-				ScrollAtTime(StartTime);
-			}
+				bIsDragging = true;
 
-			if ((int32)PanningMode & (int32)EPanningMode::Vertical)
-			{
-				const float ScrollPosY = ViewportScrollPosYOnButtonDown + (MousePositionOnButtonDown.Y - MousePosition.Y);
-				ScrollAtPosY(ScrollPosY);
+				if ((int32)PanningMode & (int32)EPanningMode::Horizontal)
+				{
+					const double StartTime = ViewportStartTimeOnButtonDown + static_cast<double>(MousePositionOnButtonDown.X - MousePosition.X) / Viewport.GetScaleX();
+					ScrollAtTime(StartTime);
+				}
+
+				if ((int32)PanningMode & (int32)EPanningMode::Vertical)
+				{
+					const float ScrollPosY = ViewportScrollPosYOnButtonDown + (MousePositionOnButtonDown.Y - MousePosition.Y);
+					ScrollAtPosY(ScrollPosY);
+				}
 			}
 		}
-
-		Reply = FReply::Handled();
-	}
-	else if (bIsSelecting)
-	{
-		if (HasMouseCapture() && !MouseEvent.GetCursorDelta().IsZero())
+		else if (bIsSelecting)
 		{
-			bIsDragging = true;
-
-			SelectionStartTime = Viewport.SlateUnitsToTime(MousePositionOnButtonDown.X);
-			SelectionEndTime = Viewport.SlateUnitsToTime(MousePosition.X);
-			if (SelectionStartTime > SelectionEndTime)
+			if (HasMouseCapture())
 			{
-				double Temp = SelectionStartTime;
-				SelectionStartTime = SelectionEndTime;
-				SelectionEndTime = Temp;
+				bIsDragging = true;
+
+				SelectionStartTime = Viewport.SlateUnitsToTime(MousePositionOnButtonDown.X);
+				SelectionEndTime = Viewport.SlateUnitsToTime(MousePosition.X);
+				if (SelectionStartTime > SelectionEndTime)
+				{
+					double Temp = SelectionStartTime;
+					SelectionStartTime = SelectionEndTime;
+					SelectionEndTime = Temp;
+				}
+				//TODO: SelectionChangingEvent.Broadcast(SelectionStartTime, SelectionEndTime);
 			}
-			//TODO: SelectionChangingEvent.Broadcast(SelectionStartTime, SelectionEndTime);
+		}
+		else
+		{
+			GraphTrack->UpdateHoveredState(MousePosition.X, MousePosition.Y, Viewport);
 		}
 
 		Reply = FReply::Handled();
 	}
-	else
-	{
-		GraphTrack->UpdateHoveredState(MousePosition.X, MousePosition.Y, Viewport);
-	}
 
 	return Reply;
 }
@@ -431,8 +435,10 @@ void SGraphTrack::OnMouseLeave(const FPointerEvent& MouseEvent)
 
 		bIsLMB_Pressed = false;
 		bIsRMB_Pressed = false;
-
+		
 		MousePosition = FVector2D::ZeroVector;
+
+		GraphTrack->UpdateHoveredState(MousePosition.X, MousePosition.Y, Viewport);
 	}
 }
 
@@ -444,35 +450,21 @@ FReply SGraphTrack::OnMouseWheel(const FGeometry& MyGeometry, const FPointerEven
 	{
 		// Scroll vertically.
 		constexpr float ScrollSpeedY = 16.0f * 3;
-		const float ScrollPosY = Viewport.ScrollPosY - ScrollSpeedY * MouseEvent.GetWheelDelta();
+		const float ScrollPosY = Viewport.GetScrollPosY() - ScrollSpeedY * MouseEvent.GetWheelDelta();
 		ScrollAtPosY(ScrollPosY);
 	}
 	else if (MouseEvent.GetModifierKeys().IsControlDown())
 	{
 		// Scroll horizontally.
 		const double ScrollSpeedX = Viewport.GetDurationForViewportDX(16.0 * 3);
-		ScrollAtTime(Viewport.StartTime - ScrollSpeedX * MouseEvent.GetWheelDelta());
+		ScrollAtTime(Viewport.GetStartTime() - ScrollSpeedX * MouseEvent.GetWheelDelta());
 	}
 	else
 	{
 		// Zoom in/out horizontally.
 		const double Delta = MouseEvent.GetWheelDelta();
-		constexpr double ZoomStep = 0.25; // as percent
-		double ScaleX;
-
-		if (Delta > 0)
-		{
-			ScaleX = Viewport.ScaleX * FMath::Pow(1.0 + ZoomStep, Delta);
-		}
-		else
-		{
-			ScaleX = Viewport.ScaleX * FMath::Pow(1.0 / (1.0 + ZoomStep), -Delta);
-		}
-
-		//UE_LOG(TimingProfiler, Log, TEXT("%.2f, %.2f, %.2f"), Delta, Viewport.ScaleX, ScaleX);
 		MousePosition = MyGeometry.AbsoluteToLocal(MouseEvent.GetScreenSpacePosition());
-
-		if (Viewport.ZoomWithFixedX(ScaleX, MousePosition.X))
+		if (Viewport.RelativeZoomWithFixedX(Delta, MousePosition.X))
 		{
 			//UpdateHorizontalScrollBar();
 			bIsViewportDirty = true;
@@ -554,9 +546,9 @@ FCursorReply SGraphTrack::OnCursorQuery(const FGeometry& MyGeometry, const FPoin
 
 void SGraphTrack::ScrollAtPosY(float ScrollPosY)
 {
-	if (Viewport.ScrollPosY != ScrollPosY)
+	if (ScrollPosY != Viewport.GetScrollPosY())
 	{
-		Viewport.ScrollPosY = ScrollPosY;
+		Viewport.SetScrollPosY(ScrollPosY);
 
 		//UpdateVerticalScrollBar();
 		bIsVerticalViewportDirty = true;
diff --git a/Engine/Source/Developer/TraceInsights/Private/Insights/Widgets/SLogView.cpp b/Engine/Source/Developer/TraceInsights/Private/Insights/Widgets/SLogView.cpp
--- a/Engine/Source/Developer/TraceInsights/Private/Insights/Widgets/SLogView.cpp
+++ b/Engine/Source/Developer/TraceInsights/Private/Insights/Widgets/SLogView.cpp
@@ -17,6 +17,7 @@
 // Insights
 #include "Insights/Common/TimeUtils.h"
 #include "Insights/InsightsManager.h"
+#include "Insights/InsightsStyle.h"
 #include "Insights/TimingProfilerCommon.h"
 #include "Insights/TimingProfilerManager.h"
 #include "Insights/ViewModels/TimingViewDrawHelper.h"
@@ -27,7 +28,6 @@
 
 #define LOCTEXT_NAMESPACE "SLogView"
 
-
 ////////////////////////////////////////////////////////////////////////////////////////////////////
 // SLogMessageRow
 ////////////////////////////////////////////////////////////////////////////////////////////////////
@@ -57,7 +57,7 @@ public:
 		ChildSlot
 		[
 			SNew(SBorder)
-				.BorderImage(FCoreStyle::Get().GetBrush("WhiteBrush"))
+				.BorderImage(FInsightsStyle::Get().GetBrush("WhiteBrush"))
 				.BorderBackgroundColor(this, &SLogMessageRow::GetBackgroundColor)
 				[
 					Row
diff --git a/Engine/Source/Developer/TraceInsights/Private/Insights/Widgets/SStartPageWindow.cpp b/Engine/Source/Developer/TraceInsights/Private/Insights/Widgets/SStartPageWindow.cpp
--- a/Engine/Source/Developer/TraceInsights/Private/Insights/Widgets/SStartPageWindow.cpp
+++ b/Engine/Source/Developer/TraceInsights/Private/Insights/Widgets/SStartPageWindow.cpp
@@ -10,6 +10,7 @@
 #include "HAL/FileManagerGeneric.h"
 #include "Input/Events.h"
 #include "SlateOptMacros.h"
+#include "Styling/CoreStyle.h"
 #include "Widgets/Docking/SDockTab.h"
 #include "Widgets/Images/SImage.h"
 #include "Widgets/Layout/SBorder.h"
diff --git a/Engine/Source/Developer/TraceInsights/Private/Insights/Widgets/SStatsTableCell.cpp b/Engine/Source/Developer/TraceInsights/Private/Insights/Widgets/SStatsTableCell.cpp
--- a/Engine/Source/Developer/TraceInsights/Private/Insights/Widgets/SStatsTableCell.cpp
+++ b/Engine/Source/Developer/TraceInsights/Private/Insights/Widgets/SStatsTableCell.cpp
@@ -8,10 +8,12 @@
 #include "Widgets/Input/SButton.h"
 #include "Widgets/Layout/SBox.h"
 #include "Widgets/SBoxPanel.h"
+#include "Widgets/SToolTip.h"
 #include "Widgets/Text/STextBlock.h"
 #include "Widgets/Views/SExpanderArrow.h"
 
 // Insights
+#include "Insights/InsightsStyle.h"
 #include "Insights/ViewModels/StatsViewColumn.h"
 #include "Insights/ViewModels/StatsViewColumnFactory.h"
 #include "Insights/Widgets/SStatsViewTooltip.h"
@@ -90,7 +92,7 @@ TSharedRef<SWidget> SStatsTableCell::GenerateWidgetForNameColumn(const FArgument
 			.HeightOverride(14.0f)
 			[
 				SNew(SBorder)
-				.BorderImage(FCoreStyle::Get().GetBrush("WhiteBrush"))
+				.BorderImage(FInsightsStyle::Get().GetBrush("WhiteBrush"))
 				.BorderBackgroundColor(this, &SStatsTableCell::GetStatsBoxColorAndOpacity)
 				.HAlign(HAlign_Fill)
 				.VAlign(VAlign_Fill)
diff --git a/Engine/Source/Developer/TraceInsights/Private/Insights/Widgets/SStatsViewTooltip.cpp b/Engine/Source/Developer/TraceInsights/Private/Insights/Widgets/SStatsViewTooltip.cpp
--- a/Engine/Source/Developer/TraceInsights/Private/Insights/Widgets/SStatsViewTooltip.cpp
+++ b/Engine/Source/Developer/TraceInsights/Private/Insights/Widgets/SStatsViewTooltip.cpp
@@ -4,7 +4,7 @@
 
 #include "EditorStyleSet.h"
 #include "SlateOptMacros.h"
-#include "Widgets/Images/SImage.h"
+#include "TraceServices/Model/AnalysisSession.h"
 #include "Widgets/Layout/SGridPanel.h"
 #include "Widgets/Layout/SSeparator.h"
 #include "Widgets/SBoxPanel.h"
diff --git a/Engine/Source/Developer/TraceInsights/Private/Insights/Widgets/SStatsViewTooltip.h b/Engine/Source/Developer/TraceInsights/Private/Insights/Widgets/SStatsViewTooltip.h
--- a/Engine/Source/Developer/TraceInsights/Private/Insights/Widgets/SStatsViewTooltip.h
+++ b/Engine/Source/Developer/TraceInsights/Private/Insights/Widgets/SStatsViewTooltip.h
@@ -3,18 +3,15 @@
 #pragma once
 
 #include "CoreMinimal.h"
-#include "EditorStyleSet.h"
-#include "SlateOptMacros.h"
-#include "TraceServices/AnalysisService.h"
-#include "Widgets/Images/SImage.h"
-#include "Widgets/Layout/SGridPanel.h"
-#include "Widgets/Layout/SSeparator.h"
-#include "Widgets/SToolTip.h"
-#include "Widgets/Text/STextBlock.h"
 
 // Insights
 #include "Insights/InsightsManager.h"
 
+class SToolTip;
+class SGridPanel;
+
+namespace Trace { class IAnalysisSession; }
+
 // Insights
 class FStatsNode;
 class FStatsViewColumn;
diff --git a/Engine/Source/Developer/TraceInsights/Private/Insights/Widgets/STimerTableCell.cpp b/Engine/Source/Developer/TraceInsights/Private/Insights/Widgets/STimerTableCell.cpp
--- a/Engine/Source/Developer/TraceInsights/Private/Insights/Widgets/STimerTableCell.cpp
+++ b/Engine/Source/Developer/TraceInsights/Private/Insights/Widgets/STimerTableCell.cpp
@@ -7,6 +7,7 @@
 #include "Widgets/Images/SImage.h"
 #include "Widgets/Input/SButton.h"
 #include "Widgets/SBoxPanel.h"
+#include "Widgets/SToolTip.h"
 #include "Widgets/Text/STextBlock.h"
 #include "Widgets/Views/SExpanderArrow.h"
 
diff --git a/Engine/Source/Developer/TraceInsights/Private/Insights/Widgets/STimersViewTooltip.cpp b/Engine/Source/Developer/TraceInsights/Private/Insights/Widgets/STimersViewTooltip.cpp
--- a/Engine/Source/Developer/TraceInsights/Private/Insights/Widgets/STimersViewTooltip.cpp
+++ b/Engine/Source/Developer/TraceInsights/Private/Insights/Widgets/STimersViewTooltip.cpp
@@ -4,7 +4,7 @@
 
 #include "EditorStyleSet.h"
 #include "SlateOptMacros.h"
-#include "Widgets/Images/SImage.h"
+#include "TraceServices/Model/AnalysisSession.h"
 #include "Widgets/Layout/SGridPanel.h"
 #include "Widgets/Layout/SSeparator.h"
 #include "Widgets/SBoxPanel.h"
