[Insights]

diff --git a/Engine/Source/Developer/TraceInsights/Private/Insights/Widgets/STimingView.cpp b/Engine/Source/Developer/TraceInsights/Private/Insights/Widgets/STimingView.cpp
--- a/Engine/Source/Developer/TraceInsights/Private/Insights/Widgets/STimingView.cpp
+++ b/Engine/Source/Developer/TraceInsights/Private/Insights/Widgets/STimingView.cpp
@@ -2,9 +2,6 @@
 
 #include "STimingView.h"
 
-#include "Brushes/SlateBorderBrush.h"
-#include "Brushes/SlateBoxBrush.h"
-#include "Brushes/SlateColorBrush.h"
 #include "Containers/ArrayBuilder.h"
 #include "Containers/MapBuilder.h"
 #include "EditorStyleSet.h"
@@ -76,8 +73,7 @@ STimingView::STimingView()
 	, TimeRulerTrack(FBaseTimingTrack::GenerateId())
 	, MarkersTrack(FBaseTimingTrack::GenerateId())
 	, GraphTrack(MakeShareable(new FTimingGraphTrack(FBaseTimingTrack::GenerateId())))
-	, Tooltip()
-	, WhiteBrush(FCoreStyle::Get().GetBrush("WhiteBrush"))
+	, WhiteBrush(FInsightsStyle::Get().GetBrush("WhiteBrush"))
 	, MainFont(FCoreStyle::GetDefaultFontStyle("Regular", 8))
 {
 	Reset();
@@ -278,6 +274,7 @@ void STimingView::Reset()
 
 	HoveredTimingEvent.Reset();
 	SelectedTimingEvent.Reset();
+	Tooltip.Reset();
 
 	LastSelectionType = ESelectionType::None;
 
@@ -319,6 +316,9 @@ void STimingView::EnableAssetLoadingMode()
 
 void STimingView::Tick(const FGeometry& AllottedGeometry, const double InCurrentTime, const float InDeltaTime)
 {
+	//FStopwatch TickStopwatch;
+	//TickStopwatch.Start();
+
 	ThisGeometry = AllottedGeometry;
 
 	const float ViewWidth = AllottedGeometry.GetLocalSize().X;
@@ -348,7 +348,7 @@ void STimingView::Tick(const FGeometry& AllottedGeometry, const double InCurrent
 		GraphTrack->SetPosY(TopOffset);
 		TopOffset += GraphTrack->GetHeight();
 	}
-	Viewport.TopOffset = TopOffset;
+	Viewport.SetTopOffset(TopOffset);
 
 	//////////////////////////////////////////////////
 
@@ -358,10 +358,10 @@ void STimingView::Tick(const FGeometry& AllottedGeometry, const double InCurrent
 		// Elastic snap to vertical scroll limits.
 
 		const float MinY = 0.0f;// -0.5f * Viewport.Height;
-		const float MaxY = Viewport.ScrollHeight - Viewport.Height + TopOffset + 7.0f;
+		const float MaxY = Viewport.GetScrollHeight() - Viewport.GetHeight() + TopOffset + 7.0f;
 		const float U = 0.5f;
 
-		float ScrollPosY = Viewport.ScrollPosY;
+		float ScrollPosY = Viewport.GetScrollPosY();
 		if (ScrollPosY < MinY)
 		{
 			ScrollPosY = ScrollPosY * U + (1.0f - U) * MinY;
@@ -429,16 +429,16 @@ void STimingView::Tick(const FGeometry& AllottedGeometry, const double InCurrent
 
 		// Check if horizontal scroll area has changed.
 		double SessionTime = Session->GetDurationSeconds();
-		if (SessionTime > Viewport.MaxValidTime &&
+		if (SessionTime > Viewport.GetMaxValidTime() &&
 			SessionTime != DBL_MAX &&
 			SessionTime != std::numeric_limits<double>::infinity())
 		{
 			//UE_LOG(TimingProfiler, Log, TEXT("Session Duration: %g"), DT);
-			Viewport.MaxValidTime = SessionTime;
+			Viewport.SetMaxValidTime(SessionTime);
 			UpdateHorizontalScrollBar();
 			//bIsMaxValidTimeDirty = true;
 
-			if (SessionTime >= Viewport.StartTime && SessionTime <= Viewport.EndTime)
+			if (SessionTime >= Viewport.GetStartTime() && SessionTime <= Viewport.GetEndTime())
 			{
 				bAreTimingEventsTracksDirty = true;
 				MarkersTrack.SetDirtyFlag();
@@ -637,9 +637,9 @@ void STimingView::Tick(const FGeometry& AllottedGeometry, const double InCurrent
 	MarkersTrack.Tick(InCurrentTime, InDeltaTime);
 
 	// Check if vertical scroll area has changed.
-	if (TotalScrollHeight != Viewport.ScrollHeight)
+	if (TotalScrollHeight != Viewport.GetScrollHeight())
 	{
-		Viewport.ScrollHeight = TotalScrollHeight;
+		Viewport.SetScrollHeight(TotalScrollHeight);
 		UpdateVerticalScrollBar();
 		bIsVerticalViewportDirty = true;
 	}
@@ -708,6 +708,15 @@ void STimingView::Tick(const FGeometry& AllottedGeometry, const double InCurrent
 	{
 		GraphTrack->Update(Viewport);
 	}
+
+	Tooltip.Update();
+	if (!MousePosition.IsZero())
+	{
+		Tooltip.SetPosition(MousePosition, 0.0f, Viewport.GetWidth() - 12.0f, 0.0f, Viewport.GetHeight() - 12.0f); // -12.0f is to avoid overlaping the scrollbars
+	}
+
+	//TickStopwatch.Stop();
+	//TODO: TotalUpdateDurationHistory.AddValue(TickStopwatch.AccumulatedTime);
 }
 
 ////////////////////////////////////////////////////////////////////////////////////////////////////
@@ -833,9 +842,9 @@ int32 STimingView::OnPaint(const FPaintArgs& Args, const FGeometry& AllottedGeom
 	// Draw the time marker (orange vertical line).
 	//DrawTimeMarker(OnPaintState);
 	float TimeMarkerX = Viewport.TimeToSlateUnitsRounded(TimeMarker);
-	if (TimeMarkerX >= 0.0f && TimeMarkerX < Viewport.Width)
+	if (TimeMarkerX >= 0.0f && TimeMarkerX < Viewport.GetWidth())
 	{
-		DrawContext.DrawBox(TimeMarkerX, 0.0f, 1.0f, Viewport.Height, WhiteBrush, FLinearColor(0.85f, 0.5f, 0.03f, 0.5f));
+		DrawContext.DrawBox(TimeMarkerX, 0.0f, 1.0f, Viewport.GetHeight(), WhiteBrush, FLinearColor(0.85f, 0.5f, 0.03f, 0.5f));
 		DrawContext.LayerId++;
 	}
 
@@ -857,7 +866,7 @@ int32 STimingView::OnPaint(const FPaintArgs& Args, const FGeometry& AllottedGeom
 		const float DbgW = 320.0f;
 		const float DbgH = DbgDY * 9 + 3.0f;
 		const float DbgX = ViewWidth - DbgW - 20.0f;
-		float DbgY = Viewport.TopOffset + 10.0f;
+		float DbgY = Viewport.GetTopOffset() + 10.0f;
 
 		DrawContext.LayerId++;
 
@@ -983,7 +992,10 @@ int32 STimingView::OnPaint(const FPaintArgs& Args, const FGeometry& AllottedGeom
 		DrawContext.DrawText
 		(
 			DbgX, DbgY,
-			FString::Printf(TEXT("SX: %g, ST: %g, ET: %s"), Viewport.ScaleX, Viewport.StartTime, *TimeUtils::FormatTimeAuto(Viewport.MaxValidTime)),
+			FString::Printf(TEXT("SX: %g, ST: %g, ET: %s"),
+				Viewport.GetScaleX(),
+				Viewport.GetStartTime(),
+				*TimeUtils::FormatTimeAuto(Viewport.GetMaxValidTime())),
 			SummaryFont, DbgTextColor
 		);
 		DbgY += DbgDY;
@@ -994,7 +1006,10 @@ int32 STimingView::OnPaint(const FPaintArgs& Args, const FGeometry& AllottedGeom
 		DrawContext.DrawText
 		(
 			DbgX, DbgY,
-			FString::Printf(TEXT("Y: %.2f, H: %g, VH: %g"), Viewport.ScrollPosY, Viewport.ScrollHeight, Viewport.Height),
+			FString::Printf(TEXT("Y: %.2f, H: %g, VH: %g"),
+				Viewport.GetScrollPosY(),
+				Viewport.GetScrollHeight(),
+				Viewport.GetHeight()),
 			SummaryFont, DbgTextColor
 		);
 		DbgY += DbgDY;
@@ -1022,15 +1037,8 @@ int32 STimingView::OnPaint(const FPaintArgs& Args, const FGeometry& AllottedGeom
 			SelectedTimingEvent.Track->DrawSelectedEventInfo(SelectedTimingEvent, Viewport, DrawContext, WhiteBrush, MainFont);
 		}
 
-		// Draw info about hovered event (like a tooltip at mouse position).
-		if (HoveredTimingEvent.IsValid() && !MousePosition.IsZero())
-		{
-			HoveredTimingEvent.Track->DrawTooltip(Tooltip, MousePosition, HoveredTimingEvent, Viewport, DrawContext, WhiteBrush, MainFont);
-		}
-		else
-		{
-			Tooltip.Reset();
-		}
+		// Draw tooltip with info about hovered event.
+		Tooltip.Draw(DrawContext);
 
 		if (GraphTrack->IsVisible())
 		{
@@ -1299,8 +1307,8 @@ FReply STimingView::OnMouseButtonDown(const FGeometry& MyGeometry, const FPointe
 		bIsPanning = true;
 		bIsDragging = false;
 
-		ViewportStartTimeOnButtonDown = Viewport.StartTime;
-		ViewportScrollPosYOnButtonDown = Viewport.ScrollPosY;
+		ViewportStartTimeOnButtonDown = Viewport.GetStartTime();
+		ViewportScrollPosYOnButtonDown = Viewport.GetScrollPosY();
 
 		if (MouseEvent.GetModifierKeys().IsControlDown())
 		{
@@ -1456,62 +1464,63 @@ FReply STimingView::OnMouseMove(const FGeometry& MyGeometry, const FPointerEvent
 
 	MousePosition = MyGeometry.AbsoluteToLocal(MouseEvent.GetScreenSpacePosition());
 
-	if (bIsPanning)
+	if (!MouseEvent.GetCursorDelta().IsZero())
 	{
-		if (HasMouseCapture() && !MouseEvent.GetCursorDelta().IsZero())
+		if (bIsPanning)
 		{
-			bIsDragging = true;
-
-			if ((int32)PanningMode & (int32)EPanningMode::Horizontal)
+			if (HasMouseCapture())
 			{
-				const double StartTime = ViewportStartTimeOnButtonDown + static_cast<double>(MousePositionOnButtonDown.X - MousePosition.X) / Viewport.ScaleX;
-				ScrollAtTime(StartTime);
-			}
+				bIsDragging = true;
 
-			if ((int32)PanningMode & (int32)EPanningMode::Vertical)
-			{
-				const float ScrollPosY = ViewportScrollPosYOnButtonDown + (MousePositionOnButtonDown.Y - MousePosition.Y);
-				ScrollAtPosY(ScrollPosY);
+				if ((int32)PanningMode & (int32)EPanningMode::Horizontal)
+				{
+					const double StartTime = ViewportStartTimeOnButtonDown + static_cast<double>(MousePositionOnButtonDown.X - MousePosition.X) / Viewport.GetScaleX();
+					ScrollAtTime(StartTime);
+				}
+
+				if ((int32)PanningMode & (int32)EPanningMode::Vertical)
+				{
+					const float ScrollPosY = ViewportScrollPosYOnButtonDown + (MousePositionOnButtonDown.Y - MousePosition.Y);
+					ScrollAtPosY(ScrollPosY);
+				}
 			}
 		}
-
-		Reply = FReply::Handled();
-	}
-	else if (bIsSelecting)
-	{
-		if (HasMouseCapture() && !MouseEvent.GetCursorDelta().IsZero())
+		else if (bIsSelecting)
 		{
-			bIsDragging = true;
-
-			SelectionStartTime = Viewport.SlateUnitsToTime(MousePositionOnButtonDown.X);
-			SelectionEndTime = Viewport.SlateUnitsToTime(MousePosition.X);
-			if (SelectionStartTime > SelectionEndTime)
+			if (HasMouseCapture())
 			{
-				double Temp = SelectionStartTime;
-				SelectionStartTime = SelectionEndTime;
-				SelectionEndTime = Temp;
+				bIsDragging = true;
+
+				SelectionStartTime = Viewport.SlateUnitsToTime(MousePositionOnButtonDown.X);
+				SelectionEndTime = Viewport.SlateUnitsToTime(MousePosition.X);
+				if (SelectionStartTime > SelectionEndTime)
+				{
+					double Temp = SelectionStartTime;
+					SelectionStartTime = SelectionEndTime;
+					SelectionEndTime = Temp;
+				}
+				LastSelectionType = ESelectionType::TimeRange;
+				//TODO: SelectionChangingEvent.Broadcast(SelectionStartTime, SelectionEndTime);
+				EventAggregation.Reset();
+				ObjectTypeAggregation.Reset();
 			}
-			LastSelectionType = ESelectionType::TimeRange;
-			//TODO: SelectionChangingEvent.Broadcast(SelectionStartTime, SelectionEndTime);
-			EventAggregation.Reset();
-			ObjectTypeAggregation.Reset();
 		}
-
-		Reply = FReply::Handled();
-	}
-	else
-	{
-		if (MarkersTrack.IsVisible())
+		else
 		{
-			MarkersTrack.UpdateHoveredState(MousePosition.X, MousePosition.Y, Viewport);
-		}
+			if (MarkersTrack.IsVisible())
+			{
+				MarkersTrack.UpdateHoveredState(MousePosition.X, MousePosition.Y, Viewport);
+			}
 
-		if (GraphTrack->IsVisible())
-		{
-			GraphTrack->UpdateHoveredState(MousePosition.X, MousePosition.Y, Viewport);
+			if (GraphTrack->IsVisible())
+			{
+				GraphTrack->UpdateHoveredState(MousePosition.X, MousePosition.Y, Viewport);
+			}
+
+			UpdateHoveredTimingEvent(MousePosition.X, MousePosition.Y);
 		}
 
-		UpdateHoveredTimingEvent(MousePosition.X, MousePosition.Y);
+		Reply = FReply::Handled();
 	}
 
 	return Reply;
@@ -1538,6 +1547,19 @@ void STimingView::OnMouseLeave(const FPointerEvent& MouseEvent)
 		bIsRMB_Pressed = false;
 
 		MousePosition = FVector2D::ZeroVector;
+
+		if (MarkersTrack.IsVisible())
+		{
+			MarkersTrack.UpdateHoveredState(MousePosition.X, MousePosition.Y, Viewport);
+		}
+
+		if (GraphTrack->IsVisible())
+		{
+			GraphTrack->UpdateHoveredState(MousePosition.X, MousePosition.Y, Viewport);
+		}
+
+		HoveredTimingEvent.Reset();
+		Tooltip.SetDesiredOpacity(0.0f);
 	}
 }
 
@@ -1578,7 +1600,7 @@ FReply STimingView::OnMouseWheel(const FGeometry& MyGeometry, const FPointerEven
 		{
 			// Scroll vertically.
 			constexpr float ScrollSpeedY = 16.0f * 3;
-			const float ScrollPosY = Viewport.ScrollPosY - ScrollSpeedY * MouseEvent.GetWheelDelta();
+			const float ScrollPosY = Viewport.GetScrollPosY() - ScrollSpeedY * MouseEvent.GetWheelDelta();
 			ScrollAtPosY(ScrollPosY);
 		}
 	}
@@ -1586,28 +1608,14 @@ FReply STimingView::OnMouseWheel(const FGeometry& MyGeometry, const FPointerEven
 	{
 		// Scroll horizontally.
 		const double ScrollSpeedX = Viewport.GetDurationForViewportDX(16.0 * 3);
-		ScrollAtTime(Viewport.StartTime - ScrollSpeedX * MouseEvent.GetWheelDelta());
+		ScrollAtTime(Viewport.GetStartTime() - ScrollSpeedX * MouseEvent.GetWheelDelta());
 	}
 	else
 	{
 		// Zoom in/out horizontally.
 		const double Delta = MouseEvent.GetWheelDelta();
-		constexpr double ZoomStep = 0.25; // as percent
-		double ScaleX;
-
-		if (Delta > 0)
-		{
-			ScaleX = Viewport.ScaleX * FMath::Pow(1.0 + ZoomStep, Delta);
-		}
-		else
-		{
-			ScaleX = Viewport.ScaleX * FMath::Pow(1.0 / (1.0 + ZoomStep), -Delta);
-		}
-
-		//UE_LOG(TimingProfiler, Log, TEXT("%.2f, %.2f, %.2f"), Delta, Viewport.ScaleX, ScaleX);
 		MousePosition = MyGeometry.AbsoluteToLocal(MouseEvent.GetScreenSpacePosition());
-
-		if (Viewport.ZoomWithFixedX(ScaleX, MousePosition.X))
+		if (Viewport.RelativeZoomWithFixedX(Delta, MousePosition.X))
 		{
 			//Viewport.EnforceHorizontalScrollLimits(1.0);
 			UpdateHorizontalScrollBar();
@@ -1758,8 +1766,8 @@ FReply STimingView::OnKeyDown(const FGeometry& MyGeometry, const FKeyEvent& InKe
 			 InKeyEvent.GetKey() == EKeys::Add)
 	{
 		// Zoom In
-		double ScaleX = Viewport.ScaleX * 1.25;
-		if (Viewport.ZoomWithFixedX(ScaleX, Viewport.Width/2))
+		const double ScaleX = Viewport.GetScaleX() * 1.25;
+		if (Viewport.ZoomWithFixedX(ScaleX, Viewport.GetWidth() / 2))
 		{
 			//Viewport.EnforceHorizontalScrollLimits(1.0);
 			UpdateHorizontalScrollBar();
@@ -1771,8 +1779,8 @@ FReply STimingView::OnKeyDown(const FGeometry& MyGeometry, const FKeyEvent& InKe
 			 InKeyEvent.GetKey() == EKeys::Subtract)
 	{
 		// Zoom Out
-		double ScaleX = Viewport.ScaleX / 1.25;
-		if (Viewport.ZoomWithFixedX(ScaleX, Viewport.Width / 2))
+		const double ScaleX = Viewport.GetScaleX() / 1.25;
+		if (Viewport.ZoomWithFixedX(ScaleX, Viewport.GetWidth() / 2))
 		{
 			//Viewport.EnforceHorizontalScrollLimits(1.0);
 			UpdateHorizontalScrollBar();
@@ -1785,9 +1793,9 @@ FReply STimingView::OnKeyDown(const FGeometry& MyGeometry, const FKeyEvent& InKe
 		if (InKeyEvent.GetModifierKeys().IsControlDown())
 		{
 			// Scroll Left
-			double DT = Viewport.EndTime - Viewport.StartTime;
-			//ScrollAtTime(Viewport.StartTime - DT * 0.05);
-			if (Viewport.ScrollAtTime(Viewport.StartTime - DT * 0.05))
+			const double DT = Viewport.GetDuration();
+			//ScrollAtTime(Viewport.GetStartTime() - DT * 0.05);
+			if (Viewport.ScrollAtTime(Viewport.GetStartTime() - DT * 0.05))
 			{
 				//Viewport.EnforceHorizontalScrollLimits(1.0);
 				UpdateHorizontalScrollBar();
@@ -1805,9 +1813,9 @@ FReply STimingView::OnKeyDown(const FGeometry& MyGeometry, const FKeyEvent& InKe
 		if (InKeyEvent.GetModifierKeys().IsControlDown())
 		{
 			// Scroll Right
-			double DT = Viewport.EndTime - Viewport.StartTime;
-			//ScrollAtTime(Viewport.StartTime + DT * 0.05);
-			if (Viewport.ScrollAtTime(Viewport.StartTime + DT * 0.05))
+			const double DT = Viewport.GetDuration();
+			//ScrollAtTime(Viewport.GetStartTime() + DT * 0.05);
+			if (Viewport.ScrollAtTime(Viewport.GetStartTime() + DT * 0.05))
 			{
 				//Viewport.EnforceHorizontalScrollLimits(1.0);
 				UpdateHorizontalScrollBar();
@@ -1825,7 +1833,7 @@ FReply STimingView::OnKeyDown(const FGeometry& MyGeometry, const FKeyEvent& InKe
 		if (InKeyEvent.GetModifierKeys().IsControlDown())
 		{
 			// Scroll Up
-			ScrollAtPosY(Viewport.ScrollPosY - 16.0 * 3);
+			ScrollAtPosY(Viewport.GetScrollPosY() - 16.0f * 3);
 		}
 		else
 		{
@@ -1838,7 +1846,7 @@ FReply STimingView::OnKeyDown(const FGeometry& MyGeometry, const FKeyEvent& InKe
 		if (InKeyEvent.GetModifierKeys().IsControlDown())
 		{
 			// Scroll Down
-			ScrollAtPosY(Viewport.ScrollPosY + 16.0 * 3);
+			ScrollAtPosY(Viewport.GetScrollPosY() + 16.0f * 3);
 		}
 		else
 		{
@@ -2056,71 +2064,43 @@ void STimingView::BindCommands()
 
 void STimingView::HorizontalScrollBar_OnUserScrolled(float ScrollOffset)
 {
-	const double SX = 1.0 / (Viewport.MaxValidTime - Viewport.MinValidTime);
-	const float ThumbSizeFraction = FMath::Clamp<float>((Viewport.EndTime - Viewport.StartTime) * SX, 0.0f, 1.0f);
-	const float OffsetFraction = FMath::Clamp<float>(ScrollOffset, 0.0f, 1.0f - ThumbSizeFraction);
-
-	const double Time = Viewport.MinValidTime + static_cast<double>(OffsetFraction) * (Viewport.MaxValidTime - Viewport.MinValidTime);
-	if (Viewport.StartTime != Time)
+	if (Viewport.OnUserScrolled(HorizontalScrollBar, ScrollOffset))
 	{
-		Viewport.ScrollAtTime(Time);
 		bIsViewportDirty = true;
 	}
-
-	HorizontalScrollBar->SetState(OffsetFraction, ThumbSizeFraction);
 }
 
 ////////////////////////////////////////////////////////////////////////////////////////////////////
 
 void STimingView::UpdateHorizontalScrollBar()
 {
-	const double SX = 1.0 / (Viewport.MaxValidTime - Viewport.MinValidTime);
-	const float ThumbSizeFraction = FMath::Clamp<float>((Viewport.EndTime - Viewport.StartTime) * SX, 0.0f, 1.0f);
-	const float ScrollOffset = static_cast<float>((Viewport.StartTime - Viewport.MinValidTime) * SX);
-	const float OffsetFraction = FMath::Clamp<float>(ScrollOffset, 0.0f, 1.0f - ThumbSizeFraction);
-
-	HorizontalScrollBar->SetState(OffsetFraction, ThumbSizeFraction);
+	Viewport.UpdateScrollBar(HorizontalScrollBar);
 }
 
 ////////////////////////////////////////////////////////////////////////////////////////////////////
 
 void STimingView::VerticalScrollBar_OnUserScrolled(float ScrollOffset)
 {
-	const float SY = 1.0 / Viewport.ScrollHeight;
-	const float H = Viewport.Height - Viewport.TopOffset;
-	const float ThumbSizeFraction = FMath::Clamp<float>(H * SY, 0.0f, 1.0f);
-	const float OffsetFraction = FMath::Clamp<float>(ScrollOffset, 0.0f, 1.0f - ThumbSizeFraction);
-
-	const float ScrollPosY = OffsetFraction * Viewport.ScrollHeight;
-	if (Viewport.ScrollPosY != ScrollPosY)
+	if (Viewport.OnUserScrolledY(VerticalScrollBar, ScrollOffset))
 	{
-		Viewport.ScrollPosY = ScrollPosY;
 		bIsVerticalViewportDirty = true;
 	}
-
-	VerticalScrollBar->SetState(OffsetFraction, ThumbSizeFraction);
 }
 
 ////////////////////////////////////////////////////////////////////////////////////////////////////
 
 void STimingView::UpdateVerticalScrollBar()
 {
-	const float SY = 1.0 / Viewport.ScrollHeight;
-	const float H = Viewport.Height - Viewport.TopOffset;
-	const float ThumbSizeFraction = FMath::Clamp<float>(H * SY, 0.0f, 1.0f);
-	const float ScrollOffset = Viewport.ScrollPosY * SY;
-	float OffsetFraction = FMath::Clamp<float>(ScrollOffset, 0.0f, 1.0f - ThumbSizeFraction);
-
-	VerticalScrollBar->SetState(OffsetFraction, ThumbSizeFraction);
+	Viewport.UpdateScrollBarY(VerticalScrollBar);
 }
 
 ////////////////////////////////////////////////////////////////////////////////////////////////////
 
 void STimingView::ScrollAtPosY(float ScrollPosY)
 {
-	if (Viewport.ScrollPosY != ScrollPosY)
+	if (ScrollPosY != Viewport.GetScrollPosY())
 	{
-		Viewport.ScrollPosY = ScrollPosY;
+		Viewport.SetScrollPosY(ScrollPosY);
 
 		UpdateVerticalScrollBar();
 		bIsVerticalViewportDirty = true;
@@ -2157,16 +2137,16 @@ void STimingView::BringIntoView(double StartTime, double EndTime)
 {
 	EndTime = Viewport.RestrictEndTime(EndTime);
 
-	// Increase interval with 8% on each side.
-	const double DT = (Viewport.EndTime - Viewport.StartTime) * 0.08;
+	// Increase interval with 8% (of view size) on each side.
+	const double DT = Viewport.GetDuration() * 0.08;
 	StartTime -= DT;
 	EndTime += DT;
 
-	double NewStartTime = Viewport.StartTime;
+	double NewStartTime = Viewport.GetStartTime();
 
-	if (EndTime > Viewport.EndTime)
+	if (EndTime > Viewport.GetEndTime())
 	{
-		NewStartTime += EndTime - Viewport.EndTime;
+		NewStartTime += EndTime - Viewport.GetEndTime();
 	}
 
 	if (StartTime < NewStartTime)
@@ -2331,7 +2311,7 @@ void STimingView::SetAndCenterOnTimeMarker(double Time)
 
 	double MinT, MaxT;
 	Viewport.GetHorizontalScrollLimits(MinT, MaxT);
-	const double ViewportDuration = static_cast<double>(Viewport.Width) / Viewport.ScaleX;
+	const double ViewportDuration = Viewport.GetDuration();
 	MinT += ViewportDuration / 2;
 	MaxT += ViewportDuration / 2;
 	Time = FMath::Clamp<double>(Time, MinT, MaxT);
@@ -2366,16 +2346,16 @@ void STimingView::SetTimeMarkersVisible(bool bIsMarkersTrackVisible)
 
 		if (MarkersTrack.IsVisible())
 		{
-			if (Viewport.ScrollPosY != 0.0f)
+			if (Viewport.GetScrollPosY() != 0.0f)
 			{
-				Viewport.ScrollPosY += MarkersTrack.GetHeight();
+				Viewport.SetScrollPosY(Viewport.GetScrollPosY() + MarkersTrack.GetHeight());
 			}
 
 			MarkersTrack.SetDirtyFlag();
 		}
 		else
 		{
-			Viewport.ScrollPosY -= MarkersTrack.GetHeight();
+			Viewport.SetScrollPosY(Viewport.GetScrollPosY() - MarkersTrack.GetHeight());
 		}
 	}
 }
@@ -2391,9 +2371,9 @@ void STimingView::SetDrawOnlyBookmarks(bool bIsBookmarksTrack)
 
 		if (MarkersTrack.IsVisible())
 		{
-			if (Viewport.ScrollPosY != 0.0f)
+			if (Viewport.GetScrollPosY() != 0.0f)
 			{
-				Viewport.ScrollPosY += MarkersTrack.GetHeight() - PrevHeight;
+				Viewport.SetScrollPosY(Viewport.GetScrollPosY() + MarkersTrack.GetHeight() - PrevHeight);
 			}
 
 			MarkersTrack.SetDirtyFlag();
@@ -2408,14 +2388,14 @@ void STimingView::UpdateHoveredTimingEvent(float MX, float MY)
 	HoveredTimingEvent.Track = nullptr;
 	HoveredTimingEvent.TypeId = FTimerNode::InvalidId;
 
-	if (MY >= Viewport.TopOffset && MY < Viewport.Height)
+	if (MY >= Viewport.GetTopOffset() && MY < Viewport.GetHeight())
 	{
 		for (const auto& KV : CachedTimelines)
 		{
 			const FTimingEventsTrack& Track = *KV.Value;
 			if (Track.IsVisible())
 			{
-				const float Y = Viewport.TopOffset + Track.GetPosY() - Viewport.ScrollPosY;
+				const float Y = Viewport.GetTopOffset() + Track.GetPosY() - Viewport.GetScrollPosY();
 				if (MY >= Y && MY < Y + Track.GetHeight())
 				{
 					HoveredTimingEvent.Track = &Track;
@@ -2426,7 +2406,7 @@ void STimingView::UpdateHoveredTimingEvent(float MX, float MY)
 
 		if (HoveredTimingEvent.Track)
 		{
-			const float Y0 = Viewport.TopOffset + HoveredTimingEvent.Track->GetPosY() - Viewport.ScrollPosY + 1.0f + Layout.TimelineDY;
+			const float Y0 = Viewport.GetTopOffset() + HoveredTimingEvent.Track->GetPosY() - Viewport.GetScrollPosY() + 1.0f + Layout.TimelineDY;
 
 			// If mouse is not above first sub-track or below last sub-track...
 			if (MY >= Y0 && MY < Y0 + HoveredTimingEvent.Track->GetHeight() + Layout.TimelineDY)
@@ -2435,7 +2415,7 @@ void STimingView::UpdateHoveredTimingEvent(float MX, float MY)
 				float EventMY = (MY - Y0) - Depth * (Layout.EventDY + Layout.EventH);
 
 				const double StartTime = Viewport.SlateUnitsToTime(MX);
-				const double EndTime = StartTime + 2.0 / Viewport.ScaleX; // +2px
+				const double EndTime = StartTime + 2.0 / Viewport.GetScaleX(); // +2px
 				constexpr bool bStopAtFirstMatch = true; // get first one matching
 				constexpr bool bSearchForLargestEvent = false;
 				HoveredTimingEvent.Track->SearchTimingEvent(StartTime, EndTime,
@@ -2449,6 +2429,16 @@ void STimingView::UpdateHoveredTimingEvent(float MX, float MY)
 			}
 		}
 	}
+
+	if (HoveredTimingEvent.IsValid())
+	{
+		HoveredTimingEvent.Track->InitTooltip(Tooltip, HoveredTimingEvent);
+		Tooltip.SetDesiredOpacity(1.0f);
+	}
+	else
+	{
+		Tooltip.SetDesiredOpacity(0.0f);
+	}
 }
 
 ////////////////////////////////////////////////////////////////////////////////////////////////////
@@ -2522,7 +2512,7 @@ void STimingView::SelectRightTimingEvent()
 		const double EndTime = SelectedTimingEvent.EndTime;
 		const bool bStopAtFirstMatch = true; // get first one matching
 		const bool bSearchForLargestEvent = false;
-		if (SelectedTimingEvent.Track->SearchTimingEvent(EndTime, Viewport.MaxValidTime,
+		if (SelectedTimingEvent.Track->SearchTimingEvent(EndTime, Viewport.GetMaxValidTime(),
 			[Depth, StartTime, EndTime](double EventStartTime, double EventEndTime, uint32 EventDepth)
 			{
 				return EventDepth == Depth &&
@@ -2641,7 +2631,7 @@ void STimingView::FrameSelection()
 		EndTime = Viewport.RestrictEndTime(SelectedTimingEvent.EndTime);
 		if (EndTime == StartTime)
 		{
-			EndTime += 1.0 / Viewport.ScaleX; // +1px
+			EndTime += 1.0 / Viewport.GetScaleX(); // +1px
 		}
 	}
 	else
@@ -2856,6 +2846,7 @@ void STimingView::ShowHideAllCpuTracks_Execute()
 
 	HoveredTimingEvent.Reset();
 	SelectedTimingEvent.Reset();
+	Tooltip.SetDesiredOpacity(0.0f);
 
 	bAreTimingEventsTracksDirty = true;
 }
@@ -2884,6 +2875,7 @@ void STimingView::ShowHideAllGpuTracks_Execute()
 
 	HoveredTimingEvent.Reset();
 	SelectedTimingEvent.Reset();
+	Tooltip.SetDesiredOpacity(0.0f);
 
 	bAreTimingEventsTracksDirty = true;
 }
@@ -2996,6 +2988,7 @@ void STimingView::ToggleTrackVisibility_Execute(uint64 InTrackId)
 
 		HoveredTimingEvent.Reset();
 		SelectedTimingEvent.Reset();
+		Tooltip.SetDesiredOpacity(0.0f);
 	}
 }
 
@@ -3032,6 +3025,7 @@ void STimingView::ToggleTrackVisibilityByGroup_Execute(const TCHAR* InGroupName)
 
 		HoveredTimingEvent.Reset();
 		SelectedTimingEvent.Reset();
+		Tooltip.SetDesiredOpacity(0.0f);
 	}
 }
 
diff --git a/Engine/Source/Developer/TraceInsights/Private/Insights/Widgets/STimingView.h b/Engine/Source/Developer/TraceInsights/Private/Insights/Widgets/STimingView.h
--- a/Engine/Source/Developer/TraceInsights/Private/Insights/Widgets/STimingView.h
+++ b/Engine/Source/Developer/TraceInsights/Private/Insights/Widgets/STimingView.h
@@ -21,6 +21,7 @@
 #include "Insights/ViewModels/TimingEventsTrack.h"
 #include "Insights/ViewModels/TimingTrackViewport.h"
 #include "Insights/ViewModels/TimingViewDrawHelper.h"
+#include "Insights/ViewModels/TooltipDrawState.h"
 
 class FFileActivitySharedState;
 class FLoadingSharedState;
@@ -488,11 +489,11 @@ protected:
 	double SelectionStartTime;
 	double SelectionEndTime;
 
-	mutable FTimingViewTooltip Tooltip;
-
 	FTimingEvent HoveredTimingEvent;
 	FTimingEvent SelectedTimingEvent;
 
+	FTooltipDrawState Tooltip;
+
 	enum class ESelectionType
 	{
 		None,
