[Insights]
  - Frames track: Improved computation for major vertical grid lines.
  - Loading/Timing profilers: Fixed toolbar when resizing to not hide buttons on left side too early (due to right side toolbar taking half of space).

diff --git a/Engine/Source/Developer/TraceInsights/Private/Insights/LoadingProfiler/Widgets/SLoadingProfilerToolbar.cpp b/Engine/Source/Developer/TraceInsights/Private/Insights/LoadingProfiler/Widgets/SLoadingProfilerToolbar.cpp
--- a/Engine/Source/Developer/TraceInsights/Private/Insights/LoadingProfiler/Widgets/SLoadingProfilerToolbar.cpp
+++ b/Engine/Source/Developer/TraceInsights/Private/Insights/LoadingProfiler/Widgets/SLoadingProfilerToolbar.cpp
@@ -93,7 +93,7 @@ void SLoadingProfilerToolbar::Construct(const FArguments& InArgs)
 		+SHorizontalBox::Slot()
 		.HAlign(HAlign_Right)
 		.VAlign(VAlign_Center)
-		.FillWidth(1.0)
+		.AutoWidth()
 		.Padding(0.0f)
 		[
 			SNew(SBorder)
diff --git a/Engine/Source/Developer/TraceInsights/Private/Insights/Widgets/SFrameTrack.cpp b/Engine/Source/Developer/TraceInsights/Private/Insights/Widgets/SFrameTrack.cpp
--- a/Engine/Source/Developer/TraceInsights/Private/Insights/Widgets/SFrameTrack.cpp
+++ b/Engine/Source/Developer/TraceInsights/Private/Insights/Widgets/SFrameTrack.cpp
@@ -658,11 +658,11 @@ void SFrameTrack::DrawVerticalAxisGrid(FDrawContext& DrawContext, const FSlateBr
 		}
 
 		const float Y = RoundedViewHeight - FMath::RoundToFloat(ViewportY.GetOffsetForValue(Value));
-		if (Y < -TextH)
+		if (Y < 0)
 		{
 			break;
 		}
-		if (Y > RoundedViewHeight || FMath::Abs(PreviousY - Y) < MinDY)
+		if (Y > RoundedViewHeight + TextH || FMath::Abs(PreviousY - Y) < MinDY)
 		{
 			continue;
 		}
@@ -693,12 +693,12 @@ void SFrameTrack::DrawHorizontalAxisGrid(FDrawContext& DrawContext, const FSlate
 
 	const float RoundedViewWidth = FMath::RoundToFloat(ViewportX.GetSize());
 
-	constexpr float MinDX = 100.0f; // min horizontal distance between vertical grid lines
+	constexpr float MinDX = 125.0f; // min horizontal distance between vertical grid lines
 
-	int32 LeftIndex = ViewportX.GetValueAtOffset(0.0f);
-	int32 GridIndex = ViewportX.GetValueAtOffset(MinDX);
-	int32 RightIndex = ViewportX.GetValueAtOffset(RoundedViewWidth);
-	int32 Delta = GridIndex - LeftIndex;
+	const int32 LeftIndex = ViewportX.GetValueAtOffset(0.0f);
+	const int32 GridIndex = ViewportX.GetValueAtOffset(MinDX);
+	const int32 RightIndex = ViewportX.GetValueAtOffset(RoundedViewWidth);
+	const int32 Delta = GridIndex - LeftIndex;
 
 	if (Delta > 0)
 	{
@@ -719,12 +719,7 @@ void SFrameTrack::DrawHorizontalAxisGrid(FDrawContext& DrawContext, const FSlate
 			Power10 = 1;
 		}
 
-		const int32 Grid = ((Delta + Power10) / Power10) * Power10; // next value divisible with a multiple of 10
-
-		const FLinearColor GridColor(0.0f, 0.0f, 0.0f, 0.1f);
-		//const FLinearColor TextBgColor(0.05f, 0.05f, 0.05f, 1.0f);
-		//const FLinearColor TextColor(1.0f, 1.0f, 1.0f, 1.0f);
-		const FLinearColor TopTextColor(1.0f, 1.0f, 1.0f, 0.7f);
+		const int32 Grid = ((Delta + Power10 - 1) / Power10) * Power10; // next value divisible with a multiple of 10
 
 		// Skip grid lines for negative indices.
 		double StartIndex = ((LeftIndex + Grid - 1) / Grid) * Grid;
@@ -735,6 +730,11 @@ void SFrameTrack::DrawHorizontalAxisGrid(FDrawContext& DrawContext, const FSlate
 
 		const float ViewHeight = Viewport.GetHeight();
 
+		const FLinearColor GridColor(0.0f, 0.0f, 0.0f, 0.1f);
+		//const FLinearColor TextBgColor(0.05f, 0.05f, 0.05f, 1.0f);
+		//const FLinearColor TextColor(1.0f, 1.0f, 1.0f, 1.0f);
+		const FLinearColor TopTextColor(1.0f, 1.0f, 1.0f, 0.7f);
+
 		//const TSharedRef<FSlateFontMeasure> FontMeasureService = FSlateApplication::Get().GetRenderer()->GetFontMeasureService();
 
 		for (int32 Index = StartIndex; Index < RightIndex; Index += Grid)
diff --git a/Engine/Source/Developer/TraceInsights/Private/Insights/Widgets/STimingProfilerToolbar.cpp b/Engine/Source/Developer/TraceInsights/Private/Insights/Widgets/STimingProfilerToolbar.cpp
--- a/Engine/Source/Developer/TraceInsights/Private/Insights/Widgets/STimingProfilerToolbar.cpp
+++ b/Engine/Source/Developer/TraceInsights/Private/Insights/Widgets/STimingProfilerToolbar.cpp
@@ -97,7 +97,7 @@ void STimingProfilerToolbar::Construct(const FArguments& InArgs)
 		+SHorizontalBox::Slot()
 		.HAlign(HAlign_Right)
 		.VAlign(VAlign_Center)
-		.FillWidth(1.0)
+		.AutoWidth()
 		.Padding(0.0f)
 		[
 			SNew(SBorder)
