[Insights]
  - Timing view: Added "% of Parent" and "% of Root" in tooltip of a cpu/gpu timing event.
  - Timing view: Improved performance when moving mouse over timing events.
  - Timing view: Fixed aggregated stats not being updated when a track is shown/hidden (if a time range selection is already made).
  - Timing view: Made "canceling selection" by clicking on time ruler to behave the same as canceling selection by clicking on an empty space (i.e. to not reset the stats in Timers view).
  - Timing view: Fixed regresion re background area not being highlighted as invalid (redish) above first track and below last track.

diff --git a/Engine/Source/Developer/TraceInsights/Private/Insights/ViewModels/DrawHelpers.cpp b/Engine/Source/Developer/TraceInsights/Private/Insights/ViewModels/DrawHelpers.cpp
--- a/Engine/Source/Developer/TraceInsights/Private/Insights/ViewModels/DrawHelpers.cpp
+++ b/Engine/Source/Developer/TraceInsights/Private/Insights/ViewModels/DrawHelpers.cpp
@@ -126,6 +126,28 @@ void FDrawHelpers::DrawBackground(const FDrawContext& DrawContext,
 
 ////////////////////////////////////////////////////////////////////////////////////////////////////
 
+void FDrawHelpers::DrawBackground(const FDrawContext& DrawContext,
+								  const FSlateBrush* BackgroundAreaBrush,
+								  const FTimingTrackViewport& Viewport,
+								  const float Y,
+								  const float H,
+								  float& OutValidAreaX,
+								  float& OutValidAreaW)
+{
+	const FLinearColor ValidAreaColor(0.07f, 0.07f, 0.07f, 1.0f);
+	const FLinearColor InvalidAreaColor(0.1f, 0.07f, 0.07f, 1.0f);
+	const FLinearColor EdgeColor(0.05f, 0.05f, 0.05f, 1.0f);
+
+	const float X0 = 0.0f;
+	const float X1 = Viewport.TimeToSlateUnitsRounded(Viewport.GetMinValidTime());
+	const float X2 = Viewport.TimeToSlateUnitsRounded(Viewport.GetMaxValidTime());
+	const float X3 = FMath::CeilToFloat(Viewport.GetWidth());
+
+	FDrawHelpers::DrawBackground(DrawContext, BackgroundAreaBrush, ValidAreaColor, InvalidAreaColor, EdgeColor, X0, X1, X2, X3, Y, H, OutValidAreaX, OutValidAreaW);
+}
+
+////////////////////////////////////////////////////////////////////////////////////////////////////
+
 void FDrawHelpers::DrawTimeRangeSelection(
 	const FDrawContext& DrawContext,
 	const FTimingTrackViewport& Viewport,
diff --git a/Engine/Source/Developer/TraceInsights/Private/Insights/ViewModels/DrawHelpers.h b/Engine/Source/Developer/TraceInsights/Private/Insights/ViewModels/DrawHelpers.h
--- a/Engine/Source/Developer/TraceInsights/Private/Insights/ViewModels/DrawHelpers.h
+++ b/Engine/Source/Developer/TraceInsights/Private/Insights/ViewModels/DrawHelpers.h
@@ -56,6 +56,14 @@ public:
 							   const float Y,
 							   const float H);
 
+	static void DrawBackground(const FDrawContext& DrawContext,
+							   const FSlateBrush* BackgroundAreaBrush,
+							   const FTimingTrackViewport& Viewport,
+							   const float Y,
+							   const float H,
+							   float& OutValidAreaX,
+							   float& OutValidAreaW);
+
 	/**
 	 * Draw time range selection.
 	 */
diff --git a/Engine/Source/Developer/TraceInsights/Private/Insights/ViewModels/ThreadTimingTrack.cpp b/Engine/Source/Developer/TraceInsights/Private/Insights/ViewModels/ThreadTimingTrack.cpp
--- a/Engine/Source/Developer/TraceInsights/Private/Insights/ViewModels/ThreadTimingTrack.cpp
+++ b/Engine/Source/Developer/TraceInsights/Private/Insights/ViewModels/ThreadTimingTrack.cpp
@@ -90,14 +90,41 @@ void FThreadTimingTrack::DrawSelectedEventInfo(const FTimingEvent& SelectedTimin
 
 void FThreadTimingTrack::InitTooltip(FTooltipDrawState& Tooltip, const FTimingEvent& HoveredTimingEvent) const
 {
+	FTimingEvent ParentTimingEvent;
+	FTimingEvent RootTimingEvent;
+	if (HoveredTimingEvent.IsValid())
+	{
+		GetParentAndRoot(HoveredTimingEvent, ParentTimingEvent, RootTimingEvent);
+	}
+
 	Tooltip.ResetContent();
 
 	const FTimerNodePtr TimerNodePtr = FTimingProfilerManager::Get()->GetTimerNode(HoveredTimingEvent.TypeId);
 	FString TimerName = TimerNodePtr ? TimerNodePtr->GetName().ToString() : TEXT("N/A");
 	Tooltip.AddTitle(TimerName);
 
-	Tooltip.AddNameValueTextLine(TEXT("Incl. Time:"), TimeUtils::FormatTimeAuto(HoveredTimingEvent.Duration()));
-	Tooltip.AddNameValueTextLine(TEXT("Excl. Time:"), TimeUtils::FormatTimeAuto(HoveredTimingEvent.ExclusiveTime));
+	if (ParentTimingEvent.IsValid() && HoveredTimingEvent.Depth > 0)
+	{
+		const FTimerNodePtr ParentTimerNodePtr = FTimingProfilerManager::Get()->GetTimerNode(ParentTimingEvent.TypeId);
+		const FString ParentTimerName = ParentTimerNodePtr ? ParentTimerNodePtr->GetName().ToString() : TEXT("N/A");
+		FNumberFormattingOptions FormattingOptions;
+		FormattingOptions.MaximumFractionalDigits = 2;
+		const FString ValueStr = FString::Printf(TEXT("%s %s"), *FText::AsPercent(HoveredTimingEvent.Duration() / ParentTimingEvent.Duration(), &FormattingOptions).ToString(), *ParentTimerName);
+		Tooltip.AddNameValueTextLine(TEXT("% of Parent:"), ValueStr);
+	}
+
+	if (RootTimingEvent.IsValid() && HoveredTimingEvent.Depth > 1)
+	{
+		const FTimerNodePtr RootTimerNodePtr = FTimingProfilerManager::Get()->GetTimerNode(RootTimingEvent.TypeId);
+		const FString RootTimerName = RootTimerNodePtr ? RootTimerNodePtr->GetName().ToString() : TEXT("N/A");
+		FNumberFormattingOptions FormattingOptions;
+		FormattingOptions.MaximumFractionalDigits = 2;
+		const FString ValueStr = FString::Printf(TEXT("%s %s"), *FText::AsPercent(HoveredTimingEvent.Duration() / RootTimingEvent.Duration(), &FormattingOptions).ToString(), *RootTimerName);
+		Tooltip.AddNameValueTextLine(TEXT("% of Root:"), ValueStr);
+	}
+
+	Tooltip.AddNameValueTextLine(TEXT("Inclusive Time:"), TimeUtils::FormatTimeAuto(HoveredTimingEvent.Duration()));
+	Tooltip.AddNameValueTextLine(TEXT("Exclusive Time:"), TimeUtils::FormatTimeAuto(HoveredTimingEvent.ExclusiveTime));
 	Tooltip.AddNameValueTextLine(TEXT("Depth:"), FString::Printf(TEXT("%d"), HoveredTimingEvent.Depth));
 
 	Tooltip.UpdateLayout();
@@ -105,6 +132,50 @@ void FThreadTimingTrack::InitTooltip(FTooltipDrawState& Tooltip, const FTimingEv
 
 ////////////////////////////////////////////////////////////////////////////////////////////////////
 
+void FThreadTimingTrack::GetParentAndRoot(const FTimingEvent& TimingEvent, FTimingEvent& OutParentTimingEvent, FTimingEvent& OutRootTimingEvent) const
+{
+	// Note: This function does not compute Exclusive Time for parent and root events.
+
+	OutRootTimingEvent.Track = TimingEvent.Track;
+	OutParentTimingEvent.Track = TimingEvent.Track;
+
+	if (TimingEvent.Depth > 0)
+	{
+		TSharedPtr<const Trace::IAnalysisSession> Session = FInsightsManager::Get()->GetSession();
+		if (Session.IsValid())
+		{
+			Trace::FAnalysisSessionReadScope SessionReadScope(*Session.Get());
+
+			if (Trace::ReadTimingProfilerProvider(*Session.Get()))
+			{
+				const Trace::ITimingProfilerProvider& TimingProfilerProvider = *Trace::ReadTimingProfilerProvider(*Session.Get());
+
+				TimingProfilerProvider.ReadTimeline(TimingEvent.Track->GetId(), [&](const Trace::ITimingProfilerProvider::Timeline& Timeline)
+				{
+					Timeline.EnumerateEvents(TimingEvent.StartTime, TimingEvent.EndTime, [&](double EventStartTime, double EventEndTime, uint32 EventDepth, const Trace::FTimingProfilerEvent& Event)
+					{
+						if (EventDepth == 0)
+						{
+							OutRootTimingEvent.TypeId = Event.TimerIndex;
+							OutRootTimingEvent.Depth = EventDepth;
+							OutRootTimingEvent.StartTime = EventStartTime;
+							OutRootTimingEvent.EndTime = EventEndTime;
+						}
+						if (EventDepth == TimingEvent.Depth - 1)
+						{
+							OutParentTimingEvent.TypeId = Event.TimerIndex;
+							OutParentTimingEvent.Depth = EventDepth;
+							OutParentTimingEvent.StartTime = EventStartTime;
+							OutParentTimingEvent.EndTime = EventEndTime;
+						}
+					});
+				});
+			}
+		}
+	}
+}
+
+////////////////////////////////////////////////////////////////////////////////////////////////////
 
 bool FThreadTimingTrack::SearchTimingEvent(const double InStartTime,
 										   const double InEndTime,
@@ -165,24 +236,40 @@ bool FThreadTimingTrack::SearchTimingEvent(const double InStartTime,
 	{
 		Trace::FAnalysisSessionReadScope SessionReadScope(*Session.Get());
 
-		const FTimingEventsTrack* Track = Ctx.TimingEvent.Track;
-
 		if (Trace::ReadTimingProfilerProvider(*Session.Get()))
 		{
 			const Trace::ITimingProfilerProvider& TimingProfilerProvider = *Trace::ReadTimingProfilerProvider(*Session.Get());
 
-			TimingProfilerProvider.ReadTimeline(Track->GetId(), [&Ctx](const Trace::ITimingProfilerProvider::Timeline& Timeline)
+			TimingProfilerProvider.ReadTimeline(Ctx.TimingEvent.Track->GetId(), [&Ctx](const Trace::ITimingProfilerProvider::Timeline& Timeline)
 			{
 				Timeline.EnumerateEvents(Ctx.StartTime, Ctx.EndTime, [&Ctx](double EventStartTime, double EventEndTime, uint32 EventDepth, const Trace::FTimingProfilerEvent& Event)
 				{
 					Ctx.CheckEvent(EventStartTime, EventEndTime, EventDepth, Event);
 				});
 			});
+		}
+	}
 
-			if (Ctx.bFound)
+	return Ctx.bFound;
+}
+
+////////////////////////////////////////////////////////////////////////////////////////////////////
+
+void FThreadTimingTrack::ComputeTimingEventStats(FTimingEvent& InOutTimingEvent) const
+{
+	if (InOutTimingEvent.IsValid())
+	{
+		TSharedPtr<const Trace::IAnalysisSession> Session = FInsightsManager::Get()->GetSession();
+		if (Session.IsValid())
+		{
+			Trace::FAnalysisSessionReadScope SessionReadScope(*Session.Get());
+
+			if (Trace::ReadTimingProfilerProvider(*Session.Get()))
 			{
+				const Trace::ITimingProfilerProvider& TimingProfilerProvider = *Trace::ReadTimingProfilerProvider(*Session.Get());
+
 				// Compute Exclusive Time.
-				TimingProfilerProvider.ReadTimeline(Track->GetId(), [&Ctx](const Trace::ITimingProfilerProvider::Timeline& Timeline)
+				TimingProfilerProvider.ReadTimeline(InOutTimingEvent.Track->GetId(), [&InOutTimingEvent](const Trace::ITimingProfilerProvider::Timeline& Timeline)
 				{
 					struct FEnumerationState
 					{
@@ -197,16 +284,16 @@ bool FThreadTimingTrack::SearchTimingEvent(const double InStartTime,
 					};
 					FEnumerationState State;
 
-					State.EventStartTime = Ctx.TimingEvent.StartTime;
-					State.EventEndTime = Ctx.TimingEvent.EndTime;
-					State.EventDepth = Ctx.TimingEvent.Depth;
+					State.EventStartTime = InOutTimingEvent.StartTime;
+					State.EventEndTime = InOutTimingEvent.EndTime;
+					State.EventDepth = InOutTimingEvent.Depth;
 
 					State.CurrentDepth = 0;
 					State.LastTime = 0.0;
 					State.ExclusiveTime = 0.0;
 					State.IsInEventScope = false;
 
-					Timeline.EnumerateEvents(Ctx.TimingEvent.StartTime, Ctx.TimingEvent.EndTime, [&State](bool IsEnter, double Time, const Trace::FTimingProfilerEvent& Event)
+					Timeline.EnumerateEvents(InOutTimingEvent.StartTime, InOutTimingEvent.EndTime, [&State](bool IsEnter, double Time, const Trace::FTimingProfilerEvent& Event)
 					{
 						if (IsEnter)
 						{
@@ -232,13 +319,11 @@ bool FThreadTimingTrack::SearchTimingEvent(const double InStartTime,
 						State.LastTime = Time;
 					});
 
-					Ctx.TimingEvent.ExclusiveTime = State.ExclusiveTime;
+					InOutTimingEvent.ExclusiveTime = State.ExclusiveTime;
 				});
 			}
 		}
 	}
-
-	return Ctx.bFound;
 }
 
 ////////////////////////////////////////////////////////////////////////////////////////////////////
diff --git a/Engine/Source/Developer/TraceInsights/Private/Insights/ViewModels/ThreadTimingTrack.h b/Engine/Source/Developer/TraceInsights/Private/Insights/ViewModels/ThreadTimingTrack.h
--- a/Engine/Source/Developer/TraceInsights/Private/Insights/ViewModels/ThreadTimingTrack.h
+++ b/Engine/Source/Developer/TraceInsights/Private/Insights/ViewModels/ThreadTimingTrack.h
@@ -31,6 +31,10 @@ public:
 	virtual void InitTooltip(FTooltipDrawState& Tooltip, const FTimingEvent& HoveredTimingEvent) const override;
 
 	virtual bool SearchTimingEvent(const double InStartTime, const double InEndTime, TFunctionRef<bool(double, double, uint32)> InPredicate, FTimingEvent& InOutTimingEvent, bool bInStopAtFirstMatch, bool bInSearchForLargestEvent) const override;
+	virtual void ComputeTimingEventStats(FTimingEvent& InOutTimingEvent) const override;
+
+private:
+	void GetParentAndRoot(const FTimingEvent& TimingEvent, FTimingEvent& OutParentTimingEvent, FTimingEvent& OutRootTimingEvent) const;
 
 private:
 	const TCHAR* GroupName;
diff --git a/Engine/Source/Developer/TraceInsights/Private/Insights/ViewModels/TimingEvent.h b/Engine/Source/Developer/TraceInsights/Private/Insights/ViewModels/TimingEvent.h
--- a/Engine/Source/Developer/TraceInsights/Private/Insights/ViewModels/TimingEvent.h
+++ b/Engine/Source/Developer/TraceInsights/Private/Insights/ViewModels/TimingEvent.h
@@ -26,6 +26,26 @@ struct FTimingEvent
 	uint64 Size;
 	const TCHAR* Path;
 
+	FTimingEvent()
+		: Track(nullptr)
+		, TypeId(InvalidTypeId)
+		, Depth(0)
+		, StartTime(0.0)
+		, EndTime(0.0)
+		, ExclusiveTime(0.0)
+		, LoadingInfo()
+		, Offset(0)
+		, Size(0)
+		, Path(nullptr)
+	{
+	}
+
+	FTimingEvent(const FTimingEvent&) = default;
+	FTimingEvent& operator=(const FTimingEvent&) = default;
+
+	FTimingEvent(FTimingEvent&&) = default;
+	FTimingEvent& operator=(FTimingEvent&&) = default;
+
 	void Reset()
 	{
 		Track = nullptr;
diff --git a/Engine/Source/Developer/TraceInsights/Private/Insights/ViewModels/TimingEventsTrack.h b/Engine/Source/Developer/TraceInsights/Private/Insights/ViewModels/TimingEventsTrack.h
--- a/Engine/Source/Developer/TraceInsights/Private/Insights/ViewModels/TimingEventsTrack.h
+++ b/Engine/Source/Developer/TraceInsights/Private/Insights/ViewModels/TimingEventsTrack.h
@@ -79,6 +79,7 @@ public:
 	virtual void InitTooltip(FTooltipDrawState& Tooltip, const FTimingEvent& HoveredTimingEvent) const {}
 
 	virtual bool SearchTimingEvent(const double InStartTime, const double InEndTime, TFunctionRef<bool(double, double, uint32)> InPredicate, FTimingEvent& InOutTimingEvent, bool bInStopAtFirstMatch, bool bInSearchForLargestEvent) const { return false; }
+	virtual void ComputeTimingEventStats(FTimingEvent& InOutTimingEvent) const {}
 
 private:
 	FName Type; // Thread, Loading, FileActivity
diff --git a/Engine/Source/Developer/TraceInsights/Private/Insights/ViewModels/TimingViewDrawHelper.cpp b/Engine/Source/Developer/TraceInsights/Private/Insights/ViewModels/TimingViewDrawHelper.cpp
--- a/Engine/Source/Developer/TraceInsights/Private/Insights/ViewModels/TimingViewDrawHelper.cpp
+++ b/Engine/Source/Developer/TraceInsights/Private/Insights/ViewModels/TimingViewDrawHelper.cpp
@@ -33,6 +33,14 @@ FTimingViewDrawHelper::FTimingViewDrawHelper(const FDrawContext& InDC, const FTi
 	, InvalidAreaColor(0.1f, 0.07f, 0.07f, 1.0f)
 	, EdgeColor(0.05f, 0.05f, 0.05f, 1.0f)
 	, EventFont(FCoreStyle::GetDefaultFontStyle("Regular", 8))
+	, ValidAreaX(0.0f)
+	, ValidAreaW(0.0f)
+	, TimelineTopY(0.0f)
+	, TimelineY(0.0f)
+	, MaxDepth(0)
+	, TimelineIndex(0)
+	, LastEventX2()
+	, LastBox()
 	, Stats()
 {
 }
@@ -49,7 +57,7 @@ void FTimingViewDrawHelper::DrawBackground() const
 {
 	const float Y = 0.0f;
 	const float H = FMath::CeilToFloat(Viewport.GetHeight());
-	FDrawHelpers::DrawBackground(DrawContext, BackgroundAreaBrush, Viewport, Y, H);
+	FDrawHelpers::DrawBackground(DrawContext, BackgroundAreaBrush, Viewport, Y, H, ValidAreaX, ValidAreaW);
 }
 
 ////////////////////////////////////////////////////////////////////////////////////////////////////
diff --git a/Engine/Source/Developer/TraceInsights/Private/Insights/ViewModels/TimingViewDrawHelper.h b/Engine/Source/Developer/TraceInsights/Private/Insights/ViewModels/TimingViewDrawHelper.h
--- a/Engine/Source/Developer/TraceInsights/Private/Insights/ViewModels/TimingViewDrawHelper.h
+++ b/Engine/Source/Developer/TraceInsights/Private/Insights/ViewModels/TimingViewDrawHelper.h
@@ -135,9 +135,6 @@ private:
 	const FLinearColor EdgeColor;
 	const FSlateFontInfo EventFont;
 
-	//static const FSlateColorBrush SolidWhiteBrush = FSlateColorBrush(FColorList::White);
-	//static const FSlateBrush BorderBrush = FSlateBorderBrush(NAME_None, FMargin(1.0f));
-
 	mutable float ValidAreaX;
 	mutable float ValidAreaW;
 
diff --git a/Engine/Source/Developer/TraceInsights/Private/Insights/Widgets/STimingView.cpp b/Engine/Source/Developer/TraceInsights/Private/Insights/Widgets/STimingView.cpp
--- a/Engine/Source/Developer/TraceInsights/Private/Insights/Widgets/STimingView.cpp
+++ b/Engine/Source/Developer/TraceInsights/Private/Insights/Widgets/STimingView.cpp
@@ -1366,7 +1366,11 @@ FReply STimingView::OnMouseButtonUp(const FGeometry& MyGeometry, const FPointerE
 			else if (bIsSelecting)
 			{
 				//TODO: SelectionChangedEvent.Broadcast(SelectionStartTime, SelectionEndTime);
-				UpdateAggregatedStats();
+
+				if (SelectionEndTime > SelectionStartTime)
+				{
+					UpdateAggregatedStats();
+				}
 
 				bIsSelecting = false;
 			}
@@ -1410,7 +1414,11 @@ FReply STimingView::OnMouseButtonUp(const FGeometry& MyGeometry, const FPointerE
 			else if (bIsSelecting)
 			{
 				//TODO: SelectionChangedEvent.Broadcast(SelectionStartTime, SelectionEndTime);
-				UpdateAggregatedStats();
+
+				if (SelectionEndTime > SelectionStartTime)
+				{
+					UpdateAggregatedStats();
+				}
 
 				bIsSelecting = false;
 			}
@@ -2385,8 +2393,7 @@ void STimingView::SetDrawOnlyBookmarks(bool bIsBookmarksTrack)
 
 void STimingView::UpdateHoveredTimingEvent(float MX, float MY)
 {
-	HoveredTimingEvent.Track = nullptr;
-	HoveredTimingEvent.TypeId = FTimerNode::InvalidId;
+	FTimingEvent TimingEvent;
 
 	if (MY >= Viewport.GetTopOffset() && MY < Viewport.GetHeight())
 	{
@@ -2398,18 +2405,18 @@ void STimingView::UpdateHoveredTimingEvent(float MX, float MY)
 				const float Y = Viewport.GetTopOffset() + Track.GetPosY() - Viewport.GetScrollPosY();
 				if (MY >= Y && MY < Y + Track.GetHeight())
 				{
-					HoveredTimingEvent.Track = &Track;
+					TimingEvent.Track = &Track;
 					break;
 				}
 			}
 		}
 
-		if (HoveredTimingEvent.Track)
+		if (TimingEvent.Track)
 		{
-			const float Y0 = Viewport.GetTopOffset() + HoveredTimingEvent.Track->GetPosY() - Viewport.GetScrollPosY() + 1.0f + Layout.TimelineDY;
+			const float Y0 = Viewport.GetTopOffset() + TimingEvent.Track->GetPosY() - Viewport.GetScrollPosY() + 1.0f + Layout.TimelineDY;
 
 			// If mouse is not above first sub-track or below last sub-track...
-			if (MY >= Y0 && MY < Y0 + HoveredTimingEvent.Track->GetHeight() + Layout.TimelineDY)
+			if (MY >= Y0 && MY < Y0 + TimingEvent.Track->GetHeight() + Layout.TimelineDY)
 			{
 				int32 Depth = (MY - Y0) / (Layout.EventDY + Layout.EventH);
 				float EventMY = (MY - Y0) - Depth * (Layout.EventDY + Layout.EventH);
@@ -2418,25 +2425,29 @@ void STimingView::UpdateHoveredTimingEvent(float MX, float MY)
 				const double EndTime = StartTime + 2.0 / Viewport.GetScaleX(); // +2px
 				constexpr bool bStopAtFirstMatch = true; // get first one matching
 				constexpr bool bSearchForLargestEvent = false;
-				HoveredTimingEvent.Track->SearchTimingEvent(StartTime, EndTime,
+				TimingEvent.Track->SearchTimingEvent(StartTime, EndTime,
 					[Depth](double, double, uint32 EventDepth)
 					{
 						return EventDepth == Depth;
 					},
-					HoveredTimingEvent, bStopAtFirstMatch, bSearchForLargestEvent);
-
-				//TODO: ComputeSingleTimingEventStats(HoveredTimingEvent) --> compute ExclusiveTime
+					TimingEvent, bStopAtFirstMatch, bSearchForLargestEvent);
 			}
 		}
 	}
 
-	if (HoveredTimingEvent.IsValid())
+	if (TimingEvent.IsValid())
 	{
-		HoveredTimingEvent.Track->InitTooltip(Tooltip, HoveredTimingEvent);
+		if (!TimingEvent.Equals(HoveredTimingEvent))
+		{
+			HoveredTimingEvent = TimingEvent;
+			HoveredTimingEvent.Track->ComputeTimingEventStats(HoveredTimingEvent);
+			HoveredTimingEvent.Track->InitTooltip(Tooltip, HoveredTimingEvent);
+		}
 		Tooltip.SetDesiredOpacity(1.0f);
 	}
 	else
 	{
+		HoveredTimingEvent.Reset();
 		Tooltip.SetDesiredOpacity(0.0f);
 	}
 }
@@ -2495,6 +2506,7 @@ void STimingView::SelectLeftTimingEvent()
 			},
 			SelectedTimingEvent, bStopAtFirstMatch, bSearchForLargestEvent))
 		{
+			SelectedTimingEvent.Track->ComputeTimingEventStats(SelectedTimingEvent);
 			BringIntoView(SelectedTimingEvent.StartTime, SelectedTimingEvent.EndTime);
 			OnSelectedTimingEventChanged();
 		}
@@ -2520,6 +2532,7 @@ void STimingView::SelectRightTimingEvent()
 			},
 			SelectedTimingEvent, bStopAtFirstMatch, bSearchForLargestEvent))
 		{
+			SelectedTimingEvent.Track->ComputeTimingEventStats(SelectedTimingEvent);
 			BringIntoView(SelectedTimingEvent.StartTime, SelectedTimingEvent.EndTime);
 			OnSelectedTimingEventChanged();
 		}
@@ -2547,6 +2560,7 @@ void STimingView::SelectUpTimingEvent()
 			},
 			SelectedTimingEvent, bStopAtFirstMatch, bSearchForLargestEvent))
 		{
+			SelectedTimingEvent.Track->ComputeTimingEventStats(SelectedTimingEvent);
 			BringIntoView(SelectedTimingEvent.StartTime, SelectedTimingEvent.EndTime);
 			OnSelectedTimingEventChanged();
 		}
@@ -2573,6 +2587,7 @@ void STimingView::SelectDownTimingEvent()
 			},
 			SelectedTimingEvent, bStopAtFirstMatch, bSearchForLargestEvent))
 		{
+			SelectedTimingEvent.Track->ComputeTimingEventStats(SelectedTimingEvent);
 			BringIntoView(SelectedTimingEvent.StartTime, SelectedTimingEvent.EndTime);
 			OnSelectedTimingEventChanged();
 		}
@@ -2844,11 +2859,7 @@ void STimingView::ShowHideAllCpuTracks_Execute()
 		KV.Value.bIsVisible = bShowHideAllCpuTracks;
 	}
 
-	HoveredTimingEvent.Reset();
-	SelectedTimingEvent.Reset();
-	Tooltip.SetDesiredOpacity(0.0f);
-
-	bAreTimingEventsTracksDirty = true;
+	OnTimingEventsTrackVisibilityChanged();
 }
 
 ////////////////////////////////////////////////////////////////////////////////////////////////////
@@ -2873,11 +2884,7 @@ void STimingView::ShowHideAllGpuTracks_Execute()
 		}
 	}
 
-	HoveredTimingEvent.Reset();
-	SelectedTimingEvent.Reset();
-	Tooltip.SetDesiredOpacity(0.0f);
-
-	bAreTimingEventsTracksDirty = true;
+	OnTimingEventsTrackVisibilityChanged();
 }
 
 ////////////////////////////////////////////////////////////////////////////////////////////////////
@@ -2902,7 +2909,7 @@ void STimingView::ShowHideAllLoadingTracks_Execute()
 		}
 	}
 
-	bAreTimingEventsTracksDirty = true;
+	OnTimingEventsTrackVisibilityChanged();
 }
 
 ////////////////////////////////////////////////////////////////////////////////////////////////////
@@ -2927,7 +2934,7 @@ void STimingView::ShowHideAllIoTracks_Execute()
 		}
 	}
 
-	bAreTimingEventsTracksDirty = true;
+	OnTimingEventsTrackVisibilityChanged();
 
 	if (bShowHideAllIoTracks)
 	{
@@ -2986,9 +2993,7 @@ void STimingView::ToggleTrackVisibility_Execute(uint64 InTrackId)
 		FTimingEventsTrack* Track = CachedTimelines[InTrackId];
 		Track->ToggleVisibility();
 
-		HoveredTimingEvent.Reset();
-		SelectedTimingEvent.Reset();
-		Tooltip.SetDesiredOpacity(0.0f);
+		OnTimingEventsTrackVisibilityChanged();
 	}
 }
 
@@ -3023,12 +3028,23 @@ void STimingView::ToggleTrackVisibilityByGroup_Execute(const TCHAR* InGroupName)
 			}
 		}
 
-		HoveredTimingEvent.Reset();
-		SelectedTimingEvent.Reset();
-		Tooltip.SetDesiredOpacity(0.0f);
+		OnTimingEventsTrackVisibilityChanged();
 	}
 }
 
 ////////////////////////////////////////////////////////////////////////////////////////////////////
 
+void STimingView::OnTimingEventsTrackVisibilityChanged()
+{
+	HoveredTimingEvent.Reset();
+	SelectedTimingEvent.Reset();
+	Tooltip.SetDesiredOpacity(0.0f);
+
+	UpdateAggregatedStats();
+
+	bAreTimingEventsTracksDirty = true;
+}
+
+////////////////////////////////////////////////////////////////////////////////////////////////////
+
 #undef LOCTEXT_NAMESPACE
diff --git a/Engine/Source/Developer/TraceInsights/Private/Insights/Widgets/STimingView.h b/Engine/Source/Developer/TraceInsights/Private/Insights/Widgets/STimingView.h
--- a/Engine/Source/Developer/TraceInsights/Private/Insights/Widgets/STimingView.h
+++ b/Engine/Source/Developer/TraceInsights/Private/Insights/Widgets/STimingView.h
@@ -324,6 +324,8 @@ protected:
 
 	void FrameSelection();
 
+	void OnTimingEventsTrackVisibilityChanged();
+
 	////////////////////////////////////////////////////////////////////////////////////////////////////
 	// FrameSelectionChanged Event
 
