[Insights]
  - Timing view: Added percentage for the Exclusive time (relative to Inclusive Time)) in the tooltip of a hovered event.

diff --git a/Engine/Source/Developer/TraceInsights/Private/Insights/ViewModels/ThreadTimingTrack.cpp b/Engine/Source/Developer/TraceInsights/Private/Insights/ViewModels/ThreadTimingTrack.cpp
--- a/Engine/Source/Developer/TraceInsights/Private/Insights/ViewModels/ThreadTimingTrack.cpp
+++ b/Engine/Source/Developer/TraceInsights/Private/Insights/ViewModels/ThreadTimingTrack.cpp
@@ -124,7 +124,15 @@ void FThreadTimingTrack::InitTooltip(FTooltipDrawState& Tooltip, const FTimingEv
 	}
 
 	Tooltip.AddNameValueTextLine(TEXT("Inclusive Time:"), TimeUtils::FormatTimeAuto(HoveredTimingEvent.Duration()));
-	Tooltip.AddNameValueTextLine(TEXT("Exclusive Time:"), TimeUtils::FormatTimeAuto(HoveredTimingEvent.ExclusiveTime));
+
+	//Tooltip.AddNameValueTextLine(TEXT("Exclusive Time:"), TimeUtils::FormatTimeAuto(HoveredTimingEvent.ExclusiveTime));
+	{
+		FNumberFormattingOptions FormattingOptions;
+		FormattingOptions.MaximumFractionalDigits = 2;
+		const FString ExclStr = FString::Printf(TEXT("%s (%s)"), *TimeUtils::FormatTimeAuto(HoveredTimingEvent.ExclusiveTime), *FText::AsPercent(HoveredTimingEvent.ExclusiveTime / HoveredTimingEvent.Duration(), &FormattingOptions).ToString());
+		Tooltip.AddNameValueTextLine(TEXT("Exclusive Time:"), ExclStr);
+	}
+
 	Tooltip.AddNameValueTextLine(TEXT("Depth:"), FString::Printf(TEXT("%d"), HoveredTimingEvent.Depth));
 
 	Tooltip.UpdateLayout();
