[Insights]
  - Fixed LoadingProfilerMenuGroupName for Asset Loading Insights.
  - Disabled Networking Insights prototype (for v0.09 build).

diff --git a/Engine/Source/Developer/TraceInsights/Private/Insights/InsightsManager.cpp b/Engine/Source/Developer/TraceInsights/Private/Insights/InsightsManager.cpp
--- a/Engine/Source/Developer/TraceInsights/Private/Insights/InsightsManager.cpp
+++ b/Engine/Source/Developer/TraceInsights/Private/Insights/InsightsManager.cpp
@@ -180,7 +180,7 @@ void FInsightsManager::SpawnAndActivateTabs()
 	}
 	if (FGlobalTabmanager::Get()->HasTabSpawner(FInsightsManagerTabs::NetworkingProfilerTabId))
 	{
-		FGlobalTabmanager::Get()->InvokeTab(FInsightsManagerTabs::NetworkingProfilerTabId);
+		//FGlobalTabmanager::Get()->InvokeTab(FInsightsManagerTabs::NetworkingProfilerTabId);
 	}
 
 	// Ensure Timing Insights / Timing View is the active tab / view.
diff --git a/Engine/Source/Developer/TraceInsights/Private/Insights/LoadingProfiler/LoadingProfilerManager.h b/Engine/Source/Developer/TraceInsights/Private/Insights/LoadingProfiler/LoadingProfilerManager.h
--- a/Engine/Source/Developer/TraceInsights/Private/Insights/LoadingProfiler/LoadingProfilerManager.h
+++ b/Engine/Source/Developer/TraceInsights/Private/Insights/LoadingProfiler/LoadingProfilerManager.h
@@ -21,14 +21,14 @@ namespace Trace
 
 ////////////////////////////////////////////////////////////////////////////////////////////////////
 /**
- * This class manages the Timing Profiler state and settings.
+ * This class manages the Asset Loading Profiler state and settings.
  */
 class FLoadingProfilerManager : public TSharedFromThis<FLoadingProfilerManager>
 {
 	friend class FLoadingProfilerActionManager;
 
 public:
-	/** Creates the Timing Profiler manager, only one instance can exist. */
+	/** Creates the Loading Profiler (Asset Loading Insights) manager, only one instance can exist. */
 	FLoadingProfilerManager(TSharedRef<FUICommandList> InCommandList);
 
 	/** Virtual destructor. */
@@ -48,7 +48,7 @@ public:
 		return FLoadingProfilerManager::Instance;
 	}
 
-	/** Shutdowns the Timing Profiler manager. */
+	/** Shutdowns the profiler manager. */
 	void Shutdown()
 	{
 		FLoadingProfilerManager::Instance.Reset();
@@ -71,13 +71,13 @@ public:
 	 */
 	static TSharedPtr<FLoadingProfilerManager> Get();
 
-	/** @returns UI command list for the Timing Profiler manager. */
+	/** @returns UI command list for the Loading Profiler manager. */
 	const TSharedRef<FUICommandList> GetCommandList() const;
 
-	/** @return an instance of the Timing Profiler commands. */
+	/** @return an instance of the Loading Profiler commands. */
 	static const FLoadingProfilerCommands& GetCommands();
 
-	/** @return an instance of the Timing Profiler action manager. */
+	/** @return an instance of the Loading Profiler action manager. */
 	static FLoadingProfilerActionManager& GetActionManager();
 
 	void AssignProfilerWindow(const TSharedRef<SLoadingProfilerWindow>& InProfilerWindow)
diff --git a/Engine/Source/Developer/TraceInsights/Private/Insights/LoadingProfiler/Widgets/SLoadingProfilerWindow.cpp b/Engine/Source/Developer/TraceInsights/Private/Insights/LoadingProfiler/Widgets/SLoadingProfilerWindow.cpp
--- a/Engine/Source/Developer/TraceInsights/Private/Insights/LoadingProfiler/Widgets/SLoadingProfilerWindow.cpp
+++ b/Engine/Source/Developer/TraceInsights/Private/Insights/LoadingProfiler/Widgets/SLoadingProfilerWindow.cpp
@@ -219,7 +219,7 @@ void SLoadingProfilerWindow::Construct(const FArguments& InArgs, const TSharedRe
 {
 	// Create & initialize tab manager.
 	TabManager = FGlobalTabmanager::Get()->NewTabManager(ConstructUnderMajorTab);
-	TSharedRef<FWorkspaceItem> AppMenuGroup = TabManager->AddLocalWorkspaceMenuCategory(LOCTEXT("LoadingProfilerMenuGroupName", "Timing Insights"));
+	TSharedRef<FWorkspaceItem> AppMenuGroup = TabManager->AddLocalWorkspaceMenuCategory(LOCTEXT("LoadingProfilerMenuGroupName", "Asset Loading Insights"));
 
 	TabManager->RegisterTabSpawner(FLoadingProfilerTabs::ToolbarID, FOnSpawnTab::CreateRaw(this, &SLoadingProfilerWindow::SpawnTab_Toolbar))
 		.SetDisplayName(LOCTEXT("DeviceToolbarTabTitle", "Toolbar"))
