[Insights]
  - Asset Loading Insights:  Added "Package Details" and "Export Details" tabs (generic tree table views).
  - Asset Loading Insights:  Removed aggregation mini-tables (as aggregation is available in the docked tabs / generic tree table views).
  - Generic Tree Table view: Added grouping by "Unique Values" for all columns in table.
  - Generic Tree Table view: Implemented support for Aggregation (Sum) for all numeric columns.
  - Changed the way Generic Tree Table view handles "name" column (now transformed into a "hierarchy" column). It does not use the first string column anymore, but name of items in list (table rows)  are generated as "row N" (where N is the table row index).
  - Timing view:  Added Ctrl+C shortcut to copy name of selected timing event to clipboard.
  - Increased version to 0.10

diff --git a/Engine/Source/Developer/TraceInsights/Private/Insights/LoadingProfiler/LoadingProfilerCommands.cpp b/Engine/Source/Developer/TraceInsights/Private/Insights/LoadingProfiler/LoadingProfilerCommands.cpp
--- a/Engine/Source/Developer/TraceInsights/Private/Insights/LoadingProfiler/LoadingProfilerCommands.cpp
+++ b/Engine/Source/Developer/TraceInsights/Private/Insights/LoadingProfiler/LoadingProfilerCommands.cpp
@@ -55,6 +55,8 @@ void FLoadingProfilerCommands::RegisterCommands()
 	UI_COMMAND(ToggleTimingViewVisibility, "Timing", "Toggles the visibility of the main Timing view", EUserInterfaceActionType::ToggleButton, FInputChord(EModifierKey::Control, EKeys::T));
 	UI_COMMAND(ToggleEventAggregationTreeViewVisibility, "Event Aggregation", "Toggles the visibility of the Event Aggregation table/tree view", EUserInterfaceActionType::ToggleButton, FInputChord());
 	UI_COMMAND(ToggleObjectTypeAggregationTreeViewVisibility, "Object Type Aggregation", "Toggles the visibility of the Object Type Aggregation table/tree view", EUserInterfaceActionType::ToggleButton, FInputChord());
+	UI_COMMAND(TogglePackageDetailsTreeViewVisibility, "Package Details", "Toggles the visibility of the Package Details table/tree view", EUserInterfaceActionType::ToggleButton, FInputChord());
+	UI_COMMAND(ToggleExportDetailsTreeViewVisibility, "Export Details", "Toggles the visibility of the Export Details table/tree view", EUserInterfaceActionType::ToggleButton, FInputChord());
 }
 PRAGMA_ENABLE_OPTIMIZATION
 
@@ -98,6 +100,8 @@ PRAGMA_ENABLE_OPTIMIZATION
 IMPLEMENT_TOGGLE_COMMAND(ToggleTimingViewVisibility, IsTimingViewVisible, ShowHideTimingView)
 IMPLEMENT_TOGGLE_COMMAND(ToggleEventAggregationTreeViewVisibility, IsEventAggregationTreeViewVisible, ShowHideEventAggregationTreeView)
 IMPLEMENT_TOGGLE_COMMAND(ToggleObjectTypeAggregationTreeViewVisibility, IsObjectTypeAggregationTreeViewVisible, ShowHideObjectTypeAggregationTreeView)
+IMPLEMENT_TOGGLE_COMMAND(TogglePackageDetailsTreeViewVisibility, IsPackageDetailsTreeViewVisible, ShowHidePackageDetailsTreeView)
+IMPLEMENT_TOGGLE_COMMAND(ToggleExportDetailsTreeViewVisibility, IsExportDetailsTreeViewVisible, ShowHideExportDetailsTreeView)
 
 #undef IMPLEMENT_TOGGLE_COMMAND
 
diff --git a/Engine/Source/Developer/TraceInsights/Private/Insights/LoadingProfiler/LoadingProfilerCommands.h b/Engine/Source/Developer/TraceInsights/Private/Insights/LoadingProfiler/LoadingProfilerCommands.h
--- a/Engine/Source/Developer/TraceInsights/Private/Insights/LoadingProfiler/LoadingProfilerCommands.h
+++ b/Engine/Source/Developer/TraceInsights/Private/Insights/LoadingProfiler/LoadingProfilerCommands.h
@@ -37,6 +37,12 @@ public:
 
 	/** Toggles visibility for the Object Type Aggregation tree view. Global and custom command. */
 	TSharedPtr<FUICommandInfo> ToggleObjectTypeAggregationTreeViewVisibility;
+
+	/** Toggles visibility for the Package Details tree view. Global and custom command. */
+	TSharedPtr<FUICommandInfo> TogglePackageDetailsTreeViewVisibility;
+
+	/** Toggles visibility for the Export Details tree view. Global and custom command. */
+	TSharedPtr<FUICommandInfo> ToggleExportDetailsTreeViewVisibility;
 };
 
 ////////////////////////////////////////////////////////////////////////////////////////////////////
@@ -91,6 +97,8 @@ protected:\
 	DECLARE_TOGGLE_COMMAND(ToggleTimingViewVisibility)
 	DECLARE_TOGGLE_COMMAND(ToggleEventAggregationTreeViewVisibility)
 	DECLARE_TOGGLE_COMMAND(ToggleObjectTypeAggregationTreeViewVisibility)
+	DECLARE_TOGGLE_COMMAND(TogglePackageDetailsTreeViewVisibility)
+	DECLARE_TOGGLE_COMMAND(ToggleExportDetailsTreeViewVisibility)
 #undef DECLARE_TOGGLE_COMMAND
 
 	//////////////////////////////////////////////////
diff --git a/Engine/Source/Developer/TraceInsights/Private/Insights/LoadingProfiler/LoadingProfilerManager.cpp b/Engine/Source/Developer/TraceInsights/Private/Insights/LoadingProfiler/LoadingProfilerManager.cpp
--- a/Engine/Source/Developer/TraceInsights/Private/Insights/LoadingProfiler/LoadingProfilerManager.cpp
+++ b/Engine/Source/Developer/TraceInsights/Private/Insights/LoadingProfiler/LoadingProfilerManager.cpp
@@ -34,6 +34,8 @@ FLoadingProfilerManager::FLoadingProfilerManager(TSharedRef<FUICommandList> InCo
 	, bIsTimingViewVisible(true)
 	, bIsEventAggregationTreeViewVisible(true)
 	, bIsObjectTypeAggregationTreeViewVisible(true)
+	, bIsPackageDetailsTreeViewVisible(true)
+	, bIsExportDetailsTreeViewVisible(true)
 {
 }
 
@@ -56,6 +58,8 @@ void FLoadingProfilerManager::BindCommands()
 	ActionManager.Map_ToggleTimingViewVisibility_Global();
 	ActionManager.Map_ToggleEventAggregationTreeViewVisibility_Global();
 	ActionManager.Map_ToggleObjectTypeAggregationTreeViewVisibility_Global();
+	ActionManager.Map_TogglePackageDetailsTreeViewVisibility_Global();
+	ActionManager.Map_ToggleExportDetailsTreeViewVisibility_Global();
 }
 
 ////////////////////////////////////////////////////////////////////////////////////////////////////
@@ -118,46 +122,66 @@ void FLoadingProfilerManager::OnSessionChanged()
 
 ////////////////////////////////////////////////////////////////////////////////////////////////////
 
-void FLoadingProfilerManager::ShowHideTimingView(const bool bTimingViewVisibleState)
+void FLoadingProfilerManager::ShowHideTimingView(const bool bIsVisible)
 {
-	//bool bWasTimingViewVisible = bIsTimingViewVisible;
-
-	bIsTimingViewVisible = bTimingViewVisibleState;
+	bIsTimingViewVisible = bIsVisible;
 
 	TSharedPtr<SLoadingProfilerWindow> Wnd = GetProfilerWindow();
 	if (Wnd.IsValid())
 	{
-		Wnd->ShowHideTab(FLoadingProfilerTabs::TimingViewID, bIsTimingViewVisible);
+		Wnd->ShowHideTab(FLoadingProfilerTabs::TimingViewID, bIsVisible);
 	}
 }
 
 ////////////////////////////////////////////////////////////////////////////////////////////////////
 
-void FLoadingProfilerManager::ShowHideEventAggregationTreeView(const bool bEventAggregationTreeViewVisibleState)
+void FLoadingProfilerManager::ShowHideEventAggregationTreeView(const bool bIsVisible)
 {
-	//bool bWasEventAggregationTreeViewVisible = bIsEventAggregationTreeViewVisible;
+	bIsEventAggregationTreeViewVisible = bIsVisible;
+
+	TSharedPtr<SLoadingProfilerWindow> Wnd = GetProfilerWindow();
+	if (Wnd.IsValid())
+	{
+		Wnd->ShowHideTab(FLoadingProfilerTabs::EventAggregationTreeViewID, bIsVisible);
+	}
+}
 
-	bIsEventAggregationTreeViewVisible = bEventAggregationTreeViewVisibleState;
+////////////////////////////////////////////////////////////////////////////////////////////////////
+
+void FLoadingProfilerManager::ShowHideObjectTypeAggregationTreeView(const bool bIsVisible)
+{
+	bIsObjectTypeAggregationTreeViewVisible = bIsVisible;
 
 	TSharedPtr<SLoadingProfilerWindow> Wnd = GetProfilerWindow();
 	if (Wnd.IsValid())
 	{
-		Wnd->ShowHideTab(FLoadingProfilerTabs::EventAggregationTreeViewID, bIsEventAggregationTreeViewVisible);
+		Wnd->ShowHideTab(FLoadingProfilerTabs::ObjectTypeAggregationTreeViewID, bIsVisible);
 	}
 }
 
 ////////////////////////////////////////////////////////////////////////////////////////////////////
 
-void FLoadingProfilerManager::ShowHideObjectTypeAggregationTreeView(const bool bObjectTypeAggregationTreeViewVisibleState)
+void FLoadingProfilerManager::ShowHidePackageDetailsTreeView(const bool bIsVisible)
 {
-	//bool bWasObjectTypeAggregationTreeViewVisible = bIsObjectTypeAggregationTreeViewVisible;
+	bIsPackageDetailsTreeViewVisible = bIsVisible;
 
-	bIsObjectTypeAggregationTreeViewVisible = bObjectTypeAggregationTreeViewVisibleState;
+	TSharedPtr<SLoadingProfilerWindow> Wnd = GetProfilerWindow();
+	if (Wnd.IsValid())
+	{
+		Wnd->ShowHideTab(FLoadingProfilerTabs::PackageDetailsTreeViewID, bIsVisible);
+	}
+}
+
+////////////////////////////////////////////////////////////////////////////////////////////////////
+
+void FLoadingProfilerManager::ShowHideExportDetailsTreeView(const bool bIsVisible)
+{
+	bIsExportDetailsTreeViewVisible = bIsVisible;
 
 	TSharedPtr<SLoadingProfilerWindow> Wnd = GetProfilerWindow();
 	if (Wnd.IsValid())
 	{
-		Wnd->ShowHideTab(FLoadingProfilerTabs::ObjectTypeAggregationTreeViewID, bIsObjectTypeAggregationTreeViewVisible);
+		Wnd->ShowHideTab(FLoadingProfilerTabs::ExportDetailsTreeViewID, bIsVisible);
 	}
 }
 
diff --git a/Engine/Source/Developer/TraceInsights/Private/Insights/LoadingProfiler/LoadingProfilerManager.h b/Engine/Source/Developer/TraceInsights/Private/Insights/LoadingProfiler/LoadingProfilerManager.h
--- a/Engine/Source/Developer/TraceInsights/Private/Insights/LoadingProfiler/LoadingProfilerManager.h
+++ b/Engine/Source/Developer/TraceInsights/Private/Insights/LoadingProfiler/LoadingProfilerManager.h
@@ -99,18 +99,28 @@ public:
 
 	/** @return true, if Timing view is visible */
 	const bool IsTimingViewVisible() const { return bIsTimingViewVisible; }
-	void SetTimingViewVisible(const bool bTimingViewVisibleState) { bIsTimingViewVisible = bTimingViewVisibleState; }
-	void ShowHideTimingView(const bool bTimingViewVisibleState);
+	void SetTimingViewVisible(const bool bIsVisible) { bIsTimingViewVisible = bIsVisible; }
+	void ShowHideTimingView(const bool bIsVisible);
 
 	/** @return true, if Event Aggregation tree view is visible */
 	const bool IsEventAggregationTreeViewVisible() const { return bIsEventAggregationTreeViewVisible; }
-	void SetEventAggregationTreeViewVisible(const bool bVisibleState) { bIsEventAggregationTreeViewVisible = bVisibleState; }
-	void ShowHideEventAggregationTreeView(const bool bVisibleState);
+	void SetEventAggregationTreeViewVisible(const bool bIsVisible) { bIsEventAggregationTreeViewVisible = bIsVisible; }
+	void ShowHideEventAggregationTreeView(const bool bIsVisible);
 
 	/** @return true, if Object Type Aggregation tree view is visible */
 	const bool IsObjectTypeAggregationTreeViewVisible() const { return bIsObjectTypeAggregationTreeViewVisible; }
-	void SetObjectTypeAggregationTreeViewVisible(const bool bVisibleState) { bIsObjectTypeAggregationTreeViewVisible = bVisibleState; }
-	void ShowHideObjectTypeAggregationTreeView(const bool bVisibleState);
+	void SetObjectTypeAggregationTreeViewVisible(const bool bIsVisible) { bIsObjectTypeAggregationTreeViewVisible = bIsVisible; }
+	void ShowHideObjectTypeAggregationTreeView(const bool bIsVisible);
+
+	/** @return true, if Package Details tree view is visible */
+	const bool IsPackageDetailsTreeViewVisible() const { return bIsPackageDetailsTreeViewVisible; }
+	void SetPackageDetailsTreeViewVisible(const bool bIsVisible) { bIsPackageDetailsTreeViewVisible = bIsVisible; }
+	void ShowHidePackageDetailsTreeView(const bool bIsVisible);
+
+	/** @return true, if Export Details tree view is visible */
+	const bool IsExportDetailsTreeViewVisible() const { return bIsExportDetailsTreeViewVisible; }
+	void SetExportDetailsTreeViewVisible(const bool bIsVisible) { bIsExportDetailsTreeViewVisible = bIsVisible; }
+	void ShowHideExportDetailsTreeView(const bool bIsVisible);
 
 	////////////////////////////////////////////////////////////////////////////////////////////////////
 
@@ -145,6 +155,12 @@ protected:
 	/** If the Object Type Aggregation tree view is visible or hidden. */
 	bool bIsObjectTypeAggregationTreeViewVisible;
 
+	/** If the Package Details tree view is visible or hidden. */
+	bool bIsPackageDetailsTreeViewVisible;
+
+	/** If the Export Details tree view is visible or hidden. */
+	bool bIsExportDetailsTreeViewVisible;
+
 	/** A shared pointer to the global instance of the profiler manager. */
 	static TSharedPtr<FLoadingProfilerManager> Instance;
 };
diff --git a/Engine/Source/Developer/TraceInsights/Private/Insights/LoadingProfiler/Widgets/SLoadingProfilerToolbar.cpp b/Engine/Source/Developer/TraceInsights/Private/Insights/LoadingProfiler/Widgets/SLoadingProfilerToolbar.cpp
--- a/Engine/Source/Developer/TraceInsights/Private/Insights/LoadingProfiler/Widgets/SLoadingProfilerToolbar.cpp
+++ b/Engine/Source/Developer/TraceInsights/Private/Insights/LoadingProfiler/Widgets/SLoadingProfilerToolbar.cpp
@@ -48,6 +48,8 @@ void SLoadingProfilerToolbar::Construct(const FArguments& InArgs)
 				ToolbarBuilder.AddToolBarButton(FLoadingProfilerCommands::Get().ToggleTimingViewVisibility);
 				ToolbarBuilder.AddToolBarButton(FLoadingProfilerCommands::Get().ToggleEventAggregationTreeViewVisibility);
 				ToolbarBuilder.AddToolBarButton(FLoadingProfilerCommands::Get().ToggleObjectTypeAggregationTreeViewVisibility);
+				ToolbarBuilder.AddToolBarButton(FLoadingProfilerCommands::Get().TogglePackageDetailsTreeViewVisibility);
+				ToolbarBuilder.AddToolBarButton(FLoadingProfilerCommands::Get().ToggleExportDetailsTreeViewVisibility);
 			}
 			ToolbarBuilder.EndSection();
 		}
diff --git a/Engine/Source/Developer/TraceInsights/Private/Insights/LoadingProfiler/Widgets/SLoadingProfilerWindow.cpp b/Engine/Source/Developer/TraceInsights/Private/Insights/LoadingProfiler/Widgets/SLoadingProfilerWindow.cpp
--- a/Engine/Source/Developer/TraceInsights/Private/Insights/LoadingProfiler/Widgets/SLoadingProfilerWindow.cpp
+++ b/Engine/Source/Developer/TraceInsights/Private/Insights/LoadingProfiler/Widgets/SLoadingProfilerWindow.cpp
@@ -41,6 +41,8 @@ const FName FLoadingProfilerTabs::ToolbarID(TEXT("Toolbar"));
 const FName FLoadingProfilerTabs::TimingViewID(TEXT("TimingView"));
 const FName FLoadingProfilerTabs::EventAggregationTreeViewID(TEXT("EventAggregation"));
 const FName FLoadingProfilerTabs::ObjectTypeAggregationTreeViewID(TEXT("ObjectTypeAggregation"));
+const FName FLoadingProfilerTabs::PackageDetailsTreeViewID(TEXT("PackageDetails"));
+const FName FLoadingProfilerTabs::ExportDetailsTreeViewID(TEXT("ExportDetails"));
 
 ////////////////////////////////////////////////////////////////////////////////////////////////////
 
@@ -79,6 +81,16 @@ void SLoadingProfilerWindow::Reset()
 	{
 		ObjectTypeAggregationTreeView->RebuildTree(true);
 	}
+
+	if (PackageDetailsTreeView)
+	{
+		PackageDetailsTreeView->RebuildTree(true);
+	}
+
+	if (ExportDetailsTreeView)
+	{
+		ExportDetailsTreeView->RebuildTree(true);
+	}
 }
 
 ////////////////////////////////////////////////////////////////////////////////////////////////////
@@ -215,6 +227,86 @@ void SLoadingProfilerWindow::OnObjectTypeAggregationTreeViewTabClosed(TSharedRef
 
 ////////////////////////////////////////////////////////////////////////////////////////////////////
 
+TSharedRef<SDockTab> SLoadingProfilerWindow::SpawnTab_PackageDetailsTreeView(const FSpawnTabArgs& Args)
+{
+	FLoadingProfilerManager::Get()->SetPackageDetailsTreeViewVisible(true);
+
+	TSharedPtr<Insights::FTable> TablePtr = MakeShareable(new Insights::FTable());
+	TSharedPtr<const Trace::IAnalysisSession> Session = FInsightsManager::Get()->GetSession();
+	if (Session.IsValid() && Trace::ReadLoadTimeProfilerProvider(*Session.Get()))
+	{
+		Trace::FAnalysisSessionReadScope SessionReadScope(*Session.Get());
+		const Trace::ILoadTimeProfilerProvider& LoadTimeProfilerProvider = *Trace::ReadLoadTimeProfilerProvider(*Session.Get());
+		const double SelectionStartTime = TimingView ? TimingView->GetSelectionStartTime() : 0.0;
+		const double SelectionEndTime = TimingView ? TimingView->GetSelectionEndTime() : 0.0;
+		TSharedPtr<Trace::ITable<Trace::FPackagesTableRow>> SourceTable = MakeShareable(LoadTimeProfilerProvider.CreatePackageDetailsTable(SelectionStartTime, SelectionEndTime));
+		TablePtr->Init(SourceTable);
+	}
+
+	const TSharedRef<SDockTab> DockTab = SNew(SDockTab)
+		.ShouldAutosize(false)
+		.TabRole(ETabRole::PanelTab)
+		[
+			SAssignNew(PackageDetailsTreeView, Insights::STableTreeView, TablePtr)
+		];
+
+	PackageDetailsTreeView->RebuildTree(true);
+
+	DockTab->SetOnTabClosed(SDockTab::FOnTabClosedCallback::CreateRaw(this, &SLoadingProfilerWindow::OnPackageDetailsTreeViewTabClosed));
+
+	return DockTab;
+}
+
+////////////////////////////////////////////////////////////////////////////////////////////////////
+
+void SLoadingProfilerWindow::OnPackageDetailsTreeViewTabClosed(TSharedRef<SDockTab> TabBeingClosed)
+{
+	FLoadingProfilerManager::Get()->SetPackageDetailsTreeViewVisible(false);
+	PackageDetailsTreeView = nullptr;
+}
+
+////////////////////////////////////////////////////////////////////////////////////////////////////
+
+TSharedRef<SDockTab> SLoadingProfilerWindow::SpawnTab_ExportDetailsTreeView(const FSpawnTabArgs& Args)
+{
+	FLoadingProfilerManager::Get()->SetExportDetailsTreeViewVisible(true);
+
+	TSharedPtr<Insights::FTable> TablePtr = MakeShareable(new Insights::FTable());
+	TSharedPtr<const Trace::IAnalysisSession> Session = FInsightsManager::Get()->GetSession();
+	if (Session.IsValid() && Trace::ReadLoadTimeProfilerProvider(*Session.Get()))
+	{
+		Trace::FAnalysisSessionReadScope SessionReadScope(*Session.Get());
+		const Trace::ILoadTimeProfilerProvider& LoadTimeProfilerProvider = *Trace::ReadLoadTimeProfilerProvider(*Session.Get());
+		const double SelectionStartTime = TimingView ? TimingView->GetSelectionStartTime() : 0.0;
+		const double SelectionEndTime = TimingView ? TimingView->GetSelectionEndTime() : 0.0;
+		TSharedPtr<Trace::ITable<Trace::FExportsTableRow>> SourceTable = MakeShareable(LoadTimeProfilerProvider.CreateExportDetailsTable(SelectionStartTime, SelectionEndTime));
+		TablePtr->Init(SourceTable);
+	}
+
+	const TSharedRef<SDockTab> DockTab = SNew(SDockTab)
+		.ShouldAutosize(false)
+		.TabRole(ETabRole::PanelTab)
+		[
+			SAssignNew(ExportDetailsTreeView, Insights::STableTreeView, TablePtr)
+		];
+
+	ExportDetailsTreeView->RebuildTree(true);
+
+	DockTab->SetOnTabClosed(SDockTab::FOnTabClosedCallback::CreateRaw(this, &SLoadingProfilerWindow::OnExportDetailsTreeViewTabClosed));
+
+	return DockTab;
+}
+
+////////////////////////////////////////////////////////////////////////////////////////////////////
+
+void SLoadingProfilerWindow::OnExportDetailsTreeViewTabClosed(TSharedRef<SDockTab> TabBeingClosed)
+{
+	FLoadingProfilerManager::Get()->SetExportDetailsTreeViewVisible(false);
+	ExportDetailsTreeView = nullptr;
+}
+
+////////////////////////////////////////////////////////////////////////////////////////////////////
+
 void SLoadingProfilerWindow::Construct(const FArguments& InArgs, const TSharedRef<SDockTab>& ConstructUnderMajorTab, const TSharedPtr<SWindow>& ConstructUnderWindow)
 {
 	// Create & initialize tab manager.
@@ -241,6 +333,16 @@ void SLoadingProfilerWindow::Construct(const FArguments& InArgs, const TSharedRe
 		.SetIcon(FSlateIcon(FInsightsStyle::GetStyleSetName(), "TableTreeView.Icon.Small"))
 		.SetGroup(AppMenuGroup);
 
+	TabManager->RegisterTabSpawner(FLoadingProfilerTabs::PackageDetailsTreeViewID, FOnSpawnTab::CreateRaw(this, &SLoadingProfilerWindow::SpawnTab_PackageDetailsTreeView))
+		.SetDisplayName(LOCTEXT("LoadingProfiler.PackageDetailsTreeViewTabTitle", "Package Details"))
+		.SetIcon(FSlateIcon(FInsightsStyle::GetStyleSetName(), "TableTreeView.Icon.Small"))
+		.SetGroup(AppMenuGroup);
+
+	TabManager->RegisterTabSpawner(FLoadingProfilerTabs::ExportDetailsTreeViewID, FOnSpawnTab::CreateRaw(this, &SLoadingProfilerWindow::SpawnTab_ExportDetailsTreeView))
+		.SetDisplayName(LOCTEXT("LoadingProfiler.ExportDetailsTreeViewTabTitle", "Export Details"))
+		.SetIcon(FSlateIcon(FInsightsStyle::GetStyleSetName(), "TableTreeView.Icon.Small"))
+		.SetGroup(AppMenuGroup);
+
 	TSharedPtr<FLoadingProfilerManager> LoadingProfilerManager = FLoadingProfilerManager::Get();
 	ensure(LoadingProfilerManager.IsValid());
 
@@ -274,6 +376,8 @@ void SLoadingProfilerWindow::Construct(const FArguments& InArgs, const TSharedRe
 					->SetSizeCoefficient(0.35f)
 					->AddTab(FLoadingProfilerTabs::EventAggregationTreeViewID, LoadingProfilerManager->IsEventAggregationTreeViewVisible() ? ETabState::OpenedTab : ETabState::ClosedTab)
 					->AddTab(FLoadingProfilerTabs::ObjectTypeAggregationTreeViewID, LoadingProfilerManager->IsObjectTypeAggregationTreeViewVisible() ? ETabState::OpenedTab : ETabState::ClosedTab)
+					->AddTab(FLoadingProfilerTabs::PackageDetailsTreeViewID, LoadingProfilerManager->IsPackageDetailsTreeViewVisible() ? ETabState::OpenedTab : ETabState::ClosedTab)
+					->AddTab(FLoadingProfilerTabs::ExportDetailsTreeViewID, LoadingProfilerManager->IsExportDetailsTreeViewVisible() ? ETabState::OpenedTab : ETabState::ClosedTab)
 					->SetForegroundTab(FLoadingProfilerTabs::EventAggregationTreeViewID)
 				)
 			)
diff --git a/Engine/Source/Developer/TraceInsights/Private/Insights/LoadingProfiler/Widgets/SLoadingProfilerWindow.h b/Engine/Source/Developer/TraceInsights/Private/Insights/LoadingProfiler/Widgets/SLoadingProfilerWindow.h
--- a/Engine/Source/Developer/TraceInsights/Private/Insights/LoadingProfiler/Widgets/SLoadingProfilerWindow.h
+++ b/Engine/Source/Developer/TraceInsights/Private/Insights/LoadingProfiler/Widgets/SLoadingProfilerWindow.h
@@ -36,6 +36,8 @@ struct FLoadingProfilerTabs
 	static const FName TimingViewID;
 	static const FName EventAggregationTreeViewID;
 	static const FName ObjectTypeAggregationTreeViewID;
+	static const FName PackageDetailsTreeViewID;
+	static const FName ExportDetailsTreeViewID;
 };
 
 ////////////////////////////////////////////////////////////////////////////////////////////////////
@@ -67,6 +69,8 @@ public:
 	TSharedPtr<STimingView> GetTimingView() const { return TimingView; }
 	TSharedPtr<Insights::STableTreeView> GetEventAggregationTreeView() const { return EventAggregationTreeView; }
 	TSharedPtr<Insights::STableTreeView> GetObjectTypeAggregationTreeView() const { return ObjectTypeAggregationTreeView; }
+	TSharedPtr<Insights::STableTreeView> GetPackageDetailsTreeView() const { return PackageDetailsTreeView; }
+	TSharedPtr<Insights::STableTreeView> GetExportDetailsTreeView() const { return ExportDetailsTreeView; }
 
 private:
 	TSharedRef<SDockTab> SpawnTab_Toolbar(const FSpawnTabArgs& Args);
@@ -81,6 +85,12 @@ private:
 	TSharedRef<SDockTab> SpawnTab_ObjectTypeAggregationTreeView(const FSpawnTabArgs& Args);
 	void OnObjectTypeAggregationTreeViewTabClosed(TSharedRef<SDockTab> TabBeingClosed);
 
+	TSharedRef<SDockTab> SpawnTab_PackageDetailsTreeView(const FSpawnTabArgs& Args);
+	void OnPackageDetailsTreeViewTabClosed(TSharedRef<SDockTab> TabBeingClosed);
+
+	TSharedRef<SDockTab> SpawnTab_ExportDetailsTreeView(const FSpawnTabArgs& Args);
+	void OnExportDetailsTreeViewTabClosed(TSharedRef<SDockTab> TabBeingClosed);
+
 	/**
 	 * Fill the main menu with menu items.
 	 *
@@ -162,6 +172,12 @@ private:
 	/** The Object Type Aggregation tree view widget */
 	TSharedPtr<Insights::STableTreeView> ObjectTypeAggregationTreeView;
 
+	/** The Package Details tree view widget */
+	TSharedPtr<Insights::STableTreeView> PackageDetailsTreeView;
+
+	/** The Export Details tree view widget */
+	TSharedPtr<Insights::STableTreeView> ExportDetailsTreeView;
+
 	/** Holds the tab manager that manages the front-end's tabs. */
 	TSharedPtr<FTabManager> TabManager;
 
diff --git a/Engine/Source/Developer/TraceInsights/Private/Insights/Table/ViewModels/Table.cpp b/Engine/Source/Developer/TraceInsights/Private/Insights/Table/ViewModels/Table.cpp
--- a/Engine/Source/Developer/TraceInsights/Private/Insights/Table/ViewModels/Table.cpp
+++ b/Engine/Source/Developer/TraceInsights/Private/Insights/Table/ViewModels/Table.cpp
@@ -12,6 +12,321 @@
 namespace Insights
 {
 
+////////////////////////////////////////////////////////////////////////////////////////////////////
+// FTableCellValueFormatter
+////////////////////////////////////////////////////////////////////////////////////////////////////
+
+class FBoolValueFormatter : public FTableCellValueFormatter
+{
+public:
+	virtual FText FormatValue(const TOptional<FTableCellValue>& InValue) override
+	{
+		if (InValue.IsSet())
+		{
+			return FormatBoolValue(InValue.GetValue().Bool);
+		}
+		return FText::GetEmpty();
+	}
+
+	virtual FText FormatValue(const FTable& Table, const FTableColumn& Column, const FTableRowId& RowId) override
+	{
+		TSharedPtr<Trace::IUntypedTableReader> Reader = Table.GetTableReader();
+		if (Reader.IsValid() && RowId.HasValidIndex())
+		{
+			Reader->SetRowIndex(RowId.RowIndex);
+			const bool Value = Reader->GetValueBool(Column.GetIndex());
+			return FormatBoolValue(Value);
+		}
+		return FText::GetEmpty();
+	}
+
+private:
+	FText FormatBoolValue(bool Value)
+	{
+		return FText::FromString(Value ? TEXT("True") : TEXT("False"));
+	}
+};
+
+////////////////////////////////////////////////////////////////////////////////////////////////////
+
+class FInt64ValueFormatter : public FTableCellValueFormatter
+{
+public:
+	virtual FText FormatValue(const TOptional<FTableCellValue>& InValue) override
+	{
+		if (InValue.IsSet())
+		{
+			return FormatInt64Value(InValue.GetValue().Int64);
+		}
+		return FText::GetEmpty();
+	}
+
+	virtual FText FormatValue(const FTable& Table, const FTableColumn& Column, const FTableRowId& RowId) override
+	{
+		TSharedPtr<Trace::IUntypedTableReader> Reader = Table.GetTableReader();
+		if (Reader.IsValid() && RowId.HasValidIndex())
+		{
+			Reader->SetRowIndex(RowId.RowIndex);
+			const int64 Value = Reader->GetValueInt(Column.GetIndex());
+			return FormatInt64Value(Value);
+		}
+		return FText::GetEmpty();
+	}
+
+private:
+	FText FormatInt64Value(int64 Value)
+	{
+		return FText::AsNumber(Value);
+	}
+};
+
+////////////////////////////////////////////////////////////////////////////////////////////////////
+
+class FFloatValueFormatter : public FTableCellValueFormatter
+{
+public:
+	virtual FText FormatValue(const TOptional<FTableCellValue>& InValue) override
+	{
+		if (InValue.IsSet())
+		{
+			return FormatFloatValue(InValue.GetValue().Float);
+		}
+		return FText::GetEmpty();
+	}
+
+	virtual FText FormatValue(const FTable& Table, const FTableColumn& Column, const FTableRowId& RowId) override
+	{
+		TSharedPtr<Trace::IUntypedTableReader> Reader = Table.GetTableReader();
+		if (Reader.IsValid() && RowId.HasValidIndex())
+		{
+			Reader->SetRowIndex(RowId.RowIndex);
+			const float Value = Reader->GetValueFloat(Column.GetIndex());
+			return FormatFloatValue(Value);
+		}
+		return FText::GetEmpty();
+	}
+
+private:
+	FText FormatFloatValue(float Value)
+	{
+		if (Value == 0.0f)
+		{
+			return FText::FromString(TEXT("0"));
+		}
+		else
+		{
+			return FText::FromString(FString::Printf(TEXT("%f"), Value));
+		}
+	}
+};
+
+////////////////////////////////////////////////////////////////////////////////////////////////////
+
+class FFloatTimeValueFormatter : public FTableCellValueFormatter
+{
+public:
+	virtual FText FormatValue(const TOptional<FTableCellValue>& InValue) override
+	{
+		if (InValue.IsSet())
+		{
+			return FormatFloatValue(InValue.GetValue().Float);
+		}
+		return FText::GetEmpty();
+	}
+
+	virtual FText FormatValue(const FTable& Table, const FTableColumn& Column, const FTableRowId& RowId) override
+	{
+		TSharedPtr<Trace::IUntypedTableReader> Reader = Table.GetTableReader();
+		if (Reader.IsValid() && RowId.HasValidIndex())
+		{
+			Reader->SetRowIndex(RowId.RowIndex);
+			const float Value = Reader->GetValueFloat(Column.GetIndex());
+			return FormatFloatValue(Value);
+		}
+		return FText::GetEmpty();
+	}
+
+	virtual FText FormatValueForTooltip(const TOptional<FTableCellValue>& InValue) override
+	{
+		if (InValue.IsSet())
+		{
+			return FormatFloatValueForTooltip(InValue.GetValue().Float);
+		}
+		return FText::GetEmpty();
+	}
+
+	virtual FText FormatValueForTooltip(const FTable& Table, const FTableColumn& Column, const FTableRowId& RowId) override
+	{
+		TSharedPtr<Trace::IUntypedTableReader> Reader = Table.GetTableReader();
+		if (Reader.IsValid() && RowId.HasValidIndex())
+		{
+			Reader->SetRowIndex(RowId.RowIndex);
+			const float Value = Reader->GetValueFloat(Column.GetIndex());
+			return FormatFloatValueForTooltip(Value);
+		}
+		return FText::GetEmpty();
+	}
+
+private:
+	FText FormatFloatValue(float Value)
+	{
+		return FText::FromString(TimeUtils::FormatTimeAuto(static_cast<double>(Value)));
+	}
+
+	FText FormatFloatValueForTooltip(float Value)
+	{
+		if (Value == 0.0f)
+		{
+			return FText::FromString(TEXT("0"));
+		}
+		else
+		{
+			return FText::FromString(FString::Printf(TEXT("%f (%s)"), Value, *TimeUtils::FormatTimeAuto(static_cast<double>(Value))));
+		}
+	}
+};
+
+////////////////////////////////////////////////////////////////////////////////////////////////////
+
+class FDoubleValueFormatter : public FTableCellValueFormatter
+{
+public:
+	virtual FText FormatValue(const TOptional<FTableCellValue>& InValue) override
+	{
+		if (InValue.IsSet())
+		{
+			return FormatDoubleValue(InValue.GetValue().Double);
+		}
+		return FText::GetEmpty();
+	}
+
+	virtual FText FormatValue(const FTable& Table, const FTableColumn& Column, const FTableRowId& RowId) override
+	{
+		TSharedPtr<Trace::IUntypedTableReader> Reader = Table.GetTableReader();
+		if (Reader.IsValid() && RowId.HasValidIndex())
+		{
+			Reader->SetRowIndex(RowId.RowIndex);
+			const double Value = Reader->GetValueDouble(Column.GetIndex());
+			return FormatDoubleValue(Value);
+		}
+		return FText::GetEmpty();
+	}
+
+private:
+	FText FormatDoubleValue(double Value)
+	{
+		if (Value == 0.0)
+		{
+			return FText::FromString(TEXT("0"));
+		}
+		else
+		{
+			return FText::FromString(FString::Printf(TEXT("%f"), Value));
+		}
+	}
+};
+
+////////////////////////////////////////////////////////////////////////////////////////////////////
+
+class FDoubleTimeValueFormatter : public FTableCellValueFormatter
+{
+public:
+	virtual FText FormatValue(const TOptional<FTableCellValue>& InValue) override
+	{
+		if (InValue.IsSet())
+		{
+			return FormatDoubleValue(InValue.GetValue().Double);
+		}
+		return FText::GetEmpty();
+	}
+
+	virtual FText FormatValue(const FTable& Table, const FTableColumn& Column, const FTableRowId& RowId) override
+	{
+		TSharedPtr<Trace::IUntypedTableReader> Reader = Table.GetTableReader();
+		if (Reader.IsValid() && RowId.HasValidIndex())
+		{
+			Reader->SetRowIndex(RowId.RowIndex);
+			const double Value = Reader->GetValueDouble(Column.GetIndex());
+			return FormatDoubleValue(Value);
+		}
+		return FText::GetEmpty();
+	}
+
+	virtual FText FormatValueForTooltip(const TOptional<FTableCellValue>& InValue) override
+	{
+		if (InValue.IsSet())
+		{
+			return FormatDoubleValueForTooltip(InValue.GetValue().Double);
+		}
+		return FText::GetEmpty();
+	}
+
+	virtual FText FormatValueForTooltip(const FTable& Table, const FTableColumn& Column, const FTableRowId& RowId) override
+	{
+		TSharedPtr<Trace::IUntypedTableReader> Reader = Table.GetTableReader();
+		if (Reader.IsValid() && RowId.HasValidIndex())
+		{
+			Reader->SetRowIndex(RowId.RowIndex);
+			const double Value = Reader->GetValueDouble(Column.GetIndex());
+			return FormatDoubleValueForTooltip(Value);
+		}
+		return FText::GetEmpty();
+	}
+
+private:
+	FText FormatDoubleValue(double Value)
+	{
+		return FText::FromString(TimeUtils::FormatTimeAuto(Value));
+	}
+
+	FText FormatDoubleValueForTooltip(double Value)
+	{
+		if (Value == 0.0)
+		{
+			return FText::FromString(TEXT("0"));
+		}
+		else
+		{
+			return FText::FromString(FString::Printf(TEXT("%f (%s)"), Value, *TimeUtils::FormatTimeAuto(Value)));
+		}
+	}
+};
+
+////////////////////////////////////////////////////////////////////////////////////////////////////
+
+class FCStringValueFormatter : public FTableCellValueFormatter
+{
+public:
+	virtual FText FormatValue(const TOptional<FTableCellValue>& InValue) override
+	{
+		if (InValue.IsSet())
+		{
+			return FormatCStringValue(InValue.GetValue().CString);
+		}
+		return FText::GetEmpty();
+	}
+
+	virtual FText FormatValue(const FTable& Table, const FTableColumn& Column, const FTableRowId& RowId) override
+	{
+		TSharedPtr<Trace::IUntypedTableReader> Reader = Table.GetTableReader();
+		if (Reader.IsValid() && RowId.HasValidIndex())
+		{
+			Reader->SetRowIndex(RowId.RowIndex);
+			const TCHAR* Value = Reader->GetValueCString(Column.GetIndex());
+			return FormatCStringValue(Value);
+		}
+		return FText::GetEmpty();
+	}
+
+private:
+	FText FormatCStringValue(const TCHAR* Value)
+	{
+		return FText::FromString(Value);
+	}
+};
+
+////////////////////////////////////////////////////////////////////////////////////////////////////
+// FTable
 ////////////////////////////////////////////////////////////////////////////////////////////////////
 
 FTable::FTable()
@@ -96,31 +411,39 @@ void FTable::UpdateSourceTable(TSharedPtr<Trace::IUntypedTable> InSourceTable)
 
 ////////////////////////////////////////////////////////////////////////////////////////////////////
 
-void FTable::CreateHierarchyColumn(const TCHAR* ColumnName)
+void FTable::CreateHierarchyColumn(int32 ColumnIndex, const TCHAR* ColumnName)
 {
-	const FString ColumnNameStr(ColumnName);
+	const FString ColumnNameStr = ColumnName ? FString::Printf(TEXT("Hierarchy (%s)"), ColumnName) : TEXT("Hierarchy");
 	const FText ColumnNameText = FText::FromString(ColumnNameStr);
 
 	constexpr EHorizontalAlignment HorizontalAlignment = HAlign_Left;
 	constexpr ETableColumnFlags ColumnFlags = ETableColumnFlags::ShouldBeVisible | ETableColumnFlags::CanBeSorted | ETableColumnFlags::CanBeFiltered | ETableColumnFlags::IsHierarchy;
 
-	constexpr float InitialColumnWidth = 200.0f;
+	constexpr ETableCellDataType DataType = ETableCellDataType::CString;
+	constexpr ETableColumnAggregation Aggregation = ETableColumnAggregation::None;
+
+	constexpr float InitialColumnWidth = 80.0f;
 	constexpr float MinColumnWidth = 0.0f;
 	constexpr float MaxColumnWidth = FLT_MAX;
 
+	TSharedPtr<FTableCellValueFormatter> FormatterPtr = MakeShareable(new FTableCellValueFormatter());
+
 	TSharedPtr<FTableColumn> ColumnPtr = MakeShareable(new FTableColumn
 	(
 		-1, // Order
-		0,
-		FName(ColumnName), // Id
+		ColumnIndex,
+		FName(TEXT("_Hierarchy")), // Id
 		ColumnNameText, // Short Name
 		ColumnNameText, // Title Name
 		FText::GetEmpty(), // Description
 		ColumnFlags,
+		DataType,
+		Aggregation,
 		HorizontalAlignment,
 		InitialColumnWidth,
 		MinColumnWidth,
-		MaxColumnWidth
+		MaxColumnWidth,
+		FormatterPtr.ToSharedRef()
 	));
 
 	AddColumn(ColumnPtr);
@@ -149,189 +472,98 @@ void FTable::CreateColumns()
 	ensure(TableReader.IsValid());
 	ensure(Columns.Num() == 0);
 
-	const TCHAR* HierarchyColumnName = TEXT("Name");
-
-	CreateHierarchyColumn(HierarchyColumnName);
-
 	const Trace::ITableLayout& TableLayout = SourceTable->GetLayout();
 	const int32 ColumnCount = TableLayout.GetColumnCount();
 
+	//////////////////////////////////////////////////
+	// Hierarchy Column
+
+	int32 HierarchyColumnIndex = -1;
+	const TCHAR* HierarchyColumnName = nullptr;
+
+	// Look for first string column.
+	//for (int32 ColumnIndex = 0; ColumnIndex < ColumnCount; ++ColumnIndex)
+	//{
+	//	Trace::ETableColumnType ColumnType = TableLayout.GetColumnType(ColumnIndex);
+	//	if (ColumnType == Trace::TableColumnType_CString)
+	//	{
+	//		HierarchyColumnIndex = ColumnIndex;
+	//		HierarchyColumnName = TableLayout.GetColumnName(ColumnIndex);
+	//		break;
+	//	}
+	//}
+
+	CreateHierarchyColumn(HierarchyColumnIndex, HierarchyColumnName);
+
+	//////////////////////////////////////////////////
+
 	for (int32 ColumnIndex = 0; ColumnIndex < ColumnCount; ++ColumnIndex)
 	{
 		Trace::ETableColumnType ColumnType = TableLayout.GetColumnType(ColumnIndex);
 		const TCHAR* ColumnName = TableLayout.GetColumnName(ColumnIndex);
 
-		if (ColumnType == Trace::TableColumnType_CString && FCString::Stricmp(ColumnName, HierarchyColumnName) == 0)
-		{
-			// Skip the "Name" column. It is already created as the name / hierarchy column.
-			continue;
-		}
-
 		const FString ColumnNameStr(ColumnName);
 		const FText ColumnNameText = FText::FromString(ColumnNameStr);
 
 		EHorizontalAlignment HorizontalAlignment = HAlign_Left;
 
-		const ETableColumnFlags ColumnFlags = ETableColumnFlags::ShouldBeVisible | ETableColumnFlags::CanBeSorted | ETableColumnFlags::CanBeFiltered | ETableColumnFlags::CanBeHidden;
+		ETableColumnFlags ColumnFlags = ETableColumnFlags::CanBeSorted | ETableColumnFlags::CanBeFiltered | ETableColumnFlags::CanBeHidden;
+		if (ColumnIndex != HierarchyColumnIndex)
+		{
+			ColumnFlags |= ETableColumnFlags::ShouldBeVisible;
+		}
+
+		ETableCellDataType DataType = ETableCellDataType::Unknown;
+		ETableColumnAggregation Aggregation = ETableColumnAggregation::None;
 
 		float InitialColumnWidth = 60.0f;
 		float MinColumnWidth = 0.0f;
 		float MaxColumnWidth = FLT_MAX;
 
-		FTableColumn::FGetValueAsTextFunction GetUnknownValueAsTextFn = [](const FTable& Table, const FTableColumn& Column, const FTableRowId& RowId) -> FText
-		{
-			return LOCTEXT("UnknownValue", "!?");
-		};
-
-		FTableColumn::FGetValueAsTextFunction GetBoolValueAsTextFn = [](const FTable& Table, const FTableColumn& Column, const FTableRowId& RowId) -> FText
-		{
-			TSharedPtr<Trace::IUntypedTableReader> Reader = Table.GetTableReader();
-			if (Reader.IsValid() && RowId.HasValidIndex())
-			{
-				Reader->SetRowIndex(RowId.RowIndex);
-				bool Value = Reader->GetValueBool(Column.GetIndex());
-				return FText::FromString(Value ? TEXT("True") : TEXT("False"));
-			}
-			return FText::GetEmpty();
-		};
-
-		FTableColumn::FGetValueAsTextFunction GetIntValueAsTextFn = [](const FTable& Table, const FTableColumn& Column, const FTableRowId& RowId) -> FText
-		{
-			TSharedPtr<Trace::IUntypedTableReader> Reader = Table.GetTableReader();
-			if (Reader.IsValid() && RowId.HasValidIndex())
-			{
-				Reader->SetRowIndex(RowId.RowIndex);
-				int64 Value = Reader->GetValueInt(Column.GetIndex());
-				return FText::AsNumber(Value);
-			}
-			return FText::GetEmpty();
-		};
-
-		FTableColumn::FGetValueAsTextFunction GetFloatValueAsTextFn = [](const FTable& Table, const FTableColumn& Column, const FTableRowId& RowId) -> FText
-		{
-			TSharedPtr<Trace::IUntypedTableReader> Reader = Table.GetTableReader();
-			if (Reader.IsValid() && RowId.HasValidIndex())
-			{
-				Reader->SetRowIndex(RowId.RowIndex);
-				float Value = Reader->GetValueFloat(Column.GetIndex());
-				if (Value == 0.0f)
-				{
-					return FText::FromString(TEXT("0"));
-				}
-				else
-				{
-					return FText::FromString(FString::Printf(TEXT("%f"), Value));
-				}
-			}
-			return FText::GetEmpty();
-		};
-
-		FTableColumn::FGetValueAsTextFunction GetFloatValueAsTextFn2 = [](const FTable& Table, const FTableColumn& Column, const FTableRowId& RowId) -> FText
-		{
-			TSharedPtr<Trace::IUntypedTableReader> Reader = Table.GetTableReader();
-			if (Reader.IsValid() && RowId.HasValidIndex())
-			{
-				Reader->SetRowIndex(RowId.RowIndex);
-				float Value = Reader->GetValueFloat(Column.GetIndex());
-				if (Value == 0.0f)
-				{
-					return FText::FromString(TEXT("0"));
-				}
-				else
-				{
-					return FText::FromString(FString::Printf(TEXT("%f (%s)"), Value, *TimeUtils::FormatTimeAuto(static_cast<double>(Value))));
-				}
-			}
-			return FText::GetEmpty();
-		};
-
-		FTableColumn::FGetValueAsTextFunction GetDoubleValueAsTextFn = [](const FTable& Table, const FTableColumn& Column, const FTableRowId& RowId) -> FText
-		{
-			TSharedPtr<Trace::IUntypedTableReader> Reader = Table.GetTableReader();
-			if (Reader.IsValid() && RowId.HasValidIndex())
-			{
-				Reader->SetRowIndex(RowId.RowIndex);
-				double Value = Reader->GetValueDouble(Column.GetIndex());
-				if (Value == 0.0)
-				{
-					return FText::FromString(TEXT("0"));
-				}
-				else
-				{
-					return FText::FromString(FString::Printf(TEXT("%f"), Value));
-				}
-			}
-			return FText::GetEmpty();
-		};
-
-		FTableColumn::FGetValueAsTextFunction GetDoubleValueAsTextFn2 = [](const FTable& Table, const FTableColumn& Column, const FTableRowId& RowId) -> FText
-		{
-			TSharedPtr<Trace::IUntypedTableReader> Reader = Table.GetTableReader();
-			if (Reader.IsValid() && RowId.HasValidIndex())
-			{
-				Reader->SetRowIndex(RowId.RowIndex);
-				double Value = Reader->GetValueDouble(Column.GetIndex());
-				if (Value == 0.0)
-				{
-					return FText::FromString(TEXT("0"));
-				}
-				else
-				{
-					return FText::FromString(FString::Printf(TEXT("%f (%s)"), Value, *TimeUtils::FormatTimeAuto(Value)));
-				}
-			}
-			return FText::GetEmpty();
-		};
-
-		FTableColumn::FGetValueAsTextFunction GetStringValueAsTextFn = [](const FTable& Table, const FTableColumn& Column, const FTableRowId& RowId) -> FText
-		{
-			TSharedPtr<Trace::IUntypedTableReader> Reader = Table.GetTableReader();
-			if (Reader.IsValid() && RowId.HasValidIndex())
-			{
-				Reader->SetRowIndex(RowId.RowIndex);
-				const TCHAR* Value = Reader->GetValueCString(Column.GetIndex());
-				return FText::FromString(Value);
-			}
-			return FText::GetEmpty();
-		};
-
-		FTableColumn::FGetValueAsTextFunction GetValueAsTextFn = GetUnknownValueAsTextFn;
-		FTableColumn::FGetValueAsTextFunction GetValueAsTooltipTextFn = GetUnknownValueAsTextFn;
+		TSharedPtr<FTableCellValueFormatter> FormatterPtr;
 
 		switch (ColumnType)
 		{
 		case Trace::TableColumnType_Bool:
+			DataType = ETableCellDataType::Bool;
 			HorizontalAlignment = HAlign_Right;
 			InitialColumnWidth = 40.0f;
-			GetValueAsTextFn = GetBoolValueAsTextFn;
-			GetValueAsTooltipTextFn = GetBoolValueAsTextFn;
+			FormatterPtr = MakeShareable(new FBoolValueFormatter());
 			break;
 		case Trace::TableColumnType_Int:
+			DataType = ETableCellDataType::Int64;
+			Aggregation = ETableColumnAggregation::Sum;
 			HorizontalAlignment = HAlign_Right;
 			InitialColumnWidth = 60.0f;
-			GetValueAsTextFn = GetIntValueAsTextFn;
-			GetValueAsTooltipTextFn = GetIntValueAsTextFn;
+			FormatterPtr = MakeShareable(new FInt64ValueFormatter());
 			break;
 		case Trace::TableColumnType_Float:
+			DataType = ETableCellDataType::Float;
+			Aggregation = ETableColumnAggregation::Sum;
 			HorizontalAlignment = HAlign_Right;
 			InitialColumnWidth = 60.0f;
-			GetValueAsTextFn = GetFloatValueAsTextFn;
-			GetValueAsTooltipTextFn = GetFloatValueAsTextFn2;
+			FormatterPtr = MakeShareable(new FFloatTimeValueFormatter());
 			break;
 		case Trace::TableColumnType_Double:
+			DataType = ETableCellDataType::Double;
+			Aggregation = ETableColumnAggregation::Sum;
 			HorizontalAlignment = HAlign_Right;
 			InitialColumnWidth = 80.0f;
-			GetValueAsTextFn = GetDoubleValueAsTextFn;
-			GetValueAsTooltipTextFn = GetDoubleValueAsTextFn2;
+			FormatterPtr = MakeShareable(new FDoubleTimeValueFormatter());
 			break;
 		case Trace::TableColumnType_CString:
+			DataType = ETableCellDataType::CString;
 			HorizontalAlignment = HAlign_Left;
 			InitialColumnWidth = FMath::Max(120.0f, 6.0f * ColumnNameStr.Len());
-			GetValueAsTextFn = GetStringValueAsTextFn;
-			GetValueAsTooltipTextFn = GetStringValueAsTextFn;
+			FormatterPtr = MakeShareable(new FCStringValueFormatter());
 			break;
 		}
 
+		if (!FormatterPtr.IsValid())
+		{
+			FormatterPtr = MakeShareable(new FTableCellValueFormatter());
+		}
+
 		TSharedPtr<FTableColumn> ColumnPtr = MakeShareable(new FTableColumn
 		(
 			ColumnIndex, // Order
@@ -341,15 +573,15 @@ void FTable::CreateColumns()
 			ColumnNameText, // Title Name
 			FText::GetEmpty(), // Description
 			ColumnFlags,
+			DataType,
+			Aggregation,
 			HorizontalAlignment,
 			InitialColumnWidth,
 			MinColumnWidth,
-			MaxColumnWidth
+			MaxColumnWidth,
+			FormatterPtr.ToSharedRef()
 		));
 
-		ColumnPtr->SetValueAsTextFunction(GetValueAsTextFn);
-		ColumnPtr->SetValueAsTooltipTextFunction(GetValueAsTooltipTextFn);
-
 		AddColumn(ColumnPtr);
 	}
 }
diff --git a/Engine/Source/Developer/TraceInsights/Private/Insights/Table/ViewModels/Table.h b/Engine/Source/Developer/TraceInsights/Private/Insights/Table/ViewModels/Table.h
--- a/Engine/Source/Developer/TraceInsights/Private/Insights/Table/ViewModels/Table.h
+++ b/Engine/Source/Developer/TraceInsights/Private/Insights/Table/ViewModels/Table.h
@@ -75,7 +75,7 @@ public:
 	}
 
 private:
-	void CreateHierarchyColumn(const TCHAR* ColumnName);
+	void CreateHierarchyColumn(int32 ColumnIndex, const TCHAR* ColumnName);
 	void AddColumn(TSharedPtr<FTableColumn> ColumnPtr);
 	void CreateColumns();
 
diff --git a/Engine/Source/Developer/TraceInsights/Private/Insights/Table/ViewModels/TableColumn.h b/Engine/Source/Developer/TraceInsights/Private/Insights/Table/ViewModels/TableColumn.h
--- a/Engine/Source/Developer/TraceInsights/Private/Insights/Table/ViewModels/TableColumn.h
+++ b/Engine/Source/Developer/TraceInsights/Private/Insights/Table/ViewModels/TableColumn.h
@@ -10,6 +10,7 @@ namespace Insights
 {
 
 class FTable;
+class FTableColumn;
 struct FTableRowId;
 
 ////////////////////////////////////////////////////////////////////////////////////////////////////
@@ -26,6 +27,70 @@ enum class ETableColumnFlags : uint32
 };
 ENUM_CLASS_FLAGS(ETableColumnFlags);
 
+////////////////////////////////////////////////////////////////////////////////////////////////////
+
+enum class ETableCellDataType : uint32
+{
+	Unknown,
+
+	Bool,
+	Int64,
+	Float,
+	Double,
+	CString,
+
+	/** Invalid enum type, may be used as a number of enumerations. */
+	InvalidOrMax,
+};
+
+////////////////////////////////////////////////////////////////////////////////////////////////////
+
+struct FTableCellValue
+{
+	FTableCellValue() = default;
+	FTableCellValue(bool Value) : Bool(Value) {}
+	FTableCellValue(int64 Value) : Int64(Value) {}
+	FTableCellValue(float Value) : Float(Value) {}
+	FTableCellValue(double Value) : Double(Value) {}
+	FTableCellValue(const TCHAR* Value) : CString(Value) {}
+
+	union
+	{
+		bool Bool;
+		int64 Int64;
+		float Float;
+		double Double;
+		const TCHAR* CString;
+	};
+};
+
+////////////////////////////////////////////////////////////////////////////////////////////////////
+
+class FTableCellValueFormatter
+{
+public:
+	FTableCellValueFormatter() {}
+	virtual ~FTableCellValueFormatter() {}
+
+	virtual FText FormatValue(const TOptional<FTableCellValue>& InValue) { return FText::GetEmpty(); }
+	virtual FText FormatValueForTooltip(const TOptional<FTableCellValue>& InValue) { return FormatValue(InValue); }
+
+	virtual FText FormatValue(const FTable& Table, const FTableColumn& Column, const FTableRowId& RowId) { return FText::GetEmpty(); }
+	virtual FText FormatValueForTooltip(const FTable& Table, const FTableColumn& Column, const FTableRowId& RowId) { return FormatValue(Table, Column, RowId); }
+};
+
+////////////////////////////////////////////////////////////////////////////////////////////////////
+
+enum class ETableColumnAggregation : uint32
+{
+	None = 0,
+	Sum,
+	//Min,
+	//Max,
+	//Average,
+	//Median,
+};
+
 ////////////////////////////////////////////////////////////////////////////////////////////////////
 /**
  * Table Column View Model.
@@ -34,6 +99,7 @@ ENUM_CLASS_FLAGS(ETableColumnFlags);
 class FTableColumn
 {
 public:
+	typedef TFunction<FText(const FTableCellValue& Value)> FGetCellValueAsTextFunction;
 	typedef TFunction<FText(const FTable& Table, const FTableColumn& Column, const FTableRowId& RowId)> FGetValueAsTextFunction;
 
 public:
@@ -50,10 +116,13 @@ public:
 		FText InTitleName,
 		FText InDescription,
 		const ETableColumnFlags InFlags,
+		const ETableCellDataType InDataType,
+		const ETableColumnAggregation InAggregation,
 		const EHorizontalAlignment InHorizontalAlignment,
 		const float InInitialWidth,
 		const float InMinWidth,
-		const float InMaxWidth
+		const float InMaxWidth,
+		TSharedRef<FTableCellValueFormatter> InValueFormatter
 	)
 		: Order(InOrder)
 		, Index(InIndex)
@@ -61,20 +130,17 @@ public:
 		, ShortName(MoveTemp(InShortName))
 		, TitleName(MoveTemp(InTitleName))
 		, Description(MoveTemp(InDescription))
+		, bIsVisible(false)
 		, Flags(InFlags)
+		, DataType(InDataType)
+		, Aggregation(InAggregation)
 		, HorizontalAlignment(InHorizontalAlignment)
 		, InitialWidth(InInitialWidth)
 		, MinWidth(InMinWidth)
 		, MaxWidth(InMaxWidth)
-		, GetValueAsTextFn(nullptr)
-		, GetValueAsTooltipTextFn(nullptr)
+		, ValueFormatter(InValueFormatter)
 		, ParentTable(nullptr)
 	{
-		GetValueAsTextFn = [](const FTable& Table, const FTableColumn& Column, const FTableRowId& RowId) -> FText
-		{
-			return FText();
-		};
-		GetValueAsTooltipTextFn = GetValueAsTextFn;
 	}
 
 	int32 GetOrder() const { return Order; }
@@ -86,6 +152,10 @@ public:
 	const FText& GetTitleName() const { return TitleName; }
 	const FText& GetDescription() const { return Description; }
 
+	void SetShortName(const FText& InShortName) { ShortName = InShortName; }
+	void SetTitleName(const FText& InTitleName) { TitleName = InTitleName; }
+	void SetDescription(const FText& InDescription) { Description = InDescription; }
+
 	bool IsVisible() const { return bIsVisible; }
 	void Show() { bIsVisible = true; OnVisibilityChanged(); }
 	void Hide() { bIsVisible = false; OnVisibilityChanged(); }
@@ -108,6 +178,9 @@ public:
 	/** Whether this column is the hierarcy (name) column, in a tree view. */
 	bool IsHierarchy() const { return EnumHasAnyFlags(Flags, ETableColumnFlags::IsHierarchy); }
 
+	ETableCellDataType GetDataType() const { return DataType; }
+	ETableColumnAggregation GetAggregation() const { return Aggregation; }
+
 	EHorizontalAlignment GetHorizontalAlignment() const { return HorizontalAlignment; }
 
 	float GetInitialWidth() const { return InitialWidth; }
@@ -118,19 +191,18 @@ public:
 	/** If MinWidth == MaxWidth, this column has fixed width and cannot be resized. */
 	bool IsFixedWidth() const { return MinWidth == MaxWidth; }
 
+	TSharedRef<FTableCellValueFormatter> GetValueFormatter() const { return ValueFormatter; }
+
 	FText GetValueAsText(const FTableRowId& InRowId) const
 	{
-		return GetValueAsTextFn(*ParentTable.Pin(), *this, InRowId);
+		return ValueFormatter->FormatValue(*ParentTable.Pin(), *this, InRowId);
 	}
 
 	FText GetValueAsTooltipText(const FTableRowId& InRowId) const
 	{
-		return GetValueAsTooltipTextFn(*ParentTable.Pin(), *this, InRowId);
+		return ValueFormatter->FormatValueForTooltip(*ParentTable.Pin(), *this, InRowId);
 	}
 
-	void SetValueAsTextFunction(const FGetValueAsTextFunction InGetValueAsTextFn) { GetValueAsTextFn = InGetValueAsTextFn; }
-	void SetValueAsTooltipTextFunction(const FGetValueAsTextFunction InGetValueAsTooltipTextFn) { GetValueAsTooltipTextFn = InGetValueAsTooltipTextFn; }
-
 	TWeakPtr<FTable> GetParentTable() const { return ParentTable; }
 	void SetParentTable(TWeakPtr<FTable> InParentTable) { ParentTable = InParentTable; }
 
@@ -159,6 +231,12 @@ private:
 	/** Other on/off switches. */
 	ETableColumnFlags Flags;
 
+	/** Data type for this column. */
+	ETableCellDataType DataType;
+
+	/** Aggregation for values in this column, when grouped. */
+	ETableColumnAggregation Aggregation; // TODO: make this an object (Aggregator)
+
 	/** Horizontal alignment of the content in this column. */
 	EHorizontalAlignment HorizontalAlignment;
 
@@ -167,11 +245,8 @@ private:
 	float MaxWidth; /**< Maximum column width. */
 	//float Width; /**< Current column width. */
 
-	/** Custom function used to format (as an FText) the value displayed by this column (usually in table views). */
-	FGetValueAsTextFunction GetValueAsTextFn;
-
-	/** Custom function used to format (as an FText) the value displayed by this column (usually in tooltips; i.e. with increased precission). */
-	FGetValueAsTextFunction GetValueAsTooltipTextFn;
+	/** Custom formater for the value displayed by this column. */
+	TSharedRef<FTableCellValueFormatter> ValueFormatter;
 
 	/* Parent table. Only one table instance can own this column. */
 	TWeakPtr<FTable> ParentTable;
diff --git a/Engine/Source/Developer/TraceInsights/Private/Insights/Table/ViewModels/TableTreeNode.cpp b/Engine/Source/Developer/TraceInsights/Private/Insights/Table/ViewModels/TableTreeNode.cpp
--- a/Engine/Source/Developer/TraceInsights/Private/Insights/Table/ViewModels/TableTreeNode.cpp
+++ b/Engine/Source/Developer/TraceInsights/Private/Insights/Table/ViewModels/TableTreeNode.cpp
@@ -2,6 +2,9 @@
 
 #include "TableTreeNode.h"
 
+// Insights
+#include "Insights/Table/ViewModels/TableColumn.h"
+
 #define LOCTEXT_NAMESPACE "Insights_TableTreeNode"
 
 namespace Insights
@@ -9,18 +12,44 @@ namespace Insights
 
 ////////////////////////////////////////////////////////////////////////////////////////////////////
 
-bool FTableTreeNode::GetValueBool(int32 InColumnIndex) const
+FTableCellValue FTableTreeNode::GetValue(const FTableColumn& Column) const
+{
+	switch (Column.GetDataType())
+	{
+	case ETableCellDataType::Bool:    return FTableCellValue(GetValueBool(Column));
+	case ETableCellDataType::Int64:   return FTableCellValue(GetValueInt64(Column));
+	case ETableCellDataType::Float:   return FTableCellValue(GetValueFloat(Column));
+	case ETableCellDataType::Double:  return FTableCellValue(GetValueDouble(Column));
+	case ETableCellDataType::CString: return FTableCellValue(GetValueCString(Column));
+	default:                          return FTableCellValue();
+	}
+}
+
+////////////////////////////////////////////////////////////////////////////////////////////////////
+
+bool FTableTreeNode::GetValueBool(const FTableColumn& Column) const
 {
 	bool Value = false;
 
-	TSharedPtr<FTable> Table = ParentTable.Pin();
-	if (Table.IsValid())
+	if (!IsGroup())
 	{
-		TSharedPtr<Trace::IUntypedTableReader> Reader = Table->GetTableReader();
-		if (Reader.IsValid() && RowId.HasValidIndex())
+		TSharedPtr<FTable> Table = ParentTable.Pin();
+		if (Table.IsValid())
 		{
-			Reader->SetRowIndex(RowId.RowIndex);
-			Value = Reader->GetValueBool(InColumnIndex);
+			TSharedPtr<Trace::IUntypedTableReader> Reader = Table->GetTableReader();
+			if (Reader.IsValid() && RowId.HasValidIndex())
+			{
+				Reader->SetRowIndex(RowId.RowIndex);
+				Value = Reader->GetValueBool(Column.GetIndex());
+			}
+		}
+	}
+	else
+	{
+		const FTableCellValue* ValuePtr = AggregatedValues.Find(Column.GetId());
+		if (ValuePtr != nullptr)
+		{
+			Value = ValuePtr->Bool;
 		}
 	}
 
@@ -29,18 +58,29 @@ bool FTableTreeNode::GetValueBool(int32 InColumnIndex) const
 
 ////////////////////////////////////////////////////////////////////////////////////////////////////
 
-int64 FTableTreeNode::GetValueInt(int32 InColumnIndex) const
+int64 FTableTreeNode::GetValueInt64(const FTableColumn& Column) const
 {
 	int64 Value = 0;
 
-	TSharedPtr<FTable> Table = ParentTable.Pin();
-	if (Table.IsValid())
+	if (!IsGroup())
+	{
+		TSharedPtr<FTable> Table = ParentTable.Pin();
+		if (Table.IsValid())
+		{
+			TSharedPtr<Trace::IUntypedTableReader> Reader = Table->GetTableReader();
+			if (Reader.IsValid() && RowId.HasValidIndex())
+			{
+				Reader->SetRowIndex(RowId.RowIndex);
+				Value = Reader->GetValueInt(Column.GetIndex());
+			}
+		}
+	}
+	else
 	{
-		TSharedPtr<Trace::IUntypedTableReader> Reader = Table->GetTableReader();
-		if (Reader.IsValid() && RowId.HasValidIndex())
+		const FTableCellValue* ValuePtr = AggregatedValues.Find(Column.GetId());
+		if (ValuePtr != nullptr)
 		{
-			Reader->SetRowIndex(RowId.RowIndex);
-			Value = Reader->GetValueInt(InColumnIndex);
+			Value = ValuePtr->Int64;
 		}
 	}
 
@@ -49,18 +89,29 @@ int64 FTableTreeNode::GetValueInt(int32 InColumnIndex) const
 
 ////////////////////////////////////////////////////////////////////////////////////////////////////
 
-float FTableTreeNode::GetValueFloat(int32 InColumnIndex) const
+float FTableTreeNode::GetValueFloat(const FTableColumn& Column) const
 {
 	float Value = 0.0f;
 
-	TSharedPtr<FTable> Table = ParentTable.Pin();
-	if (Table.IsValid())
+	if (!IsGroup())
 	{
-		TSharedPtr<Trace::IUntypedTableReader> Reader = Table->GetTableReader();
-		if (Reader.IsValid() && RowId.HasValidIndex())
+		TSharedPtr<FTable> Table = ParentTable.Pin();
+		if (Table.IsValid())
 		{
-			Reader->SetRowIndex(RowId.RowIndex);
-			Value = Reader->GetValueFloat(InColumnIndex);
+			TSharedPtr<Trace::IUntypedTableReader> Reader = Table->GetTableReader();
+			if (Reader.IsValid() && RowId.HasValidIndex())
+			{
+				Reader->SetRowIndex(RowId.RowIndex);
+				Value = Reader->GetValueFloat(Column.GetIndex());
+			}
+		}
+	}
+	else
+	{
+		const FTableCellValue* ValuePtr = AggregatedValues.Find(Column.GetId());
+		if (ValuePtr != nullptr)
+		{
+			Value = ValuePtr->Float;
 		}
 	}
 
@@ -69,18 +120,28 @@ float FTableTreeNode::GetValueFloat(int32 InColumnIndex) const
 
 ////////////////////////////////////////////////////////////////////////////////////////////////////
 
-double FTableTreeNode::GetValueDouble(int32 InColumnIndex) const
+double FTableTreeNode::GetValueDouble(const FTableColumn& Column) const
 {
 	double Value = 0.0;
 
-	TSharedPtr<FTable> Table = ParentTable.Pin();
-	if (Table.IsValid())
+	if (!IsGroup())
+	{
+		TSharedPtr<FTable> Table = ParentTable.Pin();
+		if (Table.IsValid())
+		{
+			TSharedPtr<Trace::IUntypedTableReader> Reader = Table->GetTableReader();
+			if (Reader.IsValid() && RowId.HasValidIndex())
+			{
+				Reader->SetRowIndex(RowId.RowIndex);
+				Value = Reader->GetValueDouble(Column.GetIndex());
+			}
+		}
+	}
 	{
-		TSharedPtr<Trace::IUntypedTableReader> Reader = Table->GetTableReader();
-		if (Reader.IsValid() && RowId.HasValidIndex())
+		const FTableCellValue* ValuePtr = AggregatedValues.Find(Column.GetId());
+		if (ValuePtr != nullptr)
 		{
-			Reader->SetRowIndex(RowId.RowIndex);
-			Value = Reader->GetValueDouble(InColumnIndex);
+			Value = ValuePtr->Double;
 		}
 	}
 
@@ -89,18 +150,29 @@ double FTableTreeNode::GetValueDouble(int32 InColumnIndex) const
 
 ////////////////////////////////////////////////////////////////////////////////////////////////////
 
-const TCHAR* FTableTreeNode::GetValueCString(int32 InColumnIndex) const
+const TCHAR* FTableTreeNode::GetValueCString(const FTableColumn& Column) const
 {
 	const TCHAR* Value = nullptr;
 
-	TSharedPtr<FTable> Table = ParentTable.Pin();
-	if (Table.IsValid())
+	if (!IsGroup())
+	{
+		TSharedPtr<FTable> Table = ParentTable.Pin();
+		if (Table.IsValid())
+		{
+			TSharedPtr<Trace::IUntypedTableReader> Reader = Table->GetTableReader();
+			if (Reader.IsValid() && RowId.HasValidIndex())
+			{
+				Reader->SetRowIndex(RowId.RowIndex);
+				Value = Reader->GetValueCString(Column.GetIndex());
+			}
+		}
+	}
+	else
 	{
-		TSharedPtr<Trace::IUntypedTableReader> Reader = Table->GetTableReader();
-		if (Reader.IsValid() && RowId.HasValidIndex())
+		const FTableCellValue* ValuePtr = AggregatedValues.Find(Column.GetId());
+		if (ValuePtr != nullptr)
 		{
-			Reader->SetRowIndex(RowId.RowIndex);
-			Value = Reader->GetValueCString(InColumnIndex);
+			Value = ValuePtr->CString;
 		}
 	}
 
diff --git a/Engine/Source/Developer/TraceInsights/Private/Insights/Table/ViewModels/TableTreeNode.h b/Engine/Source/Developer/TraceInsights/Private/Insights/Table/ViewModels/TableTreeNode.h
--- a/Engine/Source/Developer/TraceInsights/Private/Insights/Table/ViewModels/TableTreeNode.h
+++ b/Engine/Source/Developer/TraceInsights/Private/Insights/Table/ViewModels/TableTreeNode.h
@@ -8,6 +8,7 @@
 // Insights
 #include "Insights/Table/ViewModels/BaseTreeNode.h"
 #include "Insights/Table/ViewModels/Table.h"
+#include "Insights/Table/ViewModels/TableColumn.h"
 
 namespace Insights
 {
@@ -56,15 +57,25 @@ public:
 	FTableRowId GetRowId() const { return RowId; }
 	int32 GetRowIndex() const { return RowId.RowIndex; }
 
-	bool GetValueBool(int32 InColumnIndex) const;
-	int64 GetValueInt(int32 InColumnIndex) const;
-	float GetValueFloat(int32 InColumnIndex) const;
-	double GetValueDouble(int32 InColumnIndex) const;
-	const TCHAR* GetValueCString(int32 InColumnIndex) const;
+	FTableCellValue GetValue(const FTableColumn& Column) const;
+
+	bool GetValueBool(const FTableColumn& Column) const;
+	int64 GetValueInt64(const FTableColumn& Column) const;
+	float GetValueFloat(const FTableColumn& Column) const;
+	double GetValueDouble(const FTableColumn& Column) const;
+	const TCHAR* GetValueCString(const FTableColumn& Column) const;
+
+	void ResetAggregatedValues() { AggregatedValues.Reset(); }
+	bool HasAggregatedValue(const FName& ColumnId) const { return AggregatedValues.Contains(ColumnId); }
+	const FTableCellValue& GetAggregatedValue(const FName& ColumnId) { return AggregatedValues.FindChecked(ColumnId); }
+	void AddAggregatedValue(const FName& ColumnId, const FTableCellValue& Value) { AggregatedValues.Add(ColumnId, Value); }
+	void SetAggregatedValue(const FName& ColumnId, const FTableCellValue& Value) { AggregatedValues[ColumnId] = Value; }
 
 protected:
 	TWeakPtr<FTable> ParentTable;
 	FTableRowId RowId;
+
+	TMap<FName, FTableCellValue> AggregatedValues;
 };
 
 ////////////////////////////////////////////////////////////////////////////////////////////////////
diff --git a/Engine/Source/Developer/TraceInsights/Private/Insights/Table/ViewModels/TreeNodeGrouping.cpp b/Engine/Source/Developer/TraceInsights/Private/Insights/Table/ViewModels/TreeNodeGrouping.cpp
--- a/Engine/Source/Developer/TraceInsights/Private/Insights/Table/ViewModels/TreeNodeGrouping.cpp
+++ b/Engine/Source/Developer/TraceInsights/Private/Insights/Table/ViewModels/TreeNodeGrouping.cpp
@@ -2,6 +2,8 @@
 
 #include "TreeNodeGrouping.h"
 
+#include "Insights/Table/ViewModels/TableTreeNode.h"
+
 #define LOCTEXT_NAMESPACE "Insights_TreeNode"
 
 namespace Insights
@@ -37,9 +39,29 @@ FTreeNodeGroupInfo FTreeNodeGroupingFlat::GetGroupForNode(const FBaseTreeNodePtr
 
 ////////////////////////////////////////////////////////////////////////////////////////////////////
 
+FTreeNodeGroupingByUniqueValue::FTreeNodeGroupingByUniqueValue(TSharedRef<FTableColumn> InColumn)
+	: FTreeNodeGrouping(
+		FText::Format(LOCTEXT("Grouping_ByUniqueValue_NameFmt", "Unique Values - {0}"), InColumn->GetTitleName()),
+		LOCTEXT("Grouping_ByUniqueValue_Desc", "Creates a group for each unique value."),
+		TEXT("Profiler.FiltersAndPresets.Group1NameIcon"), //TODO: "Icons.Grouping.ByName"
+		nullptr)
+	, Column(InColumn)
+{
+}
+
+////////////////////////////////////////////////////////////////////////////////////////////////////
+
+FTreeNodeGroupInfo FTreeNodeGroupingByUniqueValue::GetGroupForNode(const FBaseTreeNodePtr InNode) const
+{
+	FTableTreeNodePtr TableTreeNodePtr = StaticCastSharedPtr<FTableTreeNode>(InNode);
+	return { FName(*Column->GetValueAsText(TableTreeNodePtr->GetRowId()).ToString()), false };
+}
+
+////////////////////////////////////////////////////////////////////////////////////////////////////
+
 FTreeNodeGroupingByNameFirstLetter::FTreeNodeGroupingByNameFirstLetter()
 	: FTreeNodeGrouping(
-		LOCTEXT("Grouping_ByName_Name", "Name"),
+		LOCTEXT("Grouping_ByName_Name", "Name (First Letter)"),
 		LOCTEXT("Grouping_ByName_Desc", "Creates a group for each first letter of node names."),
 		TEXT("Profiler.FiltersAndPresets.Group1NameIcon"), //TODO: "Icons.Grouping.ByName"
 		nullptr)
diff --git a/Engine/Source/Developer/TraceInsights/Private/Insights/Table/ViewModels/TreeNodeGrouping.h b/Engine/Source/Developer/TraceInsights/Private/Insights/Table/ViewModels/TreeNodeGrouping.h
--- a/Engine/Source/Developer/TraceInsights/Private/Insights/Table/ViewModels/TreeNodeGrouping.h
+++ b/Engine/Source/Developer/TraceInsights/Private/Insights/Table/ViewModels/TreeNodeGrouping.h
@@ -8,6 +8,7 @@
 
 // Insights
 #include "Insights/Table/ViewModels/BaseTreeNode.h"
+#include "Insights/Table/ViewModels/TableColumn.h"
 
 namespace Insights
 {
@@ -69,6 +70,23 @@ public:
 
 ////////////////////////////////////////////////////////////////////////////////////////////////////
 
+/** Creates a group for unique column value of node. */
+class FTreeNodeGroupingByUniqueValue : public FTreeNodeGrouping
+{
+public:
+	FTreeNodeGroupingByUniqueValue(TSharedRef<FTableColumn> InColumn);
+	virtual ~FTreeNodeGroupingByUniqueValue() {}
+
+	virtual FTreeNodeGroupInfo GetGroupForNode(const FBaseTreeNodePtr InNode) const override;
+
+	TSharedRef<FTableColumn> GetColumn() const { return Column; }
+
+private:
+	TSharedRef<FTableColumn> Column;
+};
+
+////////////////////////////////////////////////////////////////////////////////////////////////////
+
 /** Creates a group for each first letter of node names. */
 class FTreeNodeGroupingByNameFirstLetter : public FTreeNodeGrouping
 {
diff --git a/Engine/Source/Developer/TraceInsights/Private/Insights/Table/ViewModels/TreeNodeSorting.cpp b/Engine/Source/Developer/TraceInsights/Private/Insights/Table/ViewModels/TreeNodeSorting.cpp
--- a/Engine/Source/Developer/TraceInsights/Private/Insights/Table/ViewModels/TreeNodeSorting.cpp
+++ b/Engine/Source/Developer/TraceInsights/Private/Insights/Table/ViewModels/TreeNodeSorting.cpp
@@ -124,14 +124,14 @@ FTreeNodeSortingByType::FTreeNodeSortingByType()
 // FTableTreeNodeSorting
 ////////////////////////////////////////////////////////////////////////////////////////////////////
 
-FTableTreeNodeSorting::FTableTreeNodeSorting(TSharedPtr<FTableColumn> InColumnPtr)
+FTableTreeNodeSorting::FTableTreeNodeSorting(TSharedRef<FTableColumn> InColumnRef)
 	: FTreeNodeSorting(
-		InColumnPtr->GetId(),
-		FText::Format(LOCTEXT("Sorting_TableValue_Name", "By {0}"), InColumnPtr->GetShortName()),
-		FText::Format(LOCTEXT("Sorting_TableValue_Title", "Sort By {0}"), InColumnPtr->GetTitleName()),
-		FText::Format(LOCTEXT("Sorting_TableValue_Desc", "Sort by {0} (ascending or descending), then by name (ascending)."), InColumnPtr->GetShortName()),
-		InColumnPtr->GetId())
-	, ColumnPtr(InColumnPtr)
+		InColumnRef->GetId(),
+		FText::Format(LOCTEXT("Sorting_TableValue_Name", "By {0}"), InColumnRef->GetShortName()),
+		FText::Format(LOCTEXT("Sorting_TableValue_Title", "Sort By {0}"), InColumnRef->GetTitleName()),
+		FText::Format(LOCTEXT("Sorting_TableValue_Desc", "Sort by {0} (ascending or descending), then by name (ascending)."), InColumnRef->GetShortName()),
+		InColumnRef->GetId())
+	, ColumnRef(InColumnRef)
 {
 }
 
@@ -139,15 +139,15 @@ FTableTreeNodeSorting::FTableTreeNodeSorting(TSharedPtr<FTableColumn> InColumnPt
 // FTreeNodeSortingByBoolValue
 ////////////////////////////////////////////////////////////////////////////////////////////////////
 
-FTreeNodeSortingByBoolValue::FTreeNodeSortingByBoolValue(TSharedPtr<FTableColumn> InColumnPtr)
-	: FTableTreeNodeSorting(InColumnPtr)
+FTreeNodeSortingByBoolValue::FTreeNodeSortingByBoolValue(TSharedRef<FTableColumn> InColumnRef)
+	: FTableTreeNodeSorting(InColumnRef)
 {
-	const int32 ColumnIndex = ColumnPtr->GetIndex();
+	const FTableColumn& Column = *InColumnRef;
 
-	AscendingCompareDelegate = [ColumnIndex](const FBaseTreeNodePtr& A, const FBaseTreeNodePtr& B) -> bool
+	AscendingCompareDelegate = [&Column](const FBaseTreeNodePtr& A, const FBaseTreeNodePtr& B) -> bool
 	{
-		const bool ValueA = StaticCastSharedPtr<FTableTreeNode>(A)->GetValueBool(ColumnIndex);
-		const bool ValueB = StaticCastSharedPtr<FTableTreeNode>(B)->GetValueBool(ColumnIndex);
+		const bool ValueA = StaticCastSharedPtr<FTableTreeNode>(A)->GetValueBool(Column);
+		const bool ValueB = StaticCastSharedPtr<FTableTreeNode>(B)->GetValueBool(Column);
 
 		if (ValueA == ValueB)
 		{
@@ -161,10 +161,10 @@ FTreeNodeSortingByBoolValue::FTreeNodeSortingByBoolValue(TSharedPtr<FTableColumn
 		}
 	};
 
-	DescendingCompareDelegate = [ColumnIndex](const FBaseTreeNodePtr& A, const FBaseTreeNodePtr& B) -> bool
+	DescendingCompareDelegate = [&Column](const FBaseTreeNodePtr& A, const FBaseTreeNodePtr& B) -> bool
 	{
-		const bool ValueA = StaticCastSharedPtr<FTableTreeNode>(A)->GetValueBool(ColumnIndex);
-		const bool ValueB = StaticCastSharedPtr<FTableTreeNode>(B)->GetValueBool(ColumnIndex);
+		const bool ValueA = StaticCastSharedPtr<FTableTreeNode>(A)->GetValueBool(Column);
+		const bool ValueB = StaticCastSharedPtr<FTableTreeNode>(B)->GetValueBool(Column);
 
 		if (ValueA == ValueB)
 		{
@@ -183,15 +183,15 @@ FTreeNodeSortingByBoolValue::FTreeNodeSortingByBoolValue(TSharedPtr<FTableColumn
 // FTreeNodeSortingByIntValue
 ////////////////////////////////////////////////////////////////////////////////////////////////////
 
-FTreeNodeSortingByIntValue::FTreeNodeSortingByIntValue(TSharedPtr<FTableColumn> InColumnPtr)
-	: FTableTreeNodeSorting(InColumnPtr)
+FTreeNodeSortingByIntValue::FTreeNodeSortingByIntValue(TSharedRef<FTableColumn> InColumnRef)
+	: FTableTreeNodeSorting(InColumnRef)
 {
-	const int32 ColumnIndex = ColumnPtr->GetIndex();
+	const FTableColumn& Column = *InColumnRef;
 
-	AscendingCompareDelegate = [ColumnIndex](const FBaseTreeNodePtr& A, const FBaseTreeNodePtr& B) -> bool
+	AscendingCompareDelegate = [&Column](const FBaseTreeNodePtr& A, const FBaseTreeNodePtr& B) -> bool
 	{
-		const int64 ValueA = StaticCastSharedPtr<FTableTreeNode>(A)->GetValueInt(ColumnIndex);
-		const int64 ValueB = StaticCastSharedPtr<FTableTreeNode>(B)->GetValueInt(ColumnIndex);
+		const int64 ValueA = StaticCastSharedPtr<FTableTreeNode>(A)->GetValueInt64(Column);
+		const int64 ValueB = StaticCastSharedPtr<FTableTreeNode>(B)->GetValueInt64(Column);
 
 		if (ValueA == ValueB)
 		{
@@ -205,10 +205,10 @@ FTreeNodeSortingByIntValue::FTreeNodeSortingByIntValue(TSharedPtr<FTableColumn>
 		}
 	};
 
-	DescendingCompareDelegate = [ColumnIndex](const FBaseTreeNodePtr& A, const FBaseTreeNodePtr& B) -> bool
+	DescendingCompareDelegate = [&Column](const FBaseTreeNodePtr& A, const FBaseTreeNodePtr& B) -> bool
 	{
-		const int64 ValueA = StaticCastSharedPtr<FTableTreeNode>(A)->GetValueInt(ColumnIndex);
-		const int64 ValueB = StaticCastSharedPtr<FTableTreeNode>(B)->GetValueInt(ColumnIndex);
+		const int64 ValueA = StaticCastSharedPtr<FTableTreeNode>(A)->GetValueInt64(Column);
+		const int64 ValueB = StaticCastSharedPtr<FTableTreeNode>(B)->GetValueInt64(Column);
 
 		if (ValueA == ValueB)
 		{
@@ -227,15 +227,15 @@ FTreeNodeSortingByIntValue::FTreeNodeSortingByIntValue(TSharedPtr<FTableColumn>
 // FTreeNodeSortingByFloatValue
 ////////////////////////////////////////////////////////////////////////////////////////////////////
 
-FTreeNodeSortingByFloatValue::FTreeNodeSortingByFloatValue(TSharedPtr<FTableColumn> InColumnPtr)
-	: FTableTreeNodeSorting(InColumnPtr)
+FTreeNodeSortingByFloatValue::FTreeNodeSortingByFloatValue(TSharedRef<FTableColumn> InColumnRef)
+	: FTableTreeNodeSorting(InColumnRef)
 {
-	const int32 ColumnIndex = ColumnPtr->GetIndex();
+	const FTableColumn& Column = *InColumnRef;
 
-	AscendingCompareDelegate = [ColumnIndex](const FBaseTreeNodePtr& A, const FBaseTreeNodePtr& B) -> bool
+	AscendingCompareDelegate = [&Column](const FBaseTreeNodePtr& A, const FBaseTreeNodePtr& B) -> bool
 	{
-		const float ValueA = StaticCastSharedPtr<FTableTreeNode>(A)->GetValueFloat(ColumnIndex);
-		const float ValueB = StaticCastSharedPtr<FTableTreeNode>(B)->GetValueFloat(ColumnIndex);
+		const float ValueA = StaticCastSharedPtr<FTableTreeNode>(A)->GetValueFloat(Column);
+		const float ValueB = StaticCastSharedPtr<FTableTreeNode>(B)->GetValueFloat(Column);
 
 		if (ValueA == ValueB)
 		{
@@ -249,10 +249,10 @@ FTreeNodeSortingByFloatValue::FTreeNodeSortingByFloatValue(TSharedPtr<FTableColu
 		}
 	};
 
-	DescendingCompareDelegate = [ColumnIndex](const FBaseTreeNodePtr& A, const FBaseTreeNodePtr& B) -> bool
+	DescendingCompareDelegate = [&Column](const FBaseTreeNodePtr& A, const FBaseTreeNodePtr& B) -> bool
 	{
-		const float ValueA = StaticCastSharedPtr<FTableTreeNode>(A)->GetValueFloat(ColumnIndex);
-		const float ValueB = StaticCastSharedPtr<FTableTreeNode>(B)->GetValueFloat(ColumnIndex);
+		const float ValueA = StaticCastSharedPtr<FTableTreeNode>(A)->GetValueFloat(Column);
+		const float ValueB = StaticCastSharedPtr<FTableTreeNode>(B)->GetValueFloat(Column);
 
 		if (ValueA == ValueB)
 		{
@@ -271,15 +271,15 @@ FTreeNodeSortingByFloatValue::FTreeNodeSortingByFloatValue(TSharedPtr<FTableColu
 // FTreeNodeSortingByDoubleValue
 ////////////////////////////////////////////////////////////////////////////////////////////////////
 
-FTreeNodeSortingByDoubleValue::FTreeNodeSortingByDoubleValue(TSharedPtr<FTableColumn> InColumnPtr)
-	: FTableTreeNodeSorting(InColumnPtr)
+FTreeNodeSortingByDoubleValue::FTreeNodeSortingByDoubleValue(TSharedRef<FTableColumn> InColumnRef)
+	: FTableTreeNodeSorting(InColumnRef)
 {
-	const int32 ColumnIndex = ColumnPtr->GetIndex();
+	const FTableColumn& Column = *InColumnRef;
 
-	AscendingCompareDelegate = [ColumnIndex](const FBaseTreeNodePtr& A, const FBaseTreeNodePtr& B) -> bool
+	AscendingCompareDelegate = [&Column](const FBaseTreeNodePtr& A, const FBaseTreeNodePtr& B) -> bool
 	{
-		const double ValueA = StaticCastSharedPtr<FTableTreeNode>(A)->GetValueDouble(ColumnIndex);
-		const double ValueB = StaticCastSharedPtr<FTableTreeNode>(B)->GetValueDouble(ColumnIndex);
+		const double ValueA = StaticCastSharedPtr<FTableTreeNode>(A)->GetValueDouble(Column);
+		const double ValueB = StaticCastSharedPtr<FTableTreeNode>(B)->GetValueDouble(Column);
 
 		if (ValueA == ValueB)
 		{
@@ -293,10 +293,10 @@ FTreeNodeSortingByDoubleValue::FTreeNodeSortingByDoubleValue(TSharedPtr<FTableCo
 		}
 	};
 
-	DescendingCompareDelegate = [ColumnIndex](const FBaseTreeNodePtr& A, const FBaseTreeNodePtr& B) -> bool
+	DescendingCompareDelegate = [&Column](const FBaseTreeNodePtr& A, const FBaseTreeNodePtr& B) -> bool
 	{
-		const double ValueA = StaticCastSharedPtr<FTableTreeNode>(A)->GetValueDouble(ColumnIndex);
-		const double ValueB = StaticCastSharedPtr<FTableTreeNode>(B)->GetValueDouble(ColumnIndex);
+		const double ValueA = StaticCastSharedPtr<FTableTreeNode>(A)->GetValueDouble(Column);
+		const double ValueB = StaticCastSharedPtr<FTableTreeNode>(B)->GetValueDouble(Column);
 
 		if (ValueA == ValueB)
 		{
@@ -315,15 +315,15 @@ FTreeNodeSortingByDoubleValue::FTreeNodeSortingByDoubleValue(TSharedPtr<FTableCo
 // FTreeNodeSortingByCStringValue
 ////////////////////////////////////////////////////////////////////////////////////////////////////
 
-FTreeNodeSortingByCStringValue::FTreeNodeSortingByCStringValue(TSharedPtr<FTableColumn> InColumnPtr)
-	: FTableTreeNodeSorting(InColumnPtr)
+FTreeNodeSortingByCStringValue::FTreeNodeSortingByCStringValue(TSharedRef<FTableColumn> InColumnRef)
+	: FTableTreeNodeSorting(InColumnRef)
 {
-	const int32 ColumnIndex = ColumnPtr->GetIndex();
+	const FTableColumn& Column = *InColumnRef;
 
-	AscendingCompareDelegate = [ColumnIndex](const FBaseTreeNodePtr& A, const FBaseTreeNodePtr& B) -> bool
+	AscendingCompareDelegate = [&Column](const FBaseTreeNodePtr& A, const FBaseTreeNodePtr& B) -> bool
 	{
-		const TCHAR* ValueA = StaticCastSharedPtr<FTableTreeNode>(A)->GetValueCString(ColumnIndex);
-		const TCHAR* ValueB = StaticCastSharedPtr<FTableTreeNode>(B)->GetValueCString(ColumnIndex);
+		const TCHAR* ValueA = StaticCastSharedPtr<FTableTreeNode>(A)->GetValueCString(Column);
+		const TCHAR* ValueB = StaticCastSharedPtr<FTableTreeNode>(B)->GetValueCString(Column);
 
 		if (ValueA == ValueB)
 		{
@@ -337,10 +337,10 @@ FTreeNodeSortingByCStringValue::FTreeNodeSortingByCStringValue(TSharedPtr<FTable
 		}
 	};
 
-	DescendingCompareDelegate = [ColumnIndex](const FBaseTreeNodePtr& A, const FBaseTreeNodePtr& B) -> bool
+	DescendingCompareDelegate = [&Column](const FBaseTreeNodePtr& A, const FBaseTreeNodePtr& B) -> bool
 	{
-		const TCHAR* ValueA = StaticCastSharedPtr<FTableTreeNode>(A)->GetValueCString(ColumnIndex);
-		const TCHAR* ValueB = StaticCastSharedPtr<FTableTreeNode>(B)->GetValueCString(ColumnIndex);
+		const TCHAR* ValueA = StaticCastSharedPtr<FTableTreeNode>(A)->GetValueCString(Column);
+		const TCHAR* ValueB = StaticCastSharedPtr<FTableTreeNode>(B)->GetValueCString(Column);
 
 		if (ValueA == ValueB)
 		{
diff --git a/Engine/Source/Developer/TraceInsights/Private/Insights/Table/ViewModels/TreeNodeSorting.h b/Engine/Source/Developer/TraceInsights/Private/Insights/Table/ViewModels/TreeNodeSorting.h
--- a/Engine/Source/Developer/TraceInsights/Private/Insights/Table/ViewModels/TreeNodeSorting.h
+++ b/Engine/Source/Developer/TraceInsights/Private/Insights/Table/ViewModels/TreeNodeSorting.h
@@ -101,11 +101,11 @@ public:
 class FTableTreeNodeSorting : public FTreeNodeSorting
 {
 public:
-	FTableTreeNodeSorting(TSharedPtr<FTableColumn> InColumnPtr);
+	FTableTreeNodeSorting(TSharedRef<FTableColumn> InColumnRef);
 	virtual ~FTableTreeNodeSorting() {}
 
 protected:
-	TSharedPtr<FTableColumn> ColumnPtr;
+	TSharedRef<FTableColumn> ColumnRef;
 };
 
 ////////////////////////////////////////////////////////////////////////////////////////////////////
@@ -113,7 +113,7 @@ protected:
 class FTreeNodeSortingByBoolValue : public FTableTreeNodeSorting
 {
 public:
-	FTreeNodeSortingByBoolValue(TSharedPtr<FTableColumn> InColumnPtr);
+	FTreeNodeSortingByBoolValue(TSharedRef<FTableColumn> InColumnRef);
 	virtual ~FTreeNodeSortingByBoolValue() {}
 };
 
@@ -122,7 +122,7 @@ public:
 class FTreeNodeSortingByIntValue : public FTableTreeNodeSorting
 {
 public:
-	FTreeNodeSortingByIntValue(TSharedPtr<FTableColumn> InColumnPtr);
+	FTreeNodeSortingByIntValue(TSharedRef<FTableColumn> InColumnRef);
 	virtual ~FTreeNodeSortingByIntValue() {}
 };
 
@@ -131,7 +131,7 @@ public:
 class FTreeNodeSortingByFloatValue : public FTableTreeNodeSorting
 {
 public:
-	FTreeNodeSortingByFloatValue(TSharedPtr<FTableColumn> InColumnPtr);
+	FTreeNodeSortingByFloatValue(TSharedRef<FTableColumn> InColumnRef);
 	virtual ~FTreeNodeSortingByFloatValue() {}
 };
 
@@ -140,7 +140,7 @@ public:
 class FTreeNodeSortingByDoubleValue : public FTableTreeNodeSorting
 {
 public:
-	FTreeNodeSortingByDoubleValue(TSharedPtr<FTableColumn> InColumnPtr);
+	FTreeNodeSortingByDoubleValue(TSharedRef<FTableColumn> InColumnRef);
 	virtual ~FTreeNodeSortingByDoubleValue() {}
 };
 
@@ -149,7 +149,7 @@ public:
 class FTreeNodeSortingByCStringValue : public FTableTreeNodeSorting
 {
 public:
-	FTreeNodeSortingByCStringValue(TSharedPtr<FTableColumn> InColumnPtr);
+	FTreeNodeSortingByCStringValue(TSharedRef<FTableColumn> InColumnRef);
 	virtual ~FTreeNodeSortingByCStringValue() {}
 };
 
diff --git a/Engine/Source/Developer/TraceInsights/Private/Insights/Table/Widgets/STableTreeView.cpp b/Engine/Source/Developer/TraceInsights/Private/Insights/Table/Widgets/STableTreeView.cpp
--- a/Engine/Source/Developer/TraceInsights/Private/Insights/Table/Widgets/STableTreeView.cpp
+++ b/Engine/Source/Developer/TraceInsights/Private/Insights/Table/Widgets/STableTreeView.cpp
@@ -463,7 +463,7 @@ void STableTreeView::TreeViewHeaderRow_CreateColumnArgs(const FTableColumn& Colu
 		.SortMode(EColumnSortMode::None)
 		.HAlignHeader(HAlign_Fill)
 		.VAlignHeader(VAlign_Fill)
-		.HeaderContentPadding(TOptional<FMargin>(2.0f))
+		.HeaderContentPadding(FMargin(2.0f))
 		.HAlignCell(HAlign_Fill)
 		.VAlignCell(VAlign_Fill)
 		.SortMode(this, &STableTreeView::GetSortModeForColumn, Column.GetId())
@@ -472,15 +472,13 @@ void STableTreeView::TreeViewHeaderRow_CreateColumnArgs(const FTableColumn& Colu
 		.FixedWidth(Column.IsFixedWidth() ? Column.GetInitialWidth() : TOptional<float>())
 		.HeaderContent()
 		[
-			SNew(SHorizontalBox)
+			SNew(SBox)
 			.ToolTip(STableTreeViewTooltip::GetColumnTooltip(Column))
-
-			+ SHorizontalBox::Slot()
 			.HAlign(Column.GetHorizontalAlignment())
 			.VAlign(VAlign_Center)
 			[
 				SNew(STextBlock)
-				.Text(Column.GetShortName())
+				.Text(this, &STableTreeView::OnGetHeaderText, Column.GetId())
 			]
 		]
 		.MenuContent()
@@ -493,6 +491,14 @@ void STableTreeView::TreeViewHeaderRow_CreateColumnArgs(const FTableColumn& Colu
 
 ////////////////////////////////////////////////////////////////////////////////////////////////////
 
+FText STableTreeView::OnGetHeaderText(const FName ColumnId) const
+{
+	FTableColumn& Column = *Table->FindColumnChecked(ColumnId);
+	return Column.GetShortName();
+}
+
+////////////////////////////////////////////////////////////////////////////////////////////////////
+
 TSharedRef<SWidget> STableTreeView::TreeViewHeaderRow_GenerateColumnMenu(const FTableColumn& Column)
 {
 	bool bIsMenuVisible = false;
@@ -871,6 +877,81 @@ void STableTreeView::CreateGroups()
 		GroupPtr->AddChildAndSetGroupPtr(NodePtr);
 	}
 
+	// Update aggregations.
+	for (const FTableTreeNodePtr& GroupPtr : GroupNodes)
+	{
+		GroupPtr->ResetAggregatedValues();
+	}
+	for (const TSharedPtr<FTableColumn> ColumnPtr : Table->GetColumns())
+	{
+		if (ColumnPtr->GetAggregation() == ETableColumnAggregation::Sum)
+		{
+			bool bValidForAggregation = false;
+
+			switch (ColumnPtr->GetDataType())
+			{
+				case ETableCellDataType::Int64:
+					bValidForAggregation = true;
+					break;
+				case ETableCellDataType::Float:
+					bValidForAggregation = true;
+					break;
+				case ETableCellDataType::Double:
+					bValidForAggregation = true;
+					break;
+			}
+
+			if (bValidForAggregation)
+			{
+				for (const FTableTreeNodePtr& GroupPtr : GroupNodes)
+				{
+					switch (ColumnPtr->GetDataType())
+					{
+						case ETableCellDataType::Int64:
+						{
+							int64 AggregatedValue = 0;
+							for (FBaseTreeNodePtr NodePtr : GroupPtr->GetChildren())
+							{
+								FTableTreeNodePtr TableNodePtr = StaticCastSharedPtr<FTableTreeNode>(NodePtr);
+								FTableCellValue Value = TableNodePtr->GetValue(*ColumnPtr);
+								AggregatedValue += Value.Int64;
+							}
+							GroupPtr->AddAggregatedValue(ColumnPtr->GetId(), FTableCellValue(AggregatedValue));
+							break;
+						}
+						case ETableCellDataType::Float:
+						{
+							float AggregatedValue = 0.0f;
+							for (FBaseTreeNodePtr NodePtr : GroupPtr->GetChildren())
+							{
+								FTableTreeNodePtr TableNodePtr = StaticCastSharedPtr<FTableTreeNode>(NodePtr);
+								FTableCellValue Value = TableNodePtr->GetValue(*ColumnPtr);
+								AggregatedValue += Value.Float;
+							}
+							GroupPtr->AddAggregatedValue(ColumnPtr->GetId(), FTableCellValue(AggregatedValue));
+							break;
+						}
+						case ETableCellDataType::Double:
+						{
+							double AggregatedValue = 0.0;
+							for (FBaseTreeNodePtr NodePtr : GroupPtr->GetChildren())
+							{
+								FTableTreeNodePtr TableNodePtr = StaticCastSharedPtr<FTableTreeNode>(NodePtr);
+								FTableCellValue Value = TableNodePtr->GetValue(*ColumnPtr);
+								AggregatedValue += Value.Double;
+							}
+							GroupPtr->AddAggregatedValue(ColumnPtr->GetId(), FTableCellValue(AggregatedValue));
+							break;
+						}
+					}
+				}
+			}
+		}
+		// TODO: else if (ColumnPtr->GetAggregation() == ETableColumnAggregation::Min)
+		//{
+		//}
+	}
+
 	// Sort groups by name.
 	GroupNodes.Sort([](const FBaseTreeNodePtr& A, const FBaseTreeNodePtr& B) { return A->GetName().LexicalLess(B->GetName()); });
 }
@@ -882,8 +963,16 @@ void STableTreeView::CreateGroupings()
 	AvailableGroupings.Reset(3);
 
 	AvailableGroupings.Add(MakeShareable(new FTreeNodeGroupingFlat()));
-	AvailableGroupings.Add(MakeShareable(new FTreeNodeGroupingByNameFirstLetter()));
-	AvailableGroupings.Add(MakeShareable(new FTreeNodeGroupingByType()));
+	//AvailableGroupings.Add(MakeShareable(new FTreeNodeGroupingByNameFirstLetter()));
+	//AvailableGroupings.Add(MakeShareable(new FTreeNodeGroupingByType()));
+
+	for (const TSharedPtr<FTableColumn> ColumnPtr : Table->GetColumns())
+	{
+		if (!ColumnPtr->IsHierarchy())
+		{
+			AvailableGroupings.Add(MakeShareable(new FTreeNodeGroupingByUniqueValue(ColumnPtr.ToSharedRef())));
+		}
+	}
 
 	CurrentGrouping = AvailableGroupings[0];
 
@@ -895,10 +984,64 @@ void STableTreeView::CreateGroupings()
 
 void STableTreeView::GroupBy_OnSelectionChanged(TSharedPtr<FTreeNodeGrouping> NewGrouping, ESelectInfo::Type SelectInfo)
 {
+	constexpr float HierarcyColumnWidth = 80.0f;
+
 	if (SelectInfo != ESelectInfo::Direct)
 	{
+		if (CurrentGrouping.IsValid() && CurrentGrouping->GetName().ToString().StartsWith("Unique Value"))
+		{
+			TSharedPtr<FTreeNodeGroupingByUniqueValue> GroupingByUniqueValue = StaticCastSharedPtr<FTreeNodeGroupingByUniqueValue>(CurrentGrouping);
+
+			const FName& ColumnId = GroupingByUniqueValue->GetColumn()->GetId();
+
+			ShowColumn(ColumnId);
+
+			SHeaderRow::FColumn& HierarcyColumn = const_cast<SHeaderRow::FColumn&>(TreeViewHeaderRow->GetColumns()[0]);
+			FTableColumn& HierarcyTableColumn = *Table->FindColumnChecked(HierarcyColumn.ColumnId);
+			const FText HierarcyColumnName(LOCTEXT("HierarchyShortName", "Hierarchy"));
+			HierarcyTableColumn.SetShortName(HierarcyColumnName);
+
+			const int32 NumColumns = TreeViewHeaderRow->GetColumns().Num();
+			for (int32 ColumnIndex = 0; ColumnIndex < NumColumns; ++ColumnIndex)
+			{
+				const SHeaderRow::FColumn& CurrentColumn = TreeViewHeaderRow->GetColumns()[ColumnIndex];
+				if (CurrentColumn.ColumnId == ColumnId)
+				{
+					HierarcyColumn.SetWidth(HierarcyColumnWidth);
+					break;
+				}
+			}
+		}
+
 		CurrentGrouping = NewGrouping;
 
+		if (CurrentGrouping.IsValid() && CurrentGrouping->GetName().ToString().StartsWith("Unique Value"))
+		{
+			TSharedPtr<FTreeNodeGroupingByUniqueValue> GroupingByUniqueValue = StaticCastSharedPtr<FTreeNodeGroupingByUniqueValue>(CurrentGrouping);
+
+			const FName& ColumnId = GroupingByUniqueValue->GetColumn()->GetId();
+
+			SHeaderRow::FColumn& HierarcyColumn = const_cast<SHeaderRow::FColumn&>(TreeViewHeaderRow->GetColumns()[0]);
+			FTableColumn& HierarcyTableColumn = *Table->FindColumnChecked(HierarcyColumn.ColumnId);
+			const FText HierarcyColumnName = FText::Format(LOCTEXT("HierarchyShortNameFmt", "Hierarchy ({0})"), GroupingByUniqueValue->GetColumn()->GetShortName());
+			HierarcyTableColumn.SetShortName(HierarcyColumnName);
+
+			const int32 NumColumns = TreeViewHeaderRow->GetColumns().Num();
+			for (int32 ColumnIndex = 0; ColumnIndex < NumColumns; ++ColumnIndex)
+			{
+				const SHeaderRow::FColumn& CurrentColumn = TreeViewHeaderRow->GetColumns()[ColumnIndex];
+				if (CurrentColumn.ColumnId == ColumnId)
+				{
+					HierarcyColumn.SetWidth(HierarcyColumnWidth + CurrentColumn.GetWidth());
+					break;
+				}
+			}
+
+			HideColumn(ColumnId);
+		}
+
+		TreeViewHeaderRow->RefreshColumns();
+
 		CreateGroups();
 		SortTreeNodes();
 		ApplyFiltering();
@@ -961,19 +1104,19 @@ void STableTreeView::CreateSortings()
 				switch (TableLayout.GetColumnType(ColumnPtr->GetIndex()))
 				{
 				case Trace::TableColumnType_Bool:
-					Sorting = MakeShareable(new FTreeNodeSortingByBoolValue(ColumnPtr));
+					Sorting = MakeShareable(new FTreeNodeSortingByBoolValue(ColumnPtr.ToSharedRef()));
 					break;
 				case Trace::TableColumnType_Int:
-					Sorting = MakeShareable(new FTreeNodeSortingByIntValue(ColumnPtr));
+					Sorting = MakeShareable(new FTreeNodeSortingByIntValue(ColumnPtr.ToSharedRef()));
 					break;
 				case Trace::TableColumnType_Float:
-					Sorting = MakeShareable(new FTreeNodeSortingByFloatValue(ColumnPtr));
+					Sorting = MakeShareable(new FTreeNodeSortingByFloatValue(ColumnPtr.ToSharedRef()));
 					break;
 				case Trace::TableColumnType_Double:
-					Sorting = MakeShareable(new FTreeNodeSortingByDoubleValue(ColumnPtr));
+					Sorting = MakeShareable(new FTreeNodeSortingByDoubleValue(ColumnPtr.ToSharedRef()));
 					break;
 				case Trace::TableColumnType_CString:
-					Sorting = MakeShareable(new FTreeNodeSortingByCStringValue(ColumnPtr));
+					Sorting = MakeShareable(new FTreeNodeSortingByCStringValue(ColumnPtr.ToSharedRef()));
 					break;
 				}
 
@@ -1003,10 +1146,11 @@ void STableTreeView::SortTreeNodes()
 	{
 		const Insights::ITreeNodeSorting* Sorter = CurrentSorting.Get();
 
-		const int32 NumGroups = GroupNodes.Num();
-
 		if (ColumnSortMode == EColumnSortMode::Type::Descending)
 		{
+			GroupNodes.Sort(Sorter->GetDescendingCompareDelegate());
+
+			const int32 NumGroups = GroupNodes.Num();
 			for (int32 GroupIndex = 0; GroupIndex < NumGroups; ++GroupIndex)
 			{
 				GroupNodes[GroupIndex]->SortChildrenDescending(Sorter);
@@ -1014,6 +1158,9 @@ void STableTreeView::SortTreeNodes()
 		}
 		else
 		{
+			GroupNodes.Sort(Sorter->GetAscendingCompareDelegate());
+
+			const int32 NumGroups = GroupNodes.Num();
 			for (int32 GroupIndex = 0; GroupIndex < NumGroups; ++GroupIndex)
 			{
 				GroupNodes[GroupIndex]->SortChildrenAscending(Sorter);
@@ -1355,13 +1502,11 @@ void STableTreeView::RebuildTree(bool bResync)
 			}
 			*/
 
-			int32 NameColumnIndex = 0;
-
 			for (int32 RowIndex = 0; RowIndex < TotalRowCount; ++RowIndex)
 			{
 				TableReader->SetRowIndex(RowIndex);
 				uint64 NodeId = static_cast<uint64>(RowIndex);
-				FName NodeName(TableReader->GetValueCString(NameColumnIndex));
+				FName NodeName(*FString::Printf(TEXT("row %d"), RowIndex));
 				FTableTreeNodePtr NodePtr = MakeShareable(new FTableTreeNode(NodeId, NodeName, Table, RowIndex));
 				TableTreeNodes.Add(NodePtr);
 				//TableTreeNodesMap.Add(Name, NodePtr);
diff --git a/Engine/Source/Developer/TraceInsights/Private/Insights/Table/Widgets/STableTreeView.h b/Engine/Source/Developer/TraceInsights/Private/Insights/Table/Widgets/STableTreeView.h
--- a/Engine/Source/Developer/TraceInsights/Private/Insights/Table/Widgets/STableTreeView.h
+++ b/Engine/Source/Developer/TraceInsights/Private/Insights/Table/Widgets/STableTreeView.h
@@ -104,6 +104,9 @@ protected:
 
 	void InitializeAndShowHeaderColumns();
 	void TreeViewHeaderRow_CreateColumnArgs(const FTableColumn& Column);
+
+	FText OnGetHeaderText(const FName ColumnId) const;
+
 	TSharedRef<SWidget> TreeViewHeaderRow_GenerateColumnMenu(const FTableColumn& Column);
 
 	////////////////////////////////////////////////////////////////////////////////////////////////////
diff --git a/Engine/Source/Developer/TraceInsights/Private/Insights/Table/Widgets/STableTreeViewCell.cpp b/Engine/Source/Developer/TraceInsights/Private/Insights/Table/Widgets/STableTreeViewCell.cpp
--- a/Engine/Source/Developer/TraceInsights/Private/Insights/Table/Widgets/STableTreeViewCell.cpp
+++ b/Engine/Source/Developer/TraceInsights/Private/Insights/Table/Widgets/STableTreeViewCell.cpp
@@ -101,7 +101,18 @@ TSharedRef<SWidget> STableTreeViewCell::GenerateWidgetForNameColumn(const FArgum
 
 TSharedRef<SWidget> STableTreeViewCell::GenerateWidgetForTableColumn(const FArguments& InArgs, const TSharedRef<ITableRow>& TableRow)
 {
-	const FText FormattedValue = ColumnPtr->GetValueAsText(TableTreeNodePtr->GetRowId());
+	FText FormattedValue = FText::GetEmpty();
+	if (TableTreeNodePtr->IsGroup())
+	{
+		if (ColumnPtr->GetAggregation() != ETableColumnAggregation::None)
+		{
+			FormattedValue = ColumnPtr->GetValueFormatter()->FormatValue(TableTreeNodePtr->GetValue(*ColumnPtr));
+		}
+	}
+	else
+	{
+		FormattedValue = ColumnPtr->GetValueAsText(TableTreeNodePtr->GetRowId());
+	}
 
 	return
 		SNew(SHorizontalBox)
diff --git a/Engine/Source/Developer/TraceInsights/Private/Insights/Table/Widgets/STableTreeViewCell.h b/Engine/Source/Developer/TraceInsights/Private/Insights/Table/Widgets/STableTreeViewCell.h
--- a/Engine/Source/Developer/TraceInsights/Private/Insights/Table/Widgets/STableTreeViewCell.h
+++ b/Engine/Source/Developer/TraceInsights/Private/Insights/Table/Widgets/STableTreeViewCell.h
@@ -117,9 +117,11 @@ protected:
 	FSlateColor GetStatsColorAndOpacity() const
 	{
 		const FLinearColor TextColor =
-			TableTreeNodePtr->IsFiltered() ?
-				FLinearColor(1.0f, 1.0f, 1.0f, 0.5f) :
-				FLinearColor(1.0f, 1.0f, 1.0f, 1.0f);
+			TableTreeNodePtr->IsGroup() ?
+				FLinearColor(0.7f, 0.7f, 0.7f, 1.0f) :
+				TableTreeNodePtr->IsFiltered() ?
+					FLinearColor(1.0f, 1.0f, 1.0f, 0.5f) :
+					FLinearColor(1.0f, 1.0f, 1.0f, 1.0f);
 		return TextColor;
 	}
 
diff --git a/Engine/Source/Developer/TraceInsights/Private/Insights/Table/Widgets/STableTreeViewTooltip.cpp b/Engine/Source/Developer/TraceInsights/Private/Insights/Table/Widgets/STableTreeViewTooltip.cpp
--- a/Engine/Source/Developer/TraceInsights/Private/Insights/Table/Widgets/STableTreeViewTooltip.cpp
+++ b/Engine/Source/Developer/TraceInsights/Private/Insights/Table/Widgets/STableTreeViewTooltip.cpp
@@ -127,6 +127,7 @@ TSharedPtr<SToolTip> STableTreeViewTooltip::GetCellTooltip(const TSharedPtr<FTab
 					// Id: [Id]
 					+SGridPanel::Slot(0, 0)
 					.Padding(2.0f)
+					//.HAlign(HAlign_Right)
 					[
 						SNew(STextBlock)
 						.Text(LOCTEXT("TT_Id", "Id:"))
@@ -134,6 +135,7 @@ TSharedPtr<SToolTip> STableTreeViewTooltip::GetCellTooltip(const TSharedPtr<FTab
 					]
 					+SGridPanel::Slot(1, 0)
 					.Padding(2.0f)
+					.HAlign(HAlign_Left)
 					[
 						SNew(STextBlock)
 						.Text(FText::AsNumber(TreeNodePtr->GetId()))
@@ -143,6 +145,7 @@ TSharedPtr<SToolTip> STableTreeViewTooltip::GetCellTooltip(const TSharedPtr<FTab
 					// Name: [Name]
 					+SGridPanel::Slot(0, 1)
 					.Padding(2.0f)
+					//.HAlign(HAlign_Right)
 					[
 						SNew(STextBlock)
 						.Text(LOCTEXT("TT_Name", "Name:"))
@@ -150,6 +153,7 @@ TSharedPtr<SToolTip> STableTreeViewTooltip::GetCellTooltip(const TSharedPtr<FTab
 					]
 					+SGridPanel::Slot(1, 1)
 					.Padding(2.0f)
+					.HAlign(HAlign_Left)
 					[
 						SNew(STextBlock)
 						.Text(FText::FromName(TreeNodePtr->GetName()))
@@ -159,6 +163,7 @@ TSharedPtr<SToolTip> STableTreeViewTooltip::GetCellTooltip(const TSharedPtr<FTab
 					// Type: [Type]
 					+ SGridPanel::Slot(0, 3)
 					.Padding(2.0f)
+					//.HAlign(HAlign_Right)
 					[
 						SNew(STextBlock)
 						.Text(LOCTEXT("TT_Type", "Type:"))
@@ -166,6 +171,7 @@ TSharedPtr<SToolTip> STableTreeViewTooltip::GetCellTooltip(const TSharedPtr<FTab
 					]
 					+ SGridPanel::Slot(1, 3)
 					.Padding(2.0f)
+					.HAlign(HAlign_Left)
 					[
 						SNew(STextBlock)
 						.Text(FText::FromName(TreeNodePtr->GetTypeId()))
@@ -223,6 +229,7 @@ void STableTreeViewTooltip::AddGridRow(TSharedPtr<SGridPanel> Grid, int32& Row,
 {
 	Grid->AddSlot(0, Row)
 		.Padding(2.0f)
+		//.HAlign(HAlign_Right)
 		[
 			SNew(STextBlock)
 			.Text(Name)
diff --git a/Engine/Source/Developer/TraceInsights/Private/Insights/Widgets/SStatsView.cpp b/Engine/Source/Developer/TraceInsights/Private/Insights/Widgets/SStatsView.cpp
--- a/Engine/Source/Developer/TraceInsights/Private/Insights/Widgets/SStatsView.cpp
+++ b/Engine/Source/Developer/TraceInsights/Private/Insights/Widgets/SStatsView.cpp
@@ -487,7 +487,7 @@ void SStatsView::TreeViewHeaderRow_CreateColumnArgs(const int32 ColumnIndex)
 		.DefaultLabel(Column.ShortName)
 		.HAlignHeader(HAlign_Fill)
 		.VAlignHeader(VAlign_Fill)
-		.HeaderContentPadding(TOptional<FMargin>(2.0f))
+		.HeaderContentPadding(FMargin(2.0f))
 		.HAlignCell(HAlign_Fill)
 		.VAlignCell(VAlign_Fill)
 		.SortMode(this, &SStatsView::GetSortModeForColumn, Column.Id)
@@ -496,10 +496,8 @@ void SStatsView::TreeViewHeaderRow_CreateColumnArgs(const int32 ColumnIndex)
 		.FixedWidth(Column.bIsFixedColumnWidth() ? Column.InitialColumnWidth : TOptional<float>())
 		.HeaderContent()
 		[
-			SNew(SHorizontalBox)
+			SNew(SBox)
 			.ToolTip(SStatsViewTooltip::GetColumnTooltip(Column))
-
-			+ SHorizontalBox::Slot()
 			.HAlign(Column.HorizontalAlignment)
 			.VAlign(VAlign_Center)
 			[
diff --git a/Engine/Source/Developer/TraceInsights/Private/Insights/Widgets/SStatsViewTooltip.cpp b/Engine/Source/Developer/TraceInsights/Private/Insights/Widgets/SStatsViewTooltip.cpp
--- a/Engine/Source/Developer/TraceInsights/Private/Insights/Widgets/SStatsViewTooltip.cpp
+++ b/Engine/Source/Developer/TraceInsights/Private/Insights/Widgets/SStatsViewTooltip.cpp
@@ -115,6 +115,7 @@ TSharedPtr<SToolTip> SStatsViewTooltip::GetTableCellTooltip(const TSharedPtr<FSt
 					// Id: [Id]
 					+SGridPanel::Slot(0, 0)
 					.Padding(2.0f)
+					//.HAlign(HAlign_Right)
 					[
 						SNew(STextBlock)
 						.Text(LOCTEXT("TT_Id", "Id:"))
@@ -122,6 +123,7 @@ TSharedPtr<SToolTip> SStatsViewTooltip::GetTableCellTooltip(const TSharedPtr<FSt
 					]
 					+SGridPanel::Slot(1, 0)
 					.Padding(2.0f)
+					.HAlign(HAlign_Left)
 					[
 						SNew(STextBlock)
 						.Text(FText::AsNumber(StatsNodePtr->GetId()))
@@ -131,6 +133,7 @@ TSharedPtr<SToolTip> SStatsViewTooltip::GetTableCellTooltip(const TSharedPtr<FSt
 					// Name: [Name]
 					+SGridPanel::Slot(0, 1)
 					.Padding(2.0f)
+					//.HAlign(HAlign_Right)
 					[
 						SNew(STextBlock)
 						.Text(LOCTEXT("TT_Name", "Name:"))
@@ -138,6 +141,7 @@ TSharedPtr<SToolTip> SStatsViewTooltip::GetTableCellTooltip(const TSharedPtr<FSt
 					]
 					+SGridPanel::Slot(1, 1)
 					.Padding(2.0f)
+					.HAlign(HAlign_Left)
 					[
 						SNew(STextBlock)
 						.Text(FText::FromName(StatsNodePtr->GetName()))
@@ -147,6 +151,7 @@ TSharedPtr<SToolTip> SStatsViewTooltip::GetTableCellTooltip(const TSharedPtr<FSt
 					//// Group: [MetaGroupName]
 					//+SGridPanel::Slot(0, 2)
 					//.Padding(2.0f)
+					////.HAlign(HAlign_Right)
 					//[
 					//	SNew(STextBlock)
 					//	.Text(LOCTEXT("TT_Group", "Group:"))
@@ -154,6 +159,7 @@ TSharedPtr<SToolTip> SStatsViewTooltip::GetTableCellTooltip(const TSharedPtr<FSt
 					//]
 					//+SGridPanel::Slot(1, 2)
 					//.Padding(2.0f)
+					//.HAlign(HAlign_Left)
 					//[
 					//	SNew(STextBlock)
 					//	.Text(FText::FromName(StatsNodePtr->GetMetaGroupName()))
@@ -163,6 +169,7 @@ TSharedPtr<SToolTip> SStatsViewTooltip::GetTableCellTooltip(const TSharedPtr<FSt
 					// Type: [Type]
 					+ SGridPanel::Slot(0, 3)
 					.Padding(2.0f)
+					//.HAlign(HAlign_Right)
 					[
 						SNew(STextBlock)
 						.Text(LOCTEXT("TT_Type", "Type:"))
@@ -170,6 +177,7 @@ TSharedPtr<SToolTip> SStatsViewTooltip::GetTableCellTooltip(const TSharedPtr<FSt
 					]
 					+ SGridPanel::Slot(1, 3)
 					.Padding(2.0f)
+					.HAlign(HAlign_Left)
 					[
 						SNew(STextBlock)
 						.Text(StatsNodeTypeHelper::ToName(StatsNodePtr->GetType()))
@@ -193,6 +201,7 @@ TSharedPtr<SToolTip> SStatsViewTooltip::GetTableCellTooltip(const TSharedPtr<FSt
 
 					+ SGridPanel::Slot(0, 0)
 					.Padding(2.0f)
+					//.HAlign(HAlign_Right)
 					[
 						SNew(STextBlock)
 						.Text(LOCTEXT("TT_NumInstances", "Num Instances:"))
@@ -200,6 +209,7 @@ TSharedPtr<SToolTip> SStatsViewTooltip::GetTableCellTooltip(const TSharedPtr<FSt
 					]
 					+ SGridPanel::Slot(1, 0)
 					.Padding(2.0f)
+					.HAlign(HAlign_Left)
 					[
 						SNew(STextBlock)
 						.Text(InstanceCountText)
@@ -501,6 +511,7 @@ void SStatsViewTooltip::AddAggregatedStatsRow(TSharedPtr<SGridPanel> Grid, int32
 {
 	Grid->AddSlot(0, Row)
 		.Padding(2.0f)
+		//.HAlign(HAlign_Right)
 		[
 			SNew(STextBlock)
 			.Text(Name)
diff --git a/Engine/Source/Developer/TraceInsights/Private/Insights/Widgets/STimersView.cpp b/Engine/Source/Developer/TraceInsights/Private/Insights/Widgets/STimersView.cpp
--- a/Engine/Source/Developer/TraceInsights/Private/Insights/Widgets/STimersView.cpp
+++ b/Engine/Source/Developer/TraceInsights/Private/Insights/Widgets/STimersView.cpp
@@ -494,7 +494,7 @@ void STimersView::TreeViewHeaderRow_CreateColumnArgs(const int32 ColumnIndex)
 		.DefaultLabel(Column.ShortName)
 		.HAlignHeader(HAlign_Fill)
 		.VAlignHeader(VAlign_Fill)
-		.HeaderContentPadding(TOptional<FMargin>(2.0f))
+		.HeaderContentPadding(FMargin(2.0f))
 		.HAlignCell(HAlign_Fill)
 		.VAlignCell(VAlign_Fill)
 		.SortMode(this, &STimersView::GetSortModeForColumn, Column.Id)
@@ -503,10 +503,8 @@ void STimersView::TreeViewHeaderRow_CreateColumnArgs(const int32 ColumnIndex)
 		.FixedWidth(Column.bIsFixedColumnWidth() ? Column.InitialColumnWidth : TOptional<float>())
 		.HeaderContent()
 		[
-			SNew(SHorizontalBox)
+			SNew(SBox)
 			.ToolTip(STimersViewTooltip::GetColumnTooltip(Column))
-
-			+ SHorizontalBox::Slot()
 			.HAlign(Column.HorizontalAlignment)
 			.VAlign(VAlign_Center)
 			[
diff --git a/Engine/Source/Developer/TraceInsights/Private/Insights/Widgets/STimingView.cpp b/Engine/Source/Developer/TraceInsights/Private/Insights/Widgets/STimingView.cpp
--- a/Engine/Source/Developer/TraceInsights/Private/Insights/Widgets/STimingView.cpp
+++ b/Engine/Source/Developer/TraceInsights/Private/Insights/Widgets/STimingView.cpp
@@ -10,6 +10,7 @@
 #include "Framework/Application/MenuStack.h"
 #include "Framework/Application/SlateApplication.h"
 #include "Framework/MultiBox/MultiBoxBuilder.h"
+#include "HAL/PlatformApplicationMisc.h"
 #include "HAL/PlatformTime.h"
 #include "Layout/WidgetPath.h"
 #include "Misc/Paths.h"
@@ -228,11 +229,6 @@ void STimingView::Reset()
 	LoadingMainThreadId = 0;
 	LoadingAsyncThreadId = 0;
 
-	EventAggregationTotalCount = 0;
-	EventAggregation.Reset();
-	ObjectTypeAggregationTotalCount = 0;
-	ObjectTypeAggregation.Reset();
-
 	//////////////////////////////////////////////////
 
 	FileActivitySharedState->Reset();
@@ -1046,176 +1042,11 @@ int32 STimingView::OnPaint(const FPaintArgs& Args, const FGeometry& AllottedGeom
 		}
 	}
 
-	if (bAssetLoadingMode)
-	{
-		constexpr float MarginX = 8.0f;
-		constexpr float MarginY = 8.0f;
-
-		const float X = ViewWidth - MarginX;
-		float Y = ViewHeight;
-
-		if (ObjectTypeAggregation.Num() > 0)
-		{
-			Y -= MarginY;
-			Y -= DrawAssetLoadingAggregationTable(DrawContext, X, Y, TEXT("Object Type Aggregation"), ObjectTypeAggregation, ObjectTypeAggregationTotalCount);
-		}
-
-		if (EventAggregation.Num() > 0)
-		{
-			Y -= MarginY;
-			Y -= DrawAssetLoadingAggregationTable(DrawContext, X, Y, TEXT("Event Aggregation"), EventAggregation, EventAggregationTotalCount);
-		}
-	}
-
 	return SCompoundWidget::OnPaint(Args, AllottedGeometry, MyCullingRect, OutDrawElements, LayerId, InWidgetStyle, bParentEnabled && IsEnabled());
 }
 
 ////////////////////////////////////////////////////////////////////////////////////////////////////
 
-float STimingView::DrawAssetLoadingAggregationTable(FDrawContext& DrawContext, float RightX, float BottomY, const TCHAR* TableName, const TArray<FAssetLoadingEventAggregationRow>& Aggregation, int32 TotalRowCount) const
-{
-	const FLinearColor BackgroundColor(0.01f, 0.01f, 0.01f, 0.9f);
-	const FLinearColor TextColorHeader(0.5f, 0.5f, 0.5f, 0.9f);
-	const FLinearColor TextColor(1.0f, 1.0f, 1.0f, 0.9f);
-
-	constexpr float BorderX = 4.0f;
-	constexpr float BorderY = 4.0f;
-
-	constexpr float LineH = 14.0f;
-
-	float TableHeight = (2 + Aggregation.Num()) * LineH;
-	if (Aggregation.Num() < TotalRowCount)
-	{
-		TableHeight += LineH; // for the "[...]" line
-	}
-
-	const TSharedRef<FSlateFontMeasure> FontMeasureService = FSlateApplication::Get().GetRenderer()->GetFontMeasureService();
-	float MaxNameWidth = 140.0f;
-	for (const FAssetLoadingEventAggregationRow& Row : Aggregation)
-	{
-		float NameWidth = FontMeasureService->Measure(Row.Name, MainFont).X;
-		if (NameWidth > MaxNameWidth)
-		{
-			MaxNameWidth = NameWidth;
-		}
-	}
-	float NameColumnWidth = MaxNameWidth + 5.0f;
-
-	constexpr float RowNumberColumnWidth = 20.0f;
-	constexpr float CountColumnWidth = 50.0f;
-	constexpr float TotalColumnWidth = 65.0f;
-	constexpr float MinColumnWidth = 60.0f;
-	constexpr float MaxColumnWidth = 60.0f;
-	constexpr float AvgColumnWidth = 60.0f;
-	constexpr float MedColumnWidth = 60.0f;
-
-	const float TableWidth = RowNumberColumnWidth
-						   + NameColumnWidth
-						   + CountColumnWidth
-						   + TotalColumnWidth
-						   + MaxColumnWidth
-						   + AvgColumnWidth
-						   + MedColumnWidth
-						   + MinColumnWidth;
-
-	float W = TableWidth + 2 * BorderX;
-	float H = TableHeight + 2 * BorderY;
-
-	float Y = BottomY - H;
-
-	DrawContext.DrawBox(RightX - W, Y, W, H, WhiteBrush, BackgroundColor);
-	DrawContext.LayerId++;
-	Y += BorderY;
-
-	const float TableX = RightX - W + BorderX;
-
-	// Table name
-	{
-		FString TableNameRow = FString::Printf(TEXT("%s (%d records, sorted by Total time)"), TableName, TotalRowCount);
-		DrawContext.DrawText(TableX, Y, TableNameRow, MainFont, TextColor);
-		Y += LineH;
-	}
-
-	// Column names
-	{
-		float X = TableX;
-
-		X += RowNumberColumnWidth;
-
-		DrawContext.DrawText(X, Y, TEXT("Name"), MainFont, TextColorHeader);
-		X += NameColumnWidth;
-
-		X += CountColumnWidth;
-		DrawContext.DrawTextAligned(HAlign_Right, X, Y, TEXT("Count"), MainFont, TextColorHeader);
-
-		X += TotalColumnWidth;
-		DrawContext.DrawTextAligned(HAlign_Right, X, Y, TEXT("Total"), MainFont, TextColorHeader);
-
-		X += MaxColumnWidth;
-		DrawContext.DrawTextAligned(HAlign_Right, X, Y, TEXT("Max [ms]"), MainFont, TextColorHeader);
-
-		X += AvgColumnWidth;
-		DrawContext.DrawTextAligned(HAlign_Right, X, Y, TEXT("Avg [ms]"), MainFont, TextColorHeader);
-
-		X += MedColumnWidth;
-		DrawContext.DrawTextAligned(HAlign_Right, X, Y, TEXT("Med [ms]"), MainFont, TextColorHeader);
-
-		X += MinColumnWidth;
-		DrawContext.DrawTextAligned(HAlign_Right, X, Y, TEXT("Min [ms]"), MainFont, TextColorHeader);
-
-		Y += LineH;
-	}
-
-	// Records
-	for (int32 Index = 0; Index < Aggregation.Num(); ++Index)
-	{
-		constexpr int32 NumDigits = 2;
-		constexpr bool bAddTimeUnit = false;
-
-		const FAssetLoadingEventAggregationRow& Row = Aggregation[Index];
-
-		float X = TableX;
-
-		X += RowNumberColumnWidth;
-		DrawContext.DrawTextAligned(HAlign_Right, X - 4.0f, Y, FString::Printf(TEXT("%d."), Index + 1), MainFont, TextColorHeader);
-
-		DrawContext.DrawText(X, Y, Row.Name, MainFont, TextColor);
-		X += NameColumnWidth;
-
-		X += CountColumnWidth;
-		DrawContext.DrawTextAligned(HAlign_Right, X, Y, FText::AsNumber(Row.Count).ToString(), MainFont, TextColor);
-
-		X += TotalColumnWidth;
-		DrawContext.DrawTextAligned(HAlign_Right, X, Y, TimeUtils::FormatTimeAuto(Row.Total), MainFont, TextColor);
-
-		X += MaxColumnWidth;
-		DrawContext.DrawTextAligned(HAlign_Right, X, Y, TimeUtils::FormatTimeMs(Row.Max, NumDigits, bAddTimeUnit), MainFont, TextColor);
-
-		X += AvgColumnWidth;
-		DrawContext.DrawTextAligned(HAlign_Right, X, Y, TimeUtils::FormatTimeMs(Row.Avg, NumDigits, bAddTimeUnit), MainFont, TextColor);
-
-		X += MedColumnWidth;
-		DrawContext.DrawTextAligned(HAlign_Right, X, Y, TimeUtils::FormatTimeMs(Row.Med, NumDigits, bAddTimeUnit), MainFont, TextColor);
-
-		X += MinColumnWidth;
-		DrawContext.DrawTextAligned(HAlign_Right, X, Y, TimeUtils::FormatTimeMs(Row.Min, NumDigits, bAddTimeUnit), MainFont, TextColor);
-
-		Y += LineH;
-	}
-
-	if (Aggregation.Num() < TotalRowCount)
-	{
-		DrawContext.DrawText(TableX, Y, TEXT("[...]"), MainFont, TextColorHeader);
-		Y += LineH;
-	}
-
-	DrawContext.LayerId++;
-
-	return H;
-}
-
-////////////////////////////////////////////////////////////////////////////////////////////////////
-
 void STimingView::AddTimingEventsTrack(FTimingEventsTrack* Track)
 {
 	UE_LOG(TimingProfiler, Log, TEXT("New Timing Events Track (%d) : %s - %s (\"%s\")"),
@@ -1335,8 +1166,6 @@ FReply STimingView::OnMouseButtonDown(const FGeometry& MyGeometry, const FPointe
 		SelectionEndTime = SelectionStartTime;
 		LastSelectionType = ESelectionType::None;
 		//TODO: SelectionChangingEvent.Broadcast(SelectionStartTime, SelectionEndTime);
-		EventAggregation.Reset();
-		ObjectTypeAggregation.Reset();
 	}
 
 	return Reply;
@@ -1388,8 +1217,6 @@ FReply STimingView::OnMouseButtonUp(const FGeometry& MyGeometry, const FPointerE
 					SelectionEndTime = SelectionStartTime = 0.0;
 					LastSelectionType = ESelectionType::None;
 					//TODO: SelectionChangedEvent.Broadcast(SelectionStartTime, SelectionEndTime);
-					EventAggregation.Reset();
-					ObjectTypeAggregation.Reset();
 				}
 			}
 
@@ -1509,8 +1336,6 @@ FReply STimingView::OnMouseMove(const FGeometry& MyGeometry, const FPointerEvent
 				}
 				LastSelectionType = ESelectionType::TimeRange;
 				//TODO: SelectionChangingEvent.Broadcast(SelectionStartTime, SelectionEndTime);
-				EventAggregation.Reset();
-				ObjectTypeAggregation.Reset();
 			}
 		}
 		else
@@ -1761,9 +1586,24 @@ FReply STimingView::OnKeyDown(const FGeometry& MyGeometry, const FKeyEvent& InKe
 	}
 	else if (InKeyEvent.GetKey() == EKeys::C)
 	{
-		Layout.bIsCompactMode = !Layout.bIsCompactMode;
-		bAreTimingEventsTracksDirty = true;
-		return FReply::Handled();
+		if (InKeyEvent.GetModifierKeys().IsControlDown())
+		{
+			if (SelectedTimingEvent.IsValid())
+			{
+				const FTimerNodePtr TimerNodePtr = FTimingProfilerManager::Get()->GetTimerNode(SelectedTimingEvent.TypeId);
+				if (TimerNodePtr.IsValid())
+				{
+					// Copy name of selected timing event to clipboard.
+					FPlatformApplicationMisc::ClipboardCopy(*TimerNodePtr->GetName().ToString());
+				}
+			}
+		}
+		else
+		{
+			Layout.bIsCompactMode = !Layout.bIsCompactMode;
+			bAreTimingEventsTracksDirty = true;
+			return FReply::Handled();
+		}
 	}
 	else if (InKeyEvent.GetKey() == EKeys::V)
 	{
@@ -2200,94 +2040,21 @@ void STimingView::UpdateAggregatedStats()
 	{
 		TSharedPtr<Insights::STableTreeView> EventAggregationTreeView;
 		TSharedPtr<Insights::STableTreeView> ObjectTypeAggregationTreeView;
+		TSharedPtr<Insights::STableTreeView> PackageDetailsTreeView;
+		TSharedPtr<Insights::STableTreeView> ExportDetailsTreeView;
 		TSharedPtr<SLoadingProfilerWindow> LoadingProfilerWnd = FLoadingProfilerManager::Get()->GetProfilerWindow();
 		if (LoadingProfilerWnd)
 		{
 			EventAggregationTreeView = LoadingProfilerWnd->GetEventAggregationTreeView();
 			ObjectTypeAggregationTreeView = LoadingProfilerWnd->GetObjectTypeAggregationTreeView();
+			PackageDetailsTreeView = LoadingProfilerWnd->GetPackageDetailsTreeView();
+			ExportDetailsTreeView = LoadingProfilerWnd->GetExportDetailsTreeView();
 		}
 
 		TSharedPtr<const Trace::IAnalysisSession> Session = FInsightsManager::Get()->GetSession();
 		if (Session.IsValid() && Trace::ReadLoadTimeProfilerProvider(*Session.Get()))
 		{
-			auto ReadTable = [](Trace::ITable<Trace::FLoadTimeProfilerAggregatedStats>* Table, TArray<FAssetLoadingEventAggregationRow>& Aggregation, int32& TotalRowCount)
-			{
-				TotalRowCount = Table->GetRowCount();
-
-				constexpr int32 MaxRowCount = 10; // only get first 10 records
-				const int32 RowCount = FMath::Min(TotalRowCount, MaxRowCount);
-
-				const Trace::ITableLayout& TableLayout = Table->GetLayout();
-				const int32 ColumnCount = TableLayout.GetColumnCount();
-
-				Trace::ITableReader<Trace::FLoadTimeProfilerAggregatedStats>* Reader = Table->CreateReader();
-
-				//////////////////////////////////////////////////
-
-				struct FSortedIndexEntry
-				{
-					int32 RowIndex;
-					double Value;
-				};
-
-				TArray<FSortedIndexEntry> SortedIndex;
-				SortedIndex.Reserve(TotalRowCount);
-
-				constexpr int32 TotalTimeColumnIndex = 2;
-				ensure(TableLayout.GetColumnType(TotalTimeColumnIndex) == Trace::ETableColumnType::TableColumnType_Double);
-
-				for (int32 RowIndex = 0; RowIndex < TotalRowCount; ++RowIndex)
-				{
-					Reader->SetRowIndex(RowIndex);
-					double Value = Reader->GetValueDouble(TotalTimeColumnIndex);
-					SortedIndex.Add({ RowIndex, Value });
-				}
-
-				SortedIndex.Sort([](const FSortedIndexEntry& A, const FSortedIndexEntry& B) { return A.Value > B.Value; });
-
-				//////////////////////////////////////////////////
-
-				Aggregation.Reset();
-				Aggregation.AddDefaulted(RowCount);
-
-				for (int32 Index = 0; Index < RowCount; ++Index)
-				{
-					FAssetLoadingEventAggregationRow& Row = Aggregation[Index];
-
-					int32 RowIndex = SortedIndex[Index].RowIndex;
-					Reader->SetRowIndex(RowIndex);
-
-					ensure(ColumnCount == 7);
-					for (int32 ColumnIndex = 0; ColumnIndex < ColumnCount; ++ColumnIndex)
-					{
-						ensure(TableLayout.GetColumnType(0) == Trace::ETableColumnType::TableColumnType_CString);
-						Row.Name = Reader->GetValueCString(0);
-
-						ensure(TableLayout.GetColumnType(1) == Trace::ETableColumnType::TableColumnType_Int);
-						Row.Count = Reader->GetValueInt(1);
-
-						ensure(TableLayout.GetColumnType(2) == Trace::ETableColumnType::TableColumnType_Double);
-						Row.Total = Reader->GetValueDouble(2);
-
-						ensure(TableLayout.GetColumnType(3) == Trace::ETableColumnType::TableColumnType_Double);
-						Row.Min = Reader->GetValueDouble(3);
-
-						ensure(TableLayout.GetColumnType(4) == Trace::ETableColumnType::TableColumnType_Double);
-						Row.Max = Reader->GetValueDouble(4);
-
-						ensure(TableLayout.GetColumnType(5) == Trace::ETableColumnType::TableColumnType_Double);
-						Row.Avg = Reader->GetValueDouble(5);
-
-						ensure(TableLayout.GetColumnType(6) == Trace::ETableColumnType::TableColumnType_Double);
-						Row.Med = Reader->GetValueDouble(6);
-					}
-				}
-
-				delete Reader;
-			};
-
 			Trace::FAnalysisSessionReadScope SessionReadScope(*Session.Get());
-
 			const Trace::ILoadTimeProfilerProvider& LoadTimeProfilerProvider = *Trace::ReadLoadTimeProfilerProvider(*Session.Get());
 
 			{
@@ -2295,18 +2062,32 @@ void STimingView::UpdateAggregatedStats()
 				Trace::ITable<Trace::FLoadTimeProfilerAggregatedStats>* EventAggregationTable = LoadTimeProfilerProvider.CreateEventAggregation(SelectionStartTime, SelectionEndTime);
 				EventAggregationTreeView->GetTable()->UpdateSourceTable(MakeShareable(EventAggregationTable));
 				EventAggregationTreeView->RebuildTree(true);
-				ReadTable(EventAggregationTable, EventAggregation, EventAggregationTotalCount);
 				//delete EventAggregationTable;
 			}
 
 			{
-				//TODO: LoadTimeProfilerProvider.UpdateObjectTypeAggregation(EventAggregationTreeView->GetTable(), SelectionStartTime, SelectionEndTime)
+				//TODO: LoadTimeProfilerProvider.UpdateObjectTypeAggregation(ObjectTypeAggregationTreeView->GetTable(), SelectionStartTime, SelectionEndTime)
 				Trace::ITable<Trace::FLoadTimeProfilerAggregatedStats>* ObjectTypeAggregationTable = LoadTimeProfilerProvider.CreateObjectTypeAggregation(SelectionStartTime, SelectionEndTime);
 				ObjectTypeAggregationTreeView->GetTable()->UpdateSourceTable(MakeShareable(ObjectTypeAggregationTable));
 				ObjectTypeAggregationTreeView->RebuildTree(true);
-				ReadTable(ObjectTypeAggregationTable, ObjectTypeAggregation, ObjectTypeAggregationTotalCount);
 				//delete ObjectTypeAggregationTable;
 			}
+
+			{
+				//TODO: LoadTimeProfilerProvider.UpdatePackageDetails(PackageDetailsTreeView->GetTable(), SelectionStartTime, SelectionEndTime)
+				Trace::ITable<Trace::FPackagesTableRow>* PackageDetailsTable = LoadTimeProfilerProvider.CreatePackageDetailsTable(SelectionStartTime, SelectionEndTime);
+				PackageDetailsTreeView->GetTable()->UpdateSourceTable(MakeShareable(PackageDetailsTable));
+				PackageDetailsTreeView->RebuildTree(true);
+				//delete PackageDetailsTable;
+			}
+
+			{
+				//TODO: LoadTimeProfilerProvider.UpdateExportDetails(ExportDetailsTreeView->GetTable(), SelectionStartTime, SelectionEndTime)
+				Trace::ITable<Trace::FExportsTableRow>* ExportDetailsTable = LoadTimeProfilerProvider.CreateExportDetailsTable(SelectionStartTime, SelectionEndTime);
+				ExportDetailsTreeView->GetTable()->UpdateSourceTable(MakeShareable(ExportDetailsTable));
+				ExportDetailsTreeView->RebuildTree(true);
+				//delete ExportDetailsTable;
+			}
 		}
 	}
 }
diff --git a/Engine/Source/Developer/TraceInsights/Private/Insights/Widgets/STimingView.h b/Engine/Source/Developer/TraceInsights/Private/Insights/Widgets/STimingView.h
--- a/Engine/Source/Developer/TraceInsights/Private/Insights/Widgets/STimingView.h
+++ b/Engine/Source/Developer/TraceInsights/Private/Insights/Widgets/STimingView.h
@@ -282,8 +282,6 @@ protected:
 
 	void AddTimingEventsTrack(FTimingEventsTrack* Track);
 
-	void DrawTimeRangeSelection(FDrawContext& DrawContext) const;
-
 	void ShowContextMenu(const FPointerEvent& MouseEvent);
 
 	/** Binds our UI commands to delegates. */
@@ -309,8 +307,6 @@ protected:
 
 	////////////////////////////////////////////////////////////////////////////////////////////////////
 
-	float DrawAssetLoadingAggregationTable(FDrawContext& DrawContext, float RightX, float BottomY, const TCHAR* TableName, const TArray<FAssetLoadingEventAggregationRow>& Aggregation, int32 TotalRowCount) const;
-
 	void UpdateAggregatedStats();
 
 	void UpdateHoveredTimingEvent(float MX, float MY);
@@ -370,14 +366,14 @@ protected:
 	  */
 	TMap<uint64, FTimingEventsTrack*> CachedTimelines;
 
-	//TODO: TArray<FBaseTimingTrack*> TopTracks; /**< tracks docked on top, in order to be displayed (top to bottom) */
-	//TODO: TArray<FBaseTimingTrack*> ScrollableTracks; /**< tracks in scrollable area, in order to be displayed (top to bottom) */
-	//TODO: TArray<FBaseTimingTrack*> BottomTracks; /**< tracks docked on bottom, in order to be displayed (top to bottom) */
-	//TODO: TArray<FBaseTimingTrack*> ForegorundTracks; /**< tracks to draw over top/scrollable/bottom tracks (can use entire area), in order to be displayed (back to front) */
+	//TODO: TArray<FBaseTimingTrack*> TopTracks; /**< tracks docked on top, in the order to be displayed (top to bottom) */
+	//TODO: TArray<FBaseTimingTrack*> ScrollableTracks; /**< tracks in scrollable area, in the order to be displayed (top to bottom) */
+	//TODO: TArray<FBaseTimingTrack*> BottomTracks; /**< tracks docked on bottom, in the order to be displayed (top to bottom) */
+	//TODO: TArray<FBaseTimingTrack*> ForegorundTracks; /**< tracks to draw over top/scrollable/bottom tracks (can use entire area), in the order to be displayed (back to front) */
 
 	////////////////////////////////////////////////////////////
 
-	TArray<FTimingEventsTrack*> TimingEventsTracks; /**< all timing events tracks in order to be displayed */
+	TArray<FTimingEventsTrack*> TimingEventsTracks; /**< all timing events tracks, in the order to be displayed */
 
 	bool bAreTimingEventsTracksDirty;
 
@@ -411,12 +407,6 @@ protected:
 
 	bool bAssetLoadingMode;
 
-	int32 EventAggregationTotalCount;
-	TArray<FAssetLoadingEventAggregationRow> EventAggregation;
-
-	int32 ObjectTypeAggregationTotalCount;
-	TArray<FAssetLoadingEventAggregationRow> ObjectTypeAggregation;
-
 	////////////////////////////////////////////////////////////
 	// File activity (I/O)
 
diff --git a/Engine/Source/Developer/TraceInsights/Public/Insights/Version.h b/Engine/Source/Developer/TraceInsights/Public/Insights/Version.h
--- a/Engine/Source/Developer/TraceInsights/Public/Insights/Version.h
+++ b/Engine/Source/Developer/TraceInsights/Public/Insights/Version.h
@@ -5,13 +5,13 @@
 // See https://semver.org/
 
 #define UNREAL_INSIGHTS_VERSION_MAJOR 0
-#define UNREAL_INSIGHTS_VERSION_MINOR 9
+#define UNREAL_INSIGHTS_VERSION_MINOR 10
 #define UNREAL_INSIGHTS_VERSION_PATCH 0
 
-#define UNREAL_INSIGHTS_VERSION_STRING "0.09"
+#define UNREAL_INSIGHTS_VERSION_STRING "0.10"
 
 /** Extra identifier string (like: "alpha", "beta.1", "beta.2", "test", "featureX", etc.). */
-//#define UNREAL_INSIGHTS_VERSION_ID_STRING "alpha"
+//#define UNREAL_INSIGHTS_VERSION_ID_STRING "beta"
 
 #ifdef UNREAL_INSIGHTS_VERSION_ID_STRING
 #define UNREAL_INSIGHTS_VERSION_STRING_EX "v" UNREAL_INSIGHTS_VERSION_STRING "-" UNREAL_INSIGHTS_VERSION_ID_STRING
