[Insights]
  - Networking: Added support for multiple Networking Insights tabs.
  - Networking: Renamed "Packet Sizes" to "Packets", "Packet Breakdown" to "Packet Content". Also renamed actual classes for these two views.
  - Networking: Added dummy "Net Stats" tab view (also removed "Data Stream Breakdown").
  - Packets view: Added combo-boxes to choose Game Instance, Connection and Connection Mode.
  - Packets view: Added "Sequence Number", "Content Size (bits)", "Total Size (bytes)" and "unused bits" in tooltip of hovered packet.
  - Packets view: Added support to use actual data from analyzer.
  - Packet Content view: Added actual event name and the Net Id (ObjectInstanceIndex) in text displayed on each event.
  - Packet Content view: Added support to use actual data from analyzer.

diff --git a/Engine/Source/Developer/TraceInsights/Private/Insights/InsightsManager.cpp b/Engine/Source/Developer/TraceInsights/Private/Insights/InsightsManager.cpp
--- a/Engine/Source/Developer/TraceInsights/Private/Insights/InsightsManager.cpp
+++ b/Engine/Source/Developer/TraceInsights/Private/Insights/InsightsManager.cpp
@@ -180,7 +180,15 @@ void FInsightsManager::SpawnAndActivateTabs()
 	}
 	if (FGlobalTabmanager::Get()->HasTabSpawner(FInsightsManagerTabs::NetworkingProfilerTabId))
 	{
-		//FGlobalTabmanager::Get()->InvokeTab(FInsightsManagerTabs::NetworkingProfilerTabId);
+		// Close previous tabs.
+		TSharedPtr<SDockTab> NetworkingProfilerTab;
+		while ((NetworkingProfilerTab = FGlobalTabmanager::Get()->FindExistingLiveTab(FInsightsManagerTabs::NetworkingProfilerTabId)).IsValid())
+		{
+			NetworkingProfilerTab->RequestCloseTab();
+		}
+
+		FGlobalTabmanager::Get()->InvokeTab(FInsightsManagerTabs::NetworkingProfilerTabId);
+		FGlobalTabmanager::Get()->InvokeTab(FInsightsManagerTabs::NetworkingProfilerTabId);
 	}
 
 	// Ensure Timing Insights / Timing View is the active tab / view.
diff --git a/Engine/Source/Developer/TraceInsights/Private/Insights/InsightsStyle.cpp b/Engine/Source/Developer/TraceInsights/Private/Insights/InsightsStyle.cpp
--- a/Engine/Source/Developer/TraceInsights/Private/Insights/InsightsStyle.cpp
+++ b/Engine/Source/Developer/TraceInsights/Private/Insights/InsightsStyle.cpp
@@ -103,14 +103,14 @@ TSharedRef<FSlateStyleSet> FInsightsStyle::Create()
 	//////////////////////////////////////////////////
 	// Networking Insights
 
-	Style.Set("PacketSizesView.Icon.Large", new IMAGE_BRUSH("Icons/Profiler/profiler_stats_40x", FVector2D(32.0f, 32.0f)));
-	Style.Set("PacketSizesView.Icon.Small", new IMAGE_BRUSH("Icons/Profiler/profiler_stats_40x", FVector2D(16.0f, 16.0f)));
+	Style.Set("PacketOveriew.Icon.Large", new IMAGE_BRUSH("Icons/Profiler/profiler_stats_40x", FVector2D(32.0f, 32.0f)));
+	Style.Set("PacketOveriew.Icon.Small", new IMAGE_BRUSH("Icons/Profiler/profiler_stats_40x", FVector2D(16.0f, 16.0f)));
 
-	Style.Set("PacketBreakdownView.Icon.Large", new IMAGE_BRUSH("Icons/Profiler/profiler_stats_40x", FVector2D(32.0f, 32.0f)));
-	Style.Set("PacketBreakdownView.Icon.Small", new IMAGE_BRUSH("Icons/Profiler/profiler_stats_40x", FVector2D(16.0f, 16.0f)));
+	Style.Set("PacketContentView.Icon.Large", new IMAGE_BRUSH("Icons/Profiler/profiler_stats_40x", FVector2D(32.0f, 32.0f)));
+	Style.Set("PacketContentView.Icon.Small", new IMAGE_BRUSH("Icons/Profiler/profiler_stats_40x", FVector2D(16.0f, 16.0f)));
 
-	Style.Set("DataStreamBreakdownView.Icon.Large", new IMAGE_BRUSH("Icons/Profiler/profiler_stats_40x", FVector2D(32.0f, 32.0f)));
-	Style.Set("DataStreamBreakdownView.Icon.Small", new IMAGE_BRUSH("Icons/Profiler/profiler_stats_40x", FVector2D(16.0f, 16.0f)));
+	Style.Set("NetStatsView.Icon.Large", new IMAGE_BRUSH("Icons/Profiler/profiler_stats_40x", FVector2D(32.0f, 32.0f)));
+	Style.Set("NetStatsView.Icon.Small", new IMAGE_BRUSH("Icons/Profiler/profiler_stats_40x", FVector2D(16.0f, 16.0f)));
 
 	//////////////////////////////////////////////////
 
diff --git a/Engine/Source/Developer/TraceInsights/Private/Insights/NetworkingProfiler/NetworkingProfilerCommands.cpp b/Engine/Source/Developer/TraceInsights/Private/Insights/NetworkingProfiler/NetworkingProfilerCommands.cpp
--- a/Engine/Source/Developer/TraceInsights/Private/Insights/NetworkingProfiler/NetworkingProfilerCommands.cpp
+++ b/Engine/Source/Developer/TraceInsights/Private/Insights/NetworkingProfiler/NetworkingProfilerCommands.cpp
@@ -52,9 +52,9 @@ FNetworkingProfilerCommands::FNetworkingProfilerCommands()
 PRAGMA_DISABLE_OPTIMIZATION
 void FNetworkingProfilerCommands::RegisterCommands()
 {
-	UI_COMMAND(TogglePacketSizesViewVisibility, "Packet Sizes", "Toggles the visibility of the Packet Sizes view", EUserInterfaceActionType::ToggleButton, FInputChord());
-	UI_COMMAND(TogglePacketBreakdownViewVisibility, "Packet Breakdown", "Toggles the visibility of the Packet Breakdown view", EUserInterfaceActionType::ToggleButton, FInputChord());
-	UI_COMMAND(ToggleDataStreamBreakdownViewVisibility, "Data Stream Breakdown", "Toggles the visibility of the Data Stream Breakdown view", EUserInterfaceActionType::ToggleButton, FInputChord());
+	UI_COMMAND(TogglePacketViewVisibility, "Packets", "Toggles the visibility of the Packets view", EUserInterfaceActionType::ToggleButton, FInputChord());
+	UI_COMMAND(TogglePacketContentViewVisibility, "Packet Content", "Toggles the visibility of the Packet Content view", EUserInterfaceActionType::ToggleButton, FInputChord());
+	UI_COMMAND(ToggleNetStatsViewVisibility, "Net Stats", "Toggles the visibility of the Net Stats view", EUserInterfaceActionType::ToggleButton, FInputChord());
 }
 PRAGMA_ENABLE_OPTIMIZATION
 
@@ -95,9 +95,7 @@ PRAGMA_ENABLE_OPTIMIZATION
 		return b##IsEnabled ? ECheckBoxState::Checked : ECheckBoxState::Unchecked;\
 	}
 
-IMPLEMENT_TOGGLE_COMMAND(TogglePacketSizesViewVisibility, IsPacketSizesViewVisible, ShowHidePacketSizesView)
-IMPLEMENT_TOGGLE_COMMAND(TogglePacketBreakdownViewVisibility, IsPacketBreakdownViewVisible, ShowHidePacketBreakdownView)
-IMPLEMENT_TOGGLE_COMMAND(ToggleDataStreamBreakdownViewVisibility, IsDataStreamBreakdownViewVisible, ShowHideDataStreamBreakdownView)
+//IMPLEMENT_TOGGLE_COMMAND(ToggleAAAViewVisibility, IsAAAViewVisible, ShowHideAAAView)
 
 #undef IMPLEMENT_TOGGLE_COMMAND
 
diff --git a/Engine/Source/Developer/TraceInsights/Private/Insights/NetworkingProfiler/NetworkingProfilerCommands.h b/Engine/Source/Developer/TraceInsights/Private/Insights/NetworkingProfiler/NetworkingProfilerCommands.h
--- a/Engine/Source/Developer/TraceInsights/Private/Insights/NetworkingProfiler/NetworkingProfilerCommands.h
+++ b/Engine/Source/Developer/TraceInsights/Private/Insights/NetworkingProfiler/NetworkingProfilerCommands.h
@@ -29,14 +29,14 @@ public:
 	//     const FUIAction <CommandName>_Custom(...) const;
 	//////////////////////////////////////////////////
 
-	/** Toggles visibility for the Packet Sizes view. Global and custom command. */
-	TSharedPtr<FUICommandInfo> TogglePacketSizesViewVisibility;
+	/** Toggles visibility for the Packet view. Per profiler window, custom command. */
+	TSharedPtr<FUICommandInfo> TogglePacketViewVisibility;
 
-	/** Toggles visibility for the Packet Breakdown view. Global and custom command. */
-	TSharedPtr<FUICommandInfo> TogglePacketBreakdownViewVisibility;
+	/** Toggles visibility for the Packet Content view. Per profiler window, custom command. */
+	TSharedPtr<FUICommandInfo> TogglePacketContentViewVisibility;
 
-	/** Toggles visibility for the Data Stream Breakdown view. Global and custom command. */
-	TSharedPtr<FUICommandInfo> ToggleDataStreamBreakdownViewVisibility;
+	/** Toggles visibility for the Net Stats view. Per profiler window, custom command. */
+	TSharedPtr<FUICommandInfo> ToggleNetStatsViewVisibility;
 };
 
 ////////////////////////////////////////////////////////////////////////////////////////////////////
@@ -88,9 +88,7 @@ protected:\
 	bool CmdName##_CanExecute() const; /**< Handles FCanExecuteAction for CmdName. */\
 	ECheckBoxState CmdName##_GetCheckState() const; /**< Handles FGetActionCheckState for CmdName. */
 
-	DECLARE_TOGGLE_COMMAND(TogglePacketSizesViewVisibility)
-	DECLARE_TOGGLE_COMMAND(TogglePacketBreakdownViewVisibility)
-	DECLARE_TOGGLE_COMMAND(ToggleDataStreamBreakdownViewVisibility)
+	//DECLARE_TOGGLE_COMMAND(ToggleAAAViewVisibility)
 #undef DECLARE_TOGGLE_COMMAND
 
 	//////////////////////////////////////////////////
diff --git a/Engine/Source/Developer/TraceInsights/Private/Insights/NetworkingProfiler/NetworkingProfilerManager.cpp b/Engine/Source/Developer/TraceInsights/Private/Insights/NetworkingProfiler/NetworkingProfilerManager.cpp
--- a/Engine/Source/Developer/TraceInsights/Private/Insights/NetworkingProfiler/NetworkingProfilerManager.cpp
+++ b/Engine/Source/Developer/TraceInsights/Private/Insights/NetworkingProfiler/NetworkingProfilerManager.cpp
@@ -8,10 +8,7 @@
 
 // Insights
 #include "Insights/InsightsManager.h"
-#include "Insights/NetworkingProfiler/Widgets/SDataStreamBreakdownView.h"
 #include "Insights/NetworkingProfiler/Widgets/SNetworkingProfilerWindow.h"
-#include "Insights/NetworkingProfiler/Widgets/SPacketBreakdownView.h"
-#include "Insights/NetworkingProfiler/Widgets/SPacketSizesView.h"
 
 ////////////////////////////////////////////////////////////////////////////////////////////////////
 
@@ -31,10 +28,7 @@ TSharedPtr<FNetworkingProfilerManager> FNetworkingProfilerManager::Instance = nu
 FNetworkingProfilerManager::FNetworkingProfilerManager(TSharedRef<FUICommandList> InCommandList)
 	: CommandList(InCommandList)
 	, ActionManager(this)
-	, ProfilerWindow(nullptr)
-	, bIsPacketSizesViewVisible(true)
-	, bIsPacketBreakdownViewVisible(true)
-	, bIsDataStreamBreakdownViewVisible(true)
+	, ProfilerWindows()
 {
 }
 
@@ -54,9 +48,6 @@ void FNetworkingProfilerManager::PostConstructor()
 
 void FNetworkingProfilerManager::BindCommands()
 {
-	ActionManager.Map_TogglePacketSizesViewVisibility_Global();
-	ActionManager.Map_TogglePacketBreakdownViewVisibility_Global();
-	ActionManager.Map_ToggleDataStreamBreakdownViewVisibility_Global();
 }
 
 ////////////////////////////////////////////////////////////////////////////////////////////////////
@@ -110,55 +101,13 @@ bool FNetworkingProfilerManager::Tick(float DeltaTime)
 
 void FNetworkingProfilerManager::OnSessionChanged()
 {
-	TSharedPtr<SNetworkingProfilerWindow> Wnd = GetProfilerWindow();
-	if (Wnd.IsValid())
+	for (TWeakPtr<SNetworkingProfilerWindow> WndWeakPtr : ProfilerWindows)
 	{
-		Wnd->Reset();
-	}
-}
-
-////////////////////////////////////////////////////////////////////////////////////////////////////
-
-void FNetworkingProfilerManager::ShowHidePacketSizesView(const bool bPacketSizesViewVisibleState)
-{
-	//bool bWasPacketSizesViewVisible = bIsPacketSizesViewVisible;
-
-	bIsPacketSizesViewVisible = bPacketSizesViewVisibleState;
-
-	TSharedPtr<SNetworkingProfilerWindow> Wnd = GetProfilerWindow();
-	if (Wnd.IsValid())
-	{
-		Wnd->ShowHideTab(FNetworkingProfilerTabs::PacketSizesViewID, bIsPacketSizesViewVisible);
-	}
-}
-
-////////////////////////////////////////////////////////////////////////////////////////////////////
-
-void FNetworkingProfilerManager::ShowHidePacketBreakdownView(const bool bPacketBreakdownViewVisibleState)
-{
-	//bool bWasPacketBreakdownViewVisible = bIsPacketBreakdownViewVisible;
-
-	bIsPacketBreakdownViewVisible = bPacketBreakdownViewVisibleState;
-
-	TSharedPtr<SNetworkingProfilerWindow> Wnd = GetProfilerWindow();
-	if (Wnd.IsValid())
-	{
-		Wnd->ShowHideTab(FNetworkingProfilerTabs::PacketBreakdownViewID, bIsPacketBreakdownViewVisible);
-	}
-}
-
-////////////////////////////////////////////////////////////////////////////////////////////////////
-
-void FNetworkingProfilerManager::ShowHideDataStreamBreakdownView(const bool bDataStreamBreakdownViewVisibleState)
-{
-	//bool bWasDataStreamBreakdownViewVisible = bIsDataStreamBreakdownViewVisible;
-
-	bIsDataStreamBreakdownViewVisible = bDataStreamBreakdownViewVisibleState;
-
-	TSharedPtr<SNetworkingProfilerWindow> Wnd = GetProfilerWindow();
-	if (Wnd.IsValid())
-	{
-		Wnd->ShowHideTab(FNetworkingProfilerTabs::DataStreamBreakdownViewID, bIsDataStreamBreakdownViewVisible);
+		TSharedPtr<SNetworkingProfilerWindow> Wnd = WndWeakPtr.Pin();
+		if (Wnd.IsValid())
+		{
+			Wnd->Reset();
+		}
 	}
 }
 
diff --git a/Engine/Source/Developer/TraceInsights/Private/Insights/NetworkingProfiler/NetworkingProfilerManager.h b/Engine/Source/Developer/TraceInsights/Private/Insights/NetworkingProfiler/NetworkingProfilerManager.h
--- a/Engine/Source/Developer/TraceInsights/Private/Insights/NetworkingProfiler/NetworkingProfilerManager.h
+++ b/Engine/Source/Developer/TraceInsights/Private/Insights/NetworkingProfiler/NetworkingProfilerManager.h
@@ -21,14 +21,14 @@ namespace Trace
 
 ////////////////////////////////////////////////////////////////////////////////////////////////////
 /**
- * This class manages the Timing Profiler state and settings.
+ * This class manages the Networking Profiler state and settings.
  */
 class FNetworkingProfilerManager : public TSharedFromThis<FNetworkingProfilerManager>
 {
 	friend class FNetworkingProfilerActionManager;
 
 public:
-	/** Creates the Timing Profiler manager, only one instance can exist. */
+	/** Creates the Networking Profiler (Networking Insights) manager, only one instance can exist. */
 	FNetworkingProfilerManager(TSharedRef<FUICommandList> InCommandList);
 
 	/** Virtual destructor. */
@@ -48,7 +48,7 @@ public:
 		return FNetworkingProfilerManager::Instance;
 	}
 
-	/** Shutdowns the Timing Profiler manager. */
+	/** Shutdowns the Networking Profiler manager. */
 	void Shutdown()
 	{
 		FNetworkingProfilerManager::Instance.Reset();
@@ -63,7 +63,7 @@ protected:
 
 public:
 	/**
-	 * @return the global instance of the Networking Profiler (Asset Networking Insights) manager.
+	 * @return the global instance of the Networking Profiler (Networking Insights) manager.
 	 * This is an internal singleton and cannot be used outside ProfilerModule.
 	 * For external use:
 	 *     IProfilerModule& ProfilerModule = FModuleManager::Get().LoadModuleChecked<IProfilerModule>("Profiler");
@@ -71,47 +71,34 @@ public:
 	 */
 	static TSharedPtr<FNetworkingProfilerManager> Get();
 
-	/** @returns UI command list for the Timing Profiler manager. */
+	/** @returns UI command list for the Networking Profiler manager. */
 	const TSharedRef<FUICommandList> GetCommandList() const;
 
-	/** @return an instance of the Timing Profiler commands. */
+	/** @return an instance of the Networking Profiler commands. */
 	static const FNetworkingProfilerCommands& GetCommands();
 
-	/** @return an instance of the Timing Profiler action manager. */
+	/** @return an instance of the Networking Profiler action manager. */
 	static FNetworkingProfilerActionManager& GetActionManager();
 
-	void AssignProfilerWindow(const TSharedRef<SNetworkingProfilerWindow>& InProfilerWindow)
+	void AddProfilerWindow(const TSharedRef<SNetworkingProfilerWindow>& InProfilerWindow)
 	{
-		ProfilerWindow = InProfilerWindow;
+		ProfilerWindows.Add(InProfilerWindow);
+	}
+
+	void RemoveProfilerWindow(const TSharedRef<SNetworkingProfilerWindow>& InProfilerWindow)
+	{
+		ProfilerWindows.Remove(InProfilerWindow);
 	}
 
 	/**
 	 * Converts profiler window weak pointer to a shared pointer and returns it.
 	 * Make sure the returned pointer is valid before trying to dereference it.
 	 */
-	TSharedPtr<class SNetworkingProfilerWindow> GetProfilerWindow() const
+	TSharedPtr<class SNetworkingProfilerWindow> GetProfilerWindow(int32 Index) const
 	{
-		return ProfilerWindow.Pin();
+		return ProfilerWindows[Index].Pin();
 	}
 
-	////////////////////////////////////////////////////////////////////////////////////////////////////
-	// Getters and setters used by Toggle Commands.
-
-	/** @return true, if Packet Sizes view is visible */
-	const bool IsPacketSizesViewVisible() const { return bIsPacketSizesViewVisible; }
-	void SetPacketSizesViewVisible(const bool bPacketSizesViewVisibleState) { bIsPacketSizesViewVisible = bPacketSizesViewVisibleState; }
-	void ShowHidePacketSizesView(const bool bPacketSizesViewVisibleState);
-
-	/** @return true, if Packet Breakdown view is visible */
-	const bool IsPacketBreakdownViewVisible() const { return bIsPacketBreakdownViewVisible; }
-	void SetPacketBreakdownViewVisible(const bool bVisibleState) { bIsPacketBreakdownViewVisible = bVisibleState; }
-	void ShowHidePacketBreakdownView(const bool bVisibleState);
-
-	/** @return true, if Data Stream Breakdown view is visible */
-	const bool IsDataStreamBreakdownViewVisible() const { return bIsDataStreamBreakdownViewVisible; }
-	void SetDataStreamBreakdownViewVisible(const bool bVisibleState) { bIsDataStreamBreakdownViewVisible = bVisibleState; }
-	void ShowHideDataStreamBreakdownView(const bool bVisibleState);
-
 	////////////////////////////////////////////////////////////////////////////////////////////////////
 
 	void OnSessionChanged();
@@ -133,17 +120,8 @@ protected:
 	/** An instance of the Networking Profiler action manager. */
 	FNetworkingProfilerActionManager ActionManager;
 
-	/** A weak pointer to the Networking Profiler window. */
-	TWeakPtr<class SNetworkingProfilerWindow> ProfilerWindow;
-
-	/** If the Packet Sizes view is visible or hidden. */
-	bool bIsPacketSizesViewVisible;
-
-	/** If the Packet Breakdown view is visible or hidden. */
-	bool bIsPacketBreakdownViewVisible;
-
-	/** If the Data Stream Breakdown tree view is visible or hidden. */
-	bool bIsDataStreamBreakdownViewVisible;
+	/** A list of weak pointers to the Networking Profiler windows. */
+	TArray<TWeakPtr<class SNetworkingProfilerWindow>> ProfilerWindows;
 
 	/** A shared pointer to the global instance of the profiler manager. */
 	static TSharedPtr<FNetworkingProfilerManager> Instance;
diff --git a/Engine/Source/Developer/TraceInsights/Private/Insights/NetworkingProfiler/ViewModels/PacketBreakdownViewHelper.cpp b/Engine/Source/Developer/TraceInsights/Private/Insights/NetworkingProfiler/ViewModels/PacketBreakdownViewHelper.cpp
--- a/Engine/Source/Developer/TraceInsights/Private/Insights/NetworkingProfiler/ViewModels/PacketBreakdownViewHelper.cpp
+++ b/Engine/Source/Developer/TraceInsights/Private/Insights/NetworkingProfiler/ViewModels/PacketBreakdownViewHelper.cpp
@@ -1,6 +1,7 @@
 // Copyright 1998-2019 Epic Games, Inc. All Rights Reserved.
 
 #include "PacketBreakdownViewHelper.h"
+//#include "PacketContentViewHelper.h"
 
 #include "Fonts/FontMeasure.h"
 #include "Framework/Application/SlateApplication.h"
@@ -12,15 +13,16 @@
 #include "Insights/Common/PaintUtils.h"
 #include "Insights/InsightsStyle.h"
 #include "Insights/NetworkingProfiler/ViewModels/PacketBreakdownViewport.h"
+//#include "Insights/NetworkingProfiler/ViewModels/PacketContentViewport.h"
 #include "Insights/ViewModels/DrawHelpers.h"
 
 #include <limits>
 
 ////////////////////////////////////////////////////////////////////////////////////////////////////
-// FPacketBreakdownViewDrawStateBuilder
+// FPacketContentViewDrawStateBuilder
 ////////////////////////////////////////////////////////////////////////////////////////////////////
 
-FPacketBreakdownViewDrawStateBuilder::FPacketBreakdownViewDrawStateBuilder(FPacketBreakdownViewDrawState& InDrawState, const FPacketBreakdownViewport& InViewport)
+FPacketContentViewDrawStateBuilder::FPacketContentViewDrawStateBuilder(FPacketContentViewDrawState& InDrawState, const FPacketContentViewport& InViewport)
 	: DrawState(InDrawState)
 	, Viewport(InViewport)
 	, MaxDepth(-1)
@@ -33,18 +35,20 @@ FPacketBreakdownViewDrawStateBuilder::FPacketBreakdownViewDrawStateBuilder(FPack
 
 ////////////////////////////////////////////////////////////////////////////////////////////////////
 
-void FPacketBreakdownViewDrawStateBuilder::AddEvent(int64 EventOffset, int64 EventSize, int32 EventType, int32 Depth)
+void FPacketContentViewDrawStateBuilder::AddEvent(const Trace::FNetProfilerContentEvent& Event, const TCHAR* EventName)
 {
+	const int32 Depth = Event.Level;
+
 	DrawState.Events.AddUninitialized();
-	FPacketBreakdownViewDrawState::FPacketEvent& PacketEvent = DrawState.Events.Last();
-	PacketEvent.Offset = EventOffset;
-	PacketEvent.Size = EventSize;
-	PacketEvent.Type = EventType;
+	FPacketContentViewDrawState::FPacketEvent& PacketEvent = DrawState.Events.Last();
+	PacketEvent.Offset = Event.StartPos;
+	PacketEvent.Size = Event.EndPos - Event.StartPos;
+	PacketEvent.Type = Event.NameIndex;
 	PacketEvent.Depth = Depth;
 
 	const FAxisViewportDouble& ViewportX = Viewport.GetHorizontalAxisViewport();
 
-	float EventX1 = ViewportX.GetRoundedOffsetForValue(static_cast<double>(EventOffset));
+	float EventX1 = ViewportX.GetRoundedOffsetForValue(static_cast<double>(Event.StartPos));
 	if (EventX1 > Viewport.GetWidth())
 	{
 		return;
@@ -53,15 +57,15 @@ void FPacketBreakdownViewDrawStateBuilder::AddEvent(int64 EventOffset, int64 Eve
 	//float EventX2 = EventSize == 0 ?
 	//	EventX1 + 1.0f :
 	//	ViewportX.GetScale() < 1.0f ?
-	//	ViewportX.GetRoundedOffsetForValue(static_cast<double>(EventOffset + EventSize - 1)) + 1.0f : // x2 = pos of last inclusive bit + 1px
-	//	ViewportX.GetRoundedOffsetForValue(static_cast<double>(EventOffset + EventSize));             // x2 = pos of first exclusive bit
-	float EventX2 = ViewportX.GetRoundedOffsetForValue(static_cast<double>(EventOffset + EventSize));
+	//	ViewportX.GetRoundedOffsetForValue(static_cast<double>(Event.EndPos - 1)) + 1.0f : // x2 = pos of last inclusive bit + 1px
+	//	ViewportX.GetRoundedOffsetForValue(static_cast<double>(Event.EndPos));             // x2 = pos of first exclusive bit
+	float EventX2 = ViewportX.GetRoundedOffsetForValue(static_cast<double>(Event.EndPos));
 	if (EventX2 < 0)
 	{
 		return;
 	}
 
-	// Timing events are displayed with minimum 1px (including empty ones).
+	// Events are displayed with minimum 1px (including empty ones).
 	if (EventX1 == EventX2)
 	{
 		EventX2 = EventX1 + 1.0f;
@@ -121,7 +125,7 @@ void FPacketBreakdownViewDrawStateBuilder::AddEvent(int64 EventOffset, int64 Eve
 	//////////////////////////////////////////////////
 	// Coloring
 
-	const uint32 Color = (static_cast<uint32>(EventType) + 0x2c2c57ed) * 0x2c2c57ed;
+	const uint32 Color = (Event.NameIndex + 0x2c2c57ed) * 0x2c2c57ed;
 	const FLinearColor EventColorFill(((Color >> 16) & 0xFF) / 255.0f, ((Color >> 8) & 0xFF) / 255.0f, (Color & 0xFF) / 255.0f, 1.0f);
 	constexpr float BorderColorFactor = 1.25f; // brighter border
 
@@ -131,7 +135,7 @@ void FPacketBreakdownViewDrawStateBuilder::AddEvent(int64 EventOffset, int64 Eve
 	{
 		// Fill inside of the event box.
 		DrawState.Boxes.AddUninitialized();
-		FPacketBreakdownViewDrawState::FBox& DrawBox = DrawState.Boxes.Last();
+		FPacketContentViewDrawState::FBox& DrawBox = DrawState.Boxes.Last();
 		DrawBox.X = EventX1 + 1.0f;
 		DrawBox.Y = EventY + 1.0f;
 		DrawBox.W = EventW - 2.0f;
@@ -153,7 +157,7 @@ void FPacketBreakdownViewDrawStateBuilder::AddEvent(int64 EventOffset, int64 Eve
 		}
 
 		DrawState.Borders.AddUninitialized();
-		FPacketBreakdownViewDrawState::FBox& DrawBox = DrawState.Borders.Last();
+		FPacketContentViewDrawState::FBox& DrawBox = DrawState.Borders.Last();
 		DrawBox.X = EventX1;
 		DrawBox.Y = EventY;
 		DrawBox.W = EventW;
@@ -195,16 +199,26 @@ void FPacketBreakdownViewDrawStateBuilder::AddEvent(int64 EventOffset, int64 Eve
 	// Draw the name of the event.
 	if (EventW > 8.0f && EventH > 10.0f)
 	{
-		FString Name = TEXT("NetObject") + FText::AsNumber(EventType).ToString();
+		FString Name = EventName ? EventName : TEXT("?");
 		if (EventW > Name.Len() * 2.0f + 48.0f)
 		{
-			if (EventSize == 1)
+			if (Event.ObjectInstanceIndex != 0)
 			{
-				Name += TEXT(" (1 bit)");
+				Name += TEXT(" (NetId:");
+				Name += FText::AsNumber(Event.ObjectInstanceIndex).ToString();
+				Name += TEXT(", ");
 			}
 			else
 			{
 				Name += TEXT(" (");
+			}
+			const uint32 EventSize = Event.EndPos - Event.StartPos;
+			if (EventSize == 1)
+			{
+				Name += TEXT("1 bit)");
+			}
+			else
+			{
 				Name += FText::AsNumber(EventSize).ToString();
 				Name += TEXT(" bits)");
 			}
@@ -221,7 +235,7 @@ void FPacketBreakdownViewDrawStateBuilder::AddEvent(int64 EventOffset, int64 Eve
 			const bool bIsDarkColor = (EventColorFill.ComputeLuminance() < 0.4f);
 
 			DrawState.Texts.AddDefaulted();
-			FPacketBreakdownViewDrawState::FText& DrawText = DrawState.Texts.Last();
+			FPacketContentViewDrawState::FText& DrawText = DrawState.Texts.Last();
 			DrawText.X = EventX1 + 2.0f;
 			DrawText.Y = EventY + 1.0f;
 			DrawText.Text = Name.Left(LastWholeCharacterIndex + 1);
@@ -232,10 +246,10 @@ void FPacketBreakdownViewDrawStateBuilder::AddEvent(int64 EventOffset, int64 Eve
 
 ////////////////////////////////////////////////////////////////////////////////////////////////////
 
-void FPacketBreakdownViewDrawStateBuilder::FlushBox(const FBoxData& Box, const float EventY, const float EventH)
+void FPacketContentViewDrawStateBuilder::FlushBox(const FBoxData& Box, const float EventY, const float EventH)
 {
 	DrawState.Boxes.AddUninitialized();
-	FPacketBreakdownViewDrawState::FBox& DrawBox = DrawState.Boxes.Last();
+	FPacketContentViewDrawState::FBox& DrawBox = DrawState.Boxes.Last();
 	DrawBox.X = Box.X1;
 	DrawBox.W = Box.X2 - Box.X1;
 	DrawBox.Y = EventY;
@@ -245,7 +259,7 @@ void FPacketBreakdownViewDrawStateBuilder::FlushBox(const FBoxData& Box, const f
 
 ////////////////////////////////////////////////////////////////////////////////////////////////////
 
-void FPacketBreakdownViewDrawStateBuilder::Flush()
+void FPacketContentViewDrawStateBuilder::Flush()
 {
 	constexpr float Y0 = 0.0f;
 	constexpr float EventH = 14.0f;
@@ -264,10 +278,10 @@ void FPacketBreakdownViewDrawStateBuilder::Flush()
 }
 
 ////////////////////////////////////////////////////////////////////////////////////////////////////
-// FPacketBreakdownViewDrawHelper
+// FPacketContentViewDrawHelper
 ////////////////////////////////////////////////////////////////////////////////////////////////////
 
-FPacketBreakdownViewDrawHelper::FPacketBreakdownViewDrawHelper(const FDrawContext& InDrawContext, const FPacketBreakdownViewport& InViewport)
+FPacketContentViewDrawHelper::FPacketContentViewDrawHelper(const FDrawContext& InDrawContext, const FPacketContentViewport& InViewport)
 	: DrawContext(InDrawContext)
 	, Viewport(InViewport)
 	, WhiteBrush(FInsightsStyle::Get().GetBrush("WhiteBrush"))
@@ -280,7 +294,7 @@ FPacketBreakdownViewDrawHelper::FPacketBreakdownViewDrawHelper(const FDrawContex
 
 ////////////////////////////////////////////////////////////////////////////////////////////////////
 
-void FPacketBreakdownViewDrawHelper::DrawBackground() const
+void FPacketContentViewDrawHelper::DrawBackground() const
 {
 	const FAxisViewportDouble& ViewportX = Viewport.GetHorizontalAxisViewport();
 
@@ -297,7 +311,7 @@ void FPacketBreakdownViewDrawHelper::DrawBackground() const
 
 ////////////////////////////////////////////////////////////////////////////////////////////////////
 
-FLinearColor FPacketBreakdownViewDrawHelper::GetColorByType(int32 Type)
+FLinearColor FPacketContentViewDrawHelper::GetColorByType(int32 Type)
 {
 	constexpr float Alpha = 1.0f;
 	switch (Type)
@@ -333,24 +347,24 @@ FLinearColor FPacketBreakdownViewDrawHelper::GetColorByType(int32 Type)
 
 ////////////////////////////////////////////////////////////////////////////////////////////////////
 
-void FPacketBreakdownViewDrawHelper::Draw(const FPacketBreakdownViewDrawState& DrawState) const
+void FPacketContentViewDrawHelper::Draw(const FPacketContentViewDrawState& DrawState) const
 {
 	// Draw filled boxes.
-	for (const FPacketBreakdownViewDrawState::FBox& Box : DrawState.Boxes)
+	for (const FPacketContentViewDrawState::FBox& Box : DrawState.Boxes)
 	{
 		DrawContext.DrawBox(Box.X, Box.Y, Box.W, Box.H, WhiteBrush, Box.Color);
 	}
 	DrawContext.LayerId++;
 
 	// Draw borders.
-	for (const FPacketBreakdownViewDrawState::FBox& Box : DrawState.Borders)
+	for (const FPacketContentViewDrawState::FBox& Box : DrawState.Borders)
 	{
 		DrawContext.DrawBox(Box.X, Box.Y, Box.W, Box.H, EventBorderBrush, Box.Color);
 	}
 	DrawContext.LayerId++;
 
 	// Draw texts.
-	for (const FPacketBreakdownViewDrawState::FText& Text : DrawState.Texts)
+	for (const FPacketContentViewDrawState::FText& Text : DrawState.Texts)
 	{
 		DrawContext.DrawText(Text.X, Text.Y, Text.Text, EventFont, Text.bWhite ? FLinearColor::White : FLinearColor::Black);
 	}
@@ -359,7 +373,7 @@ void FPacketBreakdownViewDrawHelper::Draw(const FPacketBreakdownViewDrawState& D
 
 ////////////////////////////////////////////////////////////////////////////////////////////////////
 
-void FPacketBreakdownViewDrawHelper::DrawEventHighlight(const FNetworkPacketEvent& Event) const
+void FPacketContentViewDrawHelper::DrawEventHighlight(const FNetworkPacketEvent& Event) const
 {
 	const FAxisViewportDouble& ViewportX = Viewport.GetHorizontalAxisViewport();
 
diff --git a/Engine/Source/Developer/TraceInsights/Private/Insights/NetworkingProfiler/ViewModels/PacketBreakdownViewHelper.h b/Engine/Source/Developer/TraceInsights/Private/Insights/NetworkingProfiler/ViewModels/PacketBreakdownViewHelper.h
--- a/Engine/Source/Developer/TraceInsights/Private/Insights/NetworkingProfiler/ViewModels/PacketBreakdownViewHelper.h
+++ b/Engine/Source/Developer/TraceInsights/Private/Insights/NetworkingProfiler/ViewModels/PacketBreakdownViewHelper.h
@@ -4,17 +4,16 @@
 
 #include "CoreMinimal.h"
 #include "Fonts/SlateFontInfo.h"
+#include "TraceServices/Model/NetProfiler.h"
 #include "Styling/WidgetStyle.h"
 
-#include <limits>
-
 enum class ESlateDrawEffect : uint8;
 
 struct FDrawContext;
 struct FGeometry;
 struct FSlateBrush;
 
-class FPacketBreakdownViewport;
+class FPacketContentViewport;
 class FSlateWindowElementList;
 
 ////////////////////////////////////////////////////////////////////////////////////////////////////
@@ -57,7 +56,7 @@ struct FNetworkPacketEvent
 
 ////////////////////////////////////////////////////////////////////////////////////////////////////
 
-struct FPacketBreakdownViewDrawState
+struct FPacketContentViewDrawState
 {
 	struct FPacketEvent
 	{
@@ -84,7 +83,7 @@ struct FPacketBreakdownViewDrawState
 		bool bWhite;
 	};
 
-	FPacketBreakdownViewDrawState()
+	FPacketContentViewDrawState()
 	{
 	}
 
@@ -110,7 +109,7 @@ struct FPacketBreakdownViewDrawState
 
 ////////////////////////////////////////////////////////////////////////////////////////////////////
 
-class FPacketBreakdownViewDrawStateBuilder
+class FPacketContentViewDrawStateBuilder
 {
 private:
 	struct FBoxData
@@ -125,23 +124,23 @@ private:
 	};
 
 public:
-	explicit FPacketBreakdownViewDrawStateBuilder(FPacketBreakdownViewDrawState& InState, const FPacketBreakdownViewport& InViewport);
+	explicit FPacketContentViewDrawStateBuilder(FPacketContentViewDrawState& InState, const FPacketContentViewport& InViewport);
 
 	/**
 	 * Non-copyable
 	 */
-	FPacketBreakdownViewDrawStateBuilder(const FPacketBreakdownViewDrawStateBuilder&) = delete;
-	FPacketBreakdownViewDrawStateBuilder& operator=(const FPacketBreakdownViewDrawStateBuilder&) = delete;
+	FPacketContentViewDrawStateBuilder(const FPacketContentViewDrawStateBuilder&) = delete;
+	FPacketContentViewDrawStateBuilder& operator=(const FPacketContentViewDrawStateBuilder&) = delete;
 
-	void AddEvent(int64 Offset, int64 Size, int32 Type, int32 Depth);
+	void AddEvent(const Trace::FNetProfilerContentEvent& Event, const TCHAR* Name = nullptr);
 	void Flush();
 
 private:
 	void FlushBox(const FBoxData& Box, const float EventY, const float EventH);
 
 private:
-	FPacketBreakdownViewDrawState& DrawState; // cached draw state to build
-	const FPacketBreakdownViewport& Viewport;
+	FPacketContentViewDrawState& DrawState; // cached draw state to build
+	const FPacketContentViewport& Viewport;
 
 	int32 MaxDepth;
 
@@ -153,7 +152,7 @@ private:
 
 ////////////////////////////////////////////////////////////////////////////////////////////////////
 
-class FPacketBreakdownViewDrawHelper
+class FPacketContentViewDrawHelper
 {
 public:
 	enum class EHighlightMode : uint32
@@ -164,26 +163,26 @@ public:
 	};
 
 public:
-	explicit FPacketBreakdownViewDrawHelper(const FDrawContext& InDrawContext, const FPacketBreakdownViewport& InViewport);
+	explicit FPacketContentViewDrawHelper(const FDrawContext& InDrawContext, const FPacketContentViewport& InViewport);
 
 	/**
 	 * Non-copyable
 	 */
-	FPacketBreakdownViewDrawHelper(const FPacketBreakdownViewDrawHelper&) = delete;
-	FPacketBreakdownViewDrawHelper& operator=(const FPacketBreakdownViewDrawHelper&) = delete;
+	FPacketContentViewDrawHelper(const FPacketContentViewDrawHelper&) = delete;
+	FPacketContentViewDrawHelper& operator=(const FPacketContentViewDrawHelper&) = delete;
 
 	const FSlateBrush* GetWhiteBrush() const { return WhiteBrush; }
 	const FSlateFontInfo& GetEventFont() const { return EventFont; }
 
 	void DrawBackground() const;
-	void Draw(const FPacketBreakdownViewDrawState& DrawState) const;
+	void Draw(const FPacketContentViewDrawState& DrawState) const;
 	void DrawEventHighlight(const FNetworkPacketEvent& Event) const;
 
 	static FLinearColor GetColorByType(int32 Type);
 
 private:
 	const FDrawContext& DrawContext;
-	const FPacketBreakdownViewport& Viewport;
+	const FPacketContentViewport& Viewport;
 
 	const FSlateBrush* WhiteBrush;
 	const FSlateBrush* EventBorderBrush;
diff --git a/Engine/Source/Developer/TraceInsights/Private/Insights/NetworkingProfiler/ViewModels/PacketBreakdownViewport.cpp b/Engine/Source/Developer/TraceInsights/Private/Insights/NetworkingProfiler/ViewModels/PacketBreakdownViewport.cpp
--- a/Engine/Source/Developer/TraceInsights/Private/Insights/NetworkingProfiler/ViewModels/PacketBreakdownViewport.cpp
+++ b/Engine/Source/Developer/TraceInsights/Private/Insights/NetworkingProfiler/ViewModels/PacketBreakdownViewport.cpp
@@ -1,5 +1,6 @@
 // Copyright 1998-2019 Epic Games, Inc. All Rights Reserved.
 
 #include "PacketBreakdownViewport.h"
+//#include "PacketContentViewport.h"
 
 ////////////////////////////////////////////////////////////////////////////////////////////////////
diff --git a/Engine/Source/Developer/TraceInsights/Private/Insights/NetworkingProfiler/ViewModels/PacketBreakdownViewport.h b/Engine/Source/Developer/TraceInsights/Private/Insights/NetworkingProfiler/ViewModels/PacketBreakdownViewport.h
--- a/Engine/Source/Developer/TraceInsights/Private/Insights/NetworkingProfiler/ViewModels/PacketBreakdownViewport.h
+++ b/Engine/Source/Developer/TraceInsights/Private/Insights/NetworkingProfiler/ViewModels/PacketBreakdownViewport.h
@@ -9,10 +9,10 @@
 
 ////////////////////////////////////////////////////////////////////////////////////////////////////
 
-class FPacketBreakdownViewport
+class FPacketContentViewport
 {
 public:
-	FPacketBreakdownViewport()
+	FPacketContentViewport()
 	{
 		Reset();
 	}
diff --git a/Engine/Source/Developer/TraceInsights/Private/Insights/NetworkingProfiler/ViewModels/PacketSizesViewHelper.cpp b/Engine/Source/Developer/TraceInsights/Private/Insights/NetworkingProfiler/ViewModels/PacketSizesViewHelper.cpp
--- a/Engine/Source/Developer/TraceInsights/Private/Insights/NetworkingProfiler/ViewModels/PacketSizesViewHelper.cpp
+++ b/Engine/Source/Developer/TraceInsights/Private/Insights/NetworkingProfiler/ViewModels/PacketSizesViewHelper.cpp
@@ -1,6 +1,7 @@
 // Copyright 1998-2019 Epic Games, Inc. All Rights Reserved.
 
 #include "PacketSizesViewHelper.h"
+//#include "PacketViewHelper.h"
 
 #include "Framework/Application/SlateApplication.h"
 #include "Rendering/DrawElements.h"
@@ -10,6 +11,7 @@
 #include "Insights/Common/PaintUtils.h"
 #include "Insights/InsightsStyle.h"
 #include "Insights/NetworkingProfiler/ViewModels/PacketSizesViewport.h"
+//#include "Insights/NetworkingProfiler/ViewModels/PacketViewViewport.h"
 #include "Insights/ViewModels/DrawHelpers.h"
 
 #include <limits>
@@ -18,7 +20,7 @@
 // FNetworkPacketSeriesBuilder
 ////////////////////////////////////////////////////////////////////////////////////////////////////
 
-FNetworkPacketSeriesBuilder::FNetworkPacketSeriesBuilder(FNetworkPacketSeries& InSeries, const FPacketSizesViewport& InViewport)
+FNetworkPacketSeriesBuilder::FNetworkPacketSeriesBuilder(FNetworkPacketSeries& InSeries, const FPacketViewViewport& InViewport)
 	: Series(InSeries)
 	, Viewport(InViewport)
 	, NumAddedPackets(0)
@@ -26,7 +28,7 @@ FNetworkPacketSeriesBuilder::FNetworkPacketSeriesBuilder(FNetworkPacketSeries& I
 	SampleW = Viewport.GetSampleWidth();
 	PacketsPerSample = Viewport.GetNumPacketsPerSample();
 	NumSamples = FMath::Max(0, FMath::CeilToInt(Viewport.GetWidth() / SampleW));
-	FirstFrameIndex = Viewport.GetFirstFrameIndex();
+	FirstPacketIndex = Viewport.GetFirstPacketIndex();
 
 	Series.NumAggregatedPackets = 0;
 	Series.Samples.Reset();
@@ -35,38 +37,51 @@ FNetworkPacketSeriesBuilder::FNetworkPacketSeriesBuilder(FNetworkPacketSeries& I
 
 ////////////////////////////////////////////////////////////////////////////////////////////////////
 
-void FNetworkPacketSeriesBuilder::AddPacket(int32 FrameIndex, int64 Size, double TimeSent, double TimeAck, ENetworkPacketStatus Status)
+void FNetworkPacketSeriesBuilder::AddPacket(int32 PacketIndex, const Trace::FNetProfilerPacket& Packet)
 {
 	NumAddedPackets++;
 
-	int32 SampleIndex = (FrameIndex - FirstFrameIndex) / PacketsPerSample;
+	int32 SampleIndex = (PacketIndex - FirstPacketIndex) / PacketsPerSample;
 	if (SampleIndex >= 0 && SampleIndex < NumSamples)
 	{
 		FNetworkPacketAggregatedSample& Sample = Series.Samples[SampleIndex];
 		Sample.NumPackets++;
 
-		if (TimeSent < Sample.StartTime)
+		const double TimeStamp = static_cast<double>(Packet.TimeStamp);
+		if (TimeStamp < Sample.StartTime)
 		{
-			Sample.StartTime = TimeSent;
+			Sample.StartTime = TimeStamp;
 		}
-		if (TimeSent > Sample.EndTime)
+		if (TimeStamp > Sample.EndTime)
 		{
-			Sample.EndTime = TimeSent;
+			Sample.EndTime = TimeStamp;
 		}
 
-		if (Size > Sample.LargestPacket.Size)
+		if (Packet.TotalPacketSizeInBytes > Sample.LargestPacket.TotalSizeInBytes)
 		{
-			Sample.LargestPacket.FrameIndex = FrameIndex;
-			Sample.LargestPacket.Size = Size;
-			Sample.LargestPacket.TimeSent = TimeSent;
-			Sample.LargestPacket.TimeAck = TimeAck;
-			Sample.LargestPacket.Status = Status;
+			Sample.LargestPacket.Index = PacketIndex;
+			Sample.LargestPacket.SequenceNumber = Packet.SequenceNumber;
+			Sample.LargestPacket.ContentSizeInBits = Packet.ContentSizeInBits;
+			Sample.LargestPacket.TotalSizeInBytes = Packet.TotalPacketSizeInBytes;
+			Sample.LargestPacket.TimeStamp = TimeStamp;
+			Sample.LargestPacket.Status = Packet.DeliveryStatus;
 		}
 
-		ensure(Status != ENetworkPacketStatus::Unknown);
-		if (static_cast<int32>(Status) > static_cast<int32>(Sample.AggregatedStatus))
+		switch (Sample.AggregatedStatus)
 		{
-			Sample.AggregatedStatus = Status;
+			case Trace::ENetProfilerDeliveryStatus::Unknown:
+				Sample.AggregatedStatus = Packet.DeliveryStatus;
+				break;
+
+			case Trace::ENetProfilerDeliveryStatus::Dropped:
+				break;
+
+			case Trace::ENetProfilerDeliveryStatus::Delivered:
+				if (Packet.DeliveryStatus == Trace::ENetProfilerDeliveryStatus::Dropped)
+				{
+					Sample.AggregatedStatus = Trace::ENetProfilerDeliveryStatus::Dropped;
+				}
+				break;
 		}
 
 		Series.NumAggregatedPackets++;
@@ -74,10 +89,10 @@ void FNetworkPacketSeriesBuilder::AddPacket(int32 FrameIndex, int64 Size, double
 }
 
 ////////////////////////////////////////////////////////////////////////////////////////////////////
-// FPacketSizesViewDrawHelper
+// FPacketViewDrawHelper
 ////////////////////////////////////////////////////////////////////////////////////////////////////
 
-FPacketSizesViewDrawHelper::FPacketSizesViewDrawHelper(const FDrawContext& InDrawContext, const FPacketSizesViewport& InViewport)
+FPacketViewDrawHelper::FPacketViewDrawHelper(const FDrawContext& InDrawContext, const FPacketViewViewport& InViewport)
 	: DrawContext(InDrawContext)
 	, Viewport(InViewport)
 	, WhiteBrush(FInsightsStyle::Get().GetBrush("WhiteBrush"))
@@ -91,7 +106,7 @@ FPacketSizesViewDrawHelper::FPacketSizesViewDrawHelper(const FDrawContext& InDra
 
 ////////////////////////////////////////////////////////////////////////////////////////////////////
 
-void FPacketSizesViewDrawHelper::DrawBackground() const
+void FPacketViewDrawHelper::DrawBackground() const
 {
 	const FAxisViewportInt32& ViewportX = Viewport.GetHorizontalAxisViewport();
 
@@ -108,21 +123,18 @@ void FPacketSizesViewDrawHelper::DrawBackground() const
 
 ////////////////////////////////////////////////////////////////////////////////////////////////////
 
-FLinearColor FPacketSizesViewDrawHelper::GetColorByStatus(ENetworkPacketStatus Status)
+FLinearColor FPacketViewDrawHelper::GetColorByStatus(Trace::ENetProfilerDeliveryStatus Status)
 {
-	const float Alpha = 1.0f;
+	constexpr float Alpha = 1.0f;
 	switch (Status)
 	{
-	case ENetworkPacketStatus::Unknown:
+	case Trace::ENetProfilerDeliveryStatus::Unknown:
 		return FLinearColor(0.25f, 0.25f, 0.25f, Alpha);
 
-	case ENetworkPacketStatus::Sent:
-		return FLinearColor(0.75f, 0.75f, 0.75f, Alpha);
-
-	case ENetworkPacketStatus::ConfirmedReceived:
+	case Trace::ENetProfilerDeliveryStatus::Delivered:
 		return FLinearColor(0.5f, 1.0f, 0.5f, Alpha);
 
-	case ENetworkPacketStatus::ConfirmedLost:
+	case Trace::ENetProfilerDeliveryStatus::Dropped:
 		return FLinearColor(1.0f, 0.5f, 0.5f, Alpha);
 
 	default:
@@ -132,7 +144,7 @@ FLinearColor FPacketSizesViewDrawHelper::GetColorByStatus(ENetworkPacketStatus S
 
 ////////////////////////////////////////////////////////////////////////////////////////////////////
 
-void FPacketSizesViewDrawHelper::DrawCached(const FNetworkPacketSeries& Series) const
+void FPacketViewDrawHelper::DrawCached(const FNetworkPacketSeries& Series) const
 {
 	if (Series.NumAggregatedPackets == 0)
 	{
@@ -161,7 +173,8 @@ void FPacketSizesViewDrawHelper::DrawCached(const FNetworkPacketSeries& Series)
 
 		const float X = SampleIndex * SampleW;
 
-		const float ValueY = FMath::RoundToFloat(ViewportY.GetOffsetForValue(static_cast<double>(Sample.LargestPacket.Size)));
+		//const float ValueY = FMath::RoundToFloat(ViewportY.GetOffsetForValue(static_cast<double>(Sample.LargestPacket.ContentSizeInBits)));
+		const float ValueY = FMath::RoundToFloat(ViewportY.GetOffsetForValue(static_cast<double>(Sample.LargestPacket.TotalSizeInBytes * 8)));
 
 		const float H = ValueY - BaselineY;
 		const float Y = ViewHeight - H;
@@ -190,19 +203,21 @@ void FPacketSizesViewDrawHelper::DrawCached(const FNetworkPacketSeries& Series)
 
 ////////////////////////////////////////////////////////////////////////////////////////////////////
 
-void FPacketSizesViewDrawHelper::DrawSampleHighlight(const FNetworkPacketAggregatedSample& Sample, EHighlightMode Mode) const
+void FPacketViewDrawHelper::DrawSampleHighlight(const FNetworkPacketAggregatedSample& Sample, EHighlightMode Mode) const
 {
 	const float SampleW = Viewport.GetSampleWidth();
 	const int32 PacketsPerSample = Viewport.GetNumPacketsPerSample();
-	const int32 FirstFrameIndex = Viewport.GetFirstFrameIndex();
-	const int32 SampleIndex = (Sample.LargestPacket.FrameIndex - FirstFrameIndex) / PacketsPerSample;
+	const int32 FirstPacketIndex = Viewport.GetFirstPacketIndex();
+	const int32 SampleIndex = (Sample.LargestPacket.Index - FirstPacketIndex) / PacketsPerSample;
 	const float X = SampleIndex * SampleW;
 
 	const FAxisViewportDouble& ViewportY = Viewport.GetVerticalAxisViewport();
 
 	const float ViewHeight = FMath::RoundToFloat(Viewport.GetHeight());
 	const float BaselineY = FMath::RoundToFloat(ViewportY.GetOffsetForValue(0.0));
-	const float ValueY = FMath::RoundToFloat(ViewportY.GetOffsetForValue(static_cast<double>(Sample.LargestPacket.Size)));
+
+	//const float ValueY = FMath::RoundToFloat(ViewportY.GetOffsetForValue(static_cast<double>(Sample.LargestPacket.ContentSizeInBits)));
+	const float ValueY = FMath::RoundToFloat(ViewportY.GetOffsetForValue(static_cast<double>(Sample.LargestPacket.TotalSizeInBytes * 8)));
 
 	const float H = ValueY - BaselineY;
 	const float Y = ViewHeight - H;
diff --git a/Engine/Source/Developer/TraceInsights/Private/Insights/NetworkingProfiler/ViewModels/PacketSizesViewHelper.h b/Engine/Source/Developer/TraceInsights/Private/Insights/NetworkingProfiler/ViewModels/PacketSizesViewHelper.h
--- a/Engine/Source/Developer/TraceInsights/Private/Insights/NetworkingProfiler/ViewModels/PacketSizesViewHelper.h
+++ b/Engine/Source/Developer/TraceInsights/Private/Insights/NetworkingProfiler/ViewModels/PacketSizesViewHelper.h
@@ -4,6 +4,7 @@
 
 #include "CoreMinimal.h"
 #include "Styling/WidgetStyle.h"
+#include "TraceServices/Model/NetProfiler.h"
 
 #include <limits>
 
@@ -13,35 +14,27 @@ struct FDrawContext;
 struct FGeometry;
 struct FSlateBrush;
 
-class FPacketSizesViewport;
+class FPacketViewViewport;
 class FSlateWindowElementList;
 
 ////////////////////////////////////////////////////////////////////////////////////////////////////
 
-enum class ENetworkPacketStatus
-{
-	Unknown = 0,
-	ConfirmedReceived, // Sent + Ack(Received)
-	Sent, // Sent, no Ack event yet
-	ConfirmedLost // Sent + Ack(Lost)
-};
-
-////////////////////////////////////////////////////////////////////////////////////////////////////
-
 struct FNetworkPacket
 {
-	int32 FrameIndex;
-	int64 Size; // [bits]
-	double TimeSent; // time when packet is sent
-	double TimeAck; // time when ACK event is received (either confirmed received or lost)
-	ENetworkPacketStatus Status;
+	int32 Index;
+	uint32 SequenceNumber;
+	uint32 ContentSizeInBits;
+	uint32 TotalSizeInBytes;
+	double TimeStamp;
+	Trace::ENetProfilerDeliveryStatus Status;
 
 	FNetworkPacket()
-		: FrameIndex(-1)
-		, Size(-1)
-		, TimeSent(std::numeric_limits<double>::infinity())
-		, TimeAck(std::numeric_limits<double>::infinity())
-		, Status(ENetworkPacketStatus::Unknown)
+		: Index(-1)
+		, SequenceNumber(0)
+		, ContentSizeInBits(0)
+		, TotalSizeInBytes(0)
+		, TimeStamp(std::numeric_limits<double>::infinity())
+		, Status(Trace::ENetProfilerDeliveryStatus::Unknown)
 	{}
 
 	FNetworkPacket(const FNetworkPacket&) = default;
@@ -52,10 +45,11 @@ struct FNetworkPacket
 
 	bool Equals(const FNetworkPacket& Other) const
 	{
-		return FrameIndex == Other.FrameIndex
-			&& Size == Other.Size
-			/*&& TimeSent == Other.TimeSent
-			&& TimeAck == Other.TimeAck
+		return Index == Other.Index
+			&& SequenceNumber == Other.SequenceNumber
+			//&& ContentSizeInBits == Other.ContentSizeInBits
+			&& TotalSizeInBytes == Other.TotalSizeInBytes
+			/*&& TimeStamp == Other.TimeStamp
 			&& Status == Other.Status*/;
 	}
 
@@ -80,7 +74,7 @@ struct FNetworkPacketAggregatedSample
 	 *    Received --> at least one packet in the sample set is confirmed received and none are confirmed lost
 	 *    Lost     --> at least one packet in the sample set is confirmed lost
 	**/
-	ENetworkPacketStatus AggregatedStatus;
+	Trace::ENetProfilerDeliveryStatus AggregatedStatus;
 
 	FNetworkPacket LargestPacket;
 
@@ -88,7 +82,7 @@ struct FNetworkPacketAggregatedSample
 		: NumPackets(0)
 		, StartTime(DBL_MAX)
 		, EndTime(-DBL_MAX)
-		, AggregatedStatus(ENetworkPacketStatus::Unknown)
+		, AggregatedStatus(Trace::ENetProfilerDeliveryStatus::Unknown)
 		, LargestPacket()
 	{}
 
@@ -139,7 +133,7 @@ struct FNetworkPacketSeries
 class FNetworkPacketSeriesBuilder
 {
 public:
-	explicit FNetworkPacketSeriesBuilder(FNetworkPacketSeries& InSeries, const FPacketSizesViewport& InViewport);
+	explicit FNetworkPacketSeriesBuilder(FNetworkPacketSeries& InSeries, const FPacketViewViewport& InViewport);
 
 	/**
 	 * Non-copyable
@@ -147,17 +141,17 @@ public:
 	FNetworkPacketSeriesBuilder(const FNetworkPacketSeriesBuilder&) = delete;
 	FNetworkPacketSeriesBuilder& operator=(const FNetworkPacketSeriesBuilder&) = delete;
 
-	void AddPacket(int32 FrameIndex, int64 Size, double TimeSent, double TimeAck, ENetworkPacketStatus Status);
+	void AddPacket(int32 PacketIndex, const Trace::FNetProfilerPacket& Packet);
 
 	int32 GetNumAddedPackets() const { return NumAddedPackets; }
 
 private:
 	FNetworkPacketSeries& Series; // series to update
-	const FPacketSizesViewport& Viewport;
+	const FPacketViewViewport& Viewport;
 
 	float SampleW; // width of a sample, in Slate units
 	int32 PacketsPerSample; // number of packets in a sample
-	int32 FirstFrameIndex; // index of first frame in first sample; can be negative
+	int32 FirstPacketIndex; // index of the first packet in the first sample; can be negative
 	int32 NumSamples; // total number of samples
 
 	// Debug stats.
@@ -166,7 +160,7 @@ private:
 
 ////////////////////////////////////////////////////////////////////////////////////////////////////
 
-class FPacketSizesViewDrawHelper
+class FPacketViewDrawHelper
 {
 public:
 	enum class EHighlightMode : uint32
@@ -177,26 +171,26 @@ public:
 	};
 
 public:
-	explicit FPacketSizesViewDrawHelper(const FDrawContext& InDrawContext, const FPacketSizesViewport& InViewport);
+	explicit FPacketViewDrawHelper(const FDrawContext& InDrawContext, const FPacketViewViewport& InViewport);
 
 	/**
 	 * Non-copyable
 	 */
-	FPacketSizesViewDrawHelper(const FPacketSizesViewDrawHelper&) = delete;
-	FPacketSizesViewDrawHelper& operator=(const FPacketSizesViewDrawHelper&) = delete;
+	FPacketViewDrawHelper(const FPacketViewDrawHelper&) = delete;
+	FPacketViewDrawHelper& operator=(const FPacketViewDrawHelper&) = delete;
 
 	void DrawBackground() const;
 	void DrawCached(const FNetworkPacketSeries& Series) const;
 	void DrawSampleHighlight(const FNetworkPacketAggregatedSample& Sample, EHighlightMode Mode) const;
 
-	static FLinearColor GetColorByStatus(ENetworkPacketStatus Status);
+	static FLinearColor GetColorByStatus(Trace::ENetProfilerDeliveryStatus Status);
 
 	int32 GetNumPackets() const { return NumPackets; }
 	int32 GetNumDrawSamples() const { return NumDrawSamples; }
 
 private:
 	const FDrawContext& DrawContext;
-	const FPacketSizesViewport& Viewport;
+	const FPacketViewViewport& Viewport;
 
 	const FSlateBrush* WhiteBrush;
 	//const FSlateBrush* EventBorderBrush;
diff --git a/Engine/Source/Developer/TraceInsights/Private/Insights/NetworkingProfiler/ViewModels/PacketSizesViewport.cpp b/Engine/Source/Developer/TraceInsights/Private/Insights/NetworkingProfiler/ViewModels/PacketSizesViewport.cpp
--- a/Engine/Source/Developer/TraceInsights/Private/Insights/NetworkingProfiler/ViewModels/PacketSizesViewport.cpp
+++ b/Engine/Source/Developer/TraceInsights/Private/Insights/NetworkingProfiler/ViewModels/PacketSizesViewport.cpp
@@ -1,5 +1,6 @@
 // Copyright 1998-2019 Epic Games, Inc. All Rights Reserved.
 
 #include "PacketSizesViewport.h"
+//#include "PacketViewViewport.h"
 
 ////////////////////////////////////////////////////////////////////////////////////////////////////
diff --git a/Engine/Source/Developer/TraceInsights/Private/Insights/NetworkingProfiler/ViewModels/PacketSizesViewport.h b/Engine/Source/Developer/TraceInsights/Private/Insights/NetworkingProfiler/ViewModels/PacketSizesViewport.h
--- a/Engine/Source/Developer/TraceInsights/Private/Insights/NetworkingProfiler/ViewModels/PacketSizesViewport.h
+++ b/Engine/Source/Developer/TraceInsights/Private/Insights/NetworkingProfiler/ViewModels/PacketSizesViewport.h
@@ -10,13 +10,13 @@
 
 ////////////////////////////////////////////////////////////////////////////////////////////////////
 
-class FPacketSizesViewport
+class FPacketViewViewport
 {
 private:
 	static constexpr float SLATE_UNITS_TOLERANCE = 0.1f;
 
 public:
-	FPacketSizesViewport()
+	FPacketViewViewport()
 	{
 		Reset();
 	}
@@ -50,7 +50,7 @@ public:
 
 	float GetSampleWidth() const { return HorizontalAxisViewport.GetSampleSize(); }
 	int32 GetNumPacketsPerSample() const { return HorizontalAxisViewport.GetNumSamplesPerPixel(); }
-	int32 GetFirstFrameIndex() const { return HorizontalAxisViewport.GetValueAtOffset(0.0f); }
+	int32 GetFirstPacketIndex() const { return HorizontalAxisViewport.GetValueAtOffset(0.0f); }
 
 private:
 	void OnSizeChanged()
diff --git a/Engine/Source/Developer/TraceInsights/Private/Insights/NetworkingProfiler/Widgets/SDataStreamBreakdownView.cpp b/Engine/Source/Developer/TraceInsights/Private/Insights/NetworkingProfiler/Widgets/SDataStreamBreakdownView.cpp
--- a/Engine/Source/Developer/TraceInsights/Private/Insights/NetworkingProfiler/Widgets/SDataStreamBreakdownView.cpp
+++ b/Engine/Source/Developer/TraceInsights/Private/Insights/NetworkingProfiler/Widgets/SDataStreamBreakdownView.cpp
@@ -1,28 +1,29 @@
 // Copyright 1998-2019 Epic Games, Inc. All Rights Reserved.
 
 #include "SDataStreamBreakdownView.h"
+//#include "SNetStatsView.h"
 
 // Insights
 
 ////////////////////////////////////////////////////////////////////////////////////////////////////
 
-#define LOCTEXT_NAMESPACE "SDataStreamBreakdownView"
+#define LOCTEXT_NAMESPACE "SNetStatsView"
 
 ////////////////////////////////////////////////////////////////////////////////////////////////////
 
-SDataStreamBreakdownView::SDataStreamBreakdownView()
+SNetStatsView::SNetStatsView()
 {
 }
 
 ////////////////////////////////////////////////////////////////////////////////////////////////////
 
-SDataStreamBreakdownView::~SDataStreamBreakdownView()
+SNetStatsView::~SNetStatsView()
 {
 }
 
 ////////////////////////////////////////////////////////////////////////////////////////////////////
 
-void SDataStreamBreakdownView::Construct(const FArguments& InArgs)
+void SNetStatsView::Construct(const FArguments& InArgs)
 {
 	//ChildSlot
 	//[
diff --git a/Engine/Source/Developer/TraceInsights/Private/Insights/NetworkingProfiler/Widgets/SDataStreamBreakdownView.h b/Engine/Source/Developer/TraceInsights/Private/Insights/NetworkingProfiler/Widgets/SDataStreamBreakdownView.h
--- a/Engine/Source/Developer/TraceInsights/Private/Insights/NetworkingProfiler/Widgets/SDataStreamBreakdownView.h
+++ b/Engine/Source/Developer/TraceInsights/Private/Insights/NetworkingProfiler/Widgets/SDataStreamBreakdownView.h
@@ -7,16 +7,16 @@
 #include "Widgets/DeclarativeSyntaxSupport.h"
 #include "Widgets/SCompoundWidget.h"
 
-class SDataStreamBreakdownView : public SCompoundWidget
+class SNetStatsView : public SCompoundWidget
 {
 public:
 	/** Default constructor. */
-	SDataStreamBreakdownView();
+	SNetStatsView();
 
 	/** Virtual destructor. */
-	virtual ~SDataStreamBreakdownView();
+	virtual ~SNetStatsView();
 
-	SLATE_BEGIN_ARGS(SDataStreamBreakdownView) {}
+	SLATE_BEGIN_ARGS(SNetStatsView) {}
 	SLATE_END_ARGS()
 
 	/**
diff --git a/Engine/Source/Developer/TraceInsights/Private/Insights/NetworkingProfiler/Widgets/SNetworkingProfilerToolbar.cpp b/Engine/Source/Developer/TraceInsights/Private/Insights/NetworkingProfiler/Widgets/SNetworkingProfilerToolbar.cpp
--- a/Engine/Source/Developer/TraceInsights/Private/Insights/NetworkingProfiler/Widgets/SNetworkingProfilerToolbar.cpp
+++ b/Engine/Source/Developer/TraceInsights/Private/Insights/NetworkingProfiler/Widgets/SNetworkingProfilerToolbar.cpp
@@ -18,6 +18,7 @@
 #include "Insights/InsightsManager.h"
 #include "Insights/NetworkingProfiler/NetworkingProfilerCommands.h"
 #include "Insights/NetworkingProfiler/NetworkingProfilerManager.h"
+#include "Insights/NetworkingProfiler/Widgets/SNetworkingProfilerWindow.h"
 
 ////////////////////////////////////////////////////////////////////////////////////////////////////
 
@@ -37,43 +38,36 @@ SNetworkingProfilerToolbar::~SNetworkingProfilerToolbar()
 
 ////////////////////////////////////////////////////////////////////////////////////////////////////
 
-void SNetworkingProfilerToolbar::Construct(const FArguments& InArgs)
+void SNetworkingProfilerToolbar::Construct(const FArguments& InArgs, TSharedPtr<SNetworkingProfilerWindow> InProfilerWindow)
 {
+	ProfilerWindow = InProfilerWindow;
+
 	struct Local
 	{
-		static void FillViewToolbar(FToolBarBuilder& ToolbarBuilder)
+		static void FillViewToolbar(TSharedPtr<SNetworkingProfilerWindow> ProfilerWindow, FToolBarBuilder& ToolbarBuilder)
 		{
 			ToolbarBuilder.BeginSection("View");
 			{
-				ToolbarBuilder.AddToolBarButton(FNetworkingProfilerCommands::Get().TogglePacketSizesViewVisibility);
-				ToolbarBuilder.AddToolBarButton(FNetworkingProfilerCommands::Get().TogglePacketBreakdownViewVisibility);
-				//ToolbarBuilder.AddToolBarButton(FNetworkingProfilerCommands::Get().ToggleDataStreamBreakdownViewVisibility);
+				ToolbarBuilder.AddToolBarButton(FNetworkingProfilerManager::GetCommands().TogglePacketViewVisibility);
+				ToolbarBuilder.AddToolBarButton(FNetworkingProfilerManager::GetCommands().TogglePacketContentViewVisibility);
+				ToolbarBuilder.AddToolBarButton(FNetworkingProfilerManager::GetCommands().ToggleNetStatsViewVisibility);
 			}
 			ToolbarBuilder.EndSection();
 			ToolbarBuilder.BeginSection("Connection");
 			{
-				//FUIAction Action_ChooseReplicationSystem
-				//(
-				//	FExecuteAction::CreateSP(this, &STableTreeView::ContextMenu_CopySelectedToClipboard_Execute),
-				//	FCanExecuteAction::CreateSP(this, &STableTreeView::ContextMenu_CopySelectedToClipboard_CanExecute)
-				//);
-				//const FOnGetContent MenuContentGenerator;
-				//ToolbarBuilder.AddComboButton(Action, MenuContentGenerator);
-				TSharedRef<SWidget> ReplicationWidget = SNew(STextBlock)
-					.Text(LOCTEXT("Replication", "Replication System"));
-				ToolbarBuilder.AddWidget(ReplicationWidget);
-
-				ToolbarBuilder.AddSeparator();
-
-				TSharedRef<SWidget> ConnectionWidget = SNew(STextBlock)
-					.Text(LOCTEXT("Connection", "Connection"));
-				ToolbarBuilder.AddWidget(ConnectionWidget);
+				TSharedRef<SWidget> GameInstanceWidget = ProfilerWindow->CreateGameInstanceComboBox();
+				ToolbarBuilder.AddWidget(GameInstanceWidget);
 
-				ToolbarBuilder.AddSeparator();
+				TSharedRef<SWidget> ConnectionWidget = ProfilerWindow->CreateConnectionComboBox();
+				ToolbarBuilder.AddWidget(ConnectionWidget);
 
-				TSharedRef<SWidget> TypeWidget = SNew(STextBlock)
-					.Text(LOCTEXT("Type", "Sent/Received"));
-				ToolbarBuilder.AddWidget(TypeWidget);
+				TSharedRef<SWidget> ConnectionModeWidget = ProfilerWindow->CreateConnectionModeComboBox();
+				ToolbarBuilder.AddWidget(ConnectionModeWidget);
+			}
+			ToolbarBuilder.EndSection();
+			ToolbarBuilder.BeginSection("New");
+			{
+				//TODO: ToolbarBuilder.AddToolBarButton(FNetworkingProfilerManager::GetCommands().OpenNewWindow);
 			}
 			ToolbarBuilder.EndSection();
 		}
@@ -88,10 +82,11 @@ void SNetworkingProfilerToolbar::Construct(const FArguments& InArgs)
 		}
 	};
 
-	TSharedPtr<FUICommandList> CommandList = FInsightsManager::Get()->GetCommandList();
+	//TSharedPtr<FUICommandList> CommandList = FInsightsManager::Get()->GetCommandList();
+	TSharedPtr<FUICommandList> CommandList = ProfilerWindow->GetCommandList();
 
 	FToolBarBuilder ToolbarBuilder(CommandList.ToSharedRef(), FMultiBoxCustomization::None);
-	Local::FillViewToolbar(ToolbarBuilder);
+	Local::FillViewToolbar(ProfilerWindow, ToolbarBuilder);
 
 	FToolBarBuilder RightSideToolbarBuilder(CommandList.ToSharedRef(), FMultiBoxCustomization::None);
 	Local::FillRightSideToolbar(RightSideToolbarBuilder);
@@ -119,7 +114,7 @@ void SNetworkingProfilerToolbar::Construct(const FArguments& InArgs)
 		+SHorizontalBox::Slot()
 		.HAlign(HAlign_Right)
 		.VAlign(VAlign_Center)
-		.FillWidth(1.0)
+		.AutoWidth()
 		.Padding(0.0f)
 		[
 			SNew(SBorder)
diff --git a/Engine/Source/Developer/TraceInsights/Private/Insights/NetworkingProfiler/Widgets/SNetworkingProfilerToolbar.h b/Engine/Source/Developer/TraceInsights/Private/Insights/NetworkingProfiler/Widgets/SNetworkingProfilerToolbar.h
--- a/Engine/Source/Developer/TraceInsights/Private/Insights/NetworkingProfiler/Widgets/SNetworkingProfilerToolbar.h
+++ b/Engine/Source/Developer/TraceInsights/Private/Insights/NetworkingProfiler/Widgets/SNetworkingProfilerToolbar.h
@@ -7,6 +7,8 @@
 #include "Widgets/DeclarativeSyntaxSupport.h"
 #include "Widgets/SCompoundWidget.h"
 
+class SNetworkingProfilerWindow;
+
 class SNetworkingProfilerToolbar : public SCompoundWidget
 {
 public:
@@ -24,5 +26,8 @@ public:
 	 *
 	 * @param	InArgs	The declaration data for this widget
 	 */
-	void Construct(const FArguments& InArgs);
+	void Construct(const FArguments& InArgs, TSharedPtr<SNetworkingProfilerWindow> InProfilerWindow);
+
+private:
+	TSharedPtr<SNetworkingProfilerWindow> ProfilerWindow;
 };
diff --git a/Engine/Source/Developer/TraceInsights/Private/Insights/NetworkingProfiler/Widgets/SNetworkingProfilerWindow.cpp b/Engine/Source/Developer/TraceInsights/Private/Insights/NetworkingProfiler/Widgets/SNetworkingProfilerWindow.cpp
--- a/Engine/Source/Developer/TraceInsights/Private/Insights/NetworkingProfiler/Widgets/SNetworkingProfilerWindow.cpp
+++ b/Engine/Source/Developer/TraceInsights/Private/Insights/NetworkingProfiler/Widgets/SNetworkingProfilerWindow.cpp
@@ -3,6 +3,7 @@
 #include "SNetworkingProfilerWindow.h"
 
 #include "EditorStyleSet.h"
+#include "Framework/Commands/UICommandList.h"
 #include "Framework/Docking/WorkspaceItem.h"
 #include "Framework/MultiBox/MultiBoxBuilder.h"
 #include "SlateOptMacros.h"
@@ -26,9 +27,12 @@
 #include "Insights/InsightsStyle.h"
 #include "Insights/NetworkingProfiler/NetworkingProfilerManager.h"
 #include "Insights/NetworkingProfiler/Widgets/SDataStreamBreakdownView.h"
+//#include "Insights/NetworkingProfiler/Widgets/SNetStatsView.h"
 #include "Insights/NetworkingProfiler/Widgets/SNetworkingProfilerToolbar.h"
 #include "Insights/NetworkingProfiler/Widgets/SPacketBreakdownView.h"
+//#include "Insights/NetworkingProfiler/Widgets/SPacketContentView.h"
 #include "Insights/NetworkingProfiler/Widgets/SPacketSizesView.h"
+//#include "Insights/NetworkingProfiler/Widgets/SPacketView.h"
 #include "Insights/Version.h"
 #include "Insights/Widgets/SInsightsSettings.h"
 
@@ -39,14 +43,20 @@
 ////////////////////////////////////////////////////////////////////////////////////////////////////
 
 const FName FNetworkingProfilerTabs::ToolbarID(TEXT("Toolbar"));
-const FName FNetworkingProfilerTabs::PacketSizesViewID(TEXT("PacketSizesView"));
-const FName FNetworkingProfilerTabs::PacketBreakdownViewID(TEXT("PacketBreakdownView"));
-const FName FNetworkingProfilerTabs::DataStreamBreakdownViewID(TEXT("DataStreamBreakdownView"));
+const FName FNetworkingProfilerTabs::PacketViewID(TEXT("PacketView"));
+const FName FNetworkingProfilerTabs::PacketContentViewID(TEXT("PacketContent"));
+const FName FNetworkingProfilerTabs::NetStatsViewID(TEXT("NetStats"));
 
 ////////////////////////////////////////////////////////////////////////////////////////////////////
 
 SNetworkingProfilerWindow::SNetworkingProfilerWindow()
-	: DurationActive(0.0f)
+	: AvailableGameInstances()
+	, SelectedGameInstance(nullptr)
+	, AvailableConnections()
+	, SelectedConnection(nullptr)
+	, AvailableConnectionModes()
+	, SelectedConnectionMode(nullptr)
+	, DurationActive(0.0f)
 {
 }
 
@@ -66,25 +76,32 @@ SNetworkingProfilerWindow::~SNetworkingProfilerWindow()
 
 void SNetworkingProfilerWindow::Reset()
 {
-	if (PacketSizesView)
+	if (PacketView)
 	{
-		//PacketSizesView->Reset();
+		PacketView->Reset();
 	}
 
-	if (PacketBreakdownView)
+	if (PacketContentView)
 	{
-		//PacketBreakdownView->Reset();
+		PacketContentView->Reset();
 	}
 
-	if (DataStreamBreakdownView)
+	if (NetStatsView)
 	{
-		//DataStreamBreakdownView->Reset();
+		//TODO: NetStatsView->Reset();
 	}
+
+	AvailableGameInstances.Reset();
+	SelectedGameInstance.Reset();
+	AvailableConnections.Reset();
+	SelectedConnection.Reset();
+	AvailableConnectionModes.Reset();
+	SelectedConnectionMode.Reset();
 }
 
 ////////////////////////////////////////////////////////////////////////////////////////////////////
-
 BEGIN_SLATE_FUNCTION_BUILD_OPTIMIZATION
+////////////////////////////////////////////////////////////////////////////////////////////////////
 
 TSharedRef<SDockTab> SNetworkingProfilerWindow::SpawnTab_Toolbar(const FSpawnTabArgs& Args)
 {
@@ -92,7 +109,7 @@ TSharedRef<SDockTab> SNetworkingProfilerWindow::SpawnTab_Toolbar(const FSpawnTab
 		.ShouldAutosize(true)
 		.TabRole(ETabRole::PanelTab)
 		[
-			SNew(SNetworkingProfilerToolbar)
+			SNew(SNetworkingProfilerToolbar, SharedThis(this))
 		];
 
 	DockTab->SetOnTabClosed(SDockTab::FOnTabClosedCallback::CreateRaw(this, &SNetworkingProfilerWindow::OnToolbarTabClosed));
@@ -108,108 +125,101 @@ void SNetworkingProfilerWindow::OnToolbarTabClosed(TSharedRef<SDockTab> TabBeing
 
 ////////////////////////////////////////////////////////////////////////////////////////////////////
 
-TSharedRef<SDockTab> SNetworkingProfilerWindow::SpawnTab_PacketSizesView(const FSpawnTabArgs& Args)
+TSharedRef<SDockTab> SNetworkingProfilerWindow::SpawnTab_PacketView(const FSpawnTabArgs& Args)
 {
-	FNetworkingProfilerManager::Get()->SetPacketSizesViewVisible(true);
-
 	const TSharedRef<SDockTab> DockTab = SNew(SDockTab)
 		.ShouldAutosize(false)
 		.TabRole(ETabRole::PanelTab)
 		[
-			SAssignNew(PacketSizesView, SPacketSizesView)
+			SAssignNew(PacketView, SPacketView, SharedThis(this))
 		];
 
-	DockTab->SetOnTabClosed(SDockTab::FOnTabClosedCallback::CreateRaw(this, &SNetworkingProfilerWindow::OnPacketSizesViewTabClosed));
+	DockTab->SetOnTabClosed(SDockTab::FOnTabClosedCallback::CreateRaw(this, &SNetworkingProfilerWindow::OnPacketViewTabClosed));
 
 	return DockTab;
 }
 
 ////////////////////////////////////////////////////////////////////////////////////////////////////
 
-void SNetworkingProfilerWindow::OnPacketSizesViewTabClosed(TSharedRef<SDockTab> TabBeingClosed)
+void SNetworkingProfilerWindow::OnPacketViewTabClosed(TSharedRef<SDockTab> TabBeingClosed)
 {
-	FNetworkingProfilerManager::Get()->SetPacketSizesViewVisible(false);
-	PacketSizesView = nullptr;
+	PacketView = nullptr;
 }
 
 ////////////////////////////////////////////////////////////////////////////////////////////////////
 
-TSharedRef<SDockTab> SNetworkingProfilerWindow::SpawnTab_PacketBreakdownView(const FSpawnTabArgs& Args)
+TSharedRef<SDockTab> SNetworkingProfilerWindow::SpawnTab_PacketContentView(const FSpawnTabArgs& Args)
 {
-	FNetworkingProfilerManager::Get()->SetPacketBreakdownViewVisible(true);
-
 	const TSharedRef<SDockTab> DockTab = SNew(SDockTab)
 		.ShouldAutosize(false)
 		.TabRole(ETabRole::PanelTab)
 		[
-			SAssignNew(PacketBreakdownView, SPacketBreakdownView)
+			SAssignNew(PacketContentView, SPacketContentView)
 		];
 
-	DockTab->SetOnTabClosed(SDockTab::FOnTabClosedCallback::CreateRaw(this, &SNetworkingProfilerWindow::OnPacketBreakdownViewTabClosed));
+	DockTab->SetOnTabClosed(SDockTab::FOnTabClosedCallback::CreateRaw(this, &SNetworkingProfilerWindow::OnPacketContentViewTabClosed));
 
 	return DockTab;
 }
 
 ////////////////////////////////////////////////////////////////////////////////////////////////////
 
-void SNetworkingProfilerWindow::OnPacketBreakdownViewTabClosed(TSharedRef<SDockTab> TabBeingClosed)
+void SNetworkingProfilerWindow::OnPacketContentViewTabClosed(TSharedRef<SDockTab> TabBeingClosed)
 {
-	FNetworkingProfilerManager::Get()->SetPacketBreakdownViewVisible(false);
-	PacketBreakdownView = nullptr;
+	PacketContentView = nullptr;
 }
 
 ////////////////////////////////////////////////////////////////////////////////////////////////////
 
-TSharedRef<SDockTab> SNetworkingProfilerWindow::SpawnTab_DataStreamBreakdownView(const FSpawnTabArgs& Args)
+TSharedRef<SDockTab> SNetworkingProfilerWindow::SpawnTab_NetStatsView(const FSpawnTabArgs& Args)
 {
-	FNetworkingProfilerManager::Get()->SetDataStreamBreakdownViewVisible(true);
-
 	const TSharedRef<SDockTab> DockTab = SNew(SDockTab)
 		.ShouldAutosize(false)
 		.TabRole(ETabRole::PanelTab)
 		[
-			SAssignNew(DataStreamBreakdownView, SDataStreamBreakdownView)
+			SAssignNew(NetStatsView, SNetStatsView)
 		];
 
-	DockTab->SetOnTabClosed(SDockTab::FOnTabClosedCallback::CreateRaw(this, &SNetworkingProfilerWindow::OnDataStreamBreakdownViewTabClosed));
+	DockTab->SetOnTabClosed(SDockTab::FOnTabClosedCallback::CreateRaw(this, &SNetworkingProfilerWindow::OnNetStatsViewTabClosed));
 
 	return DockTab;
 }
 
 ////////////////////////////////////////////////////////////////////////////////////////////////////
 
-void SNetworkingProfilerWindow::OnDataStreamBreakdownViewTabClosed(TSharedRef<SDockTab> TabBeingClosed)
+void SNetworkingProfilerWindow::OnNetStatsViewTabClosed(TSharedRef<SDockTab> TabBeingClosed)
 {
-	FNetworkingProfilerManager::Get()->SetDataStreamBreakdownViewVisible(false);
-	DataStreamBreakdownView = nullptr;
+	NetStatsView = nullptr;
 }
 
 ////////////////////////////////////////////////////////////////////////////////////////////////////
 
 void SNetworkingProfilerWindow::Construct(const FArguments& InArgs, const TSharedRef<SDockTab>& ConstructUnderMajorTab, const TSharedPtr<SWindow>& ConstructUnderWindow)
 {
+	CommandList = MakeShareable(new FUICommandList);
+
 	// Create & initialize tab manager.
 	TabManager = FGlobalTabmanager::Get()->NewTabManager(ConstructUnderMajorTab);
-	TSharedRef<FWorkspaceItem> AppMenuGroup = TabManager->AddLocalWorkspaceMenuCategory(LOCTEXT("NetworkingProfilerMenuGroupName", "Timing Insights"));
+	TSharedRef<FWorkspaceItem> AppMenuGroup = TabManager->AddLocalWorkspaceMenuCategory(LOCTEXT("NetworkingProfilerMenuGroupName", "Networking Insights"));
 
 	TabManager->RegisterTabSpawner(FNetworkingProfilerTabs::ToolbarID, FOnSpawnTab::CreateRaw(this, &SNetworkingProfilerWindow::SpawnTab_Toolbar))
 		.SetDisplayName(LOCTEXT("DeviceToolbarTabTitle", "Toolbar"))
 		.SetIcon(FSlateIcon(FInsightsStyle::GetStyleSetName(), "Toolbar.Icon.Small"))
 		.SetGroup(AppMenuGroup);
 
-	TabManager->RegisterTabSpawner(FNetworkingProfilerTabs::PacketSizesViewID, FOnSpawnTab::CreateRaw(this, &SNetworkingProfilerWindow::SpawnTab_PacketSizesView))
-		.SetDisplayName(LOCTEXT("NetworkingProfiler.PacketSizesViewTabTitle", "Packet Sizes"))
-		.SetIcon(FSlateIcon(FInsightsStyle::GetStyleSetName(), "PacketSizesView.Icon.Small"))
+	TabManager->RegisterTabSpawner(FNetworkingProfilerTabs::PacketViewID, FOnSpawnTab::CreateRaw(this, &SNetworkingProfilerWindow::SpawnTab_PacketView))
+		.SetDisplayName(LOCTEXT("NetworkingProfiler.PacketViewTabTitle", "Packet View"))
+		.SetIcon(FSlateIcon(FInsightsStyle::GetStyleSetName(), "PacketView.Icon.Small"))
 		.SetGroup(AppMenuGroup);
 
-	TabManager->RegisterTabSpawner(FNetworkingProfilerTabs::PacketBreakdownViewID, FOnSpawnTab::CreateRaw(this, &SNetworkingProfilerWindow::SpawnTab_PacketBreakdownView))
-		.SetDisplayName(LOCTEXT("NetworkingProfiler.PacketBreakdownViewTabTitle", "Packet Breakdown"))
-		.SetIcon(FSlateIcon(FInsightsStyle::GetStyleSetName(), "PacketBreakdownView.Icon.Small"))
+	TabManager->RegisterTabSpawner(FNetworkingProfilerTabs::PacketContentViewID, FOnSpawnTab::CreateRaw(this, &SNetworkingProfilerWindow::SpawnTab_PacketContentView))
+		.SetDisplayName(LOCTEXT("NetworkingProfiler.PacketContentViewTabTitle", "Packet Content"))
+		.SetIcon(FSlateIcon(FInsightsStyle::GetStyleSetName(), "PacketContentView.Icon.Small"))
 		.SetGroup(AppMenuGroup);
 
-	TabManager->RegisterTabSpawner(FNetworkingProfilerTabs::DataStreamBreakdownViewID, FOnSpawnTab::CreateRaw(this, &SNetworkingProfilerWindow::SpawnTab_DataStreamBreakdownView))
-		.SetDisplayName(LOCTEXT("NetworkingProfiler.DataStreamBreakdownViewTabTitle", "Data Stream Breakdown"))
-		.SetIcon(FSlateIcon(FInsightsStyle::GetStyleSetName(), "DataStreamBreakdownView.Icon.Small"))
+	TabManager->RegisterTabSpawner(FNetworkingProfilerTabs::NetStatsViewID, FOnSpawnTab::CreateRaw(this, &SNetworkingProfilerWindow::SpawnTab_NetStatsView))
+		.SetDisplayName(LOCTEXT("NetworkingProfiler.NetStatsViewTabTitle", "Net Stats"))
+		.SetIcon(FSlateIcon(FInsightsStyle::GetStyleSetName(), "NetStatsView.Icon.Small"))
 		.SetGroup(AppMenuGroup);
 
 	TSharedPtr<FNetworkingProfilerManager> NetworkingProfilerManager = FNetworkingProfilerManager::Get();
@@ -229,24 +239,37 @@ void SNetworkingProfilerWindow::Construct(const FArguments& InArgs, const TShare
 			)
 			->Split
 			(
-				FTabManager::NewStack()
-				->SetSizeCoefficient(0.35f)
-				->SetHideTabWell(true)
-				->AddTab(FNetworkingProfilerTabs::PacketSizesViewID, NetworkingProfilerManager->IsPacketSizesViewVisible() ? ETabState::OpenedTab : ETabState::ClosedTab)
+				FTabManager::NewSplitter()
+				->SetOrientation(Orient_Horizontal)
+				->SetSizeCoefficient(1.0f)
+				->Split
+				(
+					FTabManager::NewSplitter()
+					->SetOrientation(Orient_Vertical)
+					->SetSizeCoefficient(0.65f)
+					->Split
+					(
+						FTabManager::NewStack()
+						->SetSizeCoefficient(0.35f)
+						->SetHideTabWell(true)
+						->AddTab(FNetworkingProfilerTabs::PacketViewID, ETabState::OpenedTab)
+					)
+					->Split
+					(
+						FTabManager::NewStack()
+						->SetSizeCoefficient(0.65f)
+						->SetHideTabWell(true)
+						->AddTab(FNetworkingProfilerTabs::PacketContentViewID, ETabState::OpenedTab)
+					)
+				)
+				->Split
+				(
+					FTabManager::NewStack()
+					->SetSizeCoefficient(0.35f)
+					->AddTab(FNetworkingProfilerTabs::NetStatsViewID, ETabState::ClosedTab)
+					//->SetForegroundTab(FNetworkingProfilerTabs::NetStatsViewID)
+				)
 			)
-			->Split
-			(
-				FTabManager::NewStack()
-				->SetSizeCoefficient(0.65f)
-				->SetHideTabWell(true)
-				->AddTab(FNetworkingProfilerTabs::PacketBreakdownViewID, NetworkingProfilerManager->IsPacketBreakdownViewVisible() ? ETabState::OpenedTab : ETabState::ClosedTab)
-			)
-			//->Split
-			//(
-			//	FTabManager::NewStack()
-			//	->SetSizeCoefficient(0.33f)
-			//	->AddTab(FNetworkingProfilerTabs::DataStreamBreakdownViewID, NetworkingProfilerManager->IsDataStreamBreakdownViewVisible() ? ETabState::OpenedTab : ETabState::ClosedTab)
-			//)
 		);
 
 	// Create & initialize main menu.
@@ -313,9 +336,63 @@ void SNetworkingProfilerWindow::Construct(const FArguments& InArgs, const TShare
 
 	// Tell tab-manager about the global menu bar.
 	TabManager->SetMenuMultiBox(MenuBarBuilder.GetMultiBox());
+
+	BindCommands();
 }
 
+////////////////////////////////////////////////////////////////////////////////////////////////////
 END_SLATE_FUNCTION_BUILD_OPTIMIZATION
+////////////////////////////////////////////////////////////////////////////////////////////////////
+
+void SNetworkingProfilerWindow::BindCommands()
+{
+	Map_TogglePacketViewVisibility();
+	Map_TogglePacketContentViewVisibility();
+	Map_ToggleNetStatsViewVisibility();
+
+	CommandList->MapAction(FInsightsCommands::Get().ToggleDebugInfo, FInsightsManager::GetActionManager().ToggleDebugInfo_Custom());
+}
+
+////////////////////////////////////////////////////////////////////////////////////////////////////
+// Toggle Commands
+////////////////////////////////////////////////////////////////////////////////////////////////////
+
+#define IMPLEMENT_TOGGLE_COMMAND(CmdName, IsEnabled, SetEnabled) \
+	\
+	void SNetworkingProfilerWindow::Map_##CmdName()\
+	{\
+		CommandList->MapAction(FNetworkingProfilerManager::GetCommands().CmdName, CmdName##_Custom());\
+	}\
+	\
+	const FUIAction SNetworkingProfilerWindow::CmdName##_Custom() \
+	{\
+		FUIAction UIAction;\
+		UIAction.ExecuteAction = FExecuteAction::CreateRaw(this, &SNetworkingProfilerWindow::CmdName##_Execute);\
+		UIAction.CanExecuteAction = FCanExecuteAction::CreateRaw(this, &SNetworkingProfilerWindow::CmdName##_CanExecute);\
+		UIAction.GetActionCheckState = FGetActionCheckState::CreateRaw(this, &SNetworkingProfilerWindow::CmdName##_GetCheckState);\
+		return UIAction;\
+	}\
+	\
+	void SNetworkingProfilerWindow::CmdName##_Execute()\
+	{\
+		SetEnabled(!IsEnabled());\
+	}\
+	\
+	bool SNetworkingProfilerWindow::CmdName##_CanExecute() const\
+	{\
+		return FInsightsManager::Get()->GetSession().IsValid();\
+	}\
+	\
+	ECheckBoxState SNetworkingProfilerWindow::CmdName##_GetCheckState() const\
+	{\
+		return IsEnabled() ? ECheckBoxState::Checked : ECheckBoxState::Unchecked;\
+	}
+
+IMPLEMENT_TOGGLE_COMMAND(TogglePacketViewVisibility, IsPacketViewVisible, ShowOrHidePacketView)
+IMPLEMENT_TOGGLE_COMMAND(TogglePacketContentViewVisibility, IsPacketContentViewVisible, ShowOrHidePacketContentView)
+IMPLEMENT_TOGGLE_COMMAND(ToggleNetStatsViewVisibility, IsNetStatsViewVisible, ShowOrHideNetStatsView)
+
+#undef IMPLEMENT_TOGGLE_COMMAND
 
 ////////////////////////////////////////////////////////////////////////////////////////////////////
 
@@ -387,6 +464,113 @@ EActiveTimerReturnType SNetworkingProfilerWindow::UpdateActiveDuration(double In
 
 ////////////////////////////////////////////////////////////////////////////////////////////////////
 
+void SNetworkingProfilerWindow::Tick(const FGeometry& AllottedGeometry, const double InCurrentTime, const float InDeltaTime)
+{
+	TSharedPtr<const Trace::IAnalysisSession> Session = FInsightsManager::Get()->GetSession();
+
+	uint32 GameInstanceCount = 0;
+	if (Session.IsValid())
+	{
+		Trace::FAnalysisSessionReadScope SessionReadScope(*Session.Get());
+		const Trace::INetProfilerProvider& NetProfilerProvider = Trace::ReadNetProfilerProvider(*Session.Get());
+		GameInstanceCount = NetProfilerProvider.GetGameInstanceCount();
+	}
+
+	if (GameInstanceCount != AvailableGameInstances.Num())
+	{
+		UpdateAvailableGameInstances();
+	}
+
+	uint32 ConnectionCount = 0;
+	if (Session.IsValid() && SelectedGameInstance.IsValid())
+	{
+		Trace::FAnalysisSessionReadScope SessionReadScope(*Session.Get());
+		const Trace::INetProfilerProvider& NetProfilerProvider = Trace::ReadNetProfilerProvider(*Session.Get());
+		ConnectionCount = NetProfilerProvider.GetConnectionCount(SelectedGameInstance->GetIndex());
+	}
+
+	if (ConnectionCount != AvailableConnections.Num())
+	{
+		UpdateAvailableConnections();
+	}
+}
+
+////////////////////////////////////////////////////////////////////////////////////////////////////
+
+void SNetworkingProfilerWindow::UpdateAvailableGameInstances()
+{
+	AvailableGameInstances.Reset();
+
+	TSharedPtr<const Trace::IAnalysisSession> Session = FInsightsManager::Get()->GetSession();
+	if (Session.IsValid())
+	{
+		Trace::FAnalysisSessionReadScope SessionReadScope(*Session.Get());
+		const Trace::INetProfilerProvider& NetProfilerProvider = Trace::ReadNetProfilerProvider(*Session.Get());
+		NetProfilerProvider.ReadGameInstances([this](const Trace::FNetProfilerGameInstance& GameInstance)
+		{
+			AvailableGameInstances.Add(MakeShareable(new FGameInstanceItem(GameInstance)));
+		});
+	}
+
+	if (GameInstanceComboBox.IsValid())
+	{
+		GameInstanceComboBox->SetSelectedItem(AvailableGameInstances.Num() > 0 ? AvailableGameInstances[0] : nullptr);
+	}
+
+	UpdateAvailableConnections();
+}
+
+////////////////////////////////////////////////////////////////////////////////////////////////////
+
+void SNetworkingProfilerWindow::UpdateAvailableConnections()
+{
+	AvailableConnections.Reset();
+
+	TSharedPtr<const Trace::IAnalysisSession> Session = FInsightsManager::Get()->GetSession();
+	if (Session.IsValid() && SelectedGameInstance.IsValid())
+	{
+		Trace::FAnalysisSessionReadScope SessionReadScope(*Session.Get());
+		const Trace::INetProfilerProvider& NetProfilerProvider = Trace::ReadNetProfilerProvider(*Session.Get());
+		NetProfilerProvider.ReadConnections(SelectedGameInstance->GetIndex(), [this](const Trace::FNetProfilerConnection& Connection)
+		{
+			AvailableConnections.Add(MakeShareable(new FConnectionItem(Connection)));
+		});
+	}
+
+	if (ConnectionComboBox.IsValid())
+	{
+		ConnectionComboBox->SetSelectedItem(AvailableConnections.Num() > 0 ? AvailableConnections[0] : nullptr);
+	}
+
+	UpdateAvailableConnectionModes();
+}
+
+////////////////////////////////////////////////////////////////////////////////////////////////////
+
+void SNetworkingProfilerWindow::UpdateAvailableConnectionModes()
+{
+	AvailableConnectionModes.Reset();
+
+	if (SelectedConnection.IsValid())
+	{
+		if (SelectedConnection->Connection.bHasIncomingData)
+		{
+			AvailableConnectionModes.Add(MakeShareable(new FConnectionModeItem(Trace::ENetProfilerConnectionMode::Incoming)));
+		}
+		if (SelectedConnection->Connection.bHasOutgoingData)
+		{
+			AvailableConnectionModes.Add(MakeShareable(new FConnectionModeItem(Trace::ENetProfilerConnectionMode::Outgoing)));
+		}
+	}
+
+	if (ConnectionModeComboBox.IsValid())
+	{
+		ConnectionModeComboBox->SetSelectedItem(AvailableConnectionModes.Num() > 0 ? AvailableConnectionModes[0] : nullptr);
+	}
+}
+
+////////////////////////////////////////////////////////////////////////////////////////////////////
+
 void SNetworkingProfilerWindow::OnMouseEnter(const FGeometry& MyGeometry, const FPointerEvent& MouseEvent)
 {
 	SCompoundWidget::OnMouseEnter(MyGeometry, MouseEvent);
@@ -414,7 +598,7 @@ void SNetworkingProfilerWindow::OnMouseLeave(const FPointerEvent& MouseEvent)
 
 FReply SNetworkingProfilerWindow::OnKeyDown(const FGeometry& MyGeometry, const FKeyEvent& InKeyEvent)
 {
-	return FNetworkingProfilerManager::Get()->GetCommandList()->ProcessCommandBindings(InKeyEvent) ? FReply::Handled() : FReply::Unhandled();
+	return GetCommandList()->ProcessCommandBindings(InKeyEvent) ? FReply::Handled() : FReply::Unhandled();
 }
 
 ////////////////////////////////////////////////////////////////////////////////////////////////////
@@ -470,4 +654,230 @@ FReply SNetworkingProfilerWindow::OnDrop(const FGeometry& MyGeometry, const FDra
 
 ////////////////////////////////////////////////////////////////////////////////////////////////////
 
+FText SNetworkingProfilerWindow::FGameInstanceItem::GetText() const
+{
+	return FText::Format(LOCTEXT("GameInstanceItemFmt", "Game Instance {0}"), FText::AsNumber(GetIndex()));
+}
+
+////////////////////////////////////////////////////////////////////////////////////////////////////
+
+FText SNetworkingProfilerWindow::FGameInstanceItem::GetTooltipText() const
+{
+	return GetText();
+}
+
+////////////////////////////////////////////////////////////////////////////////////////////////////
+
+FText SNetworkingProfilerWindow::FConnectionItem::GetText() const
+{
+	if (!Connection.Name)
+	{
+		return FText::Format(LOCTEXT("ConnectionItemFmt1", "Connection {0}"), FText::AsNumber(GetIndex()));
+	}
+	else
+	{
+		return FText::Format(LOCTEXT("ConnectionItemFmt2", "Connection {0} ({1})"), FText::AsNumber(GetIndex()), FText::FromString(Connection.Name));
+	}
+}
+
+////////////////////////////////////////////////////////////////////////////////////////////////////
+
+FText SNetworkingProfilerWindow::FConnectionItem::GetTooltipText() const
+{
+	return GetText();
+}
+
+////////////////////////////////////////////////////////////////////////////////////////////////////
+
+FText SNetworkingProfilerWindow::FConnectionModeItem::GetText() const
+{
+	switch (Mode)
+	{
+	case Trace::ENetProfilerConnectionMode::Outgoing:
+		return LOCTEXT("ConnectionMode_Outgoing", "Outgoing");
+
+	case Trace::ENetProfilerConnectionMode::Incoming:
+		return LOCTEXT("ConnectionMode_Incoming", "Incoming");
+
+	default:
+		return LOCTEXT("ConnectionMode_Unknown", "Unknown");
+	}
+}
+
+////////////////////////////////////////////////////////////////////////////////////////////////////
+
+FText SNetworkingProfilerWindow::FConnectionModeItem::GetTooltipText() const
+{
+	return GetText();
+}
+
+////////////////////////////////////////////////////////////////////////////////////////////////////
+
+BEGIN_SLATE_FUNCTION_BUILD_OPTIMIZATION
+TSharedRef<SWidget> SNetworkingProfilerWindow::CreateGameInstanceComboBox()
+{
+	GameInstanceComboBox = SNew(SComboBox<TSharedPtr<FGameInstanceItem>>)
+		.ToolTipText(this, &SNetworkingProfilerWindow::GameInstance_GetSelectedTooltipText)
+		.OptionsSource(&AvailableGameInstances)
+		.OnSelectionChanged(this, &SNetworkingProfilerWindow::GameInstance_OnSelectionChanged)
+		.OnGenerateWidget(this, &SNetworkingProfilerWindow::GameInstance_OnGenerateWidget)
+		[
+			SNew(STextBlock)
+			.Text(this, &SNetworkingProfilerWindow::GameInstance_GetSelectedText)
+		];
+
+	return GameInstanceComboBox.ToSharedRef();
+}
+END_SLATE_FUNCTION_BUILD_OPTIMIZATION
+
+////////////////////////////////////////////////////////////////////////////////////////////////////
+
+BEGIN_SLATE_FUNCTION_BUILD_OPTIMIZATION
+TSharedRef<SWidget> SNetworkingProfilerWindow::GameInstance_OnGenerateWidget(TSharedPtr<FGameInstanceItem> InGameInstance) const
+{
+	return SNew(STextBlock)
+		.Text(InGameInstance->GetText())
+		.ToolTipText(InGameInstance->GetTooltipText());
+}
+END_SLATE_FUNCTION_BUILD_OPTIMIZATION
+
+////////////////////////////////////////////////////////////////////////////////////////////////////
+
+void SNetworkingProfilerWindow::GameInstance_OnSelectionChanged(TSharedPtr<FGameInstanceItem> NewGameInstance, ESelectInfo::Type SelectInfo)
+{
+	SelectedGameInstance = NewGameInstance;
+	UpdateAvailableConnections();
+	if (PacketView.IsValid())
+	{
+		PacketView->SetConnection(GetSelectedGameInstanceIndex(), GetSelectedConnectionIndex(), GetSelectedConnectionMode());
+	}
+}
+
+////////////////////////////////////////////////////////////////////////////////////////////////////
+
+FText SNetworkingProfilerWindow::GameInstance_GetSelectedText() const
+{
+	return SelectedGameInstance.IsValid() ? SelectedGameInstance->GetText() : LOCTEXT("NoGameInstanceText", "Game Instance N/A");
+}
+
+////////////////////////////////////////////////////////////////////////////////////////////////////
+
+FText SNetworkingProfilerWindow::GameInstance_GetSelectedTooltipText() const
+{
+	return SelectedGameInstance.IsValid() ? SelectedGameInstance->GetTooltipText() : FText();
+}
+
+////////////////////////////////////////////////////////////////////////////////////////////////////
+
+BEGIN_SLATE_FUNCTION_BUILD_OPTIMIZATION
+TSharedRef<SWidget> SNetworkingProfilerWindow::CreateConnectionComboBox()
+{
+	ConnectionComboBox = SNew(SComboBox<TSharedPtr<FConnectionItem>>)
+		.ToolTipText(this, &SNetworkingProfilerWindow::Connection_GetSelectedTooltipText)
+		.OptionsSource(&AvailableConnections)
+		.OnSelectionChanged(this, &SNetworkingProfilerWindow::Connection_OnSelectionChanged)
+		.OnGenerateWidget(this, &SNetworkingProfilerWindow::Connection_OnGenerateWidget)
+		[
+			SNew(STextBlock)
+			.Text(this, &SNetworkingProfilerWindow::Connection_GetSelectedText)
+		];
+
+	return ConnectionComboBox.ToSharedRef();
+}
+END_SLATE_FUNCTION_BUILD_OPTIMIZATION
+
+////////////////////////////////////////////////////////////////////////////////////////////////////
+
+BEGIN_SLATE_FUNCTION_BUILD_OPTIMIZATION
+TSharedRef<SWidget> SNetworkingProfilerWindow::Connection_OnGenerateWidget(TSharedPtr<FConnectionItem> InConnection) const
+{
+	return SNew(STextBlock)
+		.Text(InConnection->GetText())
+		.ToolTipText(InConnection->GetTooltipText());
+}
+END_SLATE_FUNCTION_BUILD_OPTIMIZATION
+
+////////////////////////////////////////////////////////////////////////////////////////////////////
+
+void SNetworkingProfilerWindow::Connection_OnSelectionChanged(TSharedPtr<FConnectionItem> NewConnection, ESelectInfo::Type SelectInfo)
+{
+	SelectedConnection = NewConnection;
+	UpdateAvailableConnectionModes();
+	if (PacketView.IsValid())
+	{
+		PacketView->SetConnection(GetSelectedGameInstanceIndex(), GetSelectedConnectionIndex(), GetSelectedConnectionMode());
+	}
+}
+
+////////////////////////////////////////////////////////////////////////////////////////////////////
+
+FText SNetworkingProfilerWindow::Connection_GetSelectedText() const
+{
+	return SelectedConnection.IsValid() ? SelectedConnection->GetText() : LOCTEXT("NoConnectionText", "Connection N/A");
+}
+
+////////////////////////////////////////////////////////////////////////////////////////////////////
+
+FText SNetworkingProfilerWindow::Connection_GetSelectedTooltipText() const
+{
+	return SelectedConnection.IsValid() ? SelectedConnection->GetTooltipText() : FText();
+}
+
+////////////////////////////////////////////////////////////////////////////////////////////////////
+
+BEGIN_SLATE_FUNCTION_BUILD_OPTIMIZATION
+TSharedRef<SWidget> SNetworkingProfilerWindow::CreateConnectionModeComboBox()
+{
+	ConnectionModeComboBox = SNew(SComboBox<TSharedPtr<FConnectionModeItem>>)
+		.ToolTipText(this, &SNetworkingProfilerWindow::ConnectionMode_GetSelectedTooltipText)
+		.OptionsSource(&AvailableConnectionModes)
+		.OnSelectionChanged(this, &SNetworkingProfilerWindow::ConnectionMode_OnSelectionChanged)
+		.OnGenerateWidget(this, &SNetworkingProfilerWindow::ConnectionMode_OnGenerateWidget)
+		[
+			SNew(STextBlock)
+			.Text(this, &SNetworkingProfilerWindow::ConnectionMode_GetSelectedText)
+		];
+
+	return ConnectionModeComboBox.ToSharedRef();
+}
+END_SLATE_FUNCTION_BUILD_OPTIMIZATION
+
+////////////////////////////////////////////////////////////////////////////////////////////////////
+
+BEGIN_SLATE_FUNCTION_BUILD_OPTIMIZATION
+TSharedRef<SWidget> SNetworkingProfilerWindow::ConnectionMode_OnGenerateWidget(TSharedPtr<FConnectionModeItem> InConnectionMode) const
+{
+	return SNew(STextBlock)
+		.Text(InConnectionMode->GetText())
+		.ToolTipText(InConnectionMode->GetTooltipText());
+}
+END_SLATE_FUNCTION_BUILD_OPTIMIZATION
+
+////////////////////////////////////////////////////////////////////////////////////////////////////
+
+void SNetworkingProfilerWindow::ConnectionMode_OnSelectionChanged(TSharedPtr<FConnectionModeItem> NewConnectionMode, ESelectInfo::Type SelectInfo)
+{
+	SelectedConnectionMode = NewConnectionMode;
+	if (PacketView.IsValid())
+	{
+		PacketView->SetConnection(GetSelectedGameInstanceIndex(), GetSelectedConnectionIndex(), GetSelectedConnectionMode());
+	}
+}
+
+////////////////////////////////////////////////////////////////////////////////////////////////////
+
+FText SNetworkingProfilerWindow::ConnectionMode_GetSelectedText() const
+{
+	return SelectedConnectionMode.IsValid() ? SelectedConnectionMode->GetText() : LOCTEXT("NoConnectionModeText", "N/A");
+}
+
+////////////////////////////////////////////////////////////////////////////////////////////////////
+
+FText SNetworkingProfilerWindow::ConnectionMode_GetSelectedTooltipText() const
+{
+	return SelectedConnectionMode.IsValid() ? SelectedConnectionMode->GetTooltipText() : FText();
+}
+
+////////////////////////////////////////////////////////////////////////////////////////////////////
+
 #undef LOCTEXT_NAMESPACE
diff --git a/Engine/Source/Developer/TraceInsights/Private/Insights/NetworkingProfiler/Widgets/SNetworkingProfilerWindow.h b/Engine/Source/Developer/TraceInsights/Private/Insights/NetworkingProfiler/Widgets/SNetworkingProfilerWindow.h
--- a/Engine/Source/Developer/TraceInsights/Private/Insights/NetworkingProfiler/Widgets/SNetworkingProfilerWindow.h
+++ b/Engine/Source/Developer/TraceInsights/Private/Insights/NetworkingProfiler/Widgets/SNetworkingProfilerWindow.h
@@ -8,8 +8,10 @@
 #include "Layout/Visibility.h"
 #include "Misc/Guid.h"
 #include "SlateFwd.h"
+#include "TraceServices/Model/NetProfiler.h"
 #include "Widgets/DeclarativeSyntaxSupport.h"
 #include "Widgets/Docking/SDockTab.h"
+#include "Widgets/Input/SComboBox.h"
 #include "Widgets/SCompoundWidget.h"
 
 // Insights
@@ -19,10 +21,11 @@
 
 class FActiveTimerHandle;
 class FMenuBuilder;
+class FUICommandList;
 
-class SDataStreamBreakdownView;
-class SPacketBreakdownView;
-class SPacketSizesView;
+class SPacketView;
+class SPacketContentView;
+class SNetStatsView;
 
 ////////////////////////////////////////////////////////////////////////////////////////////////////
 
@@ -30,9 +33,9 @@ struct FNetworkingProfilerTabs
 {
 	// Tab identifiers
 	static const FName ToolbarID;
-	static const FName PacketSizesViewID;
-	static const FName PacketBreakdownViewID;
-	static const FName DataStreamBreakdownViewID;
+	static const FName PacketViewID;
+	static const FName PacketContentViewID;
+	static const FName NetStatsViewID;
 };
 
 ////////////////////////////////////////////////////////////////////////////////////////////////////
@@ -40,6 +43,48 @@ struct FNetworkingProfilerTabs
 /** Implements the timing profiler window. */
 class SNetworkingProfilerWindow : public SCompoundWidget
 {
+private:
+	struct FGameInstanceItem
+	{
+		/** Conversion constructor. */
+		FGameInstanceItem(const Trace::FNetProfilerGameInstance& InGameInstance)
+			: GameInstance(InGameInstance)
+		{}
+
+		uint32 GetIndex() const { return GameInstance.GameInstanceIndex; }
+		FText GetText() const;
+		FText GetTooltipText() const;
+
+		const Trace::FNetProfilerGameInstance GameInstance;
+	};
+
+	struct FConnectionItem
+	{
+		/** Conversion constructor. */
+		FConnectionItem(const Trace::FNetProfilerConnection& InConnection)
+			: Connection(InConnection)
+		{}
+
+		uint32 GetIndex() const { return Connection.ConnectionIndex; }
+		FText GetText() const;
+		FText GetTooltipText() const;
+
+		const Trace::FNetProfilerConnection Connection;
+	};
+
+	struct FConnectionModeItem
+	{
+		/** Conversion constructor. */
+		FConnectionModeItem(const Trace::ENetProfilerConnectionMode& InMode)
+			: Mode(InMode)
+		{}
+
+		FText GetText() const;
+		FText GetTooltipText() const;
+
+		Trace::ENetProfilerConnectionMode Mode;
+	};
+
 public:
 	/** Default constructor. */
 	SNetworkingProfilerWindow();
@@ -55,28 +100,74 @@ public:
 	/** Constructs this widget. */
 	void Construct(const FArguments& InArgs, const TSharedRef<SDockTab>& ConstructUnderMajorTab, const TSharedPtr<SWindow>& ConstructUnderWindow);
 
+	/** @return The viewport command list */
+	const TSharedPtr<FUICommandList> GetCommandList() const { return CommandList; }
+
+	/** @return The tab manager */
+	TSharedPtr<FTabManager> GetTabManager() const { return TabManager; }
+
 	void ShowTab(const FName& TabID);
 	void HideTab(const FName& TabID);
-	void ShowHideTab(const FName& TabID, bool bShow) { bShow ? ShowTab(TabID) : HideTab(TabID); }
+	void ShowOrHideTab(const FName& TabID, bool bShow) { bShow ? ShowTab(TabID) : HideTab(TabID); }
 
-	TSharedPtr<FTabManager> GetTabManager() const { return TabManager; }
+	TSharedPtr<SPacketView> GetPacketView() const { return PacketView; }
+	const bool IsPacketViewVisible() const { return PacketView.IsValid(); }
+	void ShowOrHidePacketView(const bool bVisibleState) { ShowOrHideTab(FNetworkingProfilerTabs::PacketViewID, bVisibleState); }
+
+	TSharedPtr<SPacketContentView> GetPacketContentView() const { return PacketContentView; }
+	const bool IsPacketContentViewVisible() const { return PacketContentView.IsValid(); }
+	void ShowOrHidePacketContentView(const bool bVisibleState) { ShowOrHideTab(FNetworkingProfilerTabs::PacketContentViewID, bVisibleState); }
+
+	TSharedPtr<SNetStatsView> GetNetStatsView() const { return NetStatsView; }
+	const bool IsNetStatsViewVisible() const { return NetStatsView.IsValid(); }
+	void ShowOrHideNetStatsView(const bool bVisibleState) { ShowOrHideTab(FNetworkingProfilerTabs::NetStatsViewID, bVisibleState); }
 
-	TSharedPtr<SPacketSizesView> GetPacketSizesView() const { return PacketSizesView; }
-	TSharedPtr<SPacketBreakdownView> GetPacketBreakdownView() const { return PacketBreakdownView; }
-	TSharedPtr<SDataStreamBreakdownView> GetDataStreamBreakdownView() const { return DataStreamBreakdownView; }
+	const Trace::FNetProfilerGameInstance* GetSelectedGameInstance() const { return SelectedGameInstance ? &SelectedGameInstance->GameInstance : nullptr; }
+	uint32 GetSelectedGameInstanceIndex() const { return SelectedGameInstance ? SelectedGameInstance->GetIndex() : 0; }
+	const Trace::FNetProfilerConnection* GetSelectedConnection() const { return SelectedConnection ? &SelectedConnection->Connection : nullptr; }
+	uint32 GetSelectedConnectionIndex() const { return SelectedConnection ? SelectedConnection->GetIndex() : 0; }
+	Trace::ENetProfilerConnectionMode GetSelectedConnectionMode() const { return SelectedConnectionMode ? SelectedConnectionMode->Mode : Trace::ENetProfilerConnectionMode::Outgoing; }
+
+	TSharedRef<SWidget> CreateGameInstanceComboBox();
+	TSharedRef<SWidget> CreateConnectionComboBox();
+	TSharedRef<SWidget> CreateConnectionModeComboBox();
 
 private:
 	TSharedRef<SDockTab> SpawnTab_Toolbar(const FSpawnTabArgs& Args);
 	void OnToolbarTabClosed(TSharedRef<SDockTab> TabBeingClosed);
 
-	TSharedRef<SDockTab> SpawnTab_PacketSizesView(const FSpawnTabArgs& Args);
-	void OnPacketSizesViewTabClosed(TSharedRef<SDockTab> TabBeingClosed);
+	TSharedRef<SDockTab> SpawnTab_PacketView(const FSpawnTabArgs& Args);
+	void OnPacketViewTabClosed(TSharedRef<SDockTab> TabBeingClosed);
+
+	TSharedRef<SDockTab> SpawnTab_PacketContentView(const FSpawnTabArgs& Args);
+	void OnPacketContentViewTabClosed(TSharedRef<SDockTab> TabBeingClosed);
+
+	TSharedRef<SDockTab> SpawnTab_NetStatsView(const FSpawnTabArgs& Args);
+	void OnNetStatsViewTabClosed(TSharedRef<SDockTab> TabBeingClosed);
+
+	//////////////////////////////////////////////////
 
-	TSharedRef<SDockTab> SpawnTab_PacketBreakdownView(const FSpawnTabArgs& Args);
-	void OnPacketBreakdownViewTabClosed(TSharedRef<SDockTab> TabBeingClosed);
+	void BindCommands();
 
-	TSharedRef<SDockTab> SpawnTab_DataStreamBreakdownView(const FSpawnTabArgs& Args);
-	void OnDataStreamBreakdownViewTabClosed(TSharedRef<SDockTab> TabBeingClosed);
+	//////////////////////////////////////////////////
+	// Toggle Commands
+
+#define DECLARE_TOGGLE_COMMAND(CmdName)\
+public:\
+	void Map_##CmdName(); /**< Maps UI command info CmdName with the specified UI command list. */\
+	const FUIAction CmdName##_Custom(); /**< UI action for CmdName command. */\
+protected:\
+	void CmdName##_Execute(); /**< Handles FExecuteAction for CmdName. */\
+	bool CmdName##_CanExecute() const; /**< Handles FCanExecuteAction for CmdName. */\
+	ECheckBoxState CmdName##_GetCheckState() const; /**< Handles FGetActionCheckState for CmdName. */
+
+	DECLARE_TOGGLE_COMMAND(TogglePacketViewVisibility)
+	DECLARE_TOGGLE_COMMAND(TogglePacketContentViewVisibility)
+	DECLARE_TOGGLE_COMMAND(ToggleNetStatsViewVisibility)
+#undef DECLARE_TOGGLE_COMMAND
+
+	//////////////////////////////////////////////////
+private:
 
 	/**
 	 * Fill the main menu with menu items.
@@ -95,6 +186,8 @@ private:
 	/** Updates the amount of time the profiler has been active. */
 	EActiveTimerReturnType UpdateActiveDuration(double InCurrentTime, float InDeltaTime);
 
+	//////////////////////////////////////////////////
+
 	/**
 	 * Ticks this widget. Override in derived classes, but always call the parent implementation.
 	 *
@@ -102,7 +195,7 @@ private:
 	 * @param  InCurrentTime  Current absolute real time
 	 * @param  InDeltaTime  Real time passed since last tick
 	 */
-	//virtual void Tick(const FGeometry& AllottedGeometry, const double InCurrentTime, const float InDeltaTime) override;
+	virtual void Tick(const FGeometry& AllottedGeometry, const double InCurrentTime, const float InDeltaTime) override;
 
 	/**
 	 * The system will use this event to notify a widget that the cursor has entered it. This event is NOT bubbled.
@@ -149,19 +242,57 @@ private:
 	 */
 	virtual FReply OnDragOver(const FGeometry& MyGeometry, const FDragDropEvent& DragDropEvent)  override;
 
-private:
-	/** The Packet Sizes widget */
-	TSharedPtr<SPacketSizesView> PacketSizesView;
+	//////////////////////////////////////////////////
+
+	void UpdateAvailableGameInstances();
+	void UpdateAvailableConnections();
+	void UpdateAvailableConnectionModes();
+
+	TSharedRef<SWidget> GameInstance_OnGenerateWidget(TSharedPtr<FGameInstanceItem> InGameInstance) const;
+	void GameInstance_OnSelectionChanged(TSharedPtr<FGameInstanceItem> NewGameInstance, ESelectInfo::Type SelectInfo);
+	FText GameInstance_GetSelectedText() const;
+	FText GameInstance_GetSelectedTooltipText() const;
+
+	TSharedRef<SWidget> Connection_OnGenerateWidget(TSharedPtr<FConnectionItem> InConnection) const;
+	void Connection_OnSelectionChanged(TSharedPtr<FConnectionItem> NewConnection, ESelectInfo::Type SelectInfo);
+	FText Connection_GetSelectedText() const;
+	FText Connection_GetSelectedTooltipText() const;
+
+	TSharedRef<SWidget> ConnectionMode_OnGenerateWidget(TSharedPtr<FConnectionModeItem> InConnectionMode) const;
+	void ConnectionMode_OnSelectionChanged(TSharedPtr<FConnectionModeItem> NewConnectionMode, ESelectInfo::Type SelectInfo);
+	FText ConnectionMode_GetSelectedText() const;
+	FText ConnectionMode_GetSelectedTooltipText() const;
 
-	/** The Packet Breakdown widget */
-	TSharedPtr<SPacketBreakdownView> PacketBreakdownView;
+	//////////////////////////////////////////////////
 
-	/** The Data Stream Breakdown widget */
-	TSharedPtr<SDataStreamBreakdownView> DataStreamBreakdownView;
+private:
+	/** Commandlist used in the window (Maps commands to window specific actions) */
+	TSharedPtr<FUICommandList> CommandList;
 
 	/** Holds the tab manager that manages the front-end's tabs. */
 	TSharedPtr<FTabManager> TabManager;
 
+	/** The Packet View widget */
+	TSharedPtr<SPacketView> PacketView;
+
+	/** The Packet Content widget */
+	TSharedPtr<SPacketContentView> PacketContentView;
+
+	/** The Net Stats widget */
+	TSharedPtr<SNetStatsView> NetStatsView;
+
+	TSharedPtr<SComboBox<TSharedPtr<FGameInstanceItem>>> GameInstanceComboBox;
+	TArray<TSharedPtr<FGameInstanceItem>> AvailableGameInstances;
+	TSharedPtr<FGameInstanceItem> SelectedGameInstance;
+
+	TSharedPtr<SComboBox<TSharedPtr<FConnectionItem>>> ConnectionComboBox;
+	TArray<TSharedPtr<FConnectionItem>> AvailableConnections;
+	TSharedPtr<FConnectionItem> SelectedConnection;
+
+	TSharedPtr<SComboBox<TSharedPtr<FConnectionModeItem>>> ConnectionModeComboBox;
+	TArray<TSharedPtr<FConnectionModeItem>> AvailableConnectionModes;
+	TSharedPtr<FConnectionModeItem> SelectedConnectionMode;
+
 	/** The handle to the active update duration tick */
 	TWeakPtr<FActiveTimerHandle> ActiveTimerHandle;
 
diff --git a/Engine/Source/Developer/TraceInsights/Private/Insights/NetworkingProfiler/Widgets/SPacketBreakdownView.cpp b/Engine/Source/Developer/TraceInsights/Private/Insights/NetworkingProfiler/Widgets/SPacketBreakdownView.cpp
--- a/Engine/Source/Developer/TraceInsights/Private/Insights/NetworkingProfiler/Widgets/SPacketBreakdownView.cpp
+++ b/Engine/Source/Developer/TraceInsights/Private/Insights/NetworkingProfiler/Widgets/SPacketBreakdownView.cpp
@@ -1,6 +1,7 @@
 // Copyright 1998-2019 Epic Games, Inc. All Rights Reserved.
 
 #include "SPacketBreakdownView.h"
+//#include "SPacketContentView.h"
 
 #include "Fonts/FontMeasure.h"
 #include "Fonts/SlateFontInfo.h"
@@ -18,25 +19,25 @@
 
 ////////////////////////////////////////////////////////////////////////////////////////////////////
 
-#define LOCTEXT_NAMESPACE "SPacketBreakdownView"
+#define LOCTEXT_NAMESPACE "SPacketContentView"
 
 ////////////////////////////////////////////////////////////////////////////////////////////////////
 
-SPacketBreakdownView::SPacketBreakdownView()
-	: DrawState(MakeShareable(new FPacketBreakdownViewDrawState()))
+SPacketContentView::SPacketContentView()
+	: DrawState(MakeShareable(new FPacketContentViewDrawState()))
 {
 	Reset();
 }
 
 ////////////////////////////////////////////////////////////////////////////////////////////////////
 
-SPacketBreakdownView::~SPacketBreakdownView()
+SPacketContentView::~SPacketContentView()
 {
 }
 
 ////////////////////////////////////////////////////////////////////////////////////////////////////
 
-void SPacketBreakdownView::Reset()
+void SPacketContentView::Reset()
 {
 	Viewport.Reset();
 	//FAxisViewportDouble& ViewportX = Viewport.GetHorizontalAxisViewport();
@@ -44,8 +45,12 @@ void SPacketBreakdownView::Reset()
 	//ViewportX.SetScale(1.0);
 	bIsViewportDirty = true;
 
-	PacketFrameIndex = -1;
-	PacketSize = 0;
+	GameInstanceIndex = 0;
+	ConnectionIndex = 0;
+	ConnectionMode = Trace::ENetProfilerConnectionMode::Outgoing;
+	PacketIndex = 0;
+	PacketBitSize = 0;
+
 	DrawState->Reset();
 	bIsStateDirty = true;
 
@@ -76,7 +81,7 @@ void SPacketBreakdownView::Reset()
 
 ////////////////////////////////////////////////////////////////////////////////////////////////////
 
-void SPacketBreakdownView::Construct(const FArguments& InArgs)
+void SPacketContentView::Construct(const FArguments& InArgs)
 {
 	ChildSlot
 	[
@@ -93,7 +98,7 @@ void SPacketBreakdownView::Construct(const FArguments& InArgs)
 			.Visibility(EVisibility::Visible)
 			.Thickness(FVector2D(5.0f, 5.0f))
 			.RenderOpacity(0.75)
-			.OnUserScrolled(this, &SPacketBreakdownView::HorizontalScrollBar_OnUserScrolled)
+			.OnUserScrolled(this, &SPacketContentView::HorizontalScrollBar_OnUserScrolled)
 		]
 	];
 
@@ -104,7 +109,7 @@ void SPacketBreakdownView::Construct(const FArguments& InArgs)
 
 ////////////////////////////////////////////////////////////////////////////////////////////////////
 
-void SPacketBreakdownView::Tick(const FGeometry& AllottedGeometry, const double InCurrentTime, const float InDeltaTime)
+void SPacketContentView::Tick(const FGeometry& AllottedGeometry, const double InCurrentTime, const float InDeltaTime)
 {
 	if (ThisGeometry != AllottedGeometry || bIsViewportDirty)
 	{
@@ -143,14 +148,17 @@ void SPacketBreakdownView::Tick(const FGeometry& AllottedGeometry, const double
 
 ////////////////////////////////////////////////////////////////////////////////////////////////////
 
-void SPacketBreakdownView::SetPacket(int32 InFrameIndex, int64 InBitSize)
+void SPacketContentView::ResetPacket()
 {
-	PacketFrameIndex = InFrameIndex;
-	PacketSize = InBitSize;
+	GameInstanceIndex = 0;
+	ConnectionIndex = 0;
+	ConnectionMode = Trace::ENetProfilerConnectionMode::Outgoing;
+	PacketIndex = 0;
+	PacketBitSize = 0;
 
 	FAxisViewportDouble& ViewportX = Viewport.GetHorizontalAxisViewport();
-	ViewportX.SetMinMaxValueInterval(0.0f, static_cast<double>(InBitSize));
-	ViewportX.CenterOnValueInterval(0.0f, static_cast<double>(InBitSize));
+	ViewportX.SetMinMaxValueInterval(0.0, 0.0);
+	ViewportX.CenterOnValue(0.0);
 	UpdateHorizontalScrollBar();
 
 	DrawState->Reset();
@@ -159,9 +167,33 @@ void SPacketBreakdownView::SetPacket(int32 InFrameIndex, int64 InBitSize)
 
 ////////////////////////////////////////////////////////////////////////////////////////////////////
 
-void AddMockEventsRec(FPacketBreakdownViewDrawStateBuilder& Builder, FRandomStream& RandomStream , int32 Depth, int64 Offset, int64 Size, int32 Type)
+void SPacketContentView::SetPacket(uint32 InGameInstanceIndex, uint32 InConnectionIndex, Trace::ENetProfilerConnectionMode InConnectionMode, uint32 InPacketIndex, int64 InPacketBitSize)
 {
-	Builder.AddEvent(Offset, Size, Type, Depth);
+	GameInstanceIndex = InGameInstanceIndex;
+	ConnectionIndex = InConnectionIndex;
+	ConnectionMode = InConnectionMode;
+	PacketIndex = InPacketIndex;
+	PacketBitSize = InPacketBitSize;
+
+	FAxisViewportDouble& ViewportX = Viewport.GetHorizontalAxisViewport();
+	ViewportX.SetMinMaxValueInterval(0.0, static_cast<double>(InPacketBitSize));
+	ViewportX.CenterOnValueInterval(0.0, static_cast<double>(InPacketBitSize));
+	UpdateHorizontalScrollBar();
+
+	DrawState->Reset();
+	bIsStateDirty = true;
+}
+
+////////////////////////////////////////////////////////////////////////////////////////////////////
+
+void AddMockEventsRec(FPacketContentViewDrawStateBuilder& Builder, FRandomStream& RandomStream , int32 Depth, int64 Offset, int64 Size, int32 Type)
+{
+	Trace::FNetProfilerContentEvent Event;
+	Event.StartPos = Offset;
+	Event.EndPos = Offset + Size;
+	Event.Level = Depth;
+	Event.NameIndex = Type;
+	Builder.AddEvent(Event);
 
 	int64 ChildEventOffset = Offset;
 	while (ChildEventOffset < Offset + Size)
@@ -184,18 +216,56 @@ void AddMockEventsRec(FPacketBreakdownViewDrawStateBuilder& Builder, FRandomStre
 
 ////////////////////////////////////////////////////////////////////////////////////////////////////
 
-void SPacketBreakdownView::UpdateState()
+void SPacketContentView::UpdateState()
 {
 	FStopwatch Stopwatch;
 	Stopwatch.Start();
 
-	if (PacketFrameIndex >= 0 && PacketSize > 0)
+	if (PacketBitSize > 0)
 	{
-		FPacketBreakdownViewDrawStateBuilder Builder(*DrawState, Viewport);
+		FPacketContentViewDrawStateBuilder Builder(*DrawState, Viewport);
+
+		TSharedPtr<const Trace::IAnalysisSession> Session = FInsightsManager::Get()->GetSession();
+		if (Session.IsValid())
+		{
+			Trace::FAnalysisSessionReadScope SessionReadScope(*Session.Get());
+			const Trace::INetProfilerProvider& NetProfilerProvider = Trace::ReadNetProfilerProvider(*Session.Get());
+
+			//uint32 StartEventIndex;
+			//uint32 EndEventIndex;
+			//NetProfilerProvider.ReadConnection(ConnectionIndex, [&StartEventIndex, &EndEventIndex](const Trace::FNetProfilerConnection& Connection)
+			//{
+			//	StartEventIndex = Connection.StartEventIndex;
+			//	EndEventIndex = Connection.EndEventIndex;
+			//});
+			//NetProfilerProvider.EnumeratePacketContentEventsByIndex(ConnectionIndex, ConnectionMode, StartEventIndex, uint32 EndEventIndex, [](const Trace::FNetProfilerContentEvent& Event)
+			//{
+			//});
+
+			const FAxisViewportDouble& ViewportX = Viewport.GetHorizontalAxisViewport();
+
+			const int64 MinPos = static_cast<int64>(FMath::FloorToDouble(ViewportX.GetValueAtOffset(0.0f)));
+			const int64 MaxPos = static_cast<int64>(FMath::CeilToDouble(ViewportX.GetValueAtOffset(ViewportX.GetSize())));
+
+			const uint32 StartPos = 0;
+			const uint32 EndPos = PacketBitSize;
+			NetProfilerProvider.EnumeratePacketContentEventsByPosition(ConnectionIndex, ConnectionMode, PacketIndex, StartPos, EndPos, [&Builder, &NetProfilerProvider](const Trace::FNetProfilerContentEvent& Event)
+			{
+				const TCHAR* Name = nullptr;
+				NetProfilerProvider.ReadName(Event.NameIndex, [&Name](const Trace::FNetProfilerName& NetProfilerName)
+				{
+					Name = NetProfilerName.Name;
+				});
+				Builder.AddEvent(Event, Name);
+			});
+		}
 
 		// Init packet content with mock data.
-		FRandomStream RandomStream(PacketFrameIndex);
-		AddMockEventsRec(Builder, RandomStream, 0, 0, PacketSize, 0);
+		if (false)
+		{
+			FRandomStream RandomStream(PacketIndex);
+			AddMockEventsRec(Builder, RandomStream, 0, 0, PacketBitSize, 0);
+		}
 
 		Builder.Flush();
 	}
@@ -206,7 +276,7 @@ void SPacketBreakdownView::UpdateState()
 
 ////////////////////////////////////////////////////////////////////////////////////////////////////
 
-void SPacketBreakdownView::UpdateHoveredEvent()
+void SPacketContentView::UpdateHoveredEvent()
 {
 	HoveredEvent = GetEventAtMousePosition(MousePosition.X, MousePosition.Y);
 	//if (!HoveredEvent.IsValid())
@@ -246,7 +316,7 @@ void SPacketBreakdownView::UpdateHoveredEvent()
 
 ////////////////////////////////////////////////////////////////////////////////////////////////////
 
-FNetworkPacketEventRef SPacketBreakdownView::GetEventAtMousePosition(float X, float Y)
+FNetworkPacketEventRef SPacketContentView::GetEventAtMousePosition(float X, float Y)
 {
 	if (!bIsStateDirty)
 	{
@@ -276,7 +346,7 @@ FNetworkPacketEventRef SPacketBreakdownView::GetEventAtMousePosition(float X, fl
 
 ////////////////////////////////////////////////////////////////////////////////////////////////////
 
-int32 SPacketBreakdownView::OnPaint(const FPaintArgs& Args, const FGeometry& AllottedGeometry, const FSlateRect& MyCullingRect, FSlateWindowElementList& OutDrawElements, int32 LayerId, const FWidgetStyle& InWidgetStyle, bool bParentEnabled) const
+int32 SPacketContentView::OnPaint(const FPaintArgs& Args, const FGeometry& AllottedGeometry, const FSlateRect& MyCullingRect, FSlateWindowElementList& OutDrawElements, int32 LayerId, const FWidgetStyle& InWidgetStyle, bool bParentEnabled) const
 {
 	const bool bEnabled = ShouldBeEnabled(bParentEnabled);
 	const ESlateDrawEffect DrawEffects = bEnabled ? ESlateDrawEffect::None : ESlateDrawEffect::DisabledEffect;
@@ -290,7 +360,7 @@ int32 SPacketBreakdownView::OnPaint(const FPaintArgs& Args, const FGeometry& All
 		FStopwatch Stopwatch;
 		Stopwatch.Start();
 
-		FPacketBreakdownViewDrawHelper Helper(DrawContext, Viewport);
+		FPacketContentViewDrawHelper Helper(DrawContext, Viewport);
 
 		Helper.DrawBackground();
 
@@ -321,8 +391,8 @@ int32 SPacketBreakdownView::OnPaint(const FPaintArgs& Args, const FGeometry& All
 		const float MaxFontCharHeight = FontMeasureService->Measure(TEXT("!"), SummaryFont).Y;
 		const float DbgDY = MaxFontCharHeight;
 
-		const float DbgW = 260.0f;
-		const float DbgH = DbgDY * 4 + 3.0f;
+		const float DbgW = 280.0f;
+		const float DbgH = DbgDY * 5 + 3.0f;
 		const float DbgX = ViewWidth - DbgW - 20.0f;
 		float DbgY = 7.0f;
 
@@ -390,6 +460,19 @@ int32 SPacketBreakdownView::OnPaint(const FPaintArgs& Args, const FGeometry& All
 			SummaryFont, DbgTextColor
 		);
 		DbgY += DbgDY;
+
+		// Draw packet info.
+		DrawContext.DrawText
+		(
+			DbgX, DbgY,
+			FString::Printf(TEXT("Game Instance %d, Connection %d (%s), Packet %d"),
+				GameInstanceIndex,
+				ConnectionIndex,
+				(ConnectionMode == Trace::ENetProfilerConnectionMode::Outgoing) ? TEXT("Outgoing") : TEXT("Incoming"),
+				PacketIndex),
+			SummaryFont, DbgTextColor
+		);
+		DbgY += DbgDY;
 	}
 
 	return SCompoundWidget::OnPaint(Args, AllottedGeometry, MyCullingRect, OutDrawElements, LayerId, InWidgetStyle, bParentEnabled && IsEnabled());
@@ -397,7 +480,7 @@ int32 SPacketBreakdownView::OnPaint(const FPaintArgs& Args, const FGeometry& All
 
 ////////////////////////////////////////////////////////////////////////////////////////////////////
 
-FReply SPacketBreakdownView::OnMouseButtonDown(const FGeometry& MyGeometry, const FPointerEvent& MouseEvent)
+FReply SPacketContentView::OnMouseButtonDown(const FGeometry& MyGeometry, const FPointerEvent& MouseEvent)
 {
 	FReply Reply = FReply::Unhandled();
 
@@ -424,7 +507,7 @@ FReply SPacketBreakdownView::OnMouseButtonDown(const FGeometry& MyGeometry, cons
 
 ////////////////////////////////////////////////////////////////////////////////////////////////////
 
-FReply SPacketBreakdownView::OnMouseButtonUp(const FGeometry& MyGeometry, const FPointerEvent& MouseEvent)
+FReply SPacketContentView::OnMouseButtonUp(const FGeometry& MyGeometry, const FPointerEvent& MouseEvent)
 {
 	FReply Reply = FReply::Unhandled();
 
@@ -478,7 +561,7 @@ FReply SPacketBreakdownView::OnMouseButtonUp(const FGeometry& MyGeometry, const
 
 ////////////////////////////////////////////////////////////////////////////////////////////////////
 
-FReply SPacketBreakdownView::OnMouseMove(const FGeometry& MyGeometry, const FPointerEvent& MouseEvent)
+FReply SPacketContentView::OnMouseMove(const FGeometry& MyGeometry, const FPointerEvent& MouseEvent)
 {
 	FReply Reply = FReply::Unhandled();
 
@@ -520,13 +603,13 @@ FReply SPacketBreakdownView::OnMouseMove(const FGeometry& MyGeometry, const FPoi
 
 ////////////////////////////////////////////////////////////////////////////////////////////////////
 
-void SPacketBreakdownView::OnMouseEnter(const FGeometry& MyGeometry, const FPointerEvent& MouseEvent)
+void SPacketContentView::OnMouseEnter(const FGeometry& MyGeometry, const FPointerEvent& MouseEvent)
 {
 }
 
 ////////////////////////////////////////////////////////////////////////////////////////////////////
 
-void SPacketBreakdownView::OnMouseLeave(const FPointerEvent& MouseEvent)
+void SPacketContentView::OnMouseLeave(const FPointerEvent& MouseEvent)
 {
 	if (!HasMouseCapture())
 	{
@@ -547,7 +630,7 @@ void SPacketBreakdownView::OnMouseLeave(const FPointerEvent& MouseEvent)
 
 ////////////////////////////////////////////////////////////////////////////////////////////////////
 
-FReply SPacketBreakdownView::OnMouseWheel(const FGeometry& MyGeometry, const FPointerEvent& MouseEvent)
+FReply SPacketContentView::OnMouseWheel(const FGeometry& MyGeometry, const FPointerEvent& MouseEvent)
 {
 	MousePosition = MyGeometry.AbsoluteToLocal(MouseEvent.GetScreenSpacePosition());
 
@@ -566,14 +649,14 @@ FReply SPacketBreakdownView::OnMouseWheel(const FGeometry& MyGeometry, const FPo
 
 ////////////////////////////////////////////////////////////////////////////////////////////////////
 
-FReply SPacketBreakdownView::OnMouseButtonDoubleClick(const FGeometry& MyGeometry, const FPointerEvent& MouseEvent)
+FReply SPacketContentView::OnMouseButtonDoubleClick(const FGeometry& MyGeometry, const FPointerEvent& MouseEvent)
 {
 	return FReply::Unhandled();
 }
 
 ////////////////////////////////////////////////////////////////////////////////////////////////////
 
-FCursorReply SPacketBreakdownView::OnCursorQuery(const FGeometry& MyGeometry, const FPointerEvent& CursorEvent) const
+FCursorReply SPacketContentView::OnCursorQuery(const FGeometry& MyGeometry, const FPointerEvent& CursorEvent) const
 {
 	FCursorReply CursorReply = FCursorReply::Unhandled();
 
@@ -591,13 +674,13 @@ FCursorReply SPacketBreakdownView::OnCursorQuery(const FGeometry& MyGeometry, co
 
 ////////////////////////////////////////////////////////////////////////////////////////////////////
 
-void SPacketBreakdownView::BindCommands()
+void SPacketContentView::BindCommands()
 {
 }
 
 ////////////////////////////////////////////////////////////////////////////////////////////////////
 
-void SPacketBreakdownView::HorizontalScrollBar_OnUserScrolled(float ScrollOffset)
+void SPacketContentView::HorizontalScrollBar_OnUserScrolled(float ScrollOffset)
 {
 	FAxisViewportDouble& ViewportX = Viewport.GetHorizontalAxisViewport();
 	ViewportX.OnUserScrolled(HorizontalScrollBar, ScrollOffset);
@@ -606,7 +689,7 @@ void SPacketBreakdownView::HorizontalScrollBar_OnUserScrolled(float ScrollOffset
 
 ////////////////////////////////////////////////////////////////////////////////////////////////////
 
-void SPacketBreakdownView::UpdateHorizontalScrollBar()
+void SPacketContentView::UpdateHorizontalScrollBar()
 {
 	FAxisViewportDouble& ViewportX = Viewport.GetHorizontalAxisViewport();
 	ViewportX.UpdateScrollBar(HorizontalScrollBar);
@@ -614,7 +697,7 @@ void SPacketBreakdownView::UpdateHorizontalScrollBar()
 
 ////////////////////////////////////////////////////////////////////////////////////////////////////
 
-void SPacketBreakdownView::ZoomHorizontally(const float Delta, const float X)
+void SPacketContentView::ZoomHorizontally(const float Delta, const float X)
 {
 	FAxisViewportDouble& ViewportX = Viewport.GetHorizontalAxisViewport();
 	ViewportX.RelativeZoomWithFixedOffset(Delta, X);
diff --git a/Engine/Source/Developer/TraceInsights/Private/Insights/NetworkingProfiler/Widgets/SPacketBreakdownView.h b/Engine/Source/Developer/TraceInsights/Private/Insights/NetworkingProfiler/Widgets/SPacketBreakdownView.h
--- a/Engine/Source/Developer/TraceInsights/Private/Insights/NetworkingProfiler/Widgets/SPacketBreakdownView.h
+++ b/Engine/Source/Developer/TraceInsights/Private/Insights/NetworkingProfiler/Widgets/SPacketBreakdownView.h
@@ -7,13 +7,16 @@
 #include "Input/Reply.h"
 #include "Layout/Geometry.h"
 #include "Rendering/RenderingCommon.h"
+#include "TraceServices/Model/NetProfiler.h"
 #include "Widgets/DeclarativeSyntaxSupport.h"
 #include "Widgets/SCompoundWidget.h"
 
 // Insights
 #include "Insights/Common/FixedCircularBuffer.h"
 #include "Insights/NetworkingProfiler/ViewModels/PacketBreakdownViewHelper.h"
+//#include "Insights/NetworkingProfiler/ViewModels/PacketContentViewHelper.h"
 #include "Insights/NetworkingProfiler/ViewModels/PacketBreakdownViewport.h"
+//#include "Insights/NetworkingProfiler/ViewModels/PacketContentViewport.h"
 #include "Insights/ViewModels/TooltipDrawState.h"
 
 class SScrollBar;
@@ -70,7 +73,7 @@ struct FNetworkPacketEventRef
 /**
  * Widget used to present content of a network packet.
  */
-class SPacketBreakdownView : public SCompoundWidget
+class SPacketContentView : public SCompoundWidget
 {
 public:
 	/** Number of pixels. */
@@ -85,15 +88,15 @@ public:
 
 public:
 	/** Default constructor. */
-	SPacketBreakdownView();
+	SPacketContentView();
 
 	/** Virtual destructor. */
-	virtual ~SPacketBreakdownView();
+	virtual ~SPacketContentView();
 
 	/** Resets internal widget's data to the default one. */
 	void Reset();
 
-	SLATE_BEGIN_ARGS(SPacketBreakdownView) {}
+	SLATE_BEGIN_ARGS(SPacketContentView) {}
 	SLATE_END_ARGS()
 
 	/**
@@ -122,7 +125,8 @@ public:
 
 	virtual FCursorReply OnCursorQuery(const FGeometry& MyGeometry, const FPointerEvent& CursorEvent) const override;
 
-	void SetPacket(int32 InPacketFrameIndex, int64 InPacketBitSize);
+	void ResetPacket();
+	void SetPacket(uint32 InGameInstanceIndex, uint32 InConnectionIndex, Trace::ENetProfilerConnectionMode InConnectionMode, uint32 InPacketIndex, int64 InPacketBitSize);
 
 private:
 	void UpdateState();
@@ -145,14 +149,17 @@ private:
 
 private:
 	/** The track's viewport. Encapsulates info about position and scale. */
-	FPacketBreakdownViewport Viewport;
+	FPacketContentViewport Viewport;
 	bool bIsViewportDirty;
 
-	int32 PacketFrameIndex;
-	int64 PacketSize; // total number of bits; [bit]
+	uint32 GameInstanceIndex;
+	uint32 ConnectionIndex;
+	Trace::ENetProfilerConnectionMode ConnectionMode;
+	uint32 PacketIndex;
+	int64 PacketBitSize; // total number of bits; [bit]
 
 	/** Cached draw state of the packet content (i.e. all it needs to render). */
-	TSharedPtr<FPacketBreakdownViewDrawState> DrawState;
+	TSharedPtr<FPacketContentViewDrawState> DrawState;
 	bool bIsStateDirty;
 	//////////////////////////////////////////////////
 
diff --git a/Engine/Source/Developer/TraceInsights/Private/Insights/NetworkingProfiler/Widgets/SPacketSizesView.cpp b/Engine/Source/Developer/TraceInsights/Private/Insights/NetworkingProfiler/Widgets/SPacketSizesView.cpp
--- a/Engine/Source/Developer/TraceInsights/Private/Insights/NetworkingProfiler/Widgets/SPacketSizesView.cpp
+++ b/Engine/Source/Developer/TraceInsights/Private/Insights/NetworkingProfiler/Widgets/SPacketSizesView.cpp
@@ -1,6 +1,7 @@
 // Copyright 1998-2019 Epic Games, Inc. All Rights Reserved.
 
 #include "SPacketSizesView.h"
+//#include "SPacketView.h"
 
 #include "Fonts/FontMeasure.h"
 #include "Fonts/SlateFontInfo.h"
@@ -8,6 +9,7 @@
 #include "Framework/MultiBox/MultiBoxBuilder.h"
 #include "HAL/PlatformTime.h"
 #include "Rendering/DrawElements.h"
+#include "TraceServices/Model/NetProfiler.h"
 #include "Styling/CoreStyle.h"
 #include "Widgets/Layout/SScrollBar.h"
 
@@ -16,39 +18,50 @@
 #include "Insights/Common/Stopwatch.h"
 #include "Insights/Common/TimeUtils.h"
 #include "Insights/InsightsManager.h"
+#include "Insights/InsightsStyle.h"
 #include "Insights/NetworkingProfiler/NetworkingProfilerManager.h"
 #include "Insights/NetworkingProfiler/Widgets/SNetworkingProfilerWindow.h"
 #include "Insights/NetworkingProfiler/Widgets/SPacketBreakdownView.h"
-#include "Insights/InsightsStyle.h"
+//#include "Insights/NetworkingProfiler/Widgets/SPacketContentView.h"
 
 #include <limits>
 
 ////////////////////////////////////////////////////////////////////////////////////////////////////
 
-#define LOCTEXT_NAMESPACE "SPacketSizesView"
+#define LOCTEXT_NAMESPACE "SPacketView"
 
 ////////////////////////////////////////////////////////////////////////////////////////////////////
 
-SPacketSizesView::SPacketSizesView()
-	: PacketSeries(MakeShareable(new FNetworkPacketSeries()))
+SPacketView::SPacketView()
+	: ProfilerWindow()
+	, PacketSeries(MakeShareable(new FNetworkPacketSeries()))
 {
 	Reset();
 }
 
 ////////////////////////////////////////////////////////////////////////////////////////////////////
 
-SPacketSizesView::~SPacketSizesView()
+SPacketView::~SPacketView()
 {
 }
 
 ////////////////////////////////////////////////////////////////////////////////////////////////////
 
-void SPacketSizesView::Reset()
+void SPacketView::Reset()
 {
+	//ProfilerWindow
+
+	GameInstanceIndex = 0;
+	ConnectionIndex = 0;
+	ConnectionMode = Trace::ENetProfilerConnectionMode::Outgoing;
+
 	Viewport.Reset();
+	FAxisViewportInt32& ViewportX = Viewport.GetHorizontalAxisViewport();
+	ViewportX.SetScaleLimits(0.0001f, 16.0f); // 10000 [sample/px] to 16 [px/sample]
+	ViewportX.SetScale(5.0f);
 	FAxisViewportDouble& ViewportY = Viewport.GetVerticalAxisViewport();
 	ViewportY.SetScaleLimits(0.0001, 50.0);
-	ViewportY.SetScale(0.08);
+	ViewportY.SetScale(0.02);
 	bIsViewportDirty = true;
 
 	PacketSeries->Reset();
@@ -58,6 +71,7 @@ void SPacketSizesView::Reset()
 	bIsAutoZoomEnabled = true;
 
 	AnalysisSyncNextTimestamp = 0;
+	ConnectionChangeCount = 0;
 
 	MousePosition = FVector2D::ZeroVector;
 
@@ -88,12 +102,35 @@ void SPacketSizesView::Reset()
 	DrawDurationHistory.Reset();
 	OnPaintDurationHistory.Reset();
 	LastOnPaintTime = FPlatformTime::Cycles64();
+
+	if (ProfilerWindow.IsValid())
+	{
+		SetConnection(ProfilerWindow->GetSelectedGameInstanceIndex(), ProfilerWindow->GetSelectedConnectionIndex(), ProfilerWindow->GetSelectedConnectionMode());
+	}
+}
+
+////////////////////////////////////////////////////////////////////////////////////////////////////
+
+void SPacketView::SetConnection(uint32 InGameInstanceIndex, uint32 InConnectionIndex, Trace::ENetProfilerConnectionMode InConnectionMode)
+{
+	GameInstanceIndex = InGameInstanceIndex;
+	ConnectionIndex = InConnectionIndex;
+	ConnectionMode = InConnectionMode;
+
+	HoveredSample.Reset();
+
+	SelectedSample.Reset();
+	OnSelectedSampleChanged();
+
+	bIsStateDirty = true;
 }
 
 ////////////////////////////////////////////////////////////////////////////////////////////////////
 
-void SPacketSizesView::Construct(const FArguments& InArgs)
+void SPacketView::Construct(const FArguments& InArgs, TSharedPtr<SNetworkingProfilerWindow> InProfilerWindow)
 {
+	ProfilerWindow = InProfilerWindow;
+
 	ChildSlot
 	[
 		SNew(SOverlay)
@@ -109,18 +146,20 @@ void SPacketSizesView::Construct(const FArguments& InArgs)
 			.Visibility(EVisibility::Visible)
 			.Thickness(FVector2D(5.0f, 5.0f))
 			.RenderOpacity(0.75)
-			.OnUserScrolled(this, &SPacketSizesView::HorizontalScrollBar_OnUserScrolled)
+			.OnUserScrolled(this, &SPacketView::HorizontalScrollBar_OnUserScrolled)
 		]
 	];
 
 	UpdateHorizontalScrollBar();
 
 	BindCommands();
+
+	Reset();
 }
 
 ////////////////////////////////////////////////////////////////////////////////////////////////////
 
-void SPacketSizesView::Tick(const FGeometry& AllottedGeometry, const double InCurrentTime, const float InDeltaTime)
+void SPacketView::Tick(const FGeometry& AllottedGeometry, const double InCurrentTime, const float InDeltaTime)
 {
 	if (ThisGeometry != AllottedGeometry || bIsViewportDirty)
 	{
@@ -150,7 +189,19 @@ void SPacketSizesView::Tick(const FGeometry& AllottedGeometry, const double InCu
 		const uint64 WaitTime = static_cast<uint64>(0.1 / FPlatformTime::GetSecondsPerCycle64()); // 100ms
 		AnalysisSyncNextTimestamp = Time + WaitTime;
 
-		//TODO
+		TSharedPtr<const Trace::IAnalysisSession> Session = FInsightsManager::Get()->GetSession();
+		if (Session.IsValid())
+		{
+			Trace::FAnalysisSessionReadScope SessionReadScope(*Session.Get());
+			const Trace::INetProfilerProvider& NetProfilerProvider = Trace::ReadNetProfilerProvider(*Session.Get());
+
+			const uint32 NewConnectionChangeCount = NetProfilerProvider.GetConnectionChangeCount();
+			if (NewConnectionChangeCount != ConnectionChangeCount)
+			{
+				ConnectionChangeCount = NewConnectionChangeCount;
+				bIsStateDirty = true;
+			}
+		}
 	}
 
 	if (bIsStateDirty)
@@ -162,7 +213,7 @@ void SPacketSizesView::Tick(const FGeometry& AllottedGeometry, const double InCu
 
 ////////////////////////////////////////////////////////////////////////////////////////////////////
 
-void SPacketSizesView::UpdateState()
+void SPacketView::UpdateState()
 {
 	FStopwatch Stopwatch;
 	Stopwatch.Start();
@@ -171,54 +222,95 @@ void SPacketSizesView::UpdateState()
 	PacketSeries->NumAggregatedPackets = 0;
 	NumUpdatedPackets = 0;
 
-	// Init series with mock data.
+	FNetworkPacketSeriesBuilder Builder(*PacketSeries, Viewport);
+
+	FAxisViewportInt32& ViewportX = Viewport.GetHorizontalAxisViewport();
+	ViewportX.SetMinMaxInterval(0, 0);
+
+	TSharedPtr<const Trace::IAnalysisSession> Session = FInsightsManager::Get()->GetSession();
+	if (Session.IsValid())
 	{
-		FAxisViewportInt32& ViewportX = Viewport.GetHorizontalAxisViewport();
+		Trace::FAnalysisSessionReadScope SessionReadScope(*Session.Get());
+		const Trace::INetProfilerProvider& NetProfilerProvider = Trace::ReadNetProfilerProvider(*Session.Get());
+
+		const uint32 GameInstanceCount = NetProfilerProvider.GetGameInstanceCount();
+		if (GameInstanceIndex < GameInstanceCount)
+		{
+			//const uint32 ConnectionCount = NetProfilerProvider.GetConnectionCount(GameInstanceIndex);
+			//if (ConnectionIndex < ConnectionCount) TODO: (if NetProfilerProvider.IsConnectionValid(GameInstanceIndex))
+			{
+				const uint32 NumPackets = NetProfilerProvider.GetPacketCount(ConnectionIndex, ConnectionMode);
+
+				ViewportX.SetMinMaxInterval(0, NumPackets);
+
+				if (NumPackets > 0)
+				{
+					const int32 MinIndex = ViewportX.GetValueAtOffset(0.0f);
+					const int32 MaxIndex = ViewportX.GetValueAtOffset(ViewportX.GetSize());
+
+					const uint32 PacketStartIndex = static_cast<uint32>(FMath::Max(0, MinIndex));
+					const uint32 PacketEndIndex = FMath::Min(NumPackets - 1, static_cast<uint32>(FMath::Max(0, MaxIndex)));
+
+					if (PacketStartIndex <= PacketEndIndex)
+					{
+						int32 PacketIndex = PacketStartIndex;
+						NetProfilerProvider.EnumeratePackets(ConnectionIndex, ConnectionMode, PacketStartIndex, PacketEndIndex, [&Builder, &PacketIndex](const Trace::FNetProfilerPacket& Packet)
+						{
+							Builder.AddPacket(PacketIndex++, Packet);
+						});
+					}
+				}
+			}
+		}
+	}
 
+	// Init series with mock data.
+	if (false)
+	{
 		constexpr int32 NumPackets = 100000;
 
 		ViewportX.SetMinMaxInterval(0, NumPackets);
 
-		FNetworkPacketSeriesBuilder Builder(*PacketSeries, Viewport);
-
 		const int32 StartIndex = FMath::Max(0, ViewportX.GetValueAtOffset(0.0f));
 		const int32 EndIndex = FMath::Min(NumPackets, ViewportX.GetValueAtOffset(ViewportX.GetSize()));
 
-		for (int32 FrameIndex = StartIndex; FrameIndex < EndIndex; ++FrameIndex)
+		for (int32 PacketIndex = StartIndex; PacketIndex < EndIndex; ++PacketIndex)
 		{
-			FRandomStream RandomStream((FrameIndex * FrameIndex * FrameIndex) ^ 0x2c2c57ed);
+			FRandomStream RandomStream((PacketIndex * PacketIndex * PacketIndex) ^ 0x2c2c57ed);
 			int64 Size = RandomStream.RandRange(0, 2000);
+
+			Trace::ENetProfilerDeliveryStatus Status = Trace::ENetProfilerDeliveryStatus::Unknown;
 			const float Fraction = RandomStream.GetFraction();
-			ENetworkPacketStatus Status = ENetworkPacketStatus::ConfirmedReceived;
 			if (Fraction < 0.01) // 1%
 			{
-				Status = ENetworkPacketStatus::ConfirmedLost;
+				Status = Trace::ENetProfilerDeliveryStatus::Dropped;
 			}
 			else if (Fraction < 0.05) // 4%
 			{
-				Status = ENetworkPacketStatus::Sent;
+				Status = Trace::ENetProfilerDeliveryStatus::Delivered;
 			}
 
-			double TimeSent = ((double)FrameIndex * 100.0) / (double)NumPackets + RandomStream.GetFraction() * 0.1;
-			double TimeAck = 0.0;
-			if (Status != ENetworkPacketStatus::Sent)
-			{
-				TimeAck = TimeSent + RandomStream.GetFraction() * 0.1;
-			}
+			const double Timestamp = ((double)PacketIndex * 100.0) / (double)NumPackets + RandomStream.GetFraction() * 0.1;
 
-			Builder.AddPacket(FrameIndex, Size, TimeSent, TimeAck, Status);
+			Trace::FNetProfilerPacket Packet;
+			Packet.TimeStamp = static_cast<Trace::FNetProfilerTimeStamp>(Timestamp);
+			Packet.SequenceNumber = PacketIndex;
+			Packet.ContentSizeInBits = Size;
+			Packet.TotalPacketSizeInBytes = (Size + 7) / 8;
+			Packet.DeliveryStatus = Status;
+			Builder.AddPacket(PacketIndex, Packet);
 		}
-
-		NumUpdatedPackets += Builder.GetNumAddedPackets();
 	}
 
+	NumUpdatedPackets += Builder.GetNumAddedPackets();
+
 	Stopwatch.Stop();
 	UpdateDurationHistory.AddValue(Stopwatch.AccumulatedTime);
 }
 
 ////////////////////////////////////////////////////////////////////////////////////////////////////
 
-FNetworkPacketSampleRef SPacketSizesView::GetSampleAtMousePosition(float X, float Y)
+FNetworkPacketSampleRef SPacketView::GetSampleAtMousePosition(float X, float Y)
 {
 	if (!bIsStateDirty)
 	{
@@ -236,7 +328,9 @@ FNetworkPacketSampleRef SPacketSizesView::GetSampleAtMousePosition(float X, floa
 
 					const float ViewHeight = FMath::RoundToFloat(Viewport.GetHeight());
 					const float BaselineY = FMath::RoundToFloat(ViewportY.GetOffsetForValue(0.0));
-					const float ValueY = FMath::RoundToFloat(ViewportY.GetOffsetForValue(static_cast<double>(Sample.LargestPacket.Size)));
+
+					//const float ValueY = FMath::RoundToFloat(ViewportY.GetOffsetForValue(static_cast<double>(Sample.LargestPacket.ContentSizeInBits)));
+					const float ValueY = FMath::RoundToFloat(ViewportY.GetOffsetForValue(static_cast<double>(Sample.LargestPacket.TotalSizeInBytes * 8)));
 
 					constexpr float ToleranceY = 3.0f; // [pixels]
 
@@ -256,7 +350,7 @@ FNetworkPacketSampleRef SPacketSizesView::GetSampleAtMousePosition(float X, floa
 
 ////////////////////////////////////////////////////////////////////////////////////////////////////
 
-void SPacketSizesView::SelectSampleAtMousePosition(float X, float Y)
+void SPacketView::SelectSampleAtMousePosition(float X, float Y)
 {
 	FNetworkPacketSampleRef SampleRef = GetSampleAtMousePosition(X, Y);
 	if (!SampleRef.IsValid())
@@ -268,17 +362,31 @@ void SPacketSizesView::SelectSampleAtMousePosition(float X, float Y)
 		SampleRef = GetSampleAtMousePosition(X + 1.0f, Y);
 	}
 
-	if (SampleRef.IsValid())
+	if (!SelectedSample.Equals(SampleRef))
 	{
 		SelectedSample = SampleRef;
+		OnSelectedSampleChanged();
+	}
+}
+
+////////////////////////////////////////////////////////////////////////////////////////////////////
 
-		TSharedPtr<SNetworkingProfilerWindow> Wnd = FNetworkingProfilerManager::Get()->GetProfilerWindow();
-		if (Wnd.IsValid())
+void SPacketView::OnSelectedSampleChanged()
+{
+	if (ProfilerWindow.IsValid())
+	{
+		TSharedPtr<SPacketContentView> PacketContentView = ProfilerWindow->GetPacketContentView();
+		if (PacketContentView.IsValid())
 		{
-			TSharedPtr<SPacketBreakdownView> PacketBreakdownView = Wnd->GetPacketBreakdownView();
-			if (PacketBreakdownView.IsValid())
+			if (SelectedSample.IsValid())
 			{
-				PacketBreakdownView->SetPacket(SampleRef.Sample->LargestPacket.FrameIndex, SampleRef.Sample->LargestPacket.Size);
+				PacketContentView->SetPacket(GameInstanceIndex, ConnectionIndex, ConnectionMode,
+											 SelectedSample.Sample->LargestPacket.Index,
+											 SelectedSample.Sample->LargestPacket.TotalSizeInBytes * 8);
+			}
+			else
+			{
+				PacketContentView->ResetPacket();
 			}
 		}
 	}
@@ -286,33 +394,31 @@ void SPacketSizesView::SelectSampleAtMousePosition(float X, float Y)
 
 ////////////////////////////////////////////////////////////////////////////////////////////////////
 
-const TCHAR* StatusToString(ENetworkPacketStatus Status)
+const TCHAR* StatusToString(Trace::ENetProfilerDeliveryStatus Status)
 {
 	switch (Status)
 	{
-		case ENetworkPacketStatus::ConfirmedReceived: return TEXT("Received");
-		case ENetworkPacketStatus::Sent:              return TEXT("Sent");
-		case ENetworkPacketStatus::ConfirmedLost:     return TEXT("Lost");
-		case ENetworkPacketStatus::Unknown:
-		default:                                      return TEXT("Unknown");
+		case Trace::ENetProfilerDeliveryStatus::Delivered:  return TEXT("Delivered");
+		case Trace::ENetProfilerDeliveryStatus::Dropped:    return TEXT("Dropped");
+		case Trace::ENetProfilerDeliveryStatus::Unknown:
+		default:                                            return TEXT("Unknown");
 	};
 }
 
-const TCHAR* AggregatedStatusToString(ENetworkPacketStatus Status)
+const TCHAR* AggregatedStatusToString(Trace::ENetProfilerDeliveryStatus Status)
 {
 	switch (Status)
 	{
-		case ENetworkPacketStatus::ConfirmedReceived: return TEXT("all are Received");
-		case ENetworkPacketStatus::Sent:              return TEXT("some are Sent, but have no Ack");
-		case ENetworkPacketStatus::ConfirmedLost:     return TEXT("some are Lost");
-		case ENetworkPacketStatus::Unknown:
-		default:                                      return TEXT("Unknown");
-	};
+		case Trace::ENetProfilerDeliveryStatus::Delivered:  return TEXT("all packages are Delivered");
+		case Trace::ENetProfilerDeliveryStatus::Dropped:    return TEXT("at least one Dropped package");
+		case Trace::ENetProfilerDeliveryStatus::Unknown:
+		default:                                            return TEXT("Unknown");
+	}
 }
 
 ////////////////////////////////////////////////////////////////////////////////////////////////////
 
-int32 SPacketSizesView::OnPaint(const FPaintArgs& Args, const FGeometry& AllottedGeometry, const FSlateRect& MyCullingRect, FSlateWindowElementList& OutDrawElements, int32 LayerId, const FWidgetStyle& InWidgetStyle, bool bParentEnabled) const
+int32 SPacketView::OnPaint(const FPaintArgs& Args, const FGeometry& AllottedGeometry, const FSlateRect& MyCullingRect, FSlateWindowElementList& OutDrawElements, int32 LayerId, const FWidgetStyle& InWidgetStyle, bool bParentEnabled) const
 {
 	const bool bEnabled = ShouldBeEnabled(bParentEnabled);
 	const ESlateDrawEffect DrawEffects = bEnabled ? ESlateDrawEffect::None : ESlateDrawEffect::DisabledEffect;
@@ -333,7 +439,7 @@ int32 SPacketSizesView::OnPaint(const FPaintArgs& Args, const FGeometry& Allotte
 		FStopwatch Stopwatch;
 		Stopwatch.Start();
 
-		FPacketSizesViewDrawHelper Helper(DrawContext, Viewport);
+		FPacketViewDrawHelper Helper(DrawContext, Viewport);
 
 		Helper.DrawBackground();
 
@@ -348,11 +454,11 @@ int32 SPacketSizesView::OnPaint(const FPaintArgs& Args, const FGeometry& Allotte
 		bool bIsSelectedAndHovered = SelectedSample.Equals(HoveredSample);
 		if (SelectedSample.IsValid())
 		{
-			Helper.DrawSampleHighlight(*SelectedSample.Sample, bIsSelectedAndHovered ? FPacketSizesViewDrawHelper::EHighlightMode::SelectedAndHovered : FPacketSizesViewDrawHelper::EHighlightMode::Selected);
+			Helper.DrawSampleHighlight(*SelectedSample.Sample, bIsSelectedAndHovered ? FPacketViewDrawHelper::EHighlightMode::SelectedAndHovered : FPacketViewDrawHelper::EHighlightMode::Selected);
 		}
 		if (HoveredSample.IsValid() && !bIsSelectedAndHovered)
 		{
-			Helper.DrawSampleHighlight(*HoveredSample.Sample, FPacketSizesViewDrawHelper::EHighlightMode::Hovered);
+			Helper.DrawSampleHighlight(*HoveredSample.Sample, FPacketViewDrawHelper::EHighlightMode::Hovered);
 		}
 
 		// Draw the vertical axis grid.
@@ -373,40 +479,47 @@ int32 SPacketSizesView::OnPaint(const FPaintArgs& Args, const FGeometry& Allotte
 			const double Precision = 0.01; // 10ms
 			int NumLines;
 			FString Text;
+			uint32 UnusedBits = HoveredSample.Sample->LargestPacket.TotalSizeInBytes * 8 - HoveredSample.Sample->LargestPacket.ContentSizeInBits;
 			if (HoveredSample.Sample->NumPackets == 1)
 			{
-				NumLines = 5;
-				Text = FString::Format(TEXT("Frame Index: {0}\n"
-											"Size: {1} bits\n"
-											"Sent Timestamp: {2}\n"
-											"Ack Timestamp: {3}\n"
-											"Status: {4}"),
+				NumLines = 6;
+				Text = FString::Format(TEXT("Packet Index: {0}\n"
+											"Sequence Number: {1}\n"
+											"Content Size: {2} bits\n"
+											"Total Size: {3} bytes ({4} unused bits)\n"
+											"Timestamp: {5}\n"
+											"Status: {6}"),
 					{
-						FText::AsNumber(HoveredSample.Sample->LargestPacket.FrameIndex).ToString(),
-						FText::AsNumber(HoveredSample.Sample->LargestPacket.Size).ToString(),
-						TimeUtils::FormatTimeHMS(HoveredSample.Sample->LargestPacket.TimeSent, Precision),
-						HoveredSample.Sample->LargestPacket.TimeAck == 0 ? TEXT("N/A") : TimeUtils::FormatTimeHMS(HoveredSample.Sample->LargestPacket.TimeAck, Precision),
+						FText::AsNumber(HoveredSample.Sample->LargestPacket.Index).ToString(),
+						FText::AsNumber(HoveredSample.Sample->LargestPacket.SequenceNumber).ToString(),
+						FText::AsNumber(HoveredSample.Sample->LargestPacket.ContentSizeInBits).ToString(),
+						FText::AsNumber(HoveredSample.Sample->LargestPacket.TotalSizeInBytes).ToString(),
+						FText::AsNumber(UnusedBits).ToString(),
+						TimeUtils::FormatTimeHMS(HoveredSample.Sample->LargestPacket.TimeStamp, Precision),
 						::StatusToString(HoveredSample.Sample->LargestPacket.Status)
 					});
 			}
 			else
 			{
-				NumLines = 8;
+				NumLines = 9;
 				Text = FString::Format(TEXT("{0} network packets\n"
 											"({1})\n"
 											"Largest Packet\n"
-											"    Frame Index: {2}\n"
-											"    Size: {3} bits\n"
-											"    Sent Timestamp: {4}\n"
-											"    Ack Timestamp: {5}\n"
-											"    Status: {6}"),
+											"    Index: {2}\n"
+											"    Sequance Number: {3}\n"
+											"    Content Size: {4} bits\n"
+											"    Total Size: {5} bytes ({6} unused bits)\n"
+											"    Timestamp: {7}\n"
+											"    Status: {8}"),
 					{
 						HoveredSample.Sample->NumPackets,
 						::AggregatedStatusToString(HoveredSample.Sample->AggregatedStatus),
-						FText::AsNumber(HoveredSample.Sample->LargestPacket.FrameIndex).ToString(),
-						FText::AsNumber(HoveredSample.Sample->LargestPacket.Size).ToString(),
-						TimeUtils::FormatTimeHMS(HoveredSample.Sample->LargestPacket.TimeSent, Precision),
-						HoveredSample.Sample->LargestPacket.TimeAck == 0 ? TEXT("N/A") : TimeUtils::FormatTimeHMS(HoveredSample.Sample->LargestPacket.TimeAck, Precision),
+						FText::AsNumber(HoveredSample.Sample->LargestPacket.Index).ToString(),
+						FText::AsNumber(HoveredSample.Sample->LargestPacket.SequenceNumber).ToString(),
+						FText::AsNumber(HoveredSample.Sample->LargestPacket.ContentSizeInBits).ToString(),
+						FText::AsNumber(HoveredSample.Sample->LargestPacket.TotalSizeInBytes).ToString(),
+						FText::AsNumber(UnusedBits).ToString(),
+						TimeUtils::FormatTimeHMS(HoveredSample.Sample->LargestPacket.TimeStamp, Precision),
 						::StatusToString(HoveredSample.Sample->LargestPacket.Status)
 					});
 			}
@@ -418,7 +531,7 @@ int32 SPacketSizesView::OnPaint(const FPaintArgs& Args, const FGeometry& Allotte
 
 			const FAxisViewportInt32& ViewportX = Viewport.GetHorizontalAxisViewport();
 
-			float X1 = ViewportX.GetOffsetForValue(HoveredSample.Sample->LargestPacket.FrameIndex);
+			float X1 = ViewportX.GetOffsetForValue(HoveredSample.Sample->LargestPacket.Index);
 			float CX = X1 + FMath::RoundToFloat(Viewport.GetSampleWidth() / 2);
 			if (CX + W2 > ViewportX.GetSize())
 			{
@@ -449,7 +562,7 @@ int32 SPacketSizesView::OnPaint(const FPaintArgs& Args, const FGeometry& Allotte
 		const float DbgDY = MaxFontCharHeight;
 
 		const float DbgW = 280.0f;
-		const float DbgH = DbgDY * 4 + 3.0f;
+		const float DbgH = DbgDY * 5 + 3.0f;
 		const float DbgX = ViewWidth - DbgW - 20.0f;
 		float DbgY = 7.0f;
 
@@ -514,6 +627,18 @@ int32 SPacketSizesView::OnPaint(const FPaintArgs& Args, const FGeometry& Allotte
 			SummaryFont, DbgTextColor
 		);
 		DbgY += DbgDY;
+
+		// Draw connection info.
+		DrawContext.DrawText
+		(
+			DbgX, DbgY,
+			FString::Printf(TEXT("Game Instance %d, Connection %d (%s)"),
+				GameInstanceIndex,
+				ConnectionIndex,
+				(ConnectionMode == Trace::ENetProfilerConnectionMode::Outgoing) ? TEXT("Outgoing") : TEXT("Incoming")),
+			SummaryFont, DbgTextColor
+		);
+		DbgY += DbgDY;
 	}
 
 	return SCompoundWidget::OnPaint(Args, AllottedGeometry, MyCullingRect, OutDrawElements, LayerId, InWidgetStyle, bParentEnabled && IsEnabled());
@@ -521,7 +646,7 @@ int32 SPacketSizesView::OnPaint(const FPaintArgs& Args, const FGeometry& Allotte
 
 ////////////////////////////////////////////////////////////////////////////////////////////////////
 
-void SPacketSizesView::DrawVerticalAxisGrid(FDrawContext& DrawContext, const FSlateBrush* Brush, const FSlateFontInfo& Font) const
+void SPacketView::DrawVerticalAxisGrid(FDrawContext& DrawContext, const FSlateBrush* Brush, const FSlateFontInfo& Font) const
 {
 	const FAxisViewportDouble& ViewportY = Viewport.GetVerticalAxisViewport();
 
@@ -537,6 +662,7 @@ void SPacketSizesView::DrawVerticalAxisGrid(FDrawContext& DrawContext, const FSl
 	if (Delta > 0.0)
 	{
 		const int64 DeltaBits = static_cast<int64>(Delta);
+
 		// Compute rounding based on magnitude of visible range of values (Delta).
 		int64 Power10 = 1;
 		int64 Delta10 = DeltaBits;
@@ -554,16 +680,16 @@ void SPacketSizesView::DrawVerticalAxisGrid(FDrawContext& DrawContext, const FSl
 			Power10 = 1;
 		}
 
-		const double Grid = static_cast<double>(((DeltaBits + Power10) / Power10) * Power10); // next value divisible with a multiple of 10
+		const double Grid = static_cast<double>(((DeltaBits + Power10 - 1) / Power10) * Power10); // next value divisible with a multiple of 10
 
 		const double StartValue = FMath::GridSnap(BottomValue, Grid);
 
+		const float ViewWidth = Viewport.GetWidth();
+
 		const FLinearColor GridColor(0.0f, 0.0f, 0.0f, 0.1f);
 		const FLinearColor TextBgColor(0.05f, 0.05f, 0.05f, 1.0f);
 		const FLinearColor TextColor(1.0f, 1.0f, 1.0f, 1.0f);
 
-		const float ViewWidth = Viewport.GetWidth();
-
 		const TSharedRef<FSlateFontMeasure> FontMeasureService = FSlateApplication::Get().GetRenderer()->GetFontMeasureService();
 
 		for (double Value = StartValue; Value < TopValue; Value += Grid)
@@ -590,7 +716,7 @@ void SPacketSizesView::DrawVerticalAxisGrid(FDrawContext& DrawContext, const FSl
 
 ////////////////////////////////////////////////////////////////////////////////////////////////////
 
-void SPacketSizesView::DrawHorizontalAxisGrid(FDrawContext& DrawContext, const FSlateBrush* Brush, const FSlateFontInfo& Font) const
+void SPacketView::DrawHorizontalAxisGrid(FDrawContext& DrawContext, const FSlateBrush* Brush, const FSlateFontInfo& Font) const
 {
 	const FAxisViewportInt32& ViewportX = Viewport.GetHorizontalAxisViewport();
 
@@ -598,10 +724,10 @@ void SPacketSizesView::DrawHorizontalAxisGrid(FDrawContext& DrawContext, const F
 
 	constexpr float MinDX = 100.0f; // min horizontal distance between vertical grid lines
 
-	int32 LeftIndex = ViewportX.GetValueAtOffset(0.0f);
-	int32 GridIndex = ViewportX.GetValueAtOffset(MinDX);
-	int32 RightIndex = ViewportX.GetValueAtOffset(RoundedViewWidth);
-	int32 Delta = GridIndex - LeftIndex;
+	const int32 LeftIndex = ViewportX.GetValueAtOffset(0.0f);
+	const int32 GridIndex = ViewportX.GetValueAtOffset(MinDX);
+	const int32 RightIndex = ViewportX.GetValueAtOffset(RoundedViewWidth);
+	const int32 Delta = GridIndex - LeftIndex;
 
 	if (Delta > 0)
 	{
@@ -622,12 +748,7 @@ void SPacketSizesView::DrawHorizontalAxisGrid(FDrawContext& DrawContext, const F
 			Power10 = 1;
 		}
 
-		const int32 Grid = ((Delta + Power10) / Power10) * Power10; // next value divisible with a multiple of 10
-
-		const FLinearColor GridColor(0.0f, 0.0f, 0.0f, 0.1f);
-		//const FLinearColor TextBgColor(0.05f, 0.05f, 0.05f, 1.0f);
-		//const FLinearColor TextColor(1.0f, 1.0f, 1.0f, 1.0f);
-		const FLinearColor TopTextColor(1.0f, 1.0f, 1.0f, 0.7f);
+		const int32 Grid = ((Delta + Power10 - 1) / Power10) * Power10; // next value divisible with a multiple of 10
 
 		// Skip grid lines for negative indices.
 		double StartIndex = ((LeftIndex + Grid - 1) / Grid) * Grid;
@@ -638,6 +759,11 @@ void SPacketSizesView::DrawHorizontalAxisGrid(FDrawContext& DrawContext, const F
 
 		const float ViewHeight = Viewport.GetHeight();
 
+		const FLinearColor GridColor(0.0f, 0.0f, 0.0f, 0.1f);
+		//const FLinearColor TextBgColor(0.05f, 0.05f, 0.05f, 1.0f);
+		//const FLinearColor TextColor(1.0f, 1.0f, 1.0f, 1.0f);
+		const FLinearColor TopTextColor(1.0f, 1.0f, 1.0f, 0.7f);
+
 		//const TSharedRef<FSlateFontMeasure> FontMeasureService = FSlateApplication::Get().GetRenderer()->GetFontMeasureService();
 
 		for (int32 Index = StartIndex; Index < RightIndex; Index += Grid)
@@ -665,7 +791,7 @@ void SPacketSizesView::DrawHorizontalAxisGrid(FDrawContext& DrawContext, const F
 
 ////////////////////////////////////////////////////////////////////////////////////////////////////
 
-FReply SPacketSizesView::OnMouseButtonDown(const FGeometry& MyGeometry, const FPointerEvent& MouseEvent)
+FReply SPacketView::OnMouseButtonDown(const FGeometry& MyGeometry, const FPointerEvent& MouseEvent)
 {
 	FReply Reply = FReply::Unhandled();
 
@@ -692,7 +818,7 @@ FReply SPacketSizesView::OnMouseButtonDown(const FGeometry& MyGeometry, const FP
 
 ////////////////////////////////////////////////////////////////////////////////////////////////////
 
-FReply SPacketSizesView::OnMouseButtonUp(const FGeometry& MyGeometry, const FPointerEvent& MouseEvent)
+FReply SPacketView::OnMouseButtonUp(const FGeometry& MyGeometry, const FPointerEvent& MouseEvent)
 {
 	FReply Reply = FReply::Unhandled();
 
@@ -746,7 +872,7 @@ FReply SPacketSizesView::OnMouseButtonUp(const FGeometry& MyGeometry, const FPoi
 
 ////////////////////////////////////////////////////////////////////////////////////////////////////
 
-FReply SPacketSizesView::OnMouseMove(const FGeometry& MyGeometry, const FPointerEvent& MouseEvent)
+FReply SPacketView::OnMouseMove(const FGeometry& MyGeometry, const FPointerEvent& MouseEvent)
 {
 	FReply Reply = FReply::Unhandled();
 
@@ -799,13 +925,13 @@ FReply SPacketSizesView::OnMouseMove(const FGeometry& MyGeometry, const FPointer
 
 ////////////////////////////////////////////////////////////////////////////////////////////////////
 
-void SPacketSizesView::OnMouseEnter(const FGeometry& MyGeometry, const FPointerEvent& MouseEvent)
+void SPacketView::OnMouseEnter(const FGeometry& MyGeometry, const FPointerEvent& MouseEvent)
 {
 }
 
 ////////////////////////////////////////////////////////////////////////////////////////////////////
 
-void SPacketSizesView::OnMouseLeave(const FPointerEvent& MouseEvent)
+void SPacketView::OnMouseLeave(const FPointerEvent& MouseEvent)
 {
 	if (!HasMouseCapture())
 	{
@@ -820,7 +946,7 @@ void SPacketSizesView::OnMouseLeave(const FPointerEvent& MouseEvent)
 
 ////////////////////////////////////////////////////////////////////////////////////////////////////
 
-FReply SPacketSizesView::OnMouseWheel(const FGeometry& MyGeometry, const FPointerEvent& MouseEvent)
+FReply SPacketView::OnMouseWheel(const FGeometry& MyGeometry, const FPointerEvent& MouseEvent)
 {
 	MousePosition = MyGeometry.AbsoluteToLocal(MouseEvent.GetScreenSpacePosition());
 
@@ -857,14 +983,14 @@ FReply SPacketSizesView::OnMouseWheel(const FGeometry& MyGeometry, const FPointe
 
 ////////////////////////////////////////////////////////////////////////////////////////////////////
 
-FReply SPacketSizesView::OnMouseButtonDoubleClick(const FGeometry& MyGeometry, const FPointerEvent& MouseEvent)
+FReply SPacketView::OnMouseButtonDoubleClick(const FGeometry& MyGeometry, const FPointerEvent& MouseEvent)
 {
 	return FReply::Unhandled();
 }
 
 ////////////////////////////////////////////////////////////////////////////////////////////////////
 
-FCursorReply SPacketSizesView::OnCursorQuery(const FGeometry& MyGeometry, const FPointerEvent& CursorEvent) const
+FCursorReply SPacketView::OnCursorQuery(const FGeometry& MyGeometry, const FPointerEvent& CursorEvent) const
 {
 	FCursorReply CursorReply = FCursorReply::Unhandled();
 
@@ -882,7 +1008,7 @@ FCursorReply SPacketSizesView::OnCursorQuery(const FGeometry& MyGeometry, const
 
 ////////////////////////////////////////////////////////////////////////////////////////////////////
 
-void SPacketSizesView::ShowContextMenu(const FPointerEvent& MouseEvent)
+void SPacketView::ShowContextMenu(const FPointerEvent& MouseEvent)
 {
 	const bool bShouldCloseWindowAfterMenuSelection = true;
 	FMenuBuilder MenuBuilder(bShouldCloseWindowAfterMenuSelection, NULL);
@@ -891,9 +1017,9 @@ void SPacketSizesView::ShowContextMenu(const FPointerEvent& MouseEvent)
 	{
 		FUIAction Action_AutoZoom
 		(
-			FExecuteAction::CreateSP(this, &SPacketSizesView::ContextMenu_AutoZoom_Execute),
-			FCanExecuteAction::CreateSP(this, &SPacketSizesView::ContextMenu_AutoZoom_CanExecute),
-			FIsActionChecked::CreateSP(this, &SPacketSizesView::ContextMenu_AutoZoom_IsChecked)
+			FExecuteAction::CreateSP(this, &SPacketView::ContextMenu_AutoZoom_Execute),
+			FCanExecuteAction::CreateSP(this, &SPacketView::ContextMenu_AutoZoom_CanExecute),
+			FIsActionChecked::CreateSP(this, &SPacketView::ContextMenu_AutoZoom_IsChecked)
 		);
 		MenuBuilder.AddMenuEntry
 		(
@@ -916,7 +1042,7 @@ void SPacketSizesView::ShowContextMenu(const FPointerEvent& MouseEvent)
 
 ////////////////////////////////////////////////////////////////////////////////////////////////////
 
-void SPacketSizesView::ContextMenu_AutoZoom_Execute()
+void SPacketView::ContextMenu_AutoZoom_Execute()
 {
 	FAxisViewportInt32& ViewportX = Viewport.GetHorizontalAxisViewport();
 
@@ -946,27 +1072,27 @@ void SPacketSizesView::ContextMenu_AutoZoom_Execute()
 
 ////////////////////////////////////////////////////////////////////////////////////////////////////
 
-bool SPacketSizesView::ContextMenu_AutoZoom_CanExecute()
+bool SPacketView::ContextMenu_AutoZoom_CanExecute()
 {
 	return true;
 }
 
 ////////////////////////////////////////////////////////////////////////////////////////////////////
 
-bool SPacketSizesView::ContextMenu_AutoZoom_IsChecked()
+bool SPacketView::ContextMenu_AutoZoom_IsChecked()
 {
 	return bIsAutoZoomEnabled && Viewport.GetHorizontalAxisViewport().GetPos() == 0.0f;
 }
 
 ////////////////////////////////////////////////////////////////////////////////////////////////////
 
-void SPacketSizesView::BindCommands()
+void SPacketView::BindCommands()
 {
 }
 
 ////////////////////////////////////////////////////////////////////////////////////////////////////
 
-void SPacketSizesView::HorizontalScrollBar_OnUserScrolled(float ScrollOffset)
+void SPacketView::HorizontalScrollBar_OnUserScrolled(float ScrollOffset)
 {
 	FAxisViewportInt32& ViewportX = Viewport.GetHorizontalAxisViewport();
 	ViewportX.OnUserScrolled(HorizontalScrollBar, ScrollOffset);
@@ -975,7 +1101,7 @@ void SPacketSizesView::HorizontalScrollBar_OnUserScrolled(float ScrollOffset)
 
 ////////////////////////////////////////////////////////////////////////////////////////////////////
 
-void SPacketSizesView::UpdateHorizontalScrollBar()
+void SPacketView::UpdateHorizontalScrollBar()
 {
 	FAxisViewportInt32& ViewportX = Viewport.GetHorizontalAxisViewport();
 	ViewportX.UpdateScrollBar(HorizontalScrollBar);
@@ -983,7 +1109,7 @@ void SPacketSizesView::UpdateHorizontalScrollBar()
 
 ////////////////////////////////////////////////////////////////////////////////////////////////////
 
-void SPacketSizesView::ZoomHorizontally(const float Delta, const float X)
+void SPacketView::ZoomHorizontally(const float Delta, const float X)
 {
 	FAxisViewportInt32& ViewportX = Viewport.GetHorizontalAxisViewport();
 	ViewportX.RelativeZoomWithFixedOffset(Delta, X);
diff --git a/Engine/Source/Developer/TraceInsights/Private/Insights/NetworkingProfiler/Widgets/SPacketSizesView.h b/Engine/Source/Developer/TraceInsights/Private/Insights/NetworkingProfiler/Widgets/SPacketSizesView.h
--- a/Engine/Source/Developer/TraceInsights/Private/Insights/NetworkingProfiler/Widgets/SPacketSizesView.h
+++ b/Engine/Source/Developer/TraceInsights/Private/Insights/NetworkingProfiler/Widgets/SPacketSizesView.h
@@ -7,15 +7,19 @@
 #include "Input/Reply.h"
 #include "Layout/Geometry.h"
 #include "Rendering/RenderingCommon.h"
+#include "TraceServices/Model/NetProfiler.h"
 #include "Widgets/DeclarativeSyntaxSupport.h"
 #include "Widgets/SCompoundWidget.h"
 
 // Insights
 #include "Insights/Common/FixedCircularBuffer.h"
 #include "Insights/NetworkingProfiler/ViewModels/PacketSizesViewHelper.h"
+//#include "Insights/NetworkingProfiler/ViewModels/PacketViewHelper.h"
 #include "Insights/NetworkingProfiler/ViewModels/PacketSizesViewport.h"
+//#include "Insights/NetworkingProfiler/ViewModels/PacketViewViewport.h"
 
 class SScrollBar;
+class SNetworkingProfilerWindow;
 
 ////////////////////////////////////////////////////////////////////////////////////////////////////
 
@@ -73,7 +77,7 @@ struct FNetworkPacketSampleRef
 /**
  * Widget used to present the network packets as a bar track.
  */
-class SPacketSizesView : public SCompoundWidget
+class SPacketView : public SCompoundWidget
 {
 public:
 	/** Number of pixels. */
@@ -88,15 +92,17 @@ public:
 
 public:
 	/** Default constructor. */
-	SPacketSizesView();
+	SPacketView();
 
 	/** Virtual destructor. */
-	virtual ~SPacketSizesView();
+	virtual ~SPacketView();
 
 	/** Resets internal widget's data to the default one. */
 	void Reset();
 
-	SLATE_BEGIN_ARGS(SPacketSizesView)
+	void SetConnection(uint32 GameInstanceIndex, uint32 ConnectionIndex, Trace::ENetProfilerConnectionMode ConnectionMode);
+
+	SLATE_BEGIN_ARGS(SPacketView)
 	{
 		_Clipping = EWidgetClipping::ClipToBounds;
 	}
@@ -106,7 +112,7 @@ public:
 	 * Construct this widget
 	 * @param InArgs The declaration data for this widget
 	 */
-	void Construct(const FArguments& InArgs);
+	void Construct(const FArguments& InArgs, TSharedPtr<SNetworkingProfilerWindow> InProfilerWindow);
 
 	/**
 	 * Ticks this widget. Override in derived classes, but always call the parent implementation.
@@ -136,6 +142,7 @@ private:
 
 	FNetworkPacketSampleRef GetSampleAtMousePosition(float X, float Y);
 	void SelectSampleAtMousePosition(float X, float Y);
+	void OnSelectedSampleChanged();
 
 	void ShowContextMenu(const FPointerEvent& MouseEvent);
 
@@ -156,8 +163,14 @@ private:
 	void ZoomHorizontally(const float Delta, const float X);
 
 private:
+	TSharedPtr<SNetworkingProfilerWindow> ProfilerWindow;
+
+	uint32 GameInstanceIndex;
+	uint32 ConnectionIndex;
+	Trace::ENetProfilerConnectionMode ConnectionMode;
+
 	/** The track's viewport. Encapsulates info about position and scale. */
-	FPacketSizesViewport Viewport;
+	FPacketViewViewport Viewport;
 	bool bIsViewportDirty;
 
 	/** Cached info for the packet series. */
@@ -167,6 +180,7 @@ private:
 	bool bIsAutoZoomEnabled;
 
 	uint64 AnalysisSyncNextTimestamp;
+	uint32 ConnectionChangeCount;
 
 	//////////////////////////////////////////////////
 
diff --git a/Engine/Source/Developer/TraceInsights/Private/Insights/TraceInsightsModule.cpp b/Engine/Source/Developer/TraceInsights/Private/Insights/TraceInsightsModule.cpp
--- a/Engine/Source/Developer/TraceInsights/Private/Insights/TraceInsightsModule.cpp
+++ b/Engine/Source/Developer/TraceInsights/Private/Insights/TraceInsightsModule.cpp
@@ -74,6 +74,13 @@ protected:
 
 IMPLEMENT_MODULE(FTraceInsightsModule, TraceInsights);
 
+////////////////////////////////////////////////////////////////////////////////////////////////////
+
+static TSharedPtr<SDockTab> NeverReuse(const FTabId&)
+{
+	return TSharedPtr<SDockTab>();
+}
+
 ////////////////////////////////////////////////////////////////////////////////////////////////////
 // FTraceInsightsModule
 ////////////////////////////////////////////////////////////////////////////////////////////////////
@@ -96,14 +103,20 @@ void FTraceInsightsModule::StartupModule()
 	FNetworkingProfilerManager::Initialize();
 
 	// Register tab spawner for the Start Page.
-	auto& StartPageTabSpawnerEntry = FGlobalTabmanager::Get()->RegisterTabSpawner(FInsightsManagerTabs::StartPageTabId,
+	auto& StartPageTabSpawnerEntry = FGlobalTabmanager::Get()->RegisterNomadTabSpawner(FInsightsManagerTabs::StartPageTabId,
 		FOnSpawnTab::CreateRaw(this, &FTraceInsightsModule::SpawnStartPageTab))
 		.SetDisplayName(NSLOCTEXT("FTraceInsightsModule", "StartPageTabTitle", "Unreal Insights"))
 		.SetTooltipText(NSLOCTEXT("FTraceInsightsModule", "StartPageTooltipText", "Open the start page for Unreal Insights."))
 		.SetIcon(FSlateIcon(FInsightsStyle::GetStyleSetName(), "StartPage.Icon.Small"));
 
+//#if WITH_EDITOR
+//	StartPageTabSpawnerEntry.SetGroup(WorkspaceMenu::GetMenuStructure().GetDeveloperToolsMiscCategory());
+//#else
+//	StartPageTabSpawnerEntry.SetGroup(WorkspaceMenu::GetMenuStructure().GetToolsCategory());
+//#endif
+
 	// Register tab spawner for the Timing Insights.
-	auto& TimingProfilerTabSpawnerEntry = FGlobalTabmanager::Get()->RegisterTabSpawner(FInsightsManagerTabs::TimingProfilerTabId,
+	auto& TimingProfilerTabSpawnerEntry = FGlobalTabmanager::Get()->RegisterNomadTabSpawner(FInsightsManagerTabs::TimingProfilerTabId,
 		FOnSpawnTab::CreateRaw(this, &FTraceInsightsModule::SpawnTimingProfilerTab))
 		.SetDisplayName(NSLOCTEXT("FTraceInsightsModule", "TimingProfilerTabTitle", "Timing Insights"))
 		.SetTooltipText(NSLOCTEXT("FTraceInsightsModule", "TimingProfilerTooltipText", "Open the Timing Insights tab."))
@@ -116,7 +129,7 @@ void FTraceInsightsModule::StartupModule()
 //#endif
 
 	// Register tab spawner for the Asset Loading Insights.
-	auto& LoadingProfilerTabSpawnerEntry = FGlobalTabmanager::Get()->RegisterTabSpawner(FInsightsManagerTabs::LoadingProfilerTabId,
+	auto& LoadingProfilerTabSpawnerEntry = FGlobalTabmanager::Get()->RegisterNomadTabSpawner(FInsightsManagerTabs::LoadingProfilerTabId,
 		FOnSpawnTab::CreateRaw(this, &FTraceInsightsModule::SpawnLoadingProfilerTab))
 		.SetDisplayName(NSLOCTEXT("FTraceInsightsModule", "LoadingProfilerTabTitle", "Asset Loading Insights"))
 		.SetTooltipText(NSLOCTEXT("FTraceInsightsModule", "LoadingProfilerTooltipText", "Open the Asset Loading Insights tab."))
@@ -129,8 +142,9 @@ void FTraceInsightsModule::StartupModule()
 //#endif
 
 	// Register tab spawner for the Networking Insights.
-	auto& NetworkingProfilerTabSpawnerEntry = FGlobalTabmanager::Get()->RegisterTabSpawner(FInsightsManagerTabs::NetworkingProfilerTabId,
+	auto& NetworkingProfilerTabSpawnerEntry = FGlobalTabmanager::Get()->RegisterNomadTabSpawner(FInsightsManagerTabs::NetworkingProfilerTabId,
 		FOnSpawnTab::CreateRaw(this, &FTraceInsightsModule::SpawnNetworkingProfilerTab))
+		.SetReuseTabMethod(FOnFindTabToReuse::CreateStatic(&NeverReuse))
 		.SetDisplayName(NSLOCTEXT("FTraceInsightsModule", "NetworkingProfilerTabTitle", "Networking Insights"))
 		.SetTooltipText(NSLOCTEXT("FTraceInsightsModule", "NetworkingProfilerTooltipText", "Open the Networking Insights tab."))
 		.SetIcon(FSlateIcon(FInsightsStyle::GetStyleSetName(), "NetworkingProfiler.Icon.Small"));
@@ -282,7 +296,7 @@ TSharedRef<SDockTab> FTraceInsightsModule::SpawnNetworkingProfilerTab(const FSpa
 
 	// Create the SNetworkingProfilerWindow widget.
 	TSharedRef<SNetworkingProfilerWindow> Window = SNew(SNetworkingProfilerWindow, DockTab, Args.GetOwnerWindow());
-	FNetworkingProfilerManager::Get()->AssignProfilerWindow(Window);
+	FNetworkingProfilerManager::Get()->AddProfilerWindow(Window);
 	DockTab->SetContent(Window);
 
 	return DockTab;
@@ -292,6 +306,9 @@ TSharedRef<SDockTab> FTraceInsightsModule::SpawnNetworkingProfilerTab(const FSpa
 
 void FTraceInsightsModule::OnNetworkingProfilerTabBeingClosed(TSharedRef<SDockTab> TabBeingClosed)
 {
+	TSharedRef<SNetworkingProfilerWindow> Window = StaticCastSharedRef<SNetworkingProfilerWindow>(TabBeingClosed->GetContent());
+	FNetworkingProfilerManager::Get()->RemoveProfilerWindow(Window);
+
 	// Disable TabClosed delegate.
 	TabBeingClosed->SetOnTabClosed(SDockTab::FOnTabClosedCallback());
 }
diff --git a/Engine/Source/Developer/TraceInsights/Private/Insights/Widgets/SStartPageWindow.cpp b/Engine/Source/Developer/TraceInsights/Private/Insights/Widgets/SStartPageWindow.cpp
--- a/Engine/Source/Developer/TraceInsights/Private/Insights/Widgets/SStartPageWindow.cpp
+++ b/Engine/Source/Developer/TraceInsights/Private/Insights/Widgets/SStartPageWindow.cpp
@@ -344,7 +344,7 @@ void SStartPageWindow::Construct(const FArguments& InArgs)
 								.Padding(8.0f)
 								.HAlign(HAlign_Fill)
 								[
-									ConstructRecoderPanel()
+									ConstructRecorderPanel()
 								]
 							]
 						]
@@ -562,7 +562,7 @@ TSharedRef<SWidget> SStartPageWindow::ConstructLocalSessionDirectoryPanel()
 
 ////////////////////////////////////////////////////////////////////////////////////////////////////
 
-TSharedRef<SWidget> SStartPageWindow::ConstructRecoderPanel()
+TSharedRef<SWidget> SStartPageWindow::ConstructRecorderPanel()
 {
 	TSharedRef<SWidget> Widget = SNew(SVerticalBox)
 
diff --git a/Engine/Source/Developer/TraceInsights/Private/Insights/Widgets/SStartPageWindow.h b/Engine/Source/Developer/TraceInsights/Private/Insights/Widgets/SStartPageWindow.h
--- a/Engine/Source/Developer/TraceInsights/Private/Insights/Widgets/SStartPageWindow.h
+++ b/Engine/Source/Developer/TraceInsights/Private/Insights/Widgets/SStartPageWindow.h
@@ -81,7 +81,7 @@ private:
 	TSharedRef<SWidget> ConstructSessionsPanel();
 	TSharedRef<SWidget> ConstructLoadPanel();
 	TSharedRef<SWidget> ConstructLocalSessionDirectoryPanel();
-	TSharedRef<SWidget> ConstructRecoderPanel();
+	TSharedRef<SWidget> ConstructRecorderPanel();
 	TSharedRef<SWidget> ConstructConnectPanel();
 
 	/** Generate a new row for the TraceSessions list view. */
