[Insights]
  - Generic Table Tree view: Added support for multiple groupings. Also, the "Group by" combobox was replaced with a breadcrump control to choose the orderd list of grouping rules.

diff --git a/Engine/Source/Developer/TraceInsights/Private/Insights/Table/ViewModels/BaseTreeNode.h b/Engine/Source/Developer/TraceInsights/Private/Insights/Table/ViewModels/BaseTreeNode.h
--- a/Engine/Source/Developer/TraceInsights/Private/Insights/Table/ViewModels/BaseTreeNode.h
+++ b/Engine/Source/Developer/TraceInsights/Private/Insights/Table/ViewModels/BaseTreeNode.h
@@ -130,6 +130,16 @@ public:
 		Children.Add(ChildPtr);
 	}
 
+	/** Clears children. */
+	void ClearChildren()
+	{
+		for (FBaseTreeNodePtr& NodePtr : Children)
+		{
+			NodePtr->GroupPtr = nullptr;
+		}
+		Children.Reset();
+	}
+
 	/** Adds specified child to the filtered children. */
 	FORCEINLINE_DEBUGGABLE void AddFilteredChild(const FBaseTreeNodePtr& ChildPtr)
 	{
diff --git a/Engine/Source/Developer/TraceInsights/Private/Insights/Table/ViewModels/Table.cpp b/Engine/Source/Developer/TraceInsights/Private/Insights/Table/ViewModels/Table.cpp
--- a/Engine/Source/Developer/TraceInsights/Private/Insights/Table/ViewModels/Table.cpp
+++ b/Engine/Source/Developer/TraceInsights/Private/Insights/Table/ViewModels/Table.cpp
@@ -422,7 +422,7 @@ void FTable::CreateHierarchyColumn(int32 ColumnIndex, const TCHAR* ColumnName)
 	constexpr ETableCellDataType DataType = ETableCellDataType::CString;
 	constexpr ETableColumnAggregation Aggregation = ETableColumnAggregation::None;
 
-	constexpr float InitialColumnWidth = 80.0f;
+	constexpr float InitialColumnWidth = 90.0f;
 	constexpr float MinColumnWidth = 0.0f;
 	constexpr float MaxColumnWidth = FLT_MAX;
 
diff --git a/Engine/Source/Developer/TraceInsights/Private/Insights/Table/ViewModels/TreeNodeGrouping.cpp b/Engine/Source/Developer/TraceInsights/Private/Insights/Table/ViewModels/TreeNodeGrouping.cpp
--- a/Engine/Source/Developer/TraceInsights/Private/Insights/Table/ViewModels/TreeNodeGrouping.cpp
+++ b/Engine/Source/Developer/TraceInsights/Private/Insights/Table/ViewModels/TreeNodeGrouping.cpp
@@ -11,8 +11,9 @@ namespace Insights
 
 ////////////////////////////////////////////////////////////////////////////////////////////////////
 
-FTreeNodeGrouping::FTreeNodeGrouping(FText InName, FText InDescription, FName InBrushName,  FSlateBrush* InIcon)
-	: Name(InName)
+FTreeNodeGrouping::FTreeNodeGrouping(FText InShortName, FText InTitleName, FText InDescription, FName InBrushName,  FSlateBrush* InIcon)
+	: ShortName(InShortName)
+	, TitleName(InTitleName)
 	, Description(InDescription)
 	, BrushName(InBrushName)
 	, Icon(InIcon)
@@ -23,8 +24,9 @@ FTreeNodeGrouping::FTreeNodeGrouping(FText InName, FText InDescription, FName In
 
 FTreeNodeGroupingFlat::FTreeNodeGroupingFlat()
 	: FTreeNodeGrouping(
-		LOCTEXT("Grouping_Name_Flat", "Flat"),
-		LOCTEXT("Grouping_Desc_Flat", "Creates a single group. Includes all items."),
+		LOCTEXT("Grouping_Flat_ShortName", "All"),
+		LOCTEXT("Grouping_Flat_TitleName", "Flat (All)"),
+		LOCTEXT("Grouping_Flat_Desc", "Creates a single group. Includes all items."),
 		TEXT("Profiler.FiltersAndPresets.GroupNameIcon"), //TODO: "Icons.Grouping.Flat"
 		nullptr)
 {
@@ -39,13 +41,14 @@ FTreeNodeGroupInfo FTreeNodeGroupingFlat::GetGroupForNode(const FBaseTreeNodePtr
 
 ////////////////////////////////////////////////////////////////////////////////////////////////////
 
-FTreeNodeGroupingByUniqueValue::FTreeNodeGroupingByUniqueValue(TSharedRef<FTableColumn> InColumn)
+FTreeNodeGroupingByUniqueValue::FTreeNodeGroupingByUniqueValue(TSharedRef<FTableColumn> InColumnRef)
 	: FTreeNodeGrouping(
-		FText::Format(LOCTEXT("Grouping_ByUniqueValue_NameFmt", "Unique Values - {0}"), InColumn->GetTitleName()),
+		InColumnRef->GetTitleName(),
+		FText::Format(LOCTEXT("Grouping_ByUniqueValue_TitleNameFmt", "Unique Values - {0}"), InColumnRef->GetTitleName()),
 		LOCTEXT("Grouping_ByUniqueValue_Desc", "Creates a group for each unique value."),
 		TEXT("Profiler.FiltersAndPresets.Group1NameIcon"), //TODO: "Icons.Grouping.ByName"
 		nullptr)
-	, Column(InColumn)
+	, ColumnRef(InColumnRef)
 {
 }
 
@@ -54,14 +57,16 @@ FTreeNodeGroupingByUniqueValue::FTreeNodeGroupingByUniqueValue(TSharedRef<FTable
 FTreeNodeGroupInfo FTreeNodeGroupingByUniqueValue::GetGroupForNode(const FBaseTreeNodePtr InNode) const
 {
 	FTableTreeNodePtr TableTreeNodePtr = StaticCastSharedPtr<FTableTreeNode>(InNode);
-	return { FName(*Column->GetValueAsText(TableTreeNodePtr->GetRowId()).ToString()), false };
+	return { FName(*ColumnRef->GetValueFormatter()->FormatValue(TableTreeNodePtr->GetValue(*ColumnRef)).ToString()), false };
+	//return { FName(*ColumnRef->GetValueAsText(TableTreeNodePtr->GetRowId()).ToString()), false };
 }
 
 ////////////////////////////////////////////////////////////////////////////////////////////////////
 
 FTreeNodeGroupingByNameFirstLetter::FTreeNodeGroupingByNameFirstLetter()
 	: FTreeNodeGrouping(
-		LOCTEXT("Grouping_ByName_Name", "Name (First Letter)"),
+		LOCTEXT("Grouping_ByName_ShortName", "Name"),
+		LOCTEXT("Grouping_ByName_TitleName", "Name (First Letter)"),
 		LOCTEXT("Grouping_ByName_Desc", "Creates a group for each first letter of node names."),
 		TEXT("Profiler.FiltersAndPresets.Group1NameIcon"), //TODO: "Icons.Grouping.ByName"
 		nullptr)
@@ -79,7 +84,8 @@ FTreeNodeGroupInfo FTreeNodeGroupingByNameFirstLetter::GetGroupForNode(const FBa
 
 FTreeNodeGroupingByType::FTreeNodeGroupingByType()
 	: FTreeNodeGrouping(
-		LOCTEXT("Grouping_ByType_Name", "Type"),
+		LOCTEXT("Grouping_ByType_ShortName", "Type"),
+		LOCTEXT("Grouping_ByType_TitleName", "Type"),
 		LOCTEXT("Grouping_ByType_Desc", "Creates a group for each type."),
 		TEXT("Profiler.FiltersAndPresets.StatTypeIcon"), //TODO
 		nullptr)
diff --git a/Engine/Source/Developer/TraceInsights/Private/Insights/Table/ViewModels/TreeNodeGrouping.h b/Engine/Source/Developer/TraceInsights/Private/Insights/Table/ViewModels/TreeNodeGrouping.h
--- a/Engine/Source/Developer/TraceInsights/Private/Insights/Table/ViewModels/TreeNodeGrouping.h
+++ b/Engine/Source/Developer/TraceInsights/Private/Insights/Table/ViewModels/TreeNodeGrouping.h
@@ -18,11 +18,15 @@ namespace Insights
 class ITreeNodeGrouping
 {
 public:
-	virtual FText GetName() const = 0;
+	virtual FText GetShortName() const = 0;
+	virtual FText GetTitleName() const = 0;
 	virtual FText GetDescription() const = 0;
+
 	virtual FName GetBrushName() const = 0;
 	virtual const FSlateBrush* GetIcon() const = 0;
 
+	virtual FName GetColumnId() const = 0;
+
 	//virtual void CreateGrouping() = 0;
 };
 
@@ -39,18 +43,23 @@ struct FTreeNodeGroupInfo
 class FTreeNodeGrouping : public ITreeNodeGrouping
 {
 public:
-	FTreeNodeGrouping(FText InName, FText InDescription, FName InBrushName,  FSlateBrush* InIcon);
+	FTreeNodeGrouping(FText InShortName, FText InTitleName, FText InDescription, FName InBrushName,  FSlateBrush* InIcon);
 	virtual ~FTreeNodeGrouping() {}
 
-	virtual FText GetName() const override { return Name; }
+	virtual FText GetShortName() const override { return ShortName; }
+	virtual FText GetTitleName() const override { return TitleName; }
 	virtual FText GetDescription() const override { return Description; }
+
 	virtual FName GetBrushName() const override { return BrushName; }
 	virtual const FSlateBrush* GetIcon() const override { return Icon; }
 
+	virtual FName GetColumnId() const override { return NAME_None; }
+
 	virtual FTreeNodeGroupInfo GetGroupForNode(const FBaseTreeNodePtr InNode) const { return { FName(), false }; }
 
 protected:
-	FText Name;
+	FText ShortName;
+	FText TitleName;
 	FText Description;
 	FName BrushName;
 	const FSlateBrush* Icon;
@@ -74,15 +83,16 @@ public:
 class FTreeNodeGroupingByUniqueValue : public FTreeNodeGrouping
 {
 public:
-	FTreeNodeGroupingByUniqueValue(TSharedRef<FTableColumn> InColumn);
+	FTreeNodeGroupingByUniqueValue(TSharedRef<FTableColumn> InColumnRef);
 	virtual ~FTreeNodeGroupingByUniqueValue() {}
 
 	virtual FTreeNodeGroupInfo GetGroupForNode(const FBaseTreeNodePtr InNode) const override;
 
-	TSharedRef<FTableColumn> GetColumn() const { return Column; }
+	virtual FName GetColumnId() const override { return ColumnRef->GetId(); }
+	TSharedRef<FTableColumn> GetColumn() const { return ColumnRef; }
 
 private:
-	TSharedRef<FTableColumn> Column;
+	TSharedRef<FTableColumn> ColumnRef;
 };
 
 ////////////////////////////////////////////////////////////////////////////////////////////////////
diff --git a/Engine/Source/Developer/TraceInsights/Private/Insights/Table/Widgets/STableTreeView.cpp b/Engine/Source/Developer/TraceInsights/Private/Insights/Table/Widgets/STableTreeView.cpp
--- a/Engine/Source/Developer/TraceInsights/Private/Insights/Table/Widgets/STableTreeView.cpp
+++ b/Engine/Source/Developer/TraceInsights/Private/Insights/Table/Widgets/STableTreeView.cpp
@@ -4,6 +4,7 @@
 
 #include "EditorStyleSet.h"
 #include "Framework/MultiBox/MultiBoxBuilder.h"
+#include "Layout/WidgetPath.h"
 #include "SlateOptMacros.h"
 #include "TraceServices/AnalysisService.h"
 #include "Widgets/Input/SCheckBox.h"
@@ -34,6 +35,7 @@ namespace Insights
 
 const FName STableTreeView::DefaultColumnBeingSorted;
 const EColumnSortMode::Type STableTreeView::DefaultColumnSortMode(EColumnSortMode::Descending);
+const FName STableTreeView::RootNodeName(TEXT("Root"));
 
 ////////////////////////////////////////////////////////////////////////////////////////////////////
 
@@ -47,8 +49,8 @@ STableTreeView::STableTreeView()
 	, HoveredColumnId()
 	, HoveredNodePtr(nullptr)
 	, HighlightedNodeName()
+	, Root(MakeShareable(new FTableTreeNode(RootNodeName, Table)))
 	, TableTreeNodes()
-	, GroupNodes()
 	, FilteredGroupNodes()
 	, TableTreeNodesIdMap()
 	, ExpandedNodes()
@@ -57,8 +59,8 @@ STableTreeView::STableTreeView()
 	, TextFilter(nullptr)
 	, Filters(nullptr)
 	, AvailableGroupings()
-	, GroupByComboBox(nullptr)
-	, CurrentGrouping(nullptr)
+	, CurrentGroupings()
+	, GroupingBreadcrumbTrail(nullptr)
 	, AvailableSortings()
 	, ColumnToSortingMap()
 	, CurrentSorting(nullptr)
@@ -127,26 +129,27 @@ void STableTreeView::Construct(const FArguments& InArgs, TSharedPtr<FTable> InTa
 					SNew(SHorizontalBox)
 
 					+SHorizontalBox::Slot()
-					.FillWidth(1.0f)
+					.AutoWidth()
 					.VAlign(VAlign_Center)
 					[
 						SNew(STextBlock)
-						.Text(LOCTEXT("GroupByText", "Group by"))
+						.Text(LOCTEXT("GroupByText", "Hierarchy:"))
+						.Margin(FMargin(0.0f, 0.0f, 4.0f, 0.0f))
 					]
 
 					+SHorizontalBox::Slot()
-					.FillWidth(2.0f)
+					.FillWidth(1.0f)
 					.VAlign(VAlign_Center)
 					[
-						SAssignNew(GroupByComboBox, SComboBox<TSharedPtr<FTreeNodeGrouping>>)
-						.ToolTipText(this, &STableTreeView::GroupBy_GetSelectedTooltipText)
-						.OptionsSource(&AvailableGroupings)
-						.OnSelectionChanged(this, &STableTreeView::GroupBy_OnSelectionChanged)
-						.OnGenerateWidget(this, &STableTreeView::GroupBy_OnGenerateWidget)
-						[
-							SNew(STextBlock)
-							.Text(this, &STableTreeView::GroupBy_GetSelectedText)
-						]
+						SAssignNew(GroupingBreadcrumbTrail, SBreadcrumbTrail<TSharedPtr<FTreeNodeGrouping>>)
+						.ButtonContentPadding(FMargin(1.0f, 1.0f))
+						//.DelimiterImage(FEditorStyle::GetBrush("SlateFileDialogs.PathDelimiter"))
+						//.TextStyle(FEditorStyle::Get(), "Tutorials.Browser.PathText")
+						//.ShowLeadingDelimiter(true)
+						//.PersistentBreadcrumbs(true)
+						.InvertTextColorOnHover(true)
+						.OnCrumbClicked(this, &STableTreeView::OnGroupingCrumbClicked)
+						.GetCrumbMenuContent(this, &STableTreeView::GetGroupingCrumbMenuContent)
 					]
 				]
 			]
@@ -602,44 +605,23 @@ void STableTreeView::UpdateTree()
 
 void STableTreeView::ApplyFiltering()
 {
-	FilteredGroupNodes.Reset();
-
 	// Apply filter to all groups and its children.
-	const int32 NumGroups = GroupNodes.Num();
-	for (int32 GroupIndex = 0; GroupIndex < NumGroups; ++GroupIndex)
-	{
-		FTableTreeNodePtr& GroupPtr = GroupNodes[GroupIndex];
-		GroupPtr->ClearFilteredChildren();
-		const bool bIsGroupVisible = Filters->PassesAllFilters(GroupPtr);
+	ApplyFilteringForNode(Root);
 
-		const TArray<FBaseTreeNodePtr>& GroupChildren = GroupPtr->GetChildren();
-		const int32 NumChildren = GroupChildren.Num();
-		int32 NumVisibleChildren = 0;
-		for (int32 Cx = 0; Cx < NumChildren; ++Cx)
-		{
-			// Add a child.
-			const FTableTreeNodePtr& NodePtr = StaticCastSharedPtr<FTableTreeNode>(GroupChildren[Cx]);
-			const bool bIsChildVisible = Filters->PassesAllFilters(NodePtr);
-			if (bIsChildVisible)
-			{
-				GroupPtr->AddFilteredChild(NodePtr);
-				NumVisibleChildren++;
-			}
-		}
-
-		if (bIsGroupVisible || NumVisibleChildren > 0)
-		{
-			// Add a group.
-			FilteredGroupNodes.Add(GroupPtr);
-			GroupPtr->SetExpansion(true);
-		}
-		else
+	FilteredGroupNodes.Reset();
+	const TArray<FBaseTreeNodePtr>& RootChildren = Root->GetFilteredChildren();
+	const int32 NumRootChildren = RootChildren.Num();
+	for (int32 Cx = 0; Cx < NumRootChildren; ++Cx)
+	{
+		// Add a child.
+		const FTableTreeNodePtr& ChildNodePtr = StaticCastSharedPtr<FTableTreeNode>(RootChildren[Cx]);
+		if (ChildNodePtr->IsGroup())
 		{
-			GroupPtr->SetExpansion(false);
+			FilteredGroupNodes.Add(ChildNodePtr);
 		}
 	}
 
-	// Only expand timer nodes if we have a text filter.
+	// Only expand nodes if we have a text filter.
 	const bool bNonEmptyTextFilter = !TextFilter->GetRawFilterText().IsEmpty();
 	if (bNonEmptyTextFilter)
 	{
@@ -676,6 +658,50 @@ void STableTreeView::ApplyFiltering()
 
 ////////////////////////////////////////////////////////////////////////////////////////////////////
 
+bool STableTreeView::ApplyFilteringForNode(FTableTreeNodePtr NodePtr)
+{
+	const bool bIsNodeVisible = Filters->PassesAllFilters(NodePtr);
+
+	if (NodePtr->IsGroup())
+	{
+		NodePtr->ClearFilteredChildren();
+
+		const TArray<FBaseTreeNodePtr>& GroupChildren = NodePtr->GetChildren();
+		const int32 NumChildren = GroupChildren.Num();
+		int32 NumVisibleChildren = 0;
+		for (int32 Cx = 0; Cx < NumChildren; ++Cx)
+		{
+			// Add a child.
+			const FTableTreeNodePtr& ChildNodePtr = StaticCastSharedPtr<FTableTreeNode>(GroupChildren[Cx]);
+			if (ApplyFilteringForNode(ChildNodePtr))
+			{
+				NodePtr->AddFilteredChild(ChildNodePtr);
+				NumVisibleChildren++;
+			}
+		}
+
+		const bool bIsGroupNodeVisible = bIsNodeVisible || NumVisibleChildren > 0;
+
+		if (bIsGroupNodeVisible)
+		{
+			// Add a group.
+			NodePtr->SetExpansion(true);
+		}
+		else
+		{
+			NodePtr->SetExpansion(false);
+		}
+
+		return bIsGroupNodeVisible;
+	}
+	else
+	{
+		return bIsNodeVisible;
+	}
+}
+
+////////////////////////////////////////////////////////////////////////////////////////////////////
+
 void STableTreeView::HandleItemToStringArray(const FTableTreeNodePtr& FTableTreeNodePtr, TArray<FString>& OutSearchStrings) const
 {
 	OutSearchStrings.Add(FTableTreeNodePtr->GetName().GetPlainNameString());
@@ -850,21 +876,58 @@ bool STableTreeView::SearchBox_IsEnabled() const
 
 void STableTreeView::CreateGroups()
 {
-	GroupNodes.Reset();
+	GroupNodesRec(TableTreeNodes, *Root, 0);
+
+	ResetAggregatedValuesRec(*Root);
+	for (const TSharedPtr<FTableColumn> ColumnPtr : Table->GetColumns())
+	{
+		FTableColumn& Column = *ColumnPtr;
+		if (Column.GetAggregation() == ETableColumnAggregation::Sum)
+		{
+			switch (Column.GetDataType())
+			{
+				case ETableCellDataType::Int64:
+					UpdateInt64SumAggregationRec(Column, *Root);
+					break;
+
+				case ETableCellDataType::Float:
+					UpdateFloatSumAggregationRec(Column, *Root);
+					break;
+
+				case ETableCellDataType::Double:
+					UpdateDoubleSumAggregationRec(Column, *Root);
+					break;
+			}
+		}
+	}
+}
+
+////////////////////////////////////////////////////////////////////////////////////////////////////
+
+void STableTreeView::GroupNodesRec(const TArray<FTableTreeNodePtr>& Nodes, FTableTreeNode& ParentGroup, int32 GroupingDepth)
+{
+	ensure(CurrentGroupings.Num() > 0);
+	
+	FTreeNodeGrouping& Grouping = *CurrentGroupings[GroupingDepth];
+
 	TMap<FName, FTableTreeNodePtr> GroupMap;
 
-	for (const FTableTreeNodePtr& NodePtr : TableTreeNodes)
+	ParentGroup.ClearChildren();
+
+	for (FTableTreeNodePtr NodePtr : Nodes)
 	{
-		FTreeNodeGroupInfo GroupInfo = CurrentGrouping->GetGroupForNode(NodePtr);
+		ensure(!NodePtr->IsGroup());
 
 		FTableTreeNodePtr GroupPtr = nullptr;
+
+		FTreeNodeGroupInfo GroupInfo = Grouping.GetGroupForNode(NodePtr);
 		FTableTreeNodePtr* GroupPtrPtr = GroupMap.Find(GroupInfo.Name);
 		if (!GroupPtrPtr)
 		{
 			GroupPtr = MakeShareable(new FTableTreeNode(GroupInfo.Name, Table));
 
 			GroupMap.Add(GroupInfo.Name, GroupPtr);
-			GroupNodes.Add(GroupPtr);
+			ParentGroup.AddChildAndSetGroupPtr(GroupPtr);
 
 			GroupPtr->SetExpansion(GroupInfo.IsExpanded);
 			TreeView->SetItemExpansion(GroupPtr, GroupInfo.IsExpanded);
@@ -877,83 +940,105 @@ void STableTreeView::CreateGroups()
 		GroupPtr->AddChildAndSetGroupPtr(NodePtr);
 	}
 
-	// Update aggregations.
-	for (const FTableTreeNodePtr& GroupPtr : GroupNodes)
+	if (GroupingDepth < CurrentGroupings.Num() - 1)
 	{
-		GroupPtr->ResetAggregatedValues();
-	}
-	for (const TSharedPtr<FTableColumn> ColumnPtr : Table->GetColumns())
-	{
-		if (ColumnPtr->GetAggregation() == ETableColumnAggregation::Sum)
+		TArray<FTableTreeNodePtr> ChildNodes;
+
+		for (FBaseTreeNodePtr GroupPtr : ParentGroup.GetChildren())
 		{
-			bool bValidForAggregation = false;
+			ensure(GroupPtr->IsGroup());
+			FTableTreeNode& Group = *StaticCastSharedPtr<FTableTreeNode>(GroupPtr);
 
-			switch (ColumnPtr->GetDataType())
+			// Make a copy of the child nodes.
+			ChildNodes.Reset();
+			for (FBaseTreeNodePtr ChildPtr : Group.GetChildren())
 			{
-				case ETableCellDataType::Int64:
-					bValidForAggregation = true;
-					break;
-				case ETableCellDataType::Float:
-					bValidForAggregation = true;
-					break;
-				case ETableCellDataType::Double:
-					bValidForAggregation = true;
-					break;
+				ChildNodes.Add(StaticCastSharedPtr<FTableTreeNode>(ChildPtr));
 			}
 
-			if (bValidForAggregation)
-			{
-				for (const FTableTreeNodePtr& GroupPtr : GroupNodes)
-				{
-					switch (ColumnPtr->GetDataType())
-					{
-						case ETableCellDataType::Int64:
-						{
-							int64 AggregatedValue = 0;
-							for (FBaseTreeNodePtr NodePtr : GroupPtr->GetChildren())
-							{
-								FTableTreeNodePtr TableNodePtr = StaticCastSharedPtr<FTableTreeNode>(NodePtr);
-								FTableCellValue Value = TableNodePtr->GetValue(*ColumnPtr);
-								AggregatedValue += Value.Int64;
-							}
-							GroupPtr->AddAggregatedValue(ColumnPtr->GetId(), FTableCellValue(AggregatedValue));
-							break;
-						}
-						case ETableCellDataType::Float:
-						{
-							float AggregatedValue = 0.0f;
-							for (FBaseTreeNodePtr NodePtr : GroupPtr->GetChildren())
-							{
-								FTableTreeNodePtr TableNodePtr = StaticCastSharedPtr<FTableTreeNode>(NodePtr);
-								FTableCellValue Value = TableNodePtr->GetValue(*ColumnPtr);
-								AggregatedValue += Value.Float;
-							}
-							GroupPtr->AddAggregatedValue(ColumnPtr->GetId(), FTableCellValue(AggregatedValue));
-							break;
-						}
-						case ETableCellDataType::Double:
-						{
-							double AggregatedValue = 0.0;
-							for (FBaseTreeNodePtr NodePtr : GroupPtr->GetChildren())
-							{
-								FTableTreeNodePtr TableNodePtr = StaticCastSharedPtr<FTableTreeNode>(NodePtr);
-								FTableCellValue Value = TableNodePtr->GetValue(*ColumnPtr);
-								AggregatedValue += Value.Double;
-							}
-							GroupPtr->AddAggregatedValue(ColumnPtr->GetId(), FTableCellValue(AggregatedValue));
-							break;
-						}
-					}
-				}
-			}
+			GroupNodesRec(ChildNodes, Group, GroupingDepth + 1);
 		}
-		// TODO: else if (ColumnPtr->GetAggregation() == ETableColumnAggregation::Min)
-		//{
-		//}
+	}
+}
+
+////////////////////////////////////////////////////////////////////////////////////////////////////
+
+void STableTreeView::ResetAggregatedValuesRec(FTableTreeNode& GroupNode)
+{
+	ensure(GroupNode.IsGroup());
+
+	GroupNode.ResetAggregatedValues();
+
+	for (FBaseTreeNodePtr ChildPtr : GroupNode.GetChildren())
+	{
+		if (ChildPtr->IsGroup())
+		{
+			ResetAggregatedValuesRec(*StaticCastSharedPtr<FTableTreeNode>(ChildPtr));
+		}
+	}
+}
+
+////////////////////////////////////////////////////////////////////////////////////////////////////
+
+void STableTreeView::UpdateInt64SumAggregationRec(FTableColumn& Column, FTableTreeNode& GroupNode)
+{
+	int64 AggregatedValue = 0;
+
+	for (FBaseTreeNodePtr NodePtr : GroupNode.GetChildren())
+	{
+		FTableTreeNode& TableNode = *StaticCastSharedPtr<FTableTreeNode>(NodePtr);
+		if (TableNode.IsGroup())
+		{
+			UpdateInt64SumAggregationRec(Column, TableNode);
+		}
+
+		FTableCellValue Value = TableNode.GetValue(Column);
+		AggregatedValue += Value.Int64;
+	}
+
+	GroupNode.AddAggregatedValue(Column.GetId(), FTableCellValue(AggregatedValue));
+}
+
+////////////////////////////////////////////////////////////////////////////////////////////////////
+
+void STableTreeView::UpdateFloatSumAggregationRec(FTableColumn& Column, FTableTreeNode& GroupNode)
+{
+	float AggregatedValue = 0.0f;
+
+	for (FBaseTreeNodePtr NodePtr : GroupNode.GetChildren())
+	{
+		FTableTreeNode& TableNode = *StaticCastSharedPtr<FTableTreeNode>(NodePtr);
+		if (TableNode.IsGroup())
+		{
+			UpdateFloatSumAggregationRec(Column, TableNode);
+		}
+
+		FTableCellValue Value = TableNode.GetValue(Column);
+		AggregatedValue += Value.Float;
+	}
+
+	GroupNode.AddAggregatedValue(Column.GetId(), FTableCellValue(AggregatedValue));
+}
+
+////////////////////////////////////////////////////////////////////////////////////////////////////
+
+void STableTreeView::UpdateDoubleSumAggregationRec(FTableColumn& Column, FTableTreeNode& GroupNode)
+{
+	double AggregatedValue = 0.0;
+
+	for (FBaseTreeNodePtr NodePtr : GroupNode.GetChildren())
+	{
+		FTableTreeNode& TableNode = *StaticCastSharedPtr<FTableTreeNode>(NodePtr);
+		if (TableNode.IsGroup())
+		{
+			UpdateDoubleSumAggregationRec(Column, TableNode);
+		}
+							
+		FTableCellValue Value = TableNode.GetValue(Column);
+		AggregatedValue += Value.Double;
 	}
 
-	// Sort groups by name.
-	GroupNodes.Sort([](const FBaseTreeNodePtr& A, const FBaseTreeNodePtr& B) { return A->GetName().LexicalLess(B->GetName()); });
+	GroupNode.AddAggregatedValue(Column.GetId(), FTableCellValue(AggregatedValue));
 }
 
 ////////////////////////////////////////////////////////////////////////////////////////////////////
@@ -974,101 +1059,457 @@ void STableTreeView::CreateGroupings()
 		}
 	}
 
-	CurrentGrouping = AvailableGroupings[0];
+	CurrentGroupings.Add(AvailableGroupings[0]);
 
-	GroupByComboBox->SetSelectedItem(CurrentGrouping);
-	GroupByComboBox->RefreshOptions();
+	RebuildGroupingCrumbs();
 }
 
 ////////////////////////////////////////////////////////////////////////////////////////////////////
 
-void STableTreeView::GroupBy_OnSelectionChanged(TSharedPtr<FTreeNodeGrouping> NewGrouping, ESelectInfo::Type SelectInfo)
+void STableTreeView::PreChangeGroupings()
 {
-	constexpr float HierarcyColumnWidth = 80.0f;
-
-	if (SelectInfo != ESelectInfo::Direct)
+	for (TSharedPtr<FTreeNodeGrouping> GroupingPtr : CurrentGroupings)
 	{
-		if (CurrentGrouping.IsValid() && CurrentGrouping->GetName().ToString().StartsWith("Unique Value"))
+		const FName& ColumnId = GroupingPtr->GetColumnId();
+		if (ColumnId != NAME_None)
 		{
-			TSharedPtr<FTreeNodeGroupingByUniqueValue> GroupingByUniqueValue = StaticCastSharedPtr<FTreeNodeGroupingByUniqueValue>(CurrentGrouping);
+			// Show columns used in previous groupings.
+			ShowColumn(ColumnId);
+		}
+	}
+}
 
-			const FName& ColumnId = GroupingByUniqueValue->GetColumn()->GetId();
+////////////////////////////////////////////////////////////////////////////////////////////////////
 
-			ShowColumn(ColumnId);
+void STableTreeView::PostChangeGroupings()
+{
+	constexpr float HierarchyMinWidth = 60.0f;
+	constexpr float HierarchyIndentation = 10.0f;
+	constexpr float DefaultHierarchyColumnWidth = 90.0f;
 
-			SHeaderRow::FColumn& HierarcyColumn = const_cast<SHeaderRow::FColumn&>(TreeViewHeaderRow->GetColumns()[0]);
-			FTableColumn& HierarcyTableColumn = *Table->FindColumnChecked(HierarcyColumn.ColumnId);
-			const FText HierarcyColumnName(LOCTEXT("HierarchyShortName", "Hierarchy"));
-			HierarcyTableColumn.SetShortName(HierarcyColumnName);
+	float HierarchyColumnWidth = DefaultHierarchyColumnWidth;
+	//FString GroupingStr;
 
+	int32 GroupingDepth = 0;
+	for (TSharedPtr<FTreeNodeGrouping> GroupingPtr : CurrentGroupings)
+	{
+		const FName& ColumnId = GroupingPtr->GetColumnId();
+
+		if (ColumnId != NAME_None)
+		{
+			// Compute width for Hierarchy column based on column used in grouping and its indentation.
 			const int32 NumColumns = TreeViewHeaderRow->GetColumns().Num();
 			for (int32 ColumnIndex = 0; ColumnIndex < NumColumns; ++ColumnIndex)
 			{
 				const SHeaderRow::FColumn& CurrentColumn = TreeViewHeaderRow->GetColumns()[ColumnIndex];
 				if (CurrentColumn.ColumnId == ColumnId)
 				{
-					HierarcyColumn.SetWidth(HierarcyColumnWidth);
+					const float Width = HierarchyMinWidth + GroupingDepth * HierarchyIndentation + CurrentColumn.GetWidth();
+					if (Width > HierarchyColumnWidth)
+					{
+						HierarchyColumnWidth = Width;
+					}
 					break;
 				}
 			}
+
+			// Hide columns used in groupings.
+			HideColumn(ColumnId);
 		}
 
-		CurrentGrouping = NewGrouping;
+		// Compute name of the Hierarchy column.
+		//if (!GroupingStr.IsEmpty())
+		//{
+		//	GroupingStr.Append(TEXT(" / "));
+		//}
+		//GroupingStr.Append(GroupingPtr->GetShortName().ToString());
+
+		++GroupingDepth;
+	}
+
+	//////////////////////////////////////////////////
+
+	// Set with for the Hierarchy column.
+	SHeaderRow::FColumn& HierarchyColumn = const_cast<SHeaderRow::FColumn&>(TreeViewHeaderRow->GetColumns()[0]);
+	HierarchyColumn.SetWidth(HierarchyColumnWidth);
+
+	// Set name for the Hierarchy column.
+	//FTableColumn& HierarchyTableColumn = *Table->FindColumnChecked(HierarchyColumn.ColumnId);
+	//if (!GroupingStr.IsEmpty())
+	//{
+	//	const FText HierarchyColumnName = FText::Format(LOCTEXT("HierarchyShortNameFmt", "Hierarchy ({0})"), FText::FromString(GroupingStr));
+	//	HierarchyTableColumn.SetShortName(HierarchyColumnName);
+	//}
+	//else
+	//{
+	//	const FText HierarchyColumnName(LOCTEXT("HierarchyShortName", "Hierarchy"));
+	//	HierarchyTableColumn.SetShortName(HierarchyColumnName);
+	//}
 
-		if (CurrentGrouping.IsValid() && CurrentGrouping->GetName().ToString().StartsWith("Unique Value"))
+	//////////////////////////////////////////////////
+
+	TreeViewHeaderRow->RefreshColumns();
+
+	CreateGroups();
+	SortTreeNodes();
+	ApplyFiltering();
+
+	RebuildGroupingCrumbs();
+}
+
+////////////////////////////////////////////////////////////////////////////////////////////////////
+
+void STableTreeView::RebuildGroupingCrumbs()
+{
+	GroupingBreadcrumbTrail->ClearCrumbs();
+
+	for (const TSharedPtr<FTreeNodeGrouping> Grouping : CurrentGroupings)
+	{
+		GroupingBreadcrumbTrail->PushCrumb(Grouping->GetShortName(), Grouping);
+	}
+}
+
+////////////////////////////////////////////////////////////////////////////////////////////////////
+
+int32 STableTreeView::GetGroupingDepth(const TSharedPtr<FTreeNodeGrouping>& Grouping) const
+{
+	for (int32 GroupingDepth = CurrentGroupings.Num() - 1; GroupingDepth >= 0; --GroupingDepth)
+	{
+		if (Grouping == CurrentGroupings[GroupingDepth])
 		{
-			TSharedPtr<FTreeNodeGroupingByUniqueValue> GroupingByUniqueValue = StaticCastSharedPtr<FTreeNodeGroupingByUniqueValue>(CurrentGrouping);
+			return GroupingDepth;
+		}
+	}
+	return -1;
+}
 
-			const FName& ColumnId = GroupingByUniqueValue->GetColumn()->GetId();
+////////////////////////////////////////////////////////////////////////////////////////////////////
 
-			SHeaderRow::FColumn& HierarcyColumn = const_cast<SHeaderRow::FColumn&>(TreeViewHeaderRow->GetColumns()[0]);
-			FTableColumn& HierarcyTableColumn = *Table->FindColumnChecked(HierarcyColumn.ColumnId);
-			const FText HierarcyColumnName = FText::Format(LOCTEXT("HierarchyShortNameFmt", "Hierarchy ({0})"), GroupingByUniqueValue->GetColumn()->GetShortName());
-			HierarcyTableColumn.SetShortName(HierarcyColumnName);
+void STableTreeView::OnGroupingCrumbClicked(const TSharedPtr<FTreeNodeGrouping>& CrumbGrouping)
+{
+	const int32 CrumbGroupingDepth = GetGroupingDepth(CrumbGrouping);
+	if (CrumbGroupingDepth >= 0 && CrumbGroupingDepth < CurrentGroupings.Num() - 1)
+	{
+		PreChangeGroupings();
 
-			const int32 NumColumns = TreeViewHeaderRow->GetColumns().Num();
-			for (int32 ColumnIndex = 0; ColumnIndex < NumColumns; ++ColumnIndex)
+		CurrentGroupings.RemoveAt(CrumbGroupingDepth + 1, CurrentGroupings.Num() - CrumbGroupingDepth - 1);
+
+		PostChangeGroupings();
+	}
+}
+
+////////////////////////////////////////////////////////////////////////////////////////////////////
+
+void STableTreeView::BuildGroupingSubMenu_Change(FMenuBuilder& MenuBuilder, const TSharedPtr<FTreeNodeGrouping> CrumbGrouping)
+{
+	MenuBuilder.BeginSection("ChangeGrouping");
+	{
+		for (const TSharedPtr<FTreeNodeGrouping>& Grouping : AvailableGroupings)
+		{
+			FUIAction Action_Change
+			(
+				FExecuteAction::CreateSP(this, &STableTreeView::GroupingCrumbMenu_Change_Execute, CrumbGrouping, Grouping),
+				FCanExecuteAction::CreateSP(this, &STableTreeView::GroupingCrumbMenu_Change_CanExecute, CrumbGrouping, Grouping)
+			);
+			MenuBuilder.AddMenuEntry
+			(
+				Grouping->GetTitleName(),
+				Grouping->GetDescription(),
+				FSlateIcon(), Action_Change, NAME_None, EUserInterfaceActionType::Button
+			);
+		}
+	}
+	MenuBuilder.EndSection();
+}
+
+////////////////////////////////////////////////////////////////////////////////////////////////////
+
+void STableTreeView::BuildGroupingSubMenu_Add(FMenuBuilder& MenuBuilder, const TSharedPtr<FTreeNodeGrouping> CrumbGrouping)
+{
+	MenuBuilder.BeginSection("AddGrouping");
+	{
+		for (const TSharedPtr<FTreeNodeGrouping>& Grouping : AvailableGroupings)
+		{
+			FUIAction Action_Add
+			(
+				FExecuteAction::CreateSP(this, &STableTreeView::GroupingCrumbMenu_Add_Execute, Grouping, CrumbGrouping),
+				FCanExecuteAction::CreateSP(this, &STableTreeView::GroupingCrumbMenu_Add_CanExecute, Grouping, CrumbGrouping)
+			);
+			MenuBuilder.AddMenuEntry
+			(
+				Grouping->GetTitleName(),
+				Grouping->GetDescription(),
+				FSlateIcon(), Action_Add, NAME_None, EUserInterfaceActionType::Button
+			);
+		}
+	}
+	MenuBuilder.EndSection();
+}
+
+////////////////////////////////////////////////////////////////////////////////////////////////////
+
+TSharedRef<SWidget> STableTreeView::GetGroupingCrumbMenuContent(const TSharedPtr<FTreeNodeGrouping>& CrumbGrouping) const
+{
+	const bool bShouldCloseWindowAfterMenuSelection = true;
+	FMenuBuilder MenuBuilder(bShouldCloseWindowAfterMenuSelection, NULL);
+
+	const int32 CrumbGroupingDepth = GetGroupingDepth(CrumbGrouping);
+
+	MenuBuilder.BeginSection("InsertOrAdd");
+	{
+		const FText AddGroupingText = (CrumbGroupingDepth == CurrentGroupings.Num() - 1) ? // after last one
+			LOCTEXT("AddGrouping_Section", "Add Grouping...") :
+			LOCTEXT("InsertGrouping_Section", "Insert Grouping...");
+		MenuBuilder.AddSubMenu
+		(
+			AddGroupingText,
+			LOCTEXT("GroupingMenu_Add_Desc", "Add or insert new grouping."),
+			FNewMenuDelegate::CreateSP(this, &STableTreeView::BuildGroupingSubMenu_Add, CrumbGrouping),
+			false,
+			FSlateIcon()
+		);
+	}
+	MenuBuilder.EndSection();
+
+	if (CrumbGroupingDepth >= 0)
+	{
+		MenuBuilder.BeginSection("CrumbGrouping", CrumbGrouping->GetTitleName());
+		{
+			MenuBuilder.AddSubMenu
+			(
+				LOCTEXT("ChangeGrouping_Section", "Change With..."),
+				LOCTEXT("GroupingMenu_Add_Desc", "Change selected grouping."),
+				FNewMenuDelegate::CreateSP(this, &STableTreeView::BuildGroupingSubMenu_Change, CrumbGrouping),
+				false,
+				FSlateIcon()
+			);
+
+			if (CrumbGroupingDepth > 0)
 			{
-				const SHeaderRow::FColumn& CurrentColumn = TreeViewHeaderRow->GetColumns()[ColumnIndex];
-				if (CurrentColumn.ColumnId == ColumnId)
-				{
-					HierarcyColumn.SetWidth(HierarcyColumnWidth + CurrentColumn.GetWidth());
-					break;
-				}
+				FUIAction Action_MoveLeft
+				(
+					FExecuteAction::CreateSP(this, &STableTreeView::GroupingCrumbMenu_MoveLeft_Execute, CrumbGrouping),
+					FCanExecuteAction()
+				);
+				MenuBuilder.AddMenuEntry
+				(
+					LOCTEXT("GroupingMenu_MoveLeft", "Move Left"),
+					LOCTEXT("GroupingMenu_MoveLeft_Desc", "Move selected grouping to the left."),
+					FSlateIcon(), Action_MoveLeft, NAME_None, EUserInterfaceActionType::Button
+				);
 			}
 
-			HideColumn(ColumnId);
+			if (CrumbGroupingDepth < CurrentGroupings.Num() - 1)
+			{
+				FUIAction Action_MoveRight
+				(
+					FExecuteAction::CreateSP(this, &STableTreeView::GroupingCrumbMenu_MoveRight_Execute, CrumbGrouping),
+					FCanExecuteAction()
+				);
+				MenuBuilder.AddMenuEntry
+				(
+					LOCTEXT("GroupingMenu_MoveRight", "Move Right"),
+					LOCTEXT("GroupingMenu_MoveRight_Desc", "Move selected grouping to the right."),
+					FSlateIcon(), Action_MoveRight, NAME_None, EUserInterfaceActionType::Button
+				);
+			}
+
+			if (CurrentGroupings.Num() > 1)
+			{
+				FUIAction Action_Remove
+				(
+					FExecuteAction::CreateSP(this, &STableTreeView::GroupingCrumbMenu_Remove_Execute, CrumbGrouping),
+					FCanExecuteAction()
+				);
+				MenuBuilder.AddMenuEntry
+				(
+					LOCTEXT("GroupingMenu_Remove", "Remove"),
+					LOCTEXT("GroupingMenu_Remove_Desc", "Remove selected grouping."),
+					FSlateIcon(), Action_Remove, NAME_None, EUserInterfaceActionType::Button
+				);
+			}
+		}
+		MenuBuilder.EndSection();
+	}
+
+	if (CurrentGroupings.Num() > 1 || CurrentGroupings[0] != AvailableGroupings[0])
+	{
+		MenuBuilder.BeginSection("ResetGroupings");
+		{
+			FUIAction Action_Reset
+			(
+				FExecuteAction::CreateSP(this, &STableTreeView::GroupingCrumbMenu_Reset_Execute),
+				FCanExecuteAction()
+			);
+			MenuBuilder.AddMenuEntry
+			(
+				LOCTEXT("GroupingMenu_Reset", "Reset"),
+				LOCTEXT("GroupingMenu_Reset_Desc", "Reset groupings to default."),
+				FSlateIcon(), Action_Reset, NAME_None, EUserInterfaceActionType::Button
+			);
 		}
+		MenuBuilder.EndSection();
+	}
+
+	return MenuBuilder.MakeWidget();
+}
+
+////////////////////////////////////////////////////////////////////////////////////////////////////
+
+void STableTreeView::GroupingCrumbMenu_Reset_Execute()
+{
+	PreChangeGroupings();
+
+	CurrentGroupings.Reset();
+	CurrentGroupings.Add(AvailableGroupings[0]);
+
+	PostChangeGroupings();
+}
+
+////////////////////////////////////////////////////////////////////////////////////////////////////
+
+void STableTreeView::GroupingCrumbMenu_Remove_Execute(const TSharedPtr<FTreeNodeGrouping> Grouping)
+{
+	const int32 GroupingDepth = GetGroupingDepth(Grouping);
+	if (GroupingDepth >= 0)
+	{
+		PreChangeGroupings();
 
-		TreeViewHeaderRow->RefreshColumns();
+		CurrentGroupings.RemoveAt(GroupingDepth);
 
-		CreateGroups();
-		SortTreeNodes();
-		ApplyFiltering();
+		PostChangeGroupings();
 	}
 }
 
 ////////////////////////////////////////////////////////////////////////////////////////////////////
 
-TSharedRef<SWidget> STableTreeView::GroupBy_OnGenerateWidget(TSharedPtr<FTreeNodeGrouping> InGrouping) const
+void STableTreeView::GroupingCrumbMenu_MoveLeft_Execute(const TSharedPtr<FTreeNodeGrouping> Grouping)
 {
-	return SNew(STextBlock)
-		.Text(InGrouping->GetName())
-		.ToolTipText(InGrouping->GetDescription());
+	const int32 GroupingDepth = GetGroupingDepth(Grouping);
+	if (GroupingDepth > 0)
+	{
+		PreChangeGroupings();
+
+		CurrentGroupings[GroupingDepth] = CurrentGroupings[GroupingDepth - 1];
+		CurrentGroupings[GroupingDepth - 1] = Grouping;
+
+		PostChangeGroupings();
+	}
 }
 
 ////////////////////////////////////////////////////////////////////////////////////////////////////
 
-FText STableTreeView::GroupBy_GetSelectedText() const
+void STableTreeView::GroupingCrumbMenu_MoveRight_Execute(const TSharedPtr<FTreeNodeGrouping> Grouping)
 {
-	return CurrentGrouping.IsValid() ? CurrentGrouping->GetName() : FText::GetEmpty();
+	const int32 GroupingDepth = GetGroupingDepth(Grouping);
+	if (GroupingDepth < CurrentGroupings.Num() - 1)
+	{
+		PreChangeGroupings();
+
+		CurrentGroupings[GroupingDepth] = CurrentGroupings[GroupingDepth + 1];
+		CurrentGroupings[GroupingDepth + 1] = Grouping;
+
+		PostChangeGroupings();
+	}
+}
+
+////////////////////////////////////////////////////////////////////////////////////////////////////
+
+void STableTreeView::GroupingCrumbMenu_Change_Execute(const TSharedPtr<FTreeNodeGrouping> OldGrouping, const TSharedPtr<FTreeNodeGrouping> NewGrouping)
+{
+	const int32 OldGroupingDepth = GetGroupingDepth(OldGrouping);
+	if (OldGroupingDepth >= 0)
+	{
+		PreChangeGroupings();
+
+		const int32 NewGroupingDepth = GetGroupingDepth(NewGrouping);
+
+		if (NewGroupingDepth >= 0 && NewGroupingDepth != OldGroupingDepth) // NewGrouping already exists
+		{
+			CurrentGroupings.RemoveAt(NewGroupingDepth);
+
+			if (NewGroupingDepth < OldGroupingDepth)
+			{
+				CurrentGroupings[OldGroupingDepth - 1] = NewGrouping;
+			}
+			else
+			{
+				CurrentGroupings[OldGroupingDepth] = NewGrouping;
+			}
+		}
+		else
+		{
+			CurrentGroupings[OldGroupingDepth] = NewGrouping;
+		}
+
+		PostChangeGroupings();
+	}
 }
 
 ////////////////////////////////////////////////////////////////////////////////////////////////////
 
-FText STableTreeView::GroupBy_GetSelectedTooltipText() const
+bool STableTreeView::GroupingCrumbMenu_Change_CanExecute(const TSharedPtr<FTreeNodeGrouping> OldGrouping, const TSharedPtr<FTreeNodeGrouping> NewGrouping) const
 {
-	return CurrentGrouping.IsValid() ? CurrentGrouping->GetDescription() : FText::GetEmpty();
+	return NewGrouping != OldGrouping;
+}
+
+////////////////////////////////////////////////////////////////////////////////////////////////////
+
+void STableTreeView::GroupingCrumbMenu_Add_Execute(const TSharedPtr<FTreeNodeGrouping> Grouping, const TSharedPtr<FTreeNodeGrouping> AfterGrouping)
+{
+	PreChangeGroupings();
+
+	if (AfterGrouping.IsValid())
+	{
+		const int32 AfterGroupingDepth = GetGroupingDepth(AfterGrouping);
+		ensure(AfterGroupingDepth >= 0);
+
+		const int32 GroupingDepth = GetGroupingDepth(Grouping);
+
+		if (GroupingDepth >= 0) // Grouping already exists
+		{
+			CurrentGroupings.RemoveAt(GroupingDepth);
+
+			if (GroupingDepth <= AfterGroupingDepth)
+			{
+				CurrentGroupings.Insert(Grouping, AfterGroupingDepth);
+			}
+			else
+			{
+				CurrentGroupings.Insert(Grouping, AfterGroupingDepth + 1);
+			}
+		}
+		else
+		{
+			CurrentGroupings.Insert(Grouping, AfterGroupingDepth + 1);
+		}
+	}
+	else
+	{
+		CurrentGroupings.Remove(Grouping);
+		CurrentGroupings.Add(Grouping);
+	}
+
+	PostChangeGroupings();
+}
+
+////////////////////////////////////////////////////////////////////////////////////////////////////
+
+bool STableTreeView::GroupingCrumbMenu_Add_CanExecute(const TSharedPtr<FTreeNodeGrouping> Grouping, const TSharedPtr<FTreeNodeGrouping> AfterGrouping) const
+{
+	if (AfterGrouping.IsValid())
+	{
+		const int32 AfterGroupingDepth = GetGroupingDepth(AfterGrouping);
+		ensure(AfterGroupingDepth >= 0);
+
+		const int32 GroupingDepth = GetGroupingDepth(Grouping);
+
+		return GroupingDepth < AfterGroupingDepth || GroupingDepth > AfterGroupingDepth + 1;
+	}
+	else
+	{
+		return Grouping != CurrentGroupings.Last();
+	}
 }
 
 ////////////////////////////////////////////////////////////////////////////////////////////////////
@@ -1145,26 +1586,28 @@ void STableTreeView::SortTreeNodes()
 	if (CurrentSorting.IsValid())
 	{
 		const Insights::ITreeNodeSorting* Sorter = CurrentSorting.Get();
+		SortTreeNodesRec(*Root, Sorter);
+	}
+}
 
-		if (ColumnSortMode == EColumnSortMode::Type::Descending)
-		{
-			GroupNodes.Sort(Sorter->GetDescendingCompareDelegate());
+////////////////////////////////////////////////////////////////////////////////////////////////////
 
-			const int32 NumGroups = GroupNodes.Num();
-			for (int32 GroupIndex = 0; GroupIndex < NumGroups; ++GroupIndex)
-			{
-				GroupNodes[GroupIndex]->SortChildrenDescending(Sorter);
-			}
-		}
-		else
-		{
-			GroupNodes.Sort(Sorter->GetAscendingCompareDelegate());
+void STableTreeView::SortTreeNodesRec(FTableTreeNode& GroupNode, const Insights::ITreeNodeSorting* Sorter)
+{
+	if (ColumnSortMode == EColumnSortMode::Type::Descending)
+	{
+		GroupNode.SortChildrenDescending(Sorter);
+	}
+	else
+	{
+		GroupNode.SortChildrenAscending(Sorter);
+	}
 
-			const int32 NumGroups = GroupNodes.Num();
-			for (int32 GroupIndex = 0; GroupIndex < NumGroups; ++GroupIndex)
-			{
-				GroupNodes[GroupIndex]->SortChildrenAscending(Sorter);
-			}
+	for (FBaseTreeNodePtr ChildPtr : GroupNode.GetChildren())
+	{
+		if (ChildPtr->IsGroup())
+		{
+			SortTreeNodesRec(*StaticCastSharedPtr<FTableTreeNode>(ChildPtr), Sorter);
 		}
 	}
 }
diff --git a/Engine/Source/Developer/TraceInsights/Private/Insights/Table/Widgets/STableTreeView.h b/Engine/Source/Developer/TraceInsights/Private/Insights/Table/Widgets/STableTreeView.h
--- a/Engine/Source/Developer/TraceInsights/Private/Insights/Table/Widgets/STableTreeView.h
+++ b/Engine/Source/Developer/TraceInsights/Private/Insights/Table/Widgets/STableTreeView.h
@@ -7,7 +7,7 @@
 #include "Misc/TextFilter.h"
 #include "SlateFwd.h"
 #include "Widgets/DeclarativeSyntaxSupport.h"
-#include "Widgets/Input/SComboBox.h"
+#include "Widgets/Navigation/SBreadcrumbTrail.h"
 #include "Widgets/SCompoundWidget.h"
 #include "Widgets/SWidget.h"
 #include "Widgets/Views/SHeaderRow.h"
@@ -147,6 +147,8 @@ protected:
 	/** Populates the group and stat tree with items based on the current data. */
 	void ApplyFiltering();
 
+	bool ApplyFilteringForNode(FTableTreeNodePtr NodePtr);
+
 	bool SearchBox_IsEnabled() const;
 	void SearchBox_OnTextChanged(const FText& InFilterText);
 
@@ -154,15 +156,33 @@ protected:
 	// Grouping
 
 	void CreateGroupings();
-	void CreateGroups();
-
-	void GroupBy_OnSelectionChanged(TSharedPtr<FTreeNodeGrouping> NewGrouping, ESelectInfo::Type SelectInfo);
 
-	TSharedRef<SWidget> GroupBy_OnGenerateWidget(TSharedPtr<FTreeNodeGrouping> InGrouping) const;
-
-	FText GroupBy_GetSelectedText() const;
-
-	FText GroupBy_GetSelectedTooltipText() const;
+	void CreateGroups();
+	void GroupNodesRec(const TArray<FTableTreeNodePtr>& Nodes, FTableTreeNode& ParentGroup, int32 GroupingDepth);
+
+	void ResetAggregatedValuesRec(FTableTreeNode& GroupNode);
+	void UpdateInt64SumAggregationRec(FTableColumn& Column, FTableTreeNode& GroupNode);
+	void UpdateFloatSumAggregationRec(FTableColumn& Column, FTableTreeNode& GroupNode);
+	void UpdateDoubleSumAggregationRec(FTableColumn& Column, FTableTreeNode& GroupNode);
+
+	void RebuildGroupingCrumbs();
+	void OnGroupingCrumbClicked(const TSharedPtr<FTreeNodeGrouping>& InEntry);
+	void BuildGroupingSubMenu_Change(FMenuBuilder& MenuBuilder, const TSharedPtr<FTreeNodeGrouping> CrumbGrouping);
+	void BuildGroupingSubMenu_Add(FMenuBuilder& MenuBuilder, const TSharedPtr<FTreeNodeGrouping> CrumbGrouping);
+	TSharedRef<SWidget> GetGroupingCrumbMenuContent(const TSharedPtr<FTreeNodeGrouping>& CrumbGrouping) const;
+
+	void PreChangeGroupings();
+	void PostChangeGroupings();
+	int32 GetGroupingDepth(const TSharedPtr<FTreeNodeGrouping>& Grouping) const;
+
+	void GroupingCrumbMenu_Reset_Execute();
+	void GroupingCrumbMenu_Remove_Execute(const TSharedPtr<FTreeNodeGrouping> Grouping);
+	void GroupingCrumbMenu_MoveLeft_Execute(const TSharedPtr<FTreeNodeGrouping> Grouping);
+	void GroupingCrumbMenu_MoveRight_Execute(const TSharedPtr<FTreeNodeGrouping> Grouping);
+	void GroupingCrumbMenu_Change_Execute(const TSharedPtr<FTreeNodeGrouping> OldGrouping, const TSharedPtr<FTreeNodeGrouping> NewGrouping);
+	bool GroupingCrumbMenu_Change_CanExecute(const TSharedPtr<FTreeNodeGrouping> OldGrouping, const TSharedPtr<FTreeNodeGrouping> NewGrouping) const;
+	void GroupingCrumbMenu_Add_Execute(const TSharedPtr<FTreeNodeGrouping> Grouping, const TSharedPtr<FTreeNodeGrouping> AfterGrouping);
+	bool GroupingCrumbMenu_Add_CanExecute(const TSharedPtr<FTreeNodeGrouping> Grouping, const TSharedPtr<FTreeNodeGrouping> AfterGrouping) const;
 
 	////////////////////////////////////////////////////////////////////////////////////////////////////
 	// Sorting
@@ -171,6 +191,7 @@ protected:
 
 	void UpdateCurrentSortingByColumn();
 	void SortTreeNodes();
+	void SortTreeNodesRec(FTableTreeNode& GroupNode, const Insights::ITreeNodeSorting* Sorter);
 
 	EColumnSortMode::Type GetSortModeForColumn(const FName ColumnId) const;
 	void SetSortModeForColumn(const FName& ColumnId, EColumnSortMode::Type SortMode);
@@ -257,11 +278,13 @@ protected:
 	//////////////////////////////////////////////////
 	// Tree Nodes
 
-	/** All nodes. Each node corresponds to a table row. */
-	TArray<FTableTreeNodePtr> TableTreeNodes;
+	static const FName RootNodeName;
+
+	/** The root node of the tree. */
+	FTableTreeNodePtr Root;
 
-	/** The array of group nodes. */
-	TArray<FTableTreeNodePtr> GroupNodes;
+	/** Table (row) nodes. Each node corresponds to a table row. */
+	TArray<FTableTreeNodePtr> TableTreeNodes;
 
 	/** A filtered array of group and nodes to be displayed in the tree widget. */
 	TArray<FTableTreeNodePtr> FilteredGroupNodes;
@@ -292,10 +315,10 @@ protected:
 
 	TArray<TSharedPtr<FTreeNodeGrouping>> AvailableGroupings;
 
-	TSharedPtr<SComboBox<TSharedPtr<FTreeNodeGrouping>>> GroupByComboBox;
-
 	/** How we group the tree nodes? */
-	TSharedPtr<FTreeNodeGrouping> CurrentGrouping;
+	TArray<TSharedPtr<FTreeNodeGrouping>> CurrentGroupings;
+
+	TSharedPtr<SBreadcrumbTrail<TSharedPtr<FTreeNodeGrouping>>> GroupingBreadcrumbTrail;
 
 	//////////////////////////////////////////////////
 	// Sorting
diff --git a/Engine/Source/Developer/TraceInsights/Private/Insights/Widgets/SStartPageWindow.cpp b/Engine/Source/Developer/TraceInsights/Private/Insights/Widgets/SStartPageWindow.cpp
--- a/Engine/Source/Developer/TraceInsights/Private/Insights/Widgets/SStartPageWindow.cpp
+++ b/Engine/Source/Developer/TraceInsights/Private/Insights/Widgets/SStartPageWindow.cpp
@@ -398,7 +398,7 @@ TSharedRef<SWidget> SStartPageWindow::ConstructSessionsPanel()
 	+ SVerticalBox::Slot()
 		.AutoHeight()
 		.HAlign(HAlign_Center)
-		.Padding(0.0, 2.0f)
+		.Padding(0.0f, 2.0f)
 		[
 			SNew(STextBlock)
 			.Text(LOCTEXT("SessionsPanelTitle", "Trace Sessions"))
@@ -409,7 +409,7 @@ TSharedRef<SWidget> SStartPageWindow::ConstructSessionsPanel()
 	+ SVerticalBox::Slot()
 		.AutoHeight()
 		.HAlign(HAlign_Fill)
-		.Padding(0.0, 1.0f, 0.0f, 2.0f)
+		.Padding(0.0f, 1.0f, 0.0f, 2.0f)
 		.MaxHeight(134.0f)
 		[
 			SAssignNew(TraceSessionsListView, SListView<TSharedPtr<FTraceSession>>)
@@ -447,7 +447,7 @@ TSharedRef<SWidget> SStartPageWindow::ConstructSessionsPanel()
 	+ SVerticalBox::Slot()
 		.AutoHeight()
 		.HAlign(HAlign_Right)
-		.Padding(0.0, 2.0f)
+		.Padding(0.0f, 2.0f)
 		[
 			ConstructLoadPanel()
 		]
@@ -455,7 +455,7 @@ TSharedRef<SWidget> SStartPageWindow::ConstructSessionsPanel()
 	+ SVerticalBox::Slot()
 		.AutoHeight()
 		.HAlign(HAlign_Left)
-		.Padding(0.0, 2.0f)
+		.Padding(0.0f, 2.0f)
 		[
 			ConstructLocalSessionDirectoryPanel()
 		]
@@ -521,7 +521,7 @@ TSharedRef<SWidget> SStartPageWindow::ConstructLocalSessionDirectoryPanel()
 	+ SVerticalBox::Slot()
 		.AutoHeight()
 		.HAlign(HAlign_Left)
-		.Padding(0.0, 0.0f)
+		.Padding(0.0f, 0.0f)
 		[
 			SNew(STextBlock)
 			.Text(LOCTEXT("LocalSessionDirectoryText", "Local Session Directory:"))
@@ -531,7 +531,7 @@ TSharedRef<SWidget> SStartPageWindow::ConstructLocalSessionDirectoryPanel()
 	+ SVerticalBox::Slot()
 		.AutoHeight()
 		.HAlign(HAlign_Left)
-		.Padding(0.0, 0.0f)
+		.Padding(0.0f, 0.0f)
 		[
 			SNew(SHorizontalBox)
 
@@ -569,7 +569,7 @@ TSharedRef<SWidget> SStartPageWindow::ConstructRecorderPanel()
 	+ SVerticalBox::Slot()
 		.AutoHeight()
 		.HAlign(HAlign_Center)
-		.Padding(0.0, 2.0f)
+		.Padding(0.0f, 2.0f)
 		[
 			SNew(STextBlock)
 			.Text(LOCTEXT("RecorderPanelTitle", "Trace Recorder"))
@@ -580,7 +580,7 @@ TSharedRef<SWidget> SStartPageWindow::ConstructRecorderPanel()
 	+ SVerticalBox::Slot()
 		.AutoHeight()
 		.HAlign(HAlign_Left)
-		.Padding(0.0, 2.0f)
+		.Padding(0.0f, 2.0f)
 		[
 			SNew(SHorizontalBox)
 
@@ -631,7 +631,7 @@ TSharedRef<SWidget> SStartPageWindow::ConstructRecorderPanel()
 	+ SVerticalBox::Slot()
 		.AutoHeight()
 		.HAlign(HAlign_Left)
-		.Padding(0.0, 2.0f, 0.0f, 1.0f)
+		.Padding(0.0f, 2.0f, 0.0f, 1.0f)
 		[
 			SNew(SHorizontalBox)
 			.Visibility(this, &SStartPageWindow::StopTraceRecorder_Visibility)
@@ -666,7 +666,7 @@ TSharedRef<SWidget> SStartPageWindow::ConstructConnectPanel()
 	+ SVerticalBox::Slot()
 		.AutoHeight()
 		.HAlign(HAlign_Center)
-		.Padding(0.0, 2.0f)
+		.Padding(0.0f, 2.0f)
 		[
 			SNew(STextBlock)
 			.Text(LOCTEXT("ConnectPanelTitle", "New Connection"))
@@ -677,7 +677,7 @@ TSharedRef<SWidget> SStartPageWindow::ConstructConnectPanel()
 	+ SVerticalBox::Slot()
 		.AutoHeight()
 		.HAlign(HAlign_Fill)
-		.Padding(0.0, 2.0f)
+		.Padding(0.0f, 2.0f)
 		[
 			SNew(SHorizontalBox)
 
