[Insights] Timing Insights: Added Callers and Callees tree views.

diff --git a/Engine/Source/Developer/TraceInsights/Private/Insights/TimingProfilerCommands.cpp b/Engine/Source/Developer/TraceInsights/Private/Insights/TimingProfilerCommands.cpp
--- a/Engine/Source/Developer/TraceInsights/Private/Insights/TimingProfilerCommands.cpp
+++ b/Engine/Source/Developer/TraceInsights/Private/Insights/TimingProfilerCommands.cpp
@@ -57,6 +57,8 @@ void FTimingProfilerCommands::RegisterCommands()
 	UI_COMMAND(ToggleGraphTrackVisibility, "Graph", "Toggles the visibility of the Overview Graph track", EUserInterfaceActionType::ToggleButton, FInputChord(EModifierKey::Control, EKeys::G));
 	UI_COMMAND(ToggleTimingViewVisibility, "Timing", "Toggles the visibility of the main Timing view", EUserInterfaceActionType::ToggleButton, FInputChord(EModifierKey::Control, EKeys::T));
 	UI_COMMAND(ToggleTimersViewVisibility, "Timers", "Toggles the visibility of the Timers view", EUserInterfaceActionType::ToggleButton, FInputChord());
+	UI_COMMAND(ToggleCallersTreeViewVisibility, "Callers", "Toggles the visibility of the Callers tree view", EUserInterfaceActionType::ToggleButton, FInputChord());
+	UI_COMMAND(ToggleCalleesTreeViewVisibility, "Callees", "Toggles the visibility of the Callees tree view", EUserInterfaceActionType::ToggleButton, FInputChord());
 	UI_COMMAND(ToggleStatsCountersViewVisibility, "Stats", "Toggles the visibility of the Stats Counters view", EUserInterfaceActionType::ToggleButton, FInputChord());
 	UI_COMMAND(ToggleLogViewVisibility, "Log", "Toggles the visibility of the Log view", EUserInterfaceActionType::ToggleButton, FInputChord(EKeys::L));
 }
@@ -128,6 +130,8 @@ IMPLEMENT_TOGGLE_COMMAND(ToggleFramesTrackVisibility, IsFramesTrackVisible, Show
 IMPLEMENT_TOGGLE_COMMAND(ToggleGraphTrackVisibility, IsGraphTrackVisible, ShowHideGraphTrack)
 IMPLEMENT_TOGGLE_COMMAND(ToggleTimingViewVisibility, IsTimingViewVisible, ShowHideTimingView)
 IMPLEMENT_TOGGLE_COMMAND(ToggleTimersViewVisibility, IsTimersViewVisible, ShowHideTimersView)
+IMPLEMENT_TOGGLE_COMMAND(ToggleCallersTreeViewVisibility, IsCallersTreeViewVisible, ShowHideCallersTreeView)
+IMPLEMENT_TOGGLE_COMMAND(ToggleCalleesTreeViewVisibility, IsCalleesTreeViewVisible, ShowHideCalleesTreeView)
 IMPLEMENT_TOGGLE_COMMAND(ToggleStatsCountersViewVisibility, IsStatsCountersViewVisible, ShowHideStatsCountersView)
 IMPLEMENT_TOGGLE_COMMAND(ToggleLogViewVisibility, IsLogViewVisible, ShowHideLogView)
 
diff --git a/Engine/Source/Developer/TraceInsights/Private/Insights/TimingProfilerCommands.h b/Engine/Source/Developer/TraceInsights/Private/Insights/TimingProfilerCommands.h
--- a/Engine/Source/Developer/TraceInsights/Private/Insights/TimingProfilerCommands.h
+++ b/Engine/Source/Developer/TraceInsights/Private/Insights/TimingProfilerCommands.h
@@ -41,6 +41,12 @@ public:
 	/** Toggles visibility for the Timers View. Global and custom command. */
 	TSharedPtr<FUICommandInfo> ToggleTimersViewVisibility;
 
+	/** Toggles visibility for the Callers Tree View. Global and custom command. */
+	TSharedPtr<FUICommandInfo> ToggleCallersTreeViewVisibility;
+
+	/** Toggles visibility for the Callees Tree View. Global and custom command. */
+	TSharedPtr<FUICommandInfo> ToggleCalleesTreeViewVisibility;
+
 	/** Toggles visibility for the Stats Counters View. Global and custom command. */
 	TSharedPtr<FUICommandInfo> ToggleStatsCountersViewVisibility;
 
@@ -120,6 +126,8 @@ protected:\
 	DECLARE_TOGGLE_COMMAND(ToggleGraphTrackVisibility)
 	DECLARE_TOGGLE_COMMAND(ToggleTimingViewVisibility)
 	DECLARE_TOGGLE_COMMAND(ToggleTimersViewVisibility)
+	DECLARE_TOGGLE_COMMAND(ToggleCallersTreeViewVisibility)
+	DECLARE_TOGGLE_COMMAND(ToggleCalleesTreeViewVisibility)
 	DECLARE_TOGGLE_COMMAND(ToggleStatsCountersViewVisibility)
 	DECLARE_TOGGLE_COMMAND(ToggleLogViewVisibility)
 #undef DECLARE_TOGGLE_COMMAND
diff --git a/Engine/Source/Developer/TraceInsights/Private/Insights/TimingProfilerManager.cpp b/Engine/Source/Developer/TraceInsights/Private/Insights/TimingProfilerManager.cpp
--- a/Engine/Source/Developer/TraceInsights/Private/Insights/TimingProfilerManager.cpp
+++ b/Engine/Source/Developer/TraceInsights/Private/Insights/TimingProfilerManager.cpp
@@ -14,6 +14,7 @@
 #include "Insights/Widgets/SLogView.h"
 #include "Insights/Widgets/SStatsView.h"
 #include "Insights/Widgets/STimersView.h"
+#include "Insights/Widgets/STimerTreeView.h"
 #include "Insights/Widgets/STimingView.h"
 #include "Insights/Widgets/STimingProfilerWindow.h"
 
@@ -40,8 +41,13 @@ FTimingProfilerManager::FTimingProfilerManager(TSharedRef<FUICommandList> InComm
 	, bIsGraphTrackVisible(false)
 	, bIsTimingViewVisible(true)
 	, bIsTimersViewVisible(true)
+	, bIsCallersTreeViewVisible(true)
+	, bIsCalleesTreeViewVisible(true)
 	, bIsStatsCountersViewVisible(true)
 	, bIsLogViewVisible(true)
+	, SelectionStartTime(0.0)
+	, SelectionEndTime(0.0)
+	, SelectedTimerId(InvalidTimerId)
 {
 }
 
@@ -65,6 +71,8 @@ void FTimingProfilerManager::BindCommands()
 	ActionManager.Map_ToggleGraphTrackVisibility_Global();
 	ActionManager.Map_ToggleTimingViewVisibility_Global();
 	ActionManager.Map_ToggleTimersViewVisibility_Global();
+	ActionManager.Map_ToggleCallersTreeViewVisibility_Global();
+	ActionManager.Map_ToggleCalleesTreeViewVisibility_Global();
 	ActionManager.Map_ToggleStatsCountersViewVisibility_Global();
 	ActionManager.Map_ToggleLogViewVisibility_Global();
 }
@@ -118,85 +126,120 @@ bool FTimingProfilerManager::Tick(float DeltaTime)
 
 ////////////////////////////////////////////////////////////////////////////////////////////////////
 
-void FTimingProfilerManager::OnSessionChanged()
+void FTimingProfilerManager::ShowHideFramesTrack(const bool bIsVisible)
 {
+	bIsFramesTrackVisible = bIsVisible;
+
 	TSharedPtr<STimingProfilerWindow> Wnd = GetProfilerWindow();
 	if (Wnd.IsValid())
 	{
-		Wnd->Reset();
+		Wnd->ShowHideTab(FTimingProfilerTabs::FramesTrackID, bIsFramesTrackVisible);
 	}
 }
 
 ////////////////////////////////////////////////////////////////////////////////////////////////////
 
-void FTimingProfilerManager::ShowHideFramesTrack(const bool bFramesTrackVisibleState)
+void FTimingProfilerManager::ShowHideGraphTrack(const bool bIsVisible)
 {
-	bIsFramesTrackVisible = bFramesTrackVisibleState;
+	bIsGraphTrackVisible = bIsVisible;
 
 	TSharedPtr<STimingProfilerWindow> Wnd = GetProfilerWindow();
 	if (Wnd.IsValid())
 	{
-		Wnd->ShowHideTab(FTimingProfilerTabs::FramesTrackID, bIsFramesTrackVisible);
+		Wnd->ShowHideTab(FTimingProfilerTabs::GraphTrackID, bIsGraphTrackVisible);
 	}
 }
 
 ////////////////////////////////////////////////////////////////////////////////////////////////////
 
-void FTimingProfilerManager::ShowHideGraphTrack(const bool bGraphTrackVisibleState)
+void FTimingProfilerManager::ShowHideTimingView(const bool bIsVisible)
 {
-	bIsGraphTrackVisible = bGraphTrackVisibleState;
+	bIsTimingViewVisible = bIsVisible;
 
 	TSharedPtr<STimingProfilerWindow> Wnd = GetProfilerWindow();
 	if (Wnd.IsValid())
 	{
-		Wnd->ShowHideTab(FTimingProfilerTabs::GraphTrackID, bIsGraphTrackVisible);
+		Wnd->ShowHideTab(FTimingProfilerTabs::TimingViewID, bIsTimingViewVisible);
 	}
 }
 
 ////////////////////////////////////////////////////////////////////////////////////////////////////
 
-void FTimingProfilerManager::ShowHideTimingView(const bool bTimingViewVisibleState)
+void FTimingProfilerManager::ShowHideTimersView(const bool bIsVisible)
 {
-	bIsTimingViewVisible = bTimingViewVisibleState;
+	bIsTimersViewVisible = bIsVisible;
 
 	TSharedPtr<STimingProfilerWindow> Wnd = GetProfilerWindow();
 	if (Wnd.IsValid())
 	{
-		Wnd->ShowHideTab(FTimingProfilerTabs::TimingViewID, bIsTimingViewVisible);
+		Wnd->ShowHideTab(FTimingProfilerTabs::TimersID, bIsTimersViewVisible);
+
+		if (bIsTimersViewVisible)
+		{
+			UpdateAggregatedTimerStats();
+		}
 	}
 }
 
 ////////////////////////////////////////////////////////////////////////////////////////////////////
 
-void FTimingProfilerManager::ShowHideTimersView(const bool bTimersViewVisibleState)
+void FTimingProfilerManager::ShowHideCallersTreeView(const bool bIsVisible)
 {
-	bIsTimersViewVisible = bTimersViewVisibleState;
+	bIsCallersTreeViewVisible = bIsVisible;
 
 	TSharedPtr<STimingProfilerWindow> Wnd = GetProfilerWindow();
 	if (Wnd.IsValid())
 	{
-		Wnd->ShowHideTab(FTimingProfilerTabs::TimersID, bIsTimersViewVisible);
+		Wnd->ShowHideTab(FTimingProfilerTabs::CallersID, bIsCallersTreeViewVisible);
+
+		if (bIsCallersTreeViewVisible)
+		{
+			UpdateCallersAndCallees();
+		}
 	}
 }
 
 ////////////////////////////////////////////////////////////////////////////////////////////////////
 
-void FTimingProfilerManager::ShowHideStatsCountersView(const bool bStatsCountersViewVisibleState)
+void FTimingProfilerManager::ShowHideCalleesTreeView(const bool bIsVisible)
 {
-	bIsStatsCountersViewVisible = bStatsCountersViewVisibleState;
+	bIsCalleesTreeViewVisible = bIsVisible;
+
+	TSharedPtr<STimingProfilerWindow> Wnd = GetProfilerWindow();
+	if (Wnd.IsValid())
+	{
+		Wnd->ShowHideTab(FTimingProfilerTabs::CalleesID, bIsCalleesTreeViewVisible);
+
+		if (bIsCalleesTreeViewVisible)
+		{
+			UpdateCallersAndCallees();
+		}
+	}
+}
+
+////////////////////////////////////////////////////////////////////////////////////////////////////
+
+void FTimingProfilerManager::ShowHideStatsCountersView(const bool bIsVisible)
+{
+	bIsStatsCountersViewVisible = bIsVisible;
 
 	TSharedPtr<STimingProfilerWindow> Wnd = GetProfilerWindow();
 	if (Wnd.IsValid())
 	{
 		Wnd->ShowHideTab(FTimingProfilerTabs::StatsCountersID, bIsStatsCountersViewVisible);
+
+		if (bIsStatsCountersViewVisible)
+		{
+			UpdateAggregatedCounterStats();
+		}
 	}
 }
 
 ////////////////////////////////////////////////////////////////////////////////////////////////////
 
-void FTimingProfilerManager::ShowHideLogView(const bool bLogViewVisibleState)
+void FTimingProfilerManager::ShowHideLogView(const bool bIsVisible)
 {
-	bIsLogViewVisible = bLogViewVisibleState;
+	bIsLogViewVisible = bIsVisible;
 
 	TSharedPtr<STimingProfilerWindow> Wnd = GetProfilerWindow();
 	if (Wnd.IsValid())
@@ -207,16 +250,31 @@ void FTimingProfilerManager::ShowHideLogView(const bool bLogViewVisibleState)
 
 ////////////////////////////////////////////////////////////////////////////////////////////////////
 
+void FTimingProfilerManager::OnSessionChanged()
+{
+	TSharedPtr<STimingProfilerWindow> Wnd = GetProfilerWindow();
+	if (Wnd.IsValid())
+	{
+		Wnd->Reset();
+	}
+
+	SelectionStartTime = 0.0;
+	SelectionEndTime = 0.0;
+	SelectedTimerId = InvalidTimerId;
+}
+
+////////////////////////////////////////////////////////////////////////////////////////////////////
+
 const FTimerNodePtr FTimingProfilerManager::GetTimerNode(uint64 TypeId) const
 {
-	const FTimerNodePtr* TimerNodePtrPtr = nullptr;
 	TSharedPtr<STimingProfilerWindow> Wnd = GetProfilerWindow();
 	if (Wnd.IsValid())
 	{
 		TSharedPtr<STimersView> TimersView = Wnd->GetTimersView();
 		if (TimersView.IsValid())
 		{
-			TimerNodePtrPtr = TimersView->GetTimerNode(TypeId);
+			const FTimerNodePtr* TimerNodePtrPtr = TimersView->GetTimerNode(TypeId);
+
 			if (TimerNodePtrPtr == nullptr)
 			{
 				// List of timers in TimersView not up to date?
@@ -224,9 +282,181 @@ const FTimerNodePtr FTimingProfilerManager::GetTimerNode(uint64 TypeId) const
 				TimersView->RebuildTree(false);
 				TimerNodePtrPtr = TimersView->GetTimerNode(TypeId);
 			}
+
+			if (TimerNodePtrPtr != nullptr)
+			{
+				return *TimerNodePtrPtr;
+			}
+		}
+	}
+	return nullptr;
+}
+
+////////////////////////////////////////////////////////////////////////////////////////////////////
+
+void FTimingProfilerManager::SetSelectedTimeRange(const double InStartTime, const double InEndTime)
+{
+	if (InStartTime != SelectionStartTime ||
+		InEndTime != SelectionEndTime)
+	{
+		SelectionStartTime = InStartTime;
+		SelectionEndTime = InEndTime;
+
+		UpdateCallersAndCallees();
+		UpdateAggregatedTimerStats();
+		UpdateAggregatedCounterStats();
+	}
+}
+
+////////////////////////////////////////////////////////////////////////////////////////////////////
+
+void FTimingProfilerManager::SetSelectedTimer(const uint64 InTimerId)
+{
+	if (InTimerId != SelectedTimerId)
+	{
+		SelectedTimerId = InTimerId;
+
+		if (SelectedTimerId != InvalidTimerId)
+		{
+			UpdateCallersAndCallees();
+
+			TSharedPtr<STimingProfilerWindow> Wnd = GetProfilerWindow();
+			if (Wnd)
+			{
+				TSharedPtr<STimersView> TimersView = Wnd->GetTimersView();
+				if (TimersView)
+				{
+					TimersView->SelectTimerNode(InTimerId);
+				}
+			}
+		}
+	}
+}
+
+////////////////////////////////////////////////////////////////////////////////////////////////////
+
+void FTimingProfilerManager::OnThreadFilterChanged()
+{
+	UpdateCallersAndCallees();
+	UpdateAggregatedTimerStats();
+	UpdateAggregatedCounterStats();
+}
+
+////////////////////////////////////////////////////////////////////////////////////////////////////
+
+void FTimingProfilerManager::ResetCallersAndCallees()
+{
+	TSharedPtr<STimingProfilerWindow> Wnd = GetProfilerWindow();
+	if (Wnd)
+	{
+		TSharedPtr<STimerTreeView> CallersTreeView = Wnd->GetCallersTreeView();
+		TSharedPtr<STimerTreeView> CalleesTreeView = Wnd->GetCalleesTreeView();
+
+		if (CallersTreeView)
+		{
+			CallersTreeView->Reset();
+		}
+
+		if (CalleesTreeView)
+		{
+			CalleesTreeView->Reset();
+		}
+	}
+}
+
+////////////////////////////////////////////////////////////////////////////////////////////////////
+
+void FTimingProfilerManager::UpdateCallersAndCallees()
+{
+	if (SelectionStartTime < SelectionEndTime && SelectedTimerId != InvalidTimerId)
+	{
+		TSharedPtr<STimingProfilerWindow> Wnd = GetProfilerWindow();
+		if (Wnd)
+		{
+			TSharedPtr<STimerTreeView> CallersTreeView = Wnd->GetCallersTreeView();
+			TSharedPtr<STimerTreeView> CalleesTreeView = Wnd->GetCalleesTreeView();
+
+			if (CallersTreeView)
+			{
+				CallersTreeView->Reset();
+			}
+
+			if (CalleesTreeView)
+			{
+				CalleesTreeView->Reset();
+			}
+
+			TSharedPtr<const Trace::IAnalysisSession> Session = FInsightsManager::Get()->GetSession();
+			if (Session.IsValid() && Trace::ReadTimingProfilerProvider(*Session.Get()))
+			{
+				Trace::FAnalysisSessionReadScope SessionReadScope(*Session.Get());
+				const Trace::ITimingProfilerProvider& TimingProfilerProvider = *Trace::ReadTimingProfilerProvider(*Session.Get());
+
+				TSharedPtr<STimingView> TimingView = Wnd->GetTimingView();
+
+				auto ThreadFilter = [&TimingView](uint32 ThreadId)
+				{
+					return !TimingView.IsValid() || TimingView->IsCpuTrackVisible(ThreadId);
+				};
+
+				const bool bIsGpuTrackVisible = TimingView.IsValid() && TimingView->IsGpuTrackVisible();
+
+				Trace::ITimingProfilerButterfly* TimingProfilerButterfly = TimingProfilerProvider.CreateButterfly(SelectionStartTime, SelectionEndTime, ThreadFilter, bIsGpuTrackVisible);
+
+				uint32 TimerId = static_cast<uint32>(SelectedTimerId);
+
+				if (CallersTreeView.IsValid())
+				{
+					const Trace::FTimingProfilerButterflyNode& Callers = TimingProfilerButterfly->GenerateCallersTree(TimerId);
+					CallersTreeView->SetTree(Callers);
+				}
+
+				if (CallersTreeView.IsValid())
+				{
+					const Trace::FTimingProfilerButterflyNode& Callees = TimingProfilerButterfly->GenerateCalleesTree(TimerId);
+					CalleesTreeView->SetTree(Callees);
+				}
+
+				delete TimingProfilerButterfly;
+			}
+		}
+	}
+}
+
+////////////////////////////////////////////////////////////////////////////////////////////////////
+
+void FTimingProfilerManager::UpdateAggregatedTimerStats()
+{
+	if (SelectionStartTime < SelectionEndTime)
+	{
+		TSharedPtr<STimingProfilerWindow> Wnd = GetProfilerWindow();
+		if (Wnd)
+		{
+			TSharedPtr<STimersView> TimersView = Wnd->GetTimersView();
+			if (TimersView)
+			{
+				TimersView->UpdateStats(SelectionStartTime, SelectionEndTime);
+			}
+		}
+	}
+}
+
+////////////////////////////////////////////////////////////////////////////////////////////////////
+
+void FTimingProfilerManager::UpdateAggregatedCounterStats()
+{
+	if (SelectionStartTime < SelectionEndTime)
+	{
+		TSharedPtr<STimingProfilerWindow> Wnd = GetProfilerWindow();
+		if (Wnd)
+		{
+			TSharedPtr<SStatsView> StatsView = Wnd->GetStatsView();
+			if (StatsView)
+			{
+				StatsView->UpdateStats(SelectionStartTime, SelectionEndTime);
+			}
 		}
 	}
-	return TimerNodePtrPtr ? *TimerNodePtrPtr : nullptr;
 }
 
 ////////////////////////////////////////////////////////////////////////////////////////////////////
diff --git a/Engine/Source/Developer/TraceInsights/Private/Insights/TimingProfilerManager.h b/Engine/Source/Developer/TraceInsights/Private/Insights/TimingProfilerManager.h
--- a/Engine/Source/Developer/TraceInsights/Private/Insights/TimingProfilerManager.h
+++ b/Engine/Source/Developer/TraceInsights/Private/Insights/TimingProfilerManager.h
@@ -100,33 +100,43 @@ public:
 
 	/** @return true, if Frames Track is visible */
 	const bool IsFramesTrackVisible() const { return bIsFramesTrackVisible; }
-	void SetFramesTrackVisible(const bool bFramesTrackVisibleState) { bIsFramesTrackVisible = bFramesTrackVisibleState; }
-	void ShowHideFramesTrack(const bool bFramesTrackVisibleState);
+	void SetFramesTrackVisible(const bool bIsVisible) { bIsFramesTrackVisible = bIsVisible; }
+	void ShowHideFramesTrack(const bool bIsVisible);
 
 	/** @return true, if Graph Track is visible */
 	const bool IsGraphTrackVisible() const { return bIsGraphTrackVisible; }
-	void SetGraphTrackVisible(const bool bGraphTrackVisibleState) { bIsGraphTrackVisible = bGraphTrackVisibleState; }
-	void ShowHideGraphTrack(const bool bGraphTrackVisibleState);
+	void SetGraphTrackVisible(const bool bIsVisible) { bIsGraphTrackVisible = bIsVisible; }
+	void ShowHideGraphTrack(const bool bIsVisible);
 
 	/** @return true, if Timing View is visible */
 	const bool IsTimingViewVisible() const { return bIsTimingViewVisible; }
-	void SetTimingViewVisible(const bool bTimingViewVisibleState) { bIsTimingViewVisible = bTimingViewVisibleState; }
-	void ShowHideTimingView(const bool bTimingViewVisibleState);
+	void SetTimingViewVisible(const bool bIsVisible) { bIsTimingViewVisible = bIsVisible; }
+	void ShowHideTimingView(const bool bIsVisible);
 
 	/** @return true, if Timers View is visible */
 	const bool IsTimersViewVisible() const { return bIsTimersViewVisible; }
-	void SetTimersViewVisible(const bool bTimersViewVisibleState) { bIsTimersViewVisible = bTimersViewVisibleState; }
-	void ShowHideTimersView(const bool bTimersViewVisibleState);
+	void SetTimersViewVisible(const bool bIsVisible) { bIsTimersViewVisible = bIsVisible; }
+	void ShowHideTimersView(const bool bIsVisible);
+
+	/** @return true, if Callers Tree View is visible */
+	const bool IsCallersTreeViewVisible() const { return bIsCallersTreeViewVisible; }
+	void SetCallersTreeViewVisible(const bool bIsVisible) { bIsCallersTreeViewVisible = bIsVisible; }
+	void ShowHideCallersTreeView(const bool bIsVisible);
+
+	/** @return true, if Callees Tree View is visible */
+	const bool IsCalleesTreeViewVisible() const { return bIsCalleesTreeViewVisible; }
+	void SetCalleesTreeViewVisible(const bool bIsVisible) { bIsCalleesTreeViewVisible = bIsVisible; }
+	void ShowHideCalleesTreeView(const bool bIsVisible);
 
 	/** @return true, if Stats Counters View is visible */
 	const bool IsStatsCountersViewVisible() const { return bIsStatsCountersViewVisible; }
-	void SetStatsCountersViewVisible(const bool bStatsCountersViewVisibleState) { bIsStatsCountersViewVisible = bStatsCountersViewVisibleState; }
-	void ShowHideStatsCountersView(const bool bStatsCountersViewVisibleState);
+	void SetStatsCountersViewVisible(const bool bIsVisible) { bIsStatsCountersViewVisible = bIsVisible; }
+	void ShowHideStatsCountersView(const bool bIsVisible);
 
 	/** @return true, if Log View is visible */
 	const bool IsLogViewVisible() const { return bIsLogViewVisible; }
-	void SetLogViewVisible(const bool bLogViewVisibleState) { bIsLogViewVisible = bLogViewVisibleState; }
-	void ShowHideLogView(const bool bLogViewVisibleState);
+	void SetLogViewVisible(const bool bIsVisible) { bIsLogViewVisible = bIsVisible; }
+	void ShowHideLogView(const bool bIsVisible);
 
 	////////////////////////////////////////////////////////////////////////////////////////////////////
 
@@ -134,6 +144,17 @@ public:
 
 	const FTimerNodePtr GetTimerNode(uint64 TypeId) const;
 
+	void SetSelectedTimeRange(const double InStartTime, const double InEndTime);
+	void SetSelectedTimer(const uint64 InTimerId);
+
+	void OnThreadFilterChanged();
+
+	void ResetCallersAndCallees();
+	void UpdateCallersAndCallees();
+
+	void UpdateAggregatedTimerStats();
+	void UpdateAggregatedCounterStats();
+
 protected:
 	/** Updates this manager, done through FCoreTicker. */
 	bool Tick(float DeltaTime);
@@ -166,6 +187,12 @@ protected:
 	/** If the Timers View is visible or hidden. */
 	bool bIsTimersViewVisible;
 
+	/** If the Callers Tree View is visible or hidden. */
+	bool bIsCallersTreeViewVisible;
+
+	/** If the Callees Tree View is visible or hidden. */
+	bool bIsCalleesTreeViewVisible;
+
 	/** If the Stats Counters View is visible or hidden. */
 	bool bIsStatsCountersViewVisible;
 
@@ -174,4 +201,12 @@ protected:
 
 	/** A shared pointer to the global instance of the profiler manager. */
 	static TSharedPtr<FTimingProfilerManager> Instance;
+
+	//////////////////////////////////////////////////
+
+	double SelectionStartTime;
+	double SelectionEndTime;
+
+	static const uint64 InvalidTimerId = uint64(-1);
+	uint64 SelectedTimerId;
 };
diff --git a/Engine/Source/Developer/TraceInsights/Private/Insights/ViewModels/TimingEvent.h b/Engine/Source/Developer/TraceInsights/Private/Insights/ViewModels/TimingEvent.h
--- a/Engine/Source/Developer/TraceInsights/Private/Insights/ViewModels/TimingEvent.h
+++ b/Engine/Source/Developer/TraceInsights/Private/Insights/ViewModels/TimingEvent.h
@@ -9,7 +9,7 @@ class FTimingEventsTrack;
 
 struct FTimingEvent
 {
-	static const uint64 InvalidTypeId = -1;
+	static const uint64 InvalidTypeId = uint64(-1);
 
 	const FTimingEventsTrack* Track;
 	uint64 TypeId;
diff --git a/Engine/Source/Developer/TraceInsights/Private/Insights/Widgets/STimerTreeView.cpp b/Engine/Source/Developer/TraceInsights/Private/Insights/Widgets/STimerTreeView.cpp
new file mode 100644
--- /dev/null
+++ b/Engine/Source/Developer/TraceInsights/Private/Insights/Widgets/STimerTreeView.cpp
@@ -0,0 +1,548 @@
+// Copyright 1998-2019 Epic Games, Inc. All Rights Reserved.
+
+#include "STimerTreeView.h"
+
+#include "EditorStyleSet.h"
+#include "Framework/Application/SlateApplication.h"
+#include "Framework/MultiBox/MultiBoxBuilder.h"
+#include "SlateOptMacros.h"
+#include "Widgets/Layout/SScrollBox.h"
+//#include "Widgets/Views/STableViewBase.h"
+
+// Insights
+#include "Insights/Widgets/STimersViewTooltip.h"
+#include "Insights/Widgets/STimerTableRow.h"
+
+#define LOCTEXT_NAMESPACE "STimerTreeView"
+
+////////////////////////////////////////////////////////////////////////////////////////////////////
+
+const FName STimerTreeView::TimerNameColumnId(TEXT("Name")); // FTimersViewColumns::TotalInclusiveTimeColumnID
+const FName STimerTreeView::CountColumnId(TEXT("Count")); // FTimersViewColumns::TotalInclusiveTimeColumnID
+const FName STimerTreeView::InclusiveTimeColumnId(TEXT("TotalInclTime")); // FTimersViewColumns::TotalInclusiveTimeColumnID
+const FName STimerTreeView::ExclusiveTimeColumnId(TEXT("TotalExclTime")); // FTimersViewColumns::TotalExclusiveTimeColumnID
+
+////////////////////////////////////////////////////////////////////////////////////////////////////
+
+STimerTreeView::STimerTreeView()
+{
+}
+
+////////////////////////////////////////////////////////////////////////////////////////////////////
+
+STimerTreeView::~STimerTreeView()
+{
+}
+
+////////////////////////////////////////////////////////////////////////////////////////////////////
+
+BEGIN_SLATE_FUNCTION_BUILD_OPTIMIZATION
+void STimerTreeView::Construct(const FArguments& InArgs, const FText& InViewName)
+{
+	ViewName = InViewName;
+
+	SAssignNew(ExternalScrollbar, SScrollBar)
+	.AlwaysShowScrollbar(true);
+
+	ChildSlot
+	[
+		SNew(SHorizontalBox)
+
+		+ SHorizontalBox::Slot()
+		.FillWidth(1.0f)
+		.Padding(0.0f)
+		[
+			SNew(SScrollBox)
+			.Orientation(Orient_Horizontal)
+
+			+ SScrollBox::Slot()
+			[
+				SNew(SBorder)
+				.BorderImage(FEditorStyle::GetBrush("ToolPanel.GroupBorder"))
+				.Padding(0.0f)
+				[
+					SAssignNew(TreeView, STreeView<FTimerNodePtr>)
+					.ExternalScrollbar(ExternalScrollbar)
+					.SelectionMode(ESelectionMode::Multi)
+					.TreeItemsSource(&TreeNodes)
+					.OnGetChildren(this, &STimerTreeView::TreeView_OnGetChildren)
+					.OnGenerateRow(this, &STimerTreeView::TreeView_OnGenerateRow)
+					//.OnSelectionChanged(this, &STimerTreeView::TreeView_OnSelectionChanged)
+					//.OnMouseButtonDoubleClick(this, &STimerTreeView::TreeView_OnMouseButtonDoubleClick)
+					//.OnContextMenuOpening(FOnContextMenuOpening::CreateSP(this, &STimerTreeView::TreeView_GetMenuContent))
+					.ItemHeight(12.0f)
+					.HeaderRow
+					(
+						SAssignNew(TreeViewHeaderRow, SHeaderRow)
+						.Visibility(EVisibility::Visible)
+					)
+				]
+			]
+		]
+
+		+ SHorizontalBox::Slot()
+		.AutoWidth()
+		.Padding(0.0f)
+		[
+			SNew(SBox)
+			.WidthOverride(FOptionalSize(13.0f))
+			[
+				ExternalScrollbar.ToSharedRef()
+			]
+		]
+	];
+
+	InitializeAndShowHeaderColumns();
+}
+END_SLATE_FUNCTION_BUILD_OPTIMIZATION
+
+////////////////////////////////////////////////////////////////////////////////////////////////////
+
+TSharedPtr<SWidget> STimerTreeView::TreeView_GetMenuContent()
+{
+	const TArray<FTimerNodePtr> SelectedTimerNodes = TreeView->GetSelectedItems();
+	const int32 NumSelectedTimerNodes = SelectedTimerNodes.Num();
+	FTimerNodePtr SelectedTimerNode = NumSelectedTimerNodes ? SelectedTimerNodes[0] : nullptr;
+
+	FText SelectionStr;
+	if (NumSelectedTimerNodes == 0)
+	{
+		SelectionStr = LOCTEXT("NothingSelected", "Nothing selected");
+	}
+	else if (NumSelectedTimerNodes == 1)
+	{
+		SelectionStr = FText::FromName(SelectedTimerNode->GetName());
+	}
+	else
+	{
+		SelectionStr = LOCTEXT("MultipleSelection", "Multiple selection");
+	}
+
+	const bool bShouldCloseWindowAfterMenuSelection = true;
+	FMenuBuilder MenuBuilder(bShouldCloseWindowAfterMenuSelection, NULL);
+
+	// Selection menu
+	MenuBuilder.BeginSection("Selection", LOCTEXT("ContextMenu_Header_Selection", "Selection"));
+	{
+		struct FLocal
+		{
+			static bool ReturnFalse()
+			{
+				return false;
+			}
+		};
+
+		FUIAction DummyUIAction;
+		DummyUIAction.CanExecuteAction = FCanExecuteAction::CreateStatic(&FLocal::ReturnFalse);
+		MenuBuilder.AddMenuEntry
+		(
+			SelectionStr,
+			LOCTEXT("ContextMenu_Selection", "Currently selected items"),
+			FSlateIcon(FEditorStyle::GetStyleSetName(), "@missing.icon"), DummyUIAction, NAME_None, EUserInterfaceActionType::Button
+		);
+	}
+	MenuBuilder.EndSection();
+
+	return MenuBuilder.MakeWidget();
+}
+
+////////////////////////////////////////////////////////////////////////////////////////////////////
+
+void STimerTreeView::InitializeAndShowHeaderColumns()
+{
+	TreeViewHeaderRow_CreateTimerNameColumnArgs();
+	TreeViewHeaderRow_CreateCountColumnArgs();
+	TreeViewHeaderRow_CreateInclusiveTimeColumnArgs();
+	TreeViewHeaderRow_CreateExclusiveTimeColumnArgs();
+
+	TreeViewHeaderRow->AddColumn(TreeViewHeaderColumnArgs[TimerNameColumnId]);
+	TreeViewHeaderRow->AddColumn(TreeViewHeaderColumnArgs[CountColumnId]);
+	TreeViewHeaderRow->AddColumn(TreeViewHeaderColumnArgs[InclusiveTimeColumnId]);
+	TreeViewHeaderRow->AddColumn(TreeViewHeaderColumnArgs[ExclusiveTimeColumnId]);
+}
+
+////////////////////////////////////////////////////////////////////////////////////////////////////
+
+void STimerTreeView::TreeViewHeaderRow_CreateTimerNameColumnArgs()
+{
+	SHeaderRow::FColumn::FArguments ColumnArgs;
+
+	const FName& ColumnId = TimerNameColumnId;
+
+	ColumnArgs
+		.ColumnId(ColumnId)
+		.DefaultLabel(ViewName)
+		.HAlignHeader(HAlign_Fill)
+		.VAlignHeader(VAlign_Fill)
+		.HeaderContentPadding(FMargin(2.0f))
+		.HAlignCell(HAlign_Fill)
+		.VAlignCell(VAlign_Fill)
+		.ManualWidth(246.0f)
+		.HeaderContent()
+		[
+			SNew(SBox)
+			//.ToolTip(STimerTreeView::GetColumnTooltip(ColumnId))
+			.HAlign(HAlign_Left)
+			.VAlign(VAlign_Center)
+			[
+				SNew(STextBlock)
+				.Text(ViewName)
+			]
+		]
+		//.MenuContent()
+		//[
+		//	TreeViewHeaderRow_GenerateColumnMenu(ColumnId)
+		//]
+	;
+
+	TreeViewHeaderColumnArgs.Add(ColumnId, ColumnArgs);
+}
+
+////////////////////////////////////////////////////////////////////////////////////////////////////
+
+void STimerTreeView::TreeViewHeaderRow_CreateCountColumnArgs()
+{
+	SHeaderRow::FColumn::FArguments ColumnArgs;
+
+	const FName& ColumnId = CountColumnId;
+	const FText ColumnName = LOCTEXT("CountColumnName", "Count");
+
+	ColumnArgs
+		.ColumnId(ColumnId)
+		.DefaultLabel(ColumnName)
+		.HAlignHeader(HAlign_Fill)
+		.VAlignHeader(VAlign_Fill)
+		.HeaderContentPadding(FMargin(2.0f))
+		.HAlignCell(HAlign_Fill)
+		.VAlignCell(VAlign_Fill)
+		.ManualWidth(60.0f)
+		.HeaderContent()
+		[
+			SNew(SBox)
+			//.ToolTip(STimerTreeView::GetColumnTooltip(ColumnId))
+			.HAlign(HAlign_Right)
+			.VAlign(VAlign_Center)
+			[
+				SNew(STextBlock)
+				.Text(ColumnName)
+			]
+		]
+		//.MenuContent()
+		//[
+		//	TreeViewHeaderRow_GenerateColumnMenu(ColumnId)
+		//]
+	;
+
+	TreeViewHeaderColumnArgs.Add(ColumnId, ColumnArgs);
+}
+
+////////////////////////////////////////////////////////////////////////////////////////////////////
+
+void STimerTreeView::TreeViewHeaderRow_CreateInclusiveTimeColumnArgs()
+{
+	SHeaderRow::FColumn::FArguments ColumnArgs;
+
+	const FName& ColumnId = InclusiveTimeColumnId;
+	const FText ColumnName = LOCTEXT("InclusiveTimeColumnName", "Incl");
+
+	ColumnArgs
+		.ColumnId(ColumnId)
+		.DefaultLabel(ColumnName)
+		.HAlignHeader(HAlign_Fill)
+		.VAlignHeader(VAlign_Fill)
+		.HeaderContentPadding(FMargin(2.0f))
+		.HAlignCell(HAlign_Fill)
+		.VAlignCell(VAlign_Fill)
+		.ManualWidth(60.0f)
+		.HeaderContent()
+		[
+			SNew(SBox)
+			//.ToolTip(STimerTreeView::GetColumnTooltip(ColumnId))
+			.HAlign(HAlign_Right)
+			.VAlign(VAlign_Center)
+			[
+				SNew(STextBlock)
+				.Text(ColumnName)
+			]
+		]
+		//.MenuContent()
+		//[
+		//	TreeViewHeaderRow_GenerateColumnMenu(ColumnId)
+		//]
+	;
+
+	TreeViewHeaderColumnArgs.Add(ColumnId, ColumnArgs);
+}
+
+////////////////////////////////////////////////////////////////////////////////////////////////////
+
+void STimerTreeView::TreeViewHeaderRow_CreateExclusiveTimeColumnArgs()
+{
+	SHeaderRow::FColumn::FArguments ColumnArgs;
+
+	const FName& ColumnId = ExclusiveTimeColumnId;
+	const FText ColumnName = LOCTEXT("ExclusiveTimeColumnName", "Excl");
+
+	ColumnArgs
+		.ColumnId(ColumnId)
+		.DefaultLabel(ColumnName)
+		.HAlignHeader(HAlign_Fill)
+		.VAlignHeader(VAlign_Fill)
+		.HeaderContentPadding(FMargin(2.0f))
+		.HAlignCell(HAlign_Fill)
+		.VAlignCell(VAlign_Fill)
+		.ManualWidth(60.0f)
+		.HeaderContent()
+		[
+			SNew(SBox)
+			//.ToolTip(STimerTreeView::GetColumnTooltip(ColumnId))
+			.HAlign(HAlign_Right)
+			.VAlign(VAlign_Center)
+			[
+				SNew(STextBlock)
+				.Text(ColumnName)
+			]
+		]
+		//.MenuContent()
+		//[
+		//	TreeViewHeaderRow_GenerateColumnMenu(ColumnId)
+		//]
+	;
+
+	TreeViewHeaderColumnArgs.Add(ColumnId, ColumnArgs);
+}
+
+////////////////////////////////////////////////////////////////////////////////////////////////////
+// TreeView
+////////////////////////////////////////////////////////////////////////////////////////////////////
+
+void STimerTreeView::TreeView_Refresh()
+{
+	if (TreeView.IsValid())
+	{
+		TreeView->RequestTreeRefresh();
+	}
+}
+
+////////////////////////////////////////////////////////////////////////////////////////////////////
+
+void STimerTreeView::TreeView_OnSelectionChanged(FTimerNodePtr SelectedItem, ESelectInfo::Type SelectInfo)
+{
+}
+
+////////////////////////////////////////////////////////////////////////////////////////////////////
+
+void STimerTreeView::TreeView_OnGetChildren(FTimerNodePtr InParent, TArray<FTimerNodePtr>& OutChildren)
+{
+	constexpr bool bUseFiltering = false;
+	if (bUseFiltering)
+	{
+		const TArray<Insights::FBaseTreeNodePtr>& Children = InParent->GetFilteredChildren();
+		OutChildren.Reset(Children.Num());
+		for (const Insights::FBaseTreeNodePtr& Child : Children)
+		{
+			OutChildren.Add(StaticCastSharedPtr<FTimerNode, Insights::FBaseTreeNode>(Child));
+		}
+	}
+	else
+	{
+		const TArray<Insights::FBaseTreeNodePtr>& Children = InParent->GetChildren();
+		OutChildren.Reset(Children.Num());
+		for (const Insights::FBaseTreeNodePtr& Child : Children)
+		{
+			OutChildren.Add(StaticCastSharedPtr<FTimerNode, Insights::FBaseTreeNode>(Child));
+		}
+	}
+}
+
+////////////////////////////////////////////////////////////////////////////////////////////////////
+
+void STimerTreeView::TreeView_OnMouseButtonDoubleClick(FTimerNodePtr TimerNode)
+{
+	if (TimerNode->IsGroup())
+	{
+		const bool bIsGroupExpanded = TreeView->IsItemExpanded(TimerNode);
+		TreeView->SetItemExpansion(TimerNode, !bIsGroupExpanded);
+	}
+}
+
+////////////////////////////////////////////////////////////////////////////////////////////////////
+// Tree View's Table Row
+////////////////////////////////////////////////////////////////////////////////////////////////////
+
+TSharedRef<ITableRow> STimerTreeView::TreeView_OnGenerateRow(FTimerNodePtr TimerNodePtr, const TSharedRef<STableViewBase>& OwnerTable)
+{
+	TSharedRef<ITableRow> TableRow =
+		SNew(STimerTableRow, OwnerTable)
+		.OnShouldBeEnabled(this, &STimerTreeView::TableRow_ShouldBeEnabled)
+		.OnIsColumnVisible(this, &STimerTreeView::TableRow_IsColumnVisible)
+		.OnSetHoveredTableCell(this, &STimerTreeView::TableRow_SetHoveredTableCell)
+		.OnGetColumnOutlineHAlignmentDelegate(this, &STimerTreeView::TableRow_GetColumnOutlineHAlignment)
+		.HighlightText(this, &STimerTreeView::TableRow_GetHighlightText)
+		.HighlightedNodeName(this, &STimerTreeView::TableRow_GetHighlightedNodeName)
+		.TimerNodePtr(TimerNodePtr);
+
+	return TableRow;
+}
+
+////////////////////////////////////////////////////////////////////////////////////////////////////
+
+bool STimerTreeView::TableRow_ShouldBeEnabled(const uint32 TimerId) const
+{
+	return true;
+}
+
+////////////////////////////////////////////////////////////////////////////////////////////////////
+
+bool STimerTreeView::TableRow_IsColumnVisible(const FName ColumnId) const
+{
+	return true;
+}
+
+////////////////////////////////////////////////////////////////////////////////////////////////////
+
+void STimerTreeView::TableRow_SetHoveredTableCell(const FName ColumnId, const FTimerNodePtr TimerNodePtr)
+{
+	HoveredColumnId = ColumnId;
+
+	const bool bIsAnyMenusVisible = FSlateApplication::Get().AnyMenusVisible();
+	if (!HasMouseCapture() && !bIsAnyMenusVisible)
+	{
+		HoveredTimerNodePtr = TimerNodePtr;
+	}
+}
+
+////////////////////////////////////////////////////////////////////////////////////////////////////
+
+EHorizontalAlignment STimerTreeView::TableRow_GetColumnOutlineHAlignment(const FName ColumnId) const
+{
+	// First column
+	if (ColumnId == TimerNameColumnId)
+	{
+		return HAlign_Left;
+	}
+	// Last column
+	else if (ColumnId == ExclusiveTimeColumnId)
+	{
+		return HAlign_Right;
+	}
+	// Middle columns
+	{
+		return HAlign_Center;
+	}
+}
+
+////////////////////////////////////////////////////////////////////////////////////////////////////
+
+FText STimerTreeView::TableRow_GetHighlightText() const
+{
+	return FText::GetEmpty();
+}
+
+////////////////////////////////////////////////////////////////////////////////////////////////////
+
+FName STimerTreeView::TableRow_GetHighlightedNodeName() const
+{
+	return HighlightedNodeName;
+}
+
+////////////////////////////////////////////////////////////////////////////////////////////////////
+
+//TSharedPtr<SWidget> STimerTreeView::TreeView_GetMenuContent()
+//{
+//}
+
+////////////////////////////////////////////////////////////////////////////////////////////////////
+
+void STimerTreeView::Reset()
+{
+	TreeNodes.Reset();
+	TreeView_Refresh();
+}
+
+////////////////////////////////////////////////////////////////////////////////////////////////////
+
+void STimerTreeView::SetTree(const Trace::FTimingProfilerButterflyNode& Root)
+{
+	TreeNodes.Reset();
+
+	FTimerNodePtr RootTimerNodePtr = CreateTimerNodeRec(Root);
+	if (RootTimerNodePtr)
+	{
+		TreeNodes.Add(RootTimerNodePtr);
+	}
+
+	TreeView_Refresh();
+
+	if (RootTimerNodePtr)
+	{
+		ExpandNodesRec(RootTimerNodePtr, 0);
+	}
+}
+
+////////////////////////////////////////////////////////////////////////////////////////////////////
+
+FTimerNodePtr STimerTreeView::CreateTimerNodeRec(const Trace::FTimingProfilerButterflyNode& Node)
+{
+	if (Node.Timer == nullptr)
+	{
+		return nullptr;
+	}
+
+	const uint64 Id = Node.Timer->Id;
+	const FName Name(Node.Timer->Name);
+	const FName Group(Node.Timer->IsGpuTimer ? TEXT("GPU") : TEXT("CPU"));
+	const ETimerNodeType Type = Node.Timer->IsGpuTimer ? ETimerNodeType::GpuScope : ETimerNodeType::CpuScope;
+
+	FTimerNodePtr TimerNodePtr = MakeShareable(new FTimerNode(Id, Name, Group, Type));
+
+	Trace::FAggregatedTimingStats AggregatedStats;
+	AggregatedStats.InstanceCount = Node.Count;
+	AggregatedStats.TotalInclusiveTime = Node.InclusiveTime;
+	AggregatedStats.TotalExclusiveTime = Node.ExclusiveTime;
+	TimerNodePtr->SetAggregatedStats(AggregatedStats);
+
+	for (const Trace::FTimingProfilerButterflyNode* ChildNodePtr : Node.Children)
+	{
+		if (ChildNodePtr != nullptr)
+		{
+			FTimerNodePtr ChildTimerNodePtr = CreateTimerNodeRec(*ChildNodePtr);
+			if (ChildTimerNodePtr)
+			{
+				TimerNodePtr->AddChildAndSetGroupPtr(ChildTimerNodePtr);
+			}
+		}
+	}
+
+	// Sort children by InclTime (descending).
+	TArray<Insights::FBaseTreeNodePtr>& Children = const_cast<TArray<Insights::FBaseTreeNodePtr>&>(TimerNodePtr->GetChildren());
+	Children.Sort([](const Insights::FBaseTreeNodePtr& A, const Insights::FBaseTreeNodePtr& B) -> bool
+	{
+		const double InclTimeA = StaticCastSharedPtr<FTimerNode>(A)->GetAggregatedStats().TotalInclusiveTime;
+		const double InclTimeB = StaticCastSharedPtr<FTimerNode>(B)->GetAggregatedStats().TotalInclusiveTime;
+		return InclTimeA >= InclTimeB;
+	});
+
+	return TimerNodePtr;
+}
+
+////////////////////////////////////////////////////////////////////////////////////////////////////
+
+void STimerTreeView::ExpandNodesRec(FTimerNodePtr NodePtr, int32 Depth)
+{
+	//constexpr int32 MaxDepth = 3;
+
+	TreeView->SetItemExpansion(NodePtr, true);
+
+	//if (Depth < MaxDepth)
+	{
+		for (const Insights::FBaseTreeNodePtr& ChildPtr : NodePtr->GetChildren())
+		{
+			ExpandNodesRec(StaticCastSharedPtr<FTimerNode>(ChildPtr), Depth + 1);
+		}
+	}
+}
+
+////////////////////////////////////////////////////////////////////////////////////////////////////
+
+#undef LOCTEXT_NAMESPACE
diff --git a/Engine/Source/Developer/TraceInsights/Private/Insights/Widgets/STimerTreeView.h b/Engine/Source/Developer/TraceInsights/Private/Insights/Widgets/STimerTreeView.h
new file mode 100644
--- /dev/null
+++ b/Engine/Source/Developer/TraceInsights/Private/Insights/Widgets/STimerTreeView.h
@@ -0,0 +1,151 @@
+// Copyright 1998-2019 Epic Games, Inc. All Rights Reserved.
+
+#pragma once
+
+#include "CoreMinimal.h"
+#include "SlateFwd.h"
+#include "TraceServices/Model/TimingProfiler.h"
+#include "Widgets/DeclarativeSyntaxSupport.h"
+#include "Widgets/SCompoundWidget.h"
+#include "Widgets/SWidget.h"
+#include "Widgets/Views/SHeaderRow.h"
+#include "Widgets/Views/STableRow.h"
+#include "Widgets/Views/STableViewBase.h"
+#include "Widgets/Views/STreeView.h"
+
+// Insights
+#include "Insights/ViewModels/TimerNode.h"
+
+////////////////////////////////////////////////////////////////////////////////////////////////////
+
+class FMenuBuilder;
+
+
+////////////////////////////////////////////////////////////////////////////////////////////////////
+
+/**
+ * A custom widget used to display the timers in a tree view (ex. Callers and Callees).
+ */
+class STimerTreeView : public SCompoundWidget
+{
+public:
+	/** Default constructor. */
+	STimerTreeView();
+
+	/** Virtual destructor. */
+	virtual ~STimerTreeView();
+
+	SLATE_BEGIN_ARGS(STimerTreeView) {}
+	SLATE_END_ARGS()
+
+	/**
+	 * Construct this widget
+	 * @param InArgs - The declaration data for this widget
+	 */
+	void Construct(const FArguments& InArgs, const FText& InViewName);
+
+	void Reset();
+	void SetTree(const Trace::FTimingProfilerButterflyNode& Root);
+
+private:
+	FTimerNodePtr CreateTimerNodeRec(const Trace::FTimingProfilerButterflyNode& Node);
+	void ExpandNodesRec(FTimerNodePtr NodePtr, int32 Depth);
+
+	////////////////////////////////////////////////////////////////////////////////////////////////////
+	// Tree View - Context Menu
+
+	TSharedPtr<SWidget> TreeView_GetMenuContent();
+
+	////////////////////////////////////////////////////////////////////////////////////////////////////
+	// Tree View - Columns' Header
+
+	void InitializeAndShowHeaderColumns();
+	void TreeViewHeaderRow_CreateTimerNameColumnArgs();
+	void TreeViewHeaderRow_CreateCountColumnArgs();
+	void TreeViewHeaderRow_CreateInclusiveTimeColumnArgs();
+	void TreeViewHeaderRow_CreateExclusiveTimeColumnArgs();
+
+	////////////////////////////////////////////////////////////////////////////////////////////////////
+	// Tree View - Misc
+
+	void TreeView_Refresh();
+
+	/**
+	 * Called by STreeView to retrieves the children for the specified parent item.
+	 * @param InParent    - The parent node to retrieve the children from.
+	 * @param OutChildren - List of children for the parent node.
+	 */
+	void TreeView_OnGetChildren(FTimerNodePtr InParent, TArray<FTimerNodePtr>& OutChildren);
+
+	/** Called by STreeView when selection has changed. */
+	void TreeView_OnSelectionChanged(FTimerNodePtr SelectedItem, ESelectInfo::Type SelectInfo);
+
+	/** Called by STreeView when a tree item is double clicked. */
+	void TreeView_OnMouseButtonDoubleClick(FTimerNodePtr TreeNode);
+
+	////////////////////////////////////////////////////////////////////////////////////////////////////
+	// Tree View - Table Row
+
+	/** Called by STreeView to generate a table row for the specified item. */
+	TSharedRef<ITableRow> TreeView_OnGenerateRow(FTimerNodePtr TreeNode, const TSharedRef<STableViewBase>& OwnerTable);
+
+	bool TableRow_IsColumnVisible(const FName ColumnId) const;
+	void TableRow_SetHoveredTableCell(const FName ColumnId, const FTimerNodePtr TimerNodePtr);
+	EHorizontalAlignment TableRow_GetColumnOutlineHAlignment(const FName ColumnId) const;
+
+	FText TableRow_GetHighlightText() const;
+	FName TableRow_GetHighlightedNodeName() const;
+
+	bool TableRow_ShouldBeEnabled(const uint32 TimerId) const;
+
+	////////////////////////////////////////////////////////////////////////////////////////////////////
+
+private:
+
+	/** The view name (ex.: "Callers" or "Callees"). */
+	FText ViewName;
+
+	//////////////////////////////////////////////////
+	// Tree View, Columns
+
+	/** The tree widget which holds the list of groups and timers corresponding with each group. */
+	TSharedPtr<STreeView<FTimerNodePtr>> TreeView;
+
+	/** Column arguments used to initialize a new header column in the tree view, stored as column name to column arguments mapping. */
+	TMap<FName, SHeaderRow::FColumn::FArguments> TreeViewHeaderColumnArgs;
+
+	/** Holds the tree view header row widget which display all columns in the tree view. */
+	TSharedPtr<SHeaderRow> TreeViewHeaderRow;
+
+	/** External scrollbar used to synchronize tree view position. */
+	TSharedPtr<SScrollBar> ExternalScrollbar;
+
+	//////////////////////////////////////////////////
+	// Hovered Column, Hovered Timer Node
+
+	/** Name of the column currently being hovered by the mouse. */
+	FName HoveredColumnId;
+
+	/** A shared pointer to the timer node currently being hovered by the mouse. */
+	FTimerNodePtr HoveredTimerNodePtr;
+
+	/** Name of the timer that should be drawn as highlighted. */
+	FName HighlightedNodeName;
+
+	//////////////////////////////////////////////////
+	// Timer Nodes
+
+	/** The root node(s) of the tree. */
+	TArray<FTimerNodePtr> TreeNodes;
+
+	//////////////////////////////////////////////////
+
+	double StatsStartTime;
+	double StatsEndTime;
+	uint32 StatsTimerId;
+
+	static const FName TimerNameColumnId;
+	static const FName CountColumnId;
+	static const FName InclusiveTimeColumnId;
+	static const FName ExclusiveTimeColumnId;
+};
diff --git a/Engine/Source/Developer/TraceInsights/Private/Insights/Widgets/STimersView.cpp b/Engine/Source/Developer/TraceInsights/Private/Insights/Widgets/STimersView.cpp
--- a/Engine/Source/Developer/TraceInsights/Private/Insights/Widgets/STimersView.cpp
+++ b/Engine/Source/Developer/TraceInsights/Private/Insights/Widgets/STimersView.cpp
@@ -49,6 +49,8 @@ STimersView::STimersView()
 	, GroupingMode(ETimerGroupingMode::ByType)
 	, ColumnSortMode(GetDefaultColumnSortMode())
 	, ColumnBeingSorted(GetDefaultColumnBeingSorted())
+	, StatsStartTime(0.0)
+	, StatsEndTime(0.0)
 {
 	FMemory::Memset(bTimerNodeIsVisible, 1);
 }
@@ -606,12 +608,12 @@ TSharedRef<SWidget> STimersView::TreeViewHeaderRow_GenerateColumnMenu(const FTim
 			MenuBuilder.EndSection();
 		}
 
-		if (Column.bCanBeFiltered())
-		{
-			MenuBuilder.BeginSection("FilterMode", LOCTEXT("ContextMenu_Header_Misc_Filter_FilterMode", "Filter Mode"));
-			bIsMenuVisible = true;
-			MenuBuilder.EndSection();
-		}
+		//if (Column.CanBeFiltered())
+		//{
+		//	MenuBuilder.BeginSection("FilterMode", LOCTEXT("ContextMenu_Header_Misc_Filter_FilterMode", "Filter Mode"));
+		//	bIsMenuVisible = true;
+		//	MenuBuilder.EndSection();
+		//}
 	}
 
 	/*
@@ -804,11 +806,7 @@ void STimersView::TreeView_OnSelectionChanged(FTimerNodePtr SelectedItem, ESelec
 		TArray<FTimerNodePtr> SelectedItems = TreeView->GetSelectedItems();
 		if (SelectedItems.Num() == 1)
 		{
-			//HighlightedNodeName = SelectedItems[0]->GetName();
-		}
-		else
-		{
-			//HighlightedNodeName = NAME_None;
+			FTimingProfilerManager::Get()->SetSelectedTimer(SelectedItems[0]->GetId());
 		}
 	}
 }
@@ -827,27 +825,27 @@ void STimersView::TreeView_OnGetChildren(FTimerNodePtr InParent, TArray<FTimerNo
 
 ////////////////////////////////////////////////////////////////////////////////////////////////////
 
-void STimersView::TreeView_OnMouseButtonDoubleClick(FTimerNodePtr TimerNode)
+void STimersView::TreeView_OnMouseButtonDoubleClick(FTimerNodePtr TimerNodePtr)
 {
-	if (!TimerNode->IsGroup())
+	if (TimerNodePtr->IsGroup())
 	{
-		//im:TODO: const bool bIsTracked = FTimingProfilerManager::Get()->IsTimerTracked(TimerNode->GetId());
+		const bool bIsGroupExpanded = TreeView->IsItemExpanded(TimerNodePtr);
+		TreeView->SetItemExpansion(TimerNodePtr, !bIsGroupExpanded);
+	}
+	else
+	{
+		//im:TODO: const bool bIsTracked = FTimingProfilerManager::Get()->IsTimerTracked(TimerNodePtr->GetId());
 	//	if (!bIsTracked)
 	//	{
 	//		// Add a new graph series.
-	//		FTimingProfilerManager::Get()->TrackTimer(TimerNode->GetId());
+	//		FTimingProfilerManager::Get()->TrackTimer(TimerNodePtr->GetId());
 	//	}
 	//	else
 	//	{
 	//		// Remove the corresponding graph series.
-	//		FTimingProfilerManager::Get()->UntrackTimer(TimerNode->GetId());
+	//		FTimingProfilerManager::Get()->UntrackTimer(TimerNodePtr->GetId());
 	//	}
 	}
-	else
-	{
-		const bool bIsGroupExpanded = TreeView->IsItemExpanded(TimerNode);
-		TreeView->SetItemExpansion(TimerNode, !bIsGroupExpanded);
-	}
 }
 
 ////////////////////////////////////////////////////////////////////////////////////////////////////
@@ -871,6 +869,13 @@ TSharedRef<ITableRow> STimersView::TreeView_OnGenerateRow(FTimerNodePtr TimerNod
 
 ////////////////////////////////////////////////////////////////////////////////////////////////////
 
+bool STimersView::TableRow_ShouldBeEnabled(const uint32 TimerId) const
+{
+	return true;//im:TODO: Session->GetAggregatedStat(TimerId) != nullptr;
+}
+
+////////////////////////////////////////////////////////////////////////////////////////////////////
+
 bool STimersView::TableRow_IsColumnVisible(const FName ColumnId) const
 {
 	bool bResult = false;
@@ -889,8 +894,6 @@ void STimersView::TableRow_SetHoveredTableCell(const FName ColumnId, const FTime
 	{
 		HoveredTimerNodePtr = TimerNodePtr;
 	}
-
-	//UE_LOG(TimingProfiler, Log, TEXT("%s -> %s"), *HoveredColumnId.GetPlainNameString(), TimerNodePtr.IsValid() ? *TimerNodePtr->GetName().GetPlainNameString() : TEXT("nullptr"));
 }
 
 ////////////////////////////////////////////////////////////////////////////////////////////////////
@@ -930,13 +933,6 @@ FName STimersView::TableRow_GetHighlightedNodeName() const
 	return HighlightedNodeName;
 }
 
-////////////////////////////////////////////////////////////////////////////////////////////////////
-
-bool STimersView::TableRow_ShouldBeEnabled(const uint32 TimerId) const
-{
-	return true;//im:TODO: Session->GetAggregatedStat(TimerId) != nullptr;
-}
-
 ////////////////////////////////////////////////////////////////////////////////////////////////////
 // SearchBox
 ////////////////////////////////////////////////////////////////////////////////////////////////////
@@ -1437,7 +1433,6 @@ void STimersView::RebuildTree(bool bResync)
 	{
 		const int32 PreviousNodeCount = TimerNodes.Num();
 		TimerNodes.Empty(PreviousNodeCount);
-		//TimerNodesMap.Empty(PreviousNodeCount);
 		TimerNodesIdMap.Empty(PreviousNodeCount);
 		bListHasChanged = true;
 	}
@@ -1462,7 +1457,6 @@ void STimersView::RebuildTree(bool bResync)
 
 				const int32 PreviousNodeCount = TimerNodes.Num();
 				TimerNodes.Empty(PreviousNodeCount);
-				//TimerNodesMap.Empty(PreviousNodeCount);
 				TimerNodesIdMap.Empty(PreviousNodeCount);
 				bListHasChanged = true;
 
@@ -1474,7 +1468,6 @@ void STimersView::RebuildTree(bool bResync)
 					ETimerNodeType Type = Timer.IsGpuTimer ? ETimerNodeType::GpuScope : ETimerNodeType::CpuScope;
 					FTimerNodePtr TimerNodePtr = MakeShareable(new FTimerNode(Timer.Id, Name, Group, Type));
 					TimerNodes.Add(TimerNodePtr);
-					//TimerNodesMap.Add(Name, TimerNodePtr);
 					TimerNodesIdMap.Add(Timer.Id, TimerNodePtr);
 				}
 			}
@@ -1572,12 +1565,13 @@ void STimersView::UpdateStats(double StartTime, double EndTime)
 
 void STimersView::SelectTimerNode(uint64 Id)
 {
-	FTimerNodePtr* TimerNodePtrPtr = TimerNodesIdMap.Find(Id);
-	if (TimerNodePtrPtr != nullptr)
+	FTimerNodePtr* NodePtrPtr = TimerNodesIdMap.Find(Id);
+	if (NodePtrPtr != nullptr)
 	{
-		//UE_LOG(TimingProfiler, Log, TEXT("Select and RequestScrollIntoView %s"), *(*TimerNodePtrPtr)->GetName().ToString());
-		TreeView->SetSelection(*TimerNodePtrPtr);
-		TreeView->RequestScrollIntoView(*TimerNodePtrPtr);
+		FTimerNodePtr NodePtr = *NodePtrPtr;
+
+		TreeView->SetSelection(NodePtr);
+		TreeView->RequestScrollIntoView(NodePtr);
 	}
 }
 
diff --git a/Engine/Source/Developer/TraceInsights/Private/Insights/Widgets/STimersView.h b/Engine/Source/Developer/TraceInsights/Private/Insights/Widgets/STimersView.h
--- a/Engine/Source/Developer/TraceInsights/Private/Insights/Widgets/STimersView.h
+++ b/Engine/Source/Developer/TraceInsights/Private/Insights/Widgets/STimersView.h
@@ -176,8 +176,6 @@ protected:
 	void SetSortModeForColumn(const FName& ColumnId, EColumnSortMode::Type SortMode);
 	void OnSortModeChanged(const EColumnSortPriority::Type SortPriority, const FName& ColumnId, const EColumnSortMode::Type SortMode);
 
-	//void TreeView_BuildSortByMenu(FMenuBuilder& MenuBuilder);
-
 	////////////////////////////////////////////////////////////////////////////////////////////////////
 	// Sorting actions
 
@@ -220,6 +218,8 @@ protected:
 	bool ContextMenu_ResetColumns_CanExecute() const;
 	void ContextMenu_ResetColumns_Execute();
 
+	////////////////////////////////////////////////////////////////////////////////////////////////////
+
 protected:
 	/** A weak pointer to the profiler session used to populate this widget. */
 	TSharedPtr<const Trace::IAnalysisSession>/*Weak*/ Session;
@@ -230,7 +230,7 @@ protected:
 	/** The tree widget which holds the list of groups and timers corresponding with each group. */
 	TSharedPtr<STreeView<FTimerNodePtr>> TreeView;
 
-	/** Column metadata used to initialize column arguments, stored as PropertyName -> FEventGraphColumn. */
+	/** Column metadata used to initialize column arguments, stored as PropertyName -> FTimersViewColumn. */
 	TMap<FName, FTimersViewColumn> TreeViewHeaderColumns;
 
 	/** Column arguments used to initialize a new header column in the tree view, stored as column name to column arguments mapping. */
@@ -281,6 +281,8 @@ protected:
 	//////////////////////////////////////////////////
 	// Search box and filters
 
+	//bool bUseFiltering;
+
 	/** The search box widget used to filter items displayed in the stats and groups tree. */
 	TSharedPtr<SSearchBox> SearchBox;
 
@@ -296,6 +298,8 @@ protected:
 	//////////////////////////////////////////////////
 	// Grouping
 
+	//bool bUseGrouping;
+
 	TArray<TSharedPtr<ETimerGroupingMode>> GroupByOptionsSource;
 
 	TSharedPtr<SComboBox<TSharedPtr<ETimerGroupingMode>>> GroupByComboBox;
@@ -306,6 +310,8 @@ protected:
 	//////////////////////////////////////////////////
 	// Sorting
 
+	//bool bUseSorting;
+
 	/** How we sort the timers? */
 	EColumnSortMode::Type ColumnSortMode;
 
diff --git a/Engine/Source/Developer/TraceInsights/Private/Insights/Widgets/STimingProfilerToolbar.cpp b/Engine/Source/Developer/TraceInsights/Private/Insights/Widgets/STimingProfilerToolbar.cpp
--- a/Engine/Source/Developer/TraceInsights/Private/Insights/Widgets/STimingProfilerToolbar.cpp
+++ b/Engine/Source/Developer/TraceInsights/Private/Insights/Widgets/STimingProfilerToolbar.cpp
@@ -49,6 +49,8 @@ void STimingProfilerToolbar::Construct(const FArguments& InArgs)
 				ToolbarBuilder.AddToolBarButton(FTimingProfilerCommands::Get().ToggleGraphTrackVisibility);
 				ToolbarBuilder.AddToolBarButton(FTimingProfilerCommands::Get().ToggleTimingViewVisibility);
 				ToolbarBuilder.AddToolBarButton(FTimingProfilerCommands::Get().ToggleTimersViewVisibility);
+				ToolbarBuilder.AddToolBarButton(FTimingProfilerCommands::Get().ToggleCallersTreeViewVisibility);
+				ToolbarBuilder.AddToolBarButton(FTimingProfilerCommands::Get().ToggleCalleesTreeViewVisibility);
 				ToolbarBuilder.AddToolBarButton(FTimingProfilerCommands::Get().ToggleStatsCountersViewVisibility);
 				ToolbarBuilder.AddToolBarButton(FTimingProfilerCommands::Get().ToggleLogViewVisibility);
 			}
diff --git a/Engine/Source/Developer/TraceInsights/Private/Insights/Widgets/STimingProfilerWindow.cpp b/Engine/Source/Developer/TraceInsights/Private/Insights/Widgets/STimingProfilerWindow.cpp
--- a/Engine/Source/Developer/TraceInsights/Private/Insights/Widgets/STimingProfilerWindow.cpp
+++ b/Engine/Source/Developer/TraceInsights/Private/Insights/Widgets/STimingProfilerWindow.cpp
@@ -32,6 +32,7 @@
 #include "Insights/Widgets/SLogView.h"
 #include "Insights/Widgets/SStatsView.h"
 #include "Insights/Widgets/STimersView.h"
+#include "Insights/Widgets/STimerTreeView.h"
 #include "Insights/Widgets/STimingProfilerToolbar.h"
 #include "Insights/Widgets/STimingView.h"
 
@@ -46,6 +47,8 @@ const FName FTimingProfilerTabs::FramesTrackID(TEXT("Frames"));
 const FName FTimingProfilerTabs::GraphTrackID(TEXT("Graph"));
 const FName FTimingProfilerTabs::TimingViewID(TEXT("TimingView"));
 const FName FTimingProfilerTabs::TimersID(TEXT("Timers"));
+const FName FTimingProfilerTabs::CallersID(TEXT("Callers"));
+const FName FTimingProfilerTabs::CalleesID(TEXT("Callees"));
 const FName FTimingProfilerTabs::StatsCountersID(TEXT("StasCounters"));
 const FName FTimingProfilerTabs::LogViewID(TEXT("LogView"));
 
@@ -88,6 +91,14 @@ void STimingProfilerWindow::Reset()
 	{
 		TimersView->RebuildTree(true);
 	}
+	if (CallersTreeView)
+	{
+		CallersTreeView->Reset();
+	}
+	if (CalleesTreeView)
+	{
+		CalleesTreeView->Reset();
+	}
 	if (StatsView)
 	{
 		StatsView->RebuildTree(true);
@@ -228,6 +239,58 @@ void STimingProfilerWindow::OnTimersTabClosed(TSharedRef<SDockTab> TabBeingClose
 
 ////////////////////////////////////////////////////////////////////////////////////////////////////
 
+TSharedRef<SDockTab> STimingProfilerWindow::SpawnTab_Callers(const FSpawnTabArgs& Args)
+{
+	FTimingProfilerManager::Get()->SetCallersTreeViewVisible(true);
+
+	const TSharedRef<SDockTab> DockTab = SNew(SDockTab)
+		.ShouldAutosize(false)
+		.TabRole(ETabRole::PanelTab)
+		[
+			SAssignNew(CallersTreeView, STimerTreeView, LOCTEXT("CallersTreeViewName", "Callers"))
+		];
+
+	DockTab->SetOnTabClosed(SDockTab::FOnTabClosedCallback::CreateRaw(this, &STimingProfilerWindow::OnCallersTabClosed));
+
+	return DockTab;
+}
+
+////////////////////////////////////////////////////////////////////////////////////////////////////
+
+void STimingProfilerWindow::OnCallersTabClosed(TSharedRef<SDockTab> TabBeingClosed)
+{
+	FTimingProfilerManager::Get()->SetCallersTreeViewVisible(false);
+	CallersTreeView = nullptr;
+}
+
+////////////////////////////////////////////////////////////////////////////////////////////////////
+
+TSharedRef<SDockTab> STimingProfilerWindow::SpawnTab_Callees(const FSpawnTabArgs& Args)
+{
+	FTimingProfilerManager::Get()->SetCalleesTreeViewVisible(true);
+
+	const TSharedRef<SDockTab> DockTab = SNew(SDockTab)
+		.ShouldAutosize(false)
+		.TabRole(ETabRole::PanelTab)
+		[
+			SAssignNew(CalleesTreeView, STimerTreeView, LOCTEXT("CalleesTreeViewName", "Callees"))
+		];
+
+	DockTab->SetOnTabClosed(SDockTab::FOnTabClosedCallback::CreateRaw(this, &STimingProfilerWindow::OnCalleesTabClosed));
+
+	return DockTab;
+}
+
+////////////////////////////////////////////////////////////////////////////////////////////////////
+
+void STimingProfilerWindow::OnCalleesTabClosed(TSharedRef<SDockTab> TabBeingClosed)
+{
+	FTimingProfilerManager::Get()->SetCalleesTreeViewVisible(false);
+	CalleesTreeView = nullptr;
+}
+
+////////////////////////////////////////////////////////////////////////////////////////////////////
+
 TSharedRef<SDockTab> STimingProfilerWindow::SpawnTab_StatsCounters(const FSpawnTabArgs& Args)
 {
 	FTimingProfilerManager::Get()->SetStatsCountersViewVisible(true);
@@ -311,6 +374,16 @@ void STimingProfilerWindow::Construct(const FArguments& InArgs, const TSharedRef
 		.SetIcon(FSlateIcon(FInsightsStyle::GetStyleSetName(), "TimersView.Icon.Small"))
 		.SetGroup(AppMenuGroup);
 
+	TabManager->RegisterTabSpawner(FTimingProfilerTabs::CallersID, FOnSpawnTab::CreateRaw(this, &STimingProfilerWindow::SpawnTab_Callers))
+		.SetDisplayName(LOCTEXT("CallersTabTitle", "Callers"))
+		.SetIcon(FSlateIcon(FInsightsStyle::GetStyleSetName(), "TimersView.Icon.Small")) // TODO
+		.SetGroup(AppMenuGroup);
+
+	TabManager->RegisterTabSpawner(FTimingProfilerTabs::CalleesID, FOnSpawnTab::CreateRaw(this, &STimingProfilerWindow::SpawnTab_Callees))
+		.SetDisplayName(LOCTEXT("CalleesTabTitle", "Callees"))
+		.SetIcon(FSlateIcon(FInsightsStyle::GetStyleSetName(), "TimersView.Icon.Small")) // TODO
+		.SetGroup(AppMenuGroup);
+
 	TabManager->RegisterTabSpawner(FTimingProfilerTabs::StatsCountersID, FOnSpawnTab::CreateRaw(this, &STimingProfilerWindow::SpawnTab_StatsCounters))
 		.SetDisplayName(LOCTEXT("StatsCountersTabTitle", "Stats Counters"))
 		.SetIcon(FSlateIcon(FInsightsStyle::GetStyleSetName(), "StatsCountersView.Icon.Small"))
@@ -377,11 +450,31 @@ void STimingProfilerWindow::Construct(const FArguments& InArgs, const TSharedRef
 				)
 				->Split
 				(
-					FTabManager::NewStack()
+					FTabManager::NewSplitter()
+					->SetOrientation(Orient_Vertical)
 					->SetSizeCoefficient(0.35f)
-					->AddTab(FTimingProfilerTabs::TimersID, TimingProfilerManager->IsTimersViewVisible() ? ETabState::OpenedTab : ETabState::ClosedTab)
-					->AddTab(FTimingProfilerTabs::StatsCountersID, TimingProfilerManager->IsStatsCountersViewVisible() ? ETabState::OpenedTab : ETabState::ClosedTab)
-					->SetForegroundTab(FTimingProfilerTabs::TimersID)
+					->Split
+					(
+						FTabManager::NewStack()
+						->SetSizeCoefficient(0.67f)
+						->AddTab(FTimingProfilerTabs::TimersID, TimingProfilerManager->IsTimersViewVisible() ? ETabState::OpenedTab : ETabState::ClosedTab)
+						->AddTab(FTimingProfilerTabs::StatsCountersID, TimingProfilerManager->IsStatsCountersViewVisible() ? ETabState::OpenedTab : ETabState::ClosedTab)
+						->SetForegroundTab(FTimingProfilerTabs::TimersID)
+					)
+					->Split
+					(
+						FTabManager::NewStack()
+						->SetSizeCoefficient(0.165f)
+						->SetHideTabWell(true)
+						->AddTab(FTimingProfilerTabs::CallersID, TimingProfilerManager->IsCallersTreeViewVisible() ? ETabState::OpenedTab : ETabState::ClosedTab)
+					)
+					->Split
+					(
+						FTabManager::NewStack()
+						->SetSizeCoefficient(0.165f)
+						->SetHideTabWell(true)
+						->AddTab(FTimingProfilerTabs::CalleesID, TimingProfilerManager->IsCalleesTreeViewVisible() ? ETabState::OpenedTab : ETabState::ClosedTab)
+					)
 				)
 			)
 		);
diff --git a/Engine/Source/Developer/TraceInsights/Private/Insights/Widgets/STimingProfilerWindow.h b/Engine/Source/Developer/TraceInsights/Private/Insights/Widgets/STimingProfilerWindow.h
--- a/Engine/Source/Developer/TraceInsights/Private/Insights/Widgets/STimingProfilerWindow.h
+++ b/Engine/Source/Developer/TraceInsights/Private/Insights/Widgets/STimingProfilerWindow.h
@@ -25,6 +25,7 @@ class SGraphTrack;
 class SLogView;
 class SStatsView;
 class STimersView;
+class STimerTreeView;
 class STimingView;
 class SVerticalBox;
 
@@ -38,6 +39,8 @@ struct FTimingProfilerTabs
 	static const FName GraphTrackID;
 	static const FName TimingViewID;
 	static const FName TimersID;
+	static const FName CallersID;
+	static const FName CalleesID;
 	static const FName StatsCountersID;
 	static const FName LogViewID;
 };
@@ -70,6 +73,8 @@ public:
 
 	TSharedPtr<STimingView> GetTimingView() const { return TimingView; }
 	TSharedPtr<STimersView> GetTimersView() const { return TimersView; }
+	TSharedPtr<STimerTreeView> GetCallersTreeView() const { return CallersTreeView; }
+	TSharedPtr<STimerTreeView> GetCalleesTreeView() const { return CalleesTreeView; }
 	TSharedPtr<SStatsView> GetStatsView() const { return StatsView; }
 	TSharedPtr<SLogView> GetLogView() const { return LogView; }
 
@@ -89,6 +94,12 @@ private:
 	TSharedRef<SDockTab> SpawnTab_Timers(const FSpawnTabArgs& Args);
 	void OnTimersTabClosed(TSharedRef<SDockTab> TabBeingClosed);
 
+	TSharedRef<SDockTab> SpawnTab_Callers(const FSpawnTabArgs& Args);
+	void OnCallersTabClosed(TSharedRef<SDockTab> TabBeingClosed);
+
+	TSharedRef<SDockTab> SpawnTab_Callees(const FSpawnTabArgs& Args);
+	void OnCalleesTabClosed(TSharedRef<SDockTab> TabBeingClosed);
+
 	TSharedRef<SDockTab> SpawnTab_StatsCounters(const FSpawnTabArgs& Args);
 	void OnStatsCountersTabClosed(TSharedRef<SDockTab> TabBeingClosed);
 
@@ -179,6 +190,12 @@ private:
 	/** The Timers view widget */
 	TSharedPtr<STimersView> TimersView;
 
+	/** The Callers tree view widget */
+	TSharedPtr<STimerTreeView> CallersTreeView;
+
+	/** The Callees tree view widget */
+	TSharedPtr<STimerTreeView> CalleesTreeView;
+
 	/** The Stats (Counters) view widget */
 	TSharedPtr<SStatsView> StatsView;
 
diff --git a/Engine/Source/Developer/TraceInsights/Private/Insights/Widgets/STimingView.cpp b/Engine/Source/Developer/TraceInsights/Private/Insights/Widgets/STimingView.cpp
--- a/Engine/Source/Developer/TraceInsights/Private/Insights/Widgets/STimingView.cpp
+++ b/Engine/Source/Developer/TraceInsights/Private/Insights/Widgets/STimingView.cpp
@@ -391,6 +391,7 @@ void STimingView::Tick(const FGeometry& AllottedGeometry, const double InCurrent
 		}
 	}
 
+	//TODO: Move this check into STimersView/SStatsView or in STimingProfilerWindow or in FTimingProfilerManager.
 	// We need to check if TimersView or StatsView needs to update their lists of timers / counters.
 	// But, ensure we do not check too often.
 	static uint64 NextTimestamp = 0;
@@ -1165,7 +1166,7 @@ FReply STimingView::OnMouseButtonDown(const FGeometry& MyGeometry, const FPointe
 		SelectionStartTime = Viewport.SlateUnitsToTime(MousePositionOnButtonDown.X);
 		SelectionEndTime = SelectionStartTime;
 		LastSelectionType = ESelectionType::None;
-		//TODO: SelectionChangingEvent.Broadcast(SelectionStartTime, SelectionEndTime);
+		RaiseSelectionChanging();
 	}
 
 	return Reply;
@@ -1194,13 +1195,7 @@ FReply STimingView::OnMouseButtonUp(const FGeometry& MyGeometry, const FPointerE
 			}
 			else if (bIsSelecting)
 			{
-				//TODO: SelectionChangedEvent.Broadcast(SelectionStartTime, SelectionEndTime);
-
-				if (SelectionEndTime > SelectionStartTime)
-				{
-					UpdateAggregatedStats();
-				}
-
+				RaiseSelectionChanged();
 				bIsSelecting = false;
 			}
 
@@ -1216,7 +1211,7 @@ FReply STimingView::OnMouseButtonUp(const FGeometry& MyGeometry, const FPointerE
 					// ...reset selection.
 					SelectionEndTime = SelectionStartTime = 0.0;
 					LastSelectionType = ESelectionType::None;
-					//TODO: SelectionChangedEvent.Broadcast(SelectionStartTime, SelectionEndTime);
+					RaiseSelectionChanged();
 				}
 			}
 
@@ -1240,13 +1235,7 @@ FReply STimingView::OnMouseButtonUp(const FGeometry& MyGeometry, const FPointerE
 			}
 			else if (bIsSelecting)
 			{
-				//TODO: SelectionChangedEvent.Broadcast(SelectionStartTime, SelectionEndTime);
-
-				if (SelectionEndTime > SelectionStartTime)
-				{
-					UpdateAggregatedStats();
-				}
-
+				RaiseSelectionChanged();
 				bIsSelecting = false;
 			}
 
@@ -1335,7 +1324,7 @@ FReply STimingView::OnMouseMove(const FGeometry& MyGeometry, const FPointerEvent
 					SelectionEndTime = Temp;
 				}
 				LastSelectionType = ESelectionType::TimeRange;
-				//TODO: SelectionChangingEvent.Broadcast(SelectionStartTime, SelectionEndTime);
+				RaiseSelectionChanging();
 			}
 		}
 		else
@@ -2012,81 +2001,81 @@ void STimingView::SelectTimeInterval(double IntervalStartTime, double IntervalDu
 	SelectionStartTime = IntervalStartTime;
 	SelectionEndTime = IntervalStartTime + IntervalDuration;
 	LastSelectionType = ESelectionType::TimeRange;
+	RaiseSelectionChanged();
+}
 
-	//TODO: SelectionChangedEvent.Broadcast(SelectionStartTime, SelectionEndTime);
-	UpdateAggregatedStats();
+////////////////////////////////////////////////////////////////////////////////////////////////////
+
+void STimingView::RaiseSelectionChanging()
+{
+	//TODO: SelectionChangingEvent.Broadcast(SelectionStartTime, SelectionEndTime);
 }
 
 ////////////////////////////////////////////////////////////////////////////////////////////////////
 
-void STimingView::UpdateAggregatedStats()
+void STimingView::RaiseSelectionChanged()
 {
-	TSharedPtr<STimingProfilerWindow> Wnd = FTimingProfilerManager::Get()->GetProfilerWindow();
-	if (Wnd)
+	//TODO: SelectionChangedEvent.Broadcast(SelectionStartTime, SelectionEndTime);
+
+	FTimingProfilerManager::Get()->SetSelectedTimeRange(SelectionStartTime, SelectionEndTime);
+
+	if (SelectionStartTime < SelectionEndTime)
 	{
-		TSharedPtr<STimersView> TimersView = Wnd->GetTimersView();
-		if (TimersView)
-		{
-			TimersView->UpdateStats(SelectionStartTime, SelectionEndTime);
-		}
-		TSharedPtr<SStatsView> StatsView = Wnd->GetStatsView();
-		if (StatsView)
-		{
-			StatsView->UpdateStats(SelectionStartTime, SelectionEndTime);
-		}
+		UpdateAggregatedStats();
 	}
+}
+
+////////////////////////////////////////////////////////////////////////////////////////////////////
 
+void STimingView::UpdateAggregatedStats()
+{
 	if (bAssetLoadingMode)
 	{
-		TSharedPtr<Insights::STableTreeView> EventAggregationTreeView;
-		TSharedPtr<Insights::STableTreeView> ObjectTypeAggregationTreeView;
-		TSharedPtr<Insights::STableTreeView> PackageDetailsTreeView;
-		TSharedPtr<Insights::STableTreeView> ExportDetailsTreeView;
 		TSharedPtr<SLoadingProfilerWindow> LoadingProfilerWnd = FLoadingProfilerManager::Get()->GetProfilerWindow();
-		if (LoadingProfilerWnd)
+		if (LoadingProfilerWnd.IsValid())
 		{
-			EventAggregationTreeView = LoadingProfilerWnd->GetEventAggregationTreeView();
-			ObjectTypeAggregationTreeView = LoadingProfilerWnd->GetObjectTypeAggregationTreeView();
-			PackageDetailsTreeView = LoadingProfilerWnd->GetPackageDetailsTreeView();
-			ExportDetailsTreeView = LoadingProfilerWnd->GetExportDetailsTreeView();
-		}
-
-		TSharedPtr<const Trace::IAnalysisSession> Session = FInsightsManager::Get()->GetSession();
-		if (Session.IsValid() && Trace::ReadLoadTimeProfilerProvider(*Session.Get()))
-		{
-			Trace::FAnalysisSessionReadScope SessionReadScope(*Session.Get());
-			const Trace::ILoadTimeProfilerProvider& LoadTimeProfilerProvider = *Trace::ReadLoadTimeProfilerProvider(*Session.Get());
+			TSharedPtr<Insights::STableTreeView> EventAggregationTreeView = LoadingProfilerWnd->GetEventAggregationTreeView();
+			TSharedPtr<Insights::STableTreeView> ObjectTypeAggregationTreeView = LoadingProfilerWnd->GetObjectTypeAggregationTreeView();
+			TSharedPtr<Insights::STableTreeView> PackageDetailsTreeView = LoadingProfilerWnd->GetPackageDetailsTreeView();
+			TSharedPtr<Insights::STableTreeView> ExportDetailsTreeView = LoadingProfilerWnd->GetExportDetailsTreeView();
 
+			TSharedPtr<const Trace::IAnalysisSession> Session = FInsightsManager::Get()->GetSession();
+			if (Session.IsValid() && Trace::ReadLoadTimeProfilerProvider(*Session.Get()))
 			{
-				//TODO: LoadTimeProfilerProvider.UpdateEventAggregation(EventAggregationTreeView->GetTable(), SelectionStartTime, SelectionEndTime)
-				Trace::ITable<Trace::FLoadTimeProfilerAggregatedStats>* EventAggregationTable = LoadTimeProfilerProvider.CreateEventAggregation(SelectionStartTime, SelectionEndTime);
-				EventAggregationTreeView->GetTable()->UpdateSourceTable(MakeShareable(EventAggregationTable));
-				EventAggregationTreeView->RebuildTree(true);
-				//delete EventAggregationTable;
-			}
+				Trace::FAnalysisSessionReadScope SessionReadScope(*Session.Get());
+				const Trace::ILoadTimeProfilerProvider& LoadTimeProfilerProvider = *Trace::ReadLoadTimeProfilerProvider(*Session.Get());
 
-			{
-				//TODO: LoadTimeProfilerProvider.UpdateObjectTypeAggregation(ObjectTypeAggregationTreeView->GetTable(), SelectionStartTime, SelectionEndTime)
-				Trace::ITable<Trace::FLoadTimeProfilerAggregatedStats>* ObjectTypeAggregationTable = LoadTimeProfilerProvider.CreateObjectTypeAggregation(SelectionStartTime, SelectionEndTime);
-				ObjectTypeAggregationTreeView->GetTable()->UpdateSourceTable(MakeShareable(ObjectTypeAggregationTable));
-				ObjectTypeAggregationTreeView->RebuildTree(true);
-				//delete ObjectTypeAggregationTable;
-			}
+				if (EventAggregationTreeView.IsValid())
+				{
+					//TODO: LoadTimeProfilerProvider.UpdateEventAggregation(EventAggregationTreeView->GetTable(), SelectionStartTime, SelectionEndTime)
+					Trace::ITable<Trace::FLoadTimeProfilerAggregatedStats>* EventAggregationTable = LoadTimeProfilerProvider.CreateEventAggregation(SelectionStartTime, SelectionEndTime);
+					EventAggregationTreeView->GetTable()->UpdateSourceTable(MakeShareable(EventAggregationTable));
+					EventAggregationTreeView->RebuildTree(true);
+				}
 
-			{
-				//TODO: LoadTimeProfilerProvider.UpdatePackageDetails(PackageDetailsTreeView->GetTable(), SelectionStartTime, SelectionEndTime)
-				Trace::ITable<Trace::FPackagesTableRow>* PackageDetailsTable = LoadTimeProfilerProvider.CreatePackageDetailsTable(SelectionStartTime, SelectionEndTime);
-				PackageDetailsTreeView->GetTable()->UpdateSourceTable(MakeShareable(PackageDetailsTable));
-				PackageDetailsTreeView->RebuildTree(true);
-				//delete PackageDetailsTable;
-			}
+				if (ObjectTypeAggregationTreeView.IsValid())
+				{
+					//TODO: LoadTimeProfilerProvider.UpdateObjectTypeAggregation(ObjectTypeAggregationTreeView->GetTable(), SelectionStartTime, SelectionEndTime)
+					Trace::ITable<Trace::FLoadTimeProfilerAggregatedStats>* ObjectTypeAggregationTable = LoadTimeProfilerProvider.CreateObjectTypeAggregation(SelectionStartTime, SelectionEndTime);
+					ObjectTypeAggregationTreeView->GetTable()->UpdateSourceTable(MakeShareable(ObjectTypeAggregationTable));
+					ObjectTypeAggregationTreeView->RebuildTree(true);
+				}
 
-			{
-				//TODO: LoadTimeProfilerProvider.UpdateExportDetails(ExportDetailsTreeView->GetTable(), SelectionStartTime, SelectionEndTime)
-				Trace::ITable<Trace::FExportsTableRow>* ExportDetailsTable = LoadTimeProfilerProvider.CreateExportDetailsTable(SelectionStartTime, SelectionEndTime);
-				ExportDetailsTreeView->GetTable()->UpdateSourceTable(MakeShareable(ExportDetailsTable));
-				ExportDetailsTreeView->RebuildTree(true);
-				//delete ExportDetailsTable;
+				if (PackageDetailsTreeView.IsValid())
+				{
+					//TODO: LoadTimeProfilerProvider.UpdatePackageDetails(PackageDetailsTreeView->GetTable(), SelectionStartTime, SelectionEndTime)
+					Trace::ITable<Trace::FPackagesTableRow>* PackageDetailsTable = LoadTimeProfilerProvider.CreatePackageDetailsTable(SelectionStartTime, SelectionEndTime);
+					PackageDetailsTreeView->GetTable()->UpdateSourceTable(MakeShareable(PackageDetailsTable));
+					PackageDetailsTreeView->RebuildTree(true);
+				}
+
+				if (ExportDetailsTreeView.IsValid())
+				{
+					//TODO: LoadTimeProfilerProvider.UpdateExportDetails(ExportDetailsTreeView->GetTable(), SelectionStartTime, SelectionEndTime)
+					Trace::ITable<Trace::FExportsTableRow>* ExportDetailsTable = LoadTimeProfilerProvider.CreateExportDetailsTable(SelectionStartTime, SelectionEndTime);
+					ExportDetailsTreeView->GetTable()->UpdateSourceTable(MakeShareable(ExportDetailsTable));
+					ExportDetailsTreeView->RebuildTree(true);
+				}
 			}
 		}
 	}
@@ -2237,20 +2226,13 @@ void STimingView::UpdateHoveredTimingEvent(float MX, float MY)
 
 void STimingView::OnSelectedTimingEventChanged()
 {
-	// Select the timer node coresponding to timing event type of selected timing event.
-	if (!bAssetLoadingMode &&
-		SelectedTimingEvent.IsValid() &&
-		(SelectedTimingEvent.Track->GetType() == FName(TEXT("Thread"))))
+	if (!bAssetLoadingMode)
 	{
-		//TODO: make this more generic
-		TSharedPtr<STimingProfilerWindow> Wnd = FTimingProfilerManager::Get()->GetProfilerWindow();
-		if (Wnd)
+		if (SelectedTimingEvent.IsValid() &&
+			SelectedTimingEvent.Track->GetType() == FName(TEXT("Thread")))
 		{
-			TSharedPtr<STimersView> TimersView = Wnd->GetTimersView();
-			if (TimersView)
-			{
-				TimersView->SelectTimerNode(SelectedTimingEvent.TypeId);
-			}
+			// Select the timer node coresponding to timing event type of selected timing event.
+			FTimingProfilerManager::Get()->SetSelectedTimer(SelectedTimingEvent.TypeId);
 		}
 	}
 }
@@ -2821,7 +2803,14 @@ void STimingView::OnTimingEventsTrackVisibilityChanged()
 	SelectedTimingEvent.Reset();
 	Tooltip.SetDesiredOpacity(0.0f);
 
-	UpdateAggregatedStats();
+	//TODO: ThreadFilterChangedEvent.Broadcast();
+
+	FTimingProfilerManager::Get()->OnThreadFilterChanged();
+
+	if (SelectionStartTime < SelectionEndTime)
+	{
+		UpdateAggregatedStats();
+	}
 
 	bAreTimingEventsTracksDirty = true;
 }
diff --git a/Engine/Source/Developer/TraceInsights/Private/Insights/Widgets/STimingView.h b/Engine/Source/Developer/TraceInsights/Private/Insights/Widgets/STimingView.h
--- a/Engine/Source/Developer/TraceInsights/Private/Insights/Widgets/STimingView.h
+++ b/Engine/Source/Developer/TraceInsights/Private/Insights/Widgets/STimingView.h
@@ -307,6 +307,8 @@ protected:
 
 	////////////////////////////////////////////////////////////////////////////////////////////////////
 
+	void RaiseSelectionChanging();
+	void RaiseSelectionChanged();
 	void UpdateAggregatedStats();
 
 	void UpdateHoveredTimingEvent(float MX, float MY);
