[Insights]
  - Refactored tree view implementations (Generic Table Tree view, Timers tree view, Callers/Calles views) to share more code. Added support for Value Getter, Value Formatter and Value Sorter. Each FTableColumn can now expose a value getter, formatter and sorter (many of those implementation being now shared between multiple tree view widgets).
  - Generic Table Tree view: Improved sorting for large tables.
  - Callers/Callees: Added visualisation of timers on the hot path.
  - Callers/Callees: Added sorting for all available columns.

diff --git a/Engine/Source/Developer/TraceInsights/Private/Insights/Table/ViewModels/BaseTreeNode.cpp b/Engine/Source/Developer/TraceInsights/Private/Insights/Table/ViewModels/BaseTreeNode.cpp
--- a/Engine/Source/Developer/TraceInsights/Private/Insights/Table/ViewModels/BaseTreeNode.cpp
+++ b/Engine/Source/Developer/TraceInsights/Private/Insights/Table/ViewModels/BaseTreeNode.cpp
@@ -3,7 +3,7 @@
 #include "BaseTreeNode.h"
 
 // Insights
-#include "Insights/Table/ViewModels/TreeNodeSorting.h"
+#include "Insights/Table/ViewModels/TableCellValueSorter.h"
 
 #define LOCTEXT_NAMESPACE "Insights_TreeNode"
 
@@ -40,16 +40,16 @@ const FText FBaseTreeNode::GetDisplayName() const
 
 ////////////////////////////////////////////////////////////////////////////////////////////////////
 
-void FBaseTreeNode::SortChildrenAscending(const ITreeNodeSorting* Sorter)
+void FBaseTreeNode::SortChildrenAscending(const ITableCellValueSorter& Sorter)
 {
-	Sorter->SortAscending(Children);
+	Sorter.Sort(Children, ESortMode::Ascending);
 }
 
 ////////////////////////////////////////////////////////////////////////////////////////////////////
 
-void FBaseTreeNode::SortChildrenDescending(const ITreeNodeSorting* Sorter)
+void FBaseTreeNode::SortChildrenDescending(const ITableCellValueSorter& Sorter)
 {
-	Sorter->SortDescending(Children);
+	Sorter.Sort(Children, ESortMode::Descending);
 }
 
 ////////////////////////////////////////////////////////////////////////////////////////////////////
diff --git a/Engine/Source/Developer/TraceInsights/Private/Insights/Table/ViewModels/BaseTreeNode.h b/Engine/Source/Developer/TraceInsights/Private/Insights/Table/ViewModels/BaseTreeNode.h
--- a/Engine/Source/Developer/TraceInsights/Private/Insights/Table/ViewModels/BaseTreeNode.h
+++ b/Engine/Source/Developer/TraceInsights/Private/Insights/Table/ViewModels/BaseTreeNode.h
@@ -7,7 +7,7 @@
 namespace Insights
 {
 
-class ITreeNodeSorting;
+class ITableCellValueSorter;
 
 ////////////////////////////////////////////////////////////////////////////////////////////////////
 
@@ -46,6 +46,8 @@ public:
 	{
 	}
 
+	virtual const FName& GetTypeName() const = 0;
+
 	/**
 	 * @return an id of this node.
 	 */
@@ -63,15 +65,7 @@ public:
 	}
 
 	/**
-	 * @return a type id of this node.
-	 */
-	const FName& GetTypeId() const
-	{
-		return TypeId;
-	}
-
-	/**
-	 * @return a name of this node to display in tree view.  Ex.: group nodes may include additional info, like the number of visible/total children.
+	 * @return a name of this node to display in tree view. Ex.: group nodes may include additional info, like the number of visible/total children.
 	 */
 	virtual const FText GetDisplayName() const;
 
@@ -120,8 +114,8 @@ public:
 		return false;
 	}
 
-	void SortChildrenAscending(const ITreeNodeSorting* Sorter);
-	void SortChildrenDescending(const ITreeNodeSorting* Sorter);
+	void SortChildrenAscending(const ITableCellValueSorter& Sorter);
+	void SortChildrenDescending(const ITableCellValueSorter& Sorter);
 
 	/** Adds specified child to the children and sets group for it. */
 	FORCEINLINE_DEBUGGABLE void AddChildAndSetGroupPtr(const FBaseTreeNodePtr& ChildPtr)
@@ -171,9 +165,6 @@ private:
 	/** The name of this node. */
 	const FName Name;
 
-	/** The type id of this node (like a category name). */
-	const FName TypeId;
-
 	/** Children of this node. */
 	TArray<FBaseTreeNodePtr> Children;
 
diff --git a/Engine/Source/Developer/TraceInsights/Private/Insights/Table/ViewModels/Table.cpp b/Engine/Source/Developer/TraceInsights/Private/Insights/Table/ViewModels/Table.cpp
--- a/Engine/Source/Developer/TraceInsights/Private/Insights/Table/ViewModels/Table.cpp
+++ b/Engine/Source/Developer/TraceInsights/Private/Insights/Table/ViewModels/Table.cpp
@@ -5,7 +5,11 @@
 
 // Insights
 #include "Insights/Common/TimeUtils.h"
+#include "Insights/Table/ViewModels/TableCellValueFormatter.h"
+#include "Insights/Table/ViewModels/TableCellValueGetter.h"
+#include "Insights/Table/ViewModels/TableCellValueSorter.h"
 #include "Insights/Table/ViewModels/TableColumn.h"
+#include "Insights/Table/ViewModels/TableTreeNode.h"
 
 #define LOCTEXT_NAMESPACE "Insights_Table"
 
@@ -13,316 +17,58 @@ namespace Insights
 {
 
 ////////////////////////////////////////////////////////////////////////////////////////////////////
-// FTableCellValueFormatter
+// FTableTreeNodeValueGetter
 ////////////////////////////////////////////////////////////////////////////////////////////////////
 
-class FBoolValueFormatter : public FTableCellValueFormatter
+class FTableTreeNodeValueGetter : public FTableCellValueGetter
 {
 public:
-	virtual FText FormatValue(const TOptional<FTableCellValue>& InValue) override
-	{
-		if (InValue.IsSet())
-		{
-			return FormatBoolValue(InValue.GetValue().Bool);
-		}
-		return FText::GetEmpty();
-	}
-
-	virtual FText FormatValue(const FTable& Table, const FTableColumn& Column, const FTableRowId& RowId) override
-	{
-		TSharedPtr<Trace::IUntypedTableReader> Reader = Table.GetTableReader();
-		if (Reader.IsValid() && RowId.HasValidIndex())
-		{
-			Reader->SetRowIndex(RowId.RowIndex);
-			const bool Value = Reader->GetValueBool(Column.GetIndex());
-			return FormatBoolValue(Value);
-		}
-		return FText::GetEmpty();
-	}
-
-private:
-	FText FormatBoolValue(bool Value)
-	{
-		return FText::FromString(Value ? TEXT("True") : TEXT("False"));
-	}
-};
-
-////////////////////////////////////////////////////////////////////////////////////////////////////
-
-class FInt64ValueFormatter : public FTableCellValueFormatter
-{
-public:
-	virtual FText FormatValue(const TOptional<FTableCellValue>& InValue) override
-	{
-		if (InValue.IsSet())
-		{
-			return FormatInt64Value(InValue.GetValue().Int64);
-		}
-		return FText::GetEmpty();
-	}
-
-	virtual FText FormatValue(const FTable& Table, const FTableColumn& Column, const FTableRowId& RowId) override
-	{
-		TSharedPtr<Trace::IUntypedTableReader> Reader = Table.GetTableReader();
-		if (Reader.IsValid() && RowId.HasValidIndex())
-		{
-			Reader->SetRowIndex(RowId.RowIndex);
-			const int64 Value = Reader->GetValueInt(Column.GetIndex());
-			return FormatInt64Value(Value);
-		}
-		return FText::GetEmpty();
-	}
-
-private:
-	FText FormatInt64Value(int64 Value)
-	{
-		return FText::AsNumber(Value);
-	}
-};
-
-////////////////////////////////////////////////////////////////////////////////////////////////////
-
-class FFloatValueFormatter : public FTableCellValueFormatter
-{
-public:
-	virtual FText FormatValue(const TOptional<FTableCellValue>& InValue) override
-	{
-		if (InValue.IsSet())
-		{
-			return FormatFloatValue(InValue.GetValue().Float);
-		}
-		return FText::GetEmpty();
-	}
-
-	virtual FText FormatValue(const FTable& Table, const FTableColumn& Column, const FTableRowId& RowId) override
-	{
-		TSharedPtr<Trace::IUntypedTableReader> Reader = Table.GetTableReader();
-		if (Reader.IsValid() && RowId.HasValidIndex())
-		{
-			Reader->SetRowIndex(RowId.RowIndex);
-			const float Value = Reader->GetValueFloat(Column.GetIndex());
-			return FormatFloatValue(Value);
-		}
-		return FText::GetEmpty();
-	}
-
-private:
-	FText FormatFloatValue(float Value)
-	{
-		if (Value == 0.0f)
-		{
-			return FText::FromString(TEXT("0"));
-		}
-		else
-		{
-			return FText::FromString(FString::Printf(TEXT("%f"), Value));
-		}
-	}
-};
-
-////////////////////////////////////////////////////////////////////////////////////////////////////
-
-class FFloatTimeValueFormatter : public FTableCellValueFormatter
-{
-public:
-	virtual FText FormatValue(const TOptional<FTableCellValue>& InValue) override
-	{
-		if (InValue.IsSet())
-		{
-			return FormatFloatValue(InValue.GetValue().Float);
-		}
-		return FText::GetEmpty();
-	}
-
-	virtual FText FormatValue(const FTable& Table, const FTableColumn& Column, const FTableRowId& RowId) override
-	{
-		TSharedPtr<Trace::IUntypedTableReader> Reader = Table.GetTableReader();
-		if (Reader.IsValid() && RowId.HasValidIndex())
-		{
-			Reader->SetRowIndex(RowId.RowIndex);
-			const float Value = Reader->GetValueFloat(Column.GetIndex());
-			return FormatFloatValue(Value);
-		}
-		return FText::GetEmpty();
-	}
-
-	virtual FText FormatValueForTooltip(const TOptional<FTableCellValue>& InValue) override
-	{
-		if (InValue.IsSet())
-		{
-			return FormatFloatValueForTooltip(InValue.GetValue().Float);
-		}
-		return FText::GetEmpty();
-	}
-
-	virtual FText FormatValueForTooltip(const FTable& Table, const FTableColumn& Column, const FTableRowId& RowId) override
-	{
-		TSharedPtr<Trace::IUntypedTableReader> Reader = Table.GetTableReader();
-		if (Reader.IsValid() && RowId.HasValidIndex())
-		{
-			Reader->SetRowIndex(RowId.RowIndex);
-			const float Value = Reader->GetValueFloat(Column.GetIndex());
-			return FormatFloatValueForTooltip(Value);
-		}
-		return FText::GetEmpty();
-	}
-
-private:
-	FText FormatFloatValue(float Value)
-	{
-		return FText::FromString(TimeUtils::FormatTimeAuto(static_cast<double>(Value)));
-	}
-
-	FText FormatFloatValueForTooltip(float Value)
-	{
-		if (Value == 0.0f)
-		{
-			return FText::FromString(TEXT("0"));
-		}
-		else
-		{
-			return FText::FromString(FString::Printf(TEXT("%f (%s)"), Value, *TimeUtils::FormatTimeAuto(static_cast<double>(Value))));
-		}
-	}
-};
-
-////////////////////////////////////////////////////////////////////////////////////////////////////
-
-class FDoubleValueFormatter : public FTableCellValueFormatter
-{
-public:
-	virtual FText FormatValue(const TOptional<FTableCellValue>& InValue) override
-	{
-		if (InValue.IsSet())
-		{
-			return FormatDoubleValue(InValue.GetValue().Double);
-		}
-		return FText::GetEmpty();
-	}
-
-	virtual FText FormatValue(const FTable& Table, const FTableColumn& Column, const FTableRowId& RowId) override
-	{
-		TSharedPtr<Trace::IUntypedTableReader> Reader = Table.GetTableReader();
-		if (Reader.IsValid() && RowId.HasValidIndex())
-		{
-			Reader->SetRowIndex(RowId.RowIndex);
-			const double Value = Reader->GetValueDouble(Column.GetIndex());
-			return FormatDoubleValue(Value);
-		}
-		return FText::GetEmpty();
-	}
-
-private:
-	FText FormatDoubleValue(double Value)
-	{
-		if (Value == 0.0)
-		{
-			return FText::FromString(TEXT("0"));
-		}
-		else
-		{
-			return FText::FromString(FString::Printf(TEXT("%f"), Value));
-		}
-	}
-};
-
-////////////////////////////////////////////////////////////////////////////////////////////////////
-
-class FDoubleTimeValueFormatter : public FTableCellValueFormatter
-{
-public:
-	virtual FText FormatValue(const TOptional<FTableCellValue>& InValue) override
-	{
-		if (InValue.IsSet())
-		{
-			return FormatDoubleValue(InValue.GetValue().Double);
-		}
-		return FText::GetEmpty();
-	}
-
-	virtual FText FormatValue(const FTable& Table, const FTableColumn& Column, const FTableRowId& RowId) override
-	{
-		TSharedPtr<Trace::IUntypedTableReader> Reader = Table.GetTableReader();
-		if (Reader.IsValid() && RowId.HasValidIndex())
-		{
-			Reader->SetRowIndex(RowId.RowIndex);
-			const double Value = Reader->GetValueDouble(Column.GetIndex());
-			return FormatDoubleValue(Value);
-		}
-		return FText::GetEmpty();
-	}
-
-	virtual FText FormatValueForTooltip(const TOptional<FTableCellValue>& InValue) override
-	{
-		if (InValue.IsSet())
-		{
-			return FormatDoubleValueForTooltip(InValue.GetValue().Double);
-		}
-		return FText::GetEmpty();
-	}
-
-	virtual FText FormatValueForTooltip(const FTable& Table, const FTableColumn& Column, const FTableRowId& RowId) override
-	{
-		TSharedPtr<Trace::IUntypedTableReader> Reader = Table.GetTableReader();
-		if (Reader.IsValid() && RowId.HasValidIndex())
-		{
-			Reader->SetRowIndex(RowId.RowIndex);
-			const double Value = Reader->GetValueDouble(Column.GetIndex());
-			return FormatDoubleValueForTooltip(Value);
-		}
-		return FText::GetEmpty();
-	}
+	FTableTreeNodeValueGetter(ETableCellDataType InDataType) : FTableCellValueGetter(), DataType(InDataType) {}
 
-private:
-	FText FormatDoubleValue(double Value)
+	virtual const TOptional<FTableCellValue> GetValue(const FTableColumn& Column, const FBaseTreeNode& Node) const
 	{
-		return FText::FromString(TimeUtils::FormatTimeAuto(Value));
-	}
+		ensure(Node.GetTypeName() == FTableTreeNode::TypeName);
+		const FTableTreeNode& TableTreeNode = static_cast<const FTableTreeNode&>(Node);
 
-	FText FormatDoubleValueForTooltip(double Value)
-	{
-		if (Value == 0.0)
+		if (!Node.IsGroup()) // Table Row Node
 		{
-			return FText::FromString(TEXT("0"));
+			const TSharedPtr<FTable> TablePtr = Column.GetParentTable().Pin();
+			if (TablePtr.IsValid())
+			{
+				TSharedPtr<Trace::IUntypedTableReader> Reader = TablePtr->GetTableReader();
+				if (Reader.IsValid() && TableTreeNode.GetRowId().HasValidIndex())
+				{
+					Reader->SetRowIndex(TableTreeNode.GetRowId().RowIndex);
+					const int32 ColumnIndex = Column.GetIndex();
+					switch (DataType)
+					{
+						case ETableCellDataType::Bool:    return TOptional<FTableCellValue>(Reader->GetValueBool(ColumnIndex));
+						case ETableCellDataType::Int64:   return TOptional<FTableCellValue>(Reader->GetValueInt(ColumnIndex));
+						case ETableCellDataType::Float:   return TOptional<FTableCellValue>(Reader->GetValueFloat(ColumnIndex));
+						case ETableCellDataType::Double:  return TOptional<FTableCellValue>(Reader->GetValueDouble(ColumnIndex));
+						case ETableCellDataType::CString: return TOptional<FTableCellValue>(Reader->GetValueCString(ColumnIndex));
+					}
+				}
+			}
 		}
-		else
-		{
-			return FText::FromString(FString::Printf(TEXT("%f (%s)"), Value, *TimeUtils::FormatTimeAuto(Value)));
-		}
-	}
-};
-
-////////////////////////////////////////////////////////////////////////////////////////////////////
-
-class FCStringValueFormatter : public FTableCellValueFormatter
-{
-public:
-	virtual FText FormatValue(const TOptional<FTableCellValue>& InValue) override
-	{
-		if (InValue.IsSet())
+		else // Aggregated Group Node
 		{
-			return FormatCStringValue(InValue.GetValue().CString);
+			if (Column.GetAggregation() != ETableColumnAggregation::None)
+			{
+				const FTableCellValue* ValuePtr = TableTreeNode.FindAggregatedValue(Column.GetId());
+				if (ValuePtr != nullptr)
+				{
+					ensure(ValuePtr->DataType == DataType);
+					return TOptional<FTableCellValue>(*ValuePtr);
+				}
+			}
 		}
-		return FText::GetEmpty();
-	}
 
-	virtual FText FormatValue(const FTable& Table, const FTableColumn& Column, const FTableRowId& RowId) override
-	{
-		TSharedPtr<Trace::IUntypedTableReader> Reader = Table.GetTableReader();
-		if (Reader.IsValid() && RowId.HasValidIndex())
-		{
-			Reader->SetRowIndex(RowId.RowIndex);
-			const TCHAR* Value = Reader->GetValueCString(Column.GetIndex());
-			return FormatCStringValue(Value);
-		}
-		return FText::GetEmpty();
+		return TOptional<FTableCellValue>();
 	}
 
 private:
-	FText FormatCStringValue(const TCHAR* Value)
-	{
-		return FText::FromString(Value);
-	}
+	ETableCellDataType DataType;
 };
 
 ////////////////////////////////////////////////////////////////////////////////////////////////////
@@ -332,11 +78,10 @@ private:
 FTable::FTable()
 	: Name()
 	, Description()
-	, SourceTable()
-	, TableReader()
 	, Columns()
-	, VisibleColumns()
 	, ColumnIdToPtrMapping()
+	, SourceTable()
+	, TableReader()
 {
 }
 
@@ -351,12 +96,18 @@ FTable::~FTable()
 
 void FTable::Reset()
 {
+	Columns.Reset();
+	ColumnIdToPtrMapping.Reset();
+
 	SourceTable.Reset();
 	TableReader.Reset();
+}
 
-	Columns.Reset();
-	VisibleColumns.Reset();
-	ColumnIdToPtrMapping.Reset();
+////////////////////////////////////////////////////////////////////////////////////////////////////
+
+int32 FTable::GetColumnPositionIndex(const FName& ColumnId) const
+{
+	return Columns.IndexOfByPredicate([&ColumnId](const TSharedPtr<FTableColumn>& ColumnPtr) -> bool { return ColumnPtr->GetId() == ColumnId; });
 }
 
 ////////////////////////////////////////////////////////////////////////////////////////////////////
@@ -370,7 +121,7 @@ void FTable::Init(TSharedPtr<Trace::IUntypedTable> InSourceTable)
 	if (SourceTable)
 	{
 		TableReader = MakeShareable(SourceTable->CreateReader());
-		CreateColumns();
+		CreateColumnsFromTableLayout();
 	}
 }
 
@@ -411,42 +162,14 @@ void FTable::UpdateSourceTable(TSharedPtr<Trace::IUntypedTable> InSourceTable)
 
 ////////////////////////////////////////////////////////////////////////////////////////////////////
 
-void FTable::CreateHierarchyColumn(int32 ColumnIndex, const TCHAR* ColumnName)
+void FTable::SetColumns(const TArray<TSharedPtr<Insights::FTableColumn>>& InColumns)
 {
-	const FString ColumnNameStr = ColumnName ? FString::Printf(TEXT("Hierarchy (%s)"), ColumnName) : TEXT("Hierarchy");
-	const FText ColumnNameText = FText::FromString(ColumnNameStr);
-
-	constexpr EHorizontalAlignment HorizontalAlignment = HAlign_Left;
-	constexpr ETableColumnFlags ColumnFlags = ETableColumnFlags::ShouldBeVisible | ETableColumnFlags::CanBeSorted | ETableColumnFlags::CanBeFiltered | ETableColumnFlags::IsHierarchy;
-
-	constexpr ETableCellDataType DataType = ETableCellDataType::CString;
-	constexpr ETableColumnAggregation Aggregation = ETableColumnAggregation::None;
-
-	constexpr float InitialColumnWidth = 90.0f;
-	constexpr float MinColumnWidth = 0.0f;
-	constexpr float MaxColumnWidth = FLT_MAX;
-
-	TSharedPtr<FTableCellValueFormatter> FormatterPtr = MakeShareable(new FTableCellValueFormatter());
-
-	TSharedPtr<FTableColumn> ColumnPtr = MakeShareable(new FTableColumn
-	(
-		-1, // Order
-		ColumnIndex,
-		FName(TEXT("_Hierarchy")), // Id
-		ColumnNameText, // Short Name
-		ColumnNameText, // Title Name
-		FText::GetEmpty(), // Description
-		ColumnFlags,
-		DataType,
-		Aggregation,
-		HorizontalAlignment,
-		InitialColumnWidth,
-		MinColumnWidth,
-		MaxColumnWidth,
-		FormatterPtr.ToSharedRef()
-	));
-
-	AddColumn(ColumnPtr);
+	Columns.Reset(InColumns.Num());
+	ColumnIdToPtrMapping.Reset();
+	for (TSharedPtr<Insights::FTableColumn> ColumnPtr : InColumns)
+	{
+		AddColumn(ColumnPtr);
+	}
 }
 
 ////////////////////////////////////////////////////////////////////////////////////////////////////
@@ -454,20 +177,50 @@ void FTable::CreateHierarchyColumn(int32 ColumnIndex, const TCHAR* ColumnName)
 void FTable::AddColumn(TSharedPtr<FTableColumn> ColumnPtr)
 {
 	ColumnPtr->SetParentTable(SharedThis(this));
-
 	Columns.Add(ColumnPtr);
+	ColumnIdToPtrMapping.Add(ColumnPtr->GetId(), ColumnPtr);
+}
 
-	if (ColumnPtr->IsVisible())
-	{
-		VisibleColumns.Add(ColumnPtr);
-	}
+////////////////////////////////////////////////////////////////////////////////////////////////////
 
-	ColumnIdToPtrMapping.Add(ColumnPtr->GetId(), ColumnPtr);
+void FTable::CreateHierarchyColumn(int32 ColumnIndex, const TCHAR* ColumnName)
+{
+	const FName HierarchyColumnId(TEXT("_Hierarchy"));
+
+	TSharedPtr<FTableColumn> ColumnPtr = MakeShareable(new FTableColumn(HierarchyColumnId));
+	FTableColumn& Column = *ColumnPtr;
+	
+	Column.SetIndex(ColumnIndex);
+
+	const FString ColumnNameStr = ColumnName ? FString::Printf(TEXT("Hierarchy (%s)"), ColumnName) : TEXT("Hierarchy");
+	const FText ColumnNameText = FText::FromString(ColumnNameStr);
+
+	Column.SetShortName(ColumnNameText);
+	Column.SetTitleName(ColumnNameText);
+	//TODO: Column.SetDescription(...);
+
+	Column.SetFlags(ETableColumnFlags::ShouldBeVisible | ETableColumnFlags::CanBeFiltered | ETableColumnFlags::IsHierarchy);
+
+	Column.SetHorizontalAlignment(HAlign_Left);
+	Column.SetInitialWidth(90.0f);
+
+	Column.SetDataType(ETableCellDataType::CString);
+
+	TSharedPtr<ITableCellValueGetter> GetterPtr = MakeShareable(new FDisplayNameValueGetter());
+	Column.SetValueGetter(GetterPtr.ToSharedRef());
+
+	TSharedPtr<ITableCellValueFormatter> FormatterPtr = MakeShareable(new FTextValueFormatter());
+	Column.SetValueFormatter(FormatterPtr.ToSharedRef());
+
+	TSharedPtr<ITableCellValueSorter> SorterPtr = MakeShareable(new FSorterByName(ColumnPtr.ToSharedRef()));
+	Column.SetValueSorter(SorterPtr);
+
+	AddColumn(ColumnPtr);
 }
 
 ////////////////////////////////////////////////////////////////////////////////////////////////////
 
-void FTable::CreateColumns()
+void FTable::CreateColumnsFromTableLayout()
 {
 	ensure(TableReader.IsValid());
 	ensure(Columns.Num() == 0);
@@ -502,85 +255,105 @@ void FTable::CreateColumns()
 		Trace::ETableColumnType ColumnType = TableLayout.GetColumnType(ColumnIndex);
 		const TCHAR* ColumnName = TableLayout.GetColumnName(ColumnIndex);
 
+		TSharedPtr<FTableColumn> ColumnPtr = MakeShareable(new FTableColumn(FName(ColumnName)));
+		FTableColumn& Column = *ColumnPtr;
+
 		const FString ColumnNameStr(ColumnName);
 		const FText ColumnNameText = FText::FromString(ColumnNameStr);
 
-		EHorizontalAlignment HorizontalAlignment = HAlign_Left;
-
-		ETableColumnFlags ColumnFlags = ETableColumnFlags::CanBeSorted | ETableColumnFlags::CanBeFiltered | ETableColumnFlags::CanBeHidden;
+		ETableColumnFlags ColumnFlags = ETableColumnFlags::CanBeFiltered | ETableColumnFlags::CanBeHidden;
 		if (ColumnIndex != HierarchyColumnIndex)
 		{
 			ColumnFlags |= ETableColumnFlags::ShouldBeVisible;
 		}
 
-		ETableCellDataType DataType = ETableCellDataType::Unknown;
-		ETableColumnAggregation Aggregation = ETableColumnAggregation::None;
-
+		EHorizontalAlignment HorizontalAlignment = HAlign_Left;
 		float InitialColumnWidth = 60.0f;
-		float MinColumnWidth = 0.0f;
-		float MaxColumnWidth = FLT_MAX;
 
-		TSharedPtr<FTableCellValueFormatter> FormatterPtr;
+		ETableColumnAggregation Aggregation = ETableColumnAggregation::None;
+
+		TSharedPtr<ITableCellValueFormatter> FormatterPtr;
+		TSharedPtr<ITableCellValueSorter> SorterPtr;
 
 		switch (ColumnType)
 		{
 		case Trace::TableColumnType_Bool:
-			DataType = ETableCellDataType::Bool;
+			Column.SetDataType(ETableCellDataType::Bool);
 			HorizontalAlignment = HAlign_Right;
 			InitialColumnWidth = 40.0f;
-			FormatterPtr = MakeShareable(new FBoolValueFormatter());
+			//TODO: if (Hint == AsOnOff)
+			//else // if (Hint == AsTrueFalse)
+			FormatterPtr = MakeShareable(new FBoolValueFormatterAsTrueFalse());
+			SorterPtr = MakeShareable(new FSorterByBoolValue(ColumnPtr.ToSharedRef()));
 			break;
+
 		case Trace::TableColumnType_Int:
-			DataType = ETableCellDataType::Int64;
+			Column.SetDataType(ETableCellDataType::Int64);
 			Aggregation = ETableColumnAggregation::Sum;
 			HorizontalAlignment = HAlign_Right;
 			InitialColumnWidth = 60.0f;
-			FormatterPtr = MakeShareable(new FInt64ValueFormatter());
+			//TODO: if (Hint == AsMemory)
+			//{
+			//	FormatterPtr = MakeShareable(new FInt64ValueFormatterAsMemory());
+			//}
+			//else // AsNumber
+			FormatterPtr = MakeShareable(new FInt64ValueFormatterAsNumber());
+			SorterPtr = MakeShareable(new FSorterByInt64Value(ColumnPtr.ToSharedRef()));
 			break;
+
 		case Trace::TableColumnType_Float:
-			DataType = ETableCellDataType::Float;
+			Column.SetDataType(ETableCellDataType::Float);
 			Aggregation = ETableColumnAggregation::Sum;
 			HorizontalAlignment = HAlign_Right;
 			InitialColumnWidth = 60.0f;
-			FormatterPtr = MakeShareable(new FFloatTimeValueFormatter());
+			//TODO: if (Hint == AsTimeMs)
+			//else // if (Hint == AsTimeAuto)
+			FormatterPtr = MakeShareable(new FFloatValueFormatterAsTimeAuto());
+			SorterPtr = MakeShareable(new FSorterByFloatValue(ColumnPtr.ToSharedRef()));
 			break;
+
 		case Trace::TableColumnType_Double:
-			DataType = ETableCellDataType::Double;
+			Column.SetDataType(ETableCellDataType::Double);
 			Aggregation = ETableColumnAggregation::Sum;
 			HorizontalAlignment = HAlign_Right;
 			InitialColumnWidth = 80.0f;
-			FormatterPtr = MakeShareable(new FDoubleTimeValueFormatter());
+			//TODO: if (Hint == AsTimeMs)
+			//else // if (Hint == AsTimeAuto)
+			FormatterPtr = MakeShareable(new FDoubleValueFormatterAsTimeAuto());
+			SorterPtr = MakeShareable(new FSorterByDoubleValue(ColumnPtr.ToSharedRef()));
 			break;
+
 		case Trace::TableColumnType_CString:
-			DataType = ETableCellDataType::CString;
+			Column.SetDataType(ETableCellDataType::CString);
 			HorizontalAlignment = HAlign_Left;
 			InitialColumnWidth = FMath::Max(120.0f, 6.0f * ColumnNameStr.Len());
-			FormatterPtr = MakeShareable(new FCStringValueFormatter());
+			FormatterPtr = MakeShareable(new FCStringValueFormatterAsText());
+			SorterPtr = MakeShareable(new FSorterByCStringValue(ColumnPtr.ToSharedRef()));
 			break;
 		}
 
-		if (!FormatterPtr.IsValid())
+		Column.SetIndex(ColumnIndex);
+
+		Column.SetShortName(ColumnNameText);
+		Column.SetTitleName(ColumnNameText);
+
+		//TODO: Column.SetDescription(...);
+
+		Column.SetFlags(ColumnFlags);
+
+		Column.SetHorizontalAlignment(HorizontalAlignment);
+		Column.SetInitialWidth(InitialColumnWidth);
+
+		Column.SetAggregation(Aggregation);
+
+		Column.SetValueGetter(MakeShareable(new FTableTreeNodeValueGetter(Column.GetDataType())));
+
+		if (FormatterPtr.IsValid())
 		{
-			FormatterPtr = MakeShareable(new FTableCellValueFormatter());
+			Column.SetValueFormatter(FormatterPtr.ToSharedRef());
 		}
 
-		TSharedPtr<FTableColumn> ColumnPtr = MakeShareable(new FTableColumn
-		(
-			ColumnIndex, // Order
-			ColumnIndex,
-			FName(ColumnName), // Id
-			ColumnNameText, // Short Name
-			ColumnNameText, // Title Name
-			FText::GetEmpty(), // Description
-			ColumnFlags,
-			DataType,
-			Aggregation,
-			HorizontalAlignment,
-			InitialColumnWidth,
-			MinColumnWidth,
-			MaxColumnWidth,
-			FormatterPtr.ToSharedRef()
-		));
+		Column.SetValueSorter(SorterPtr);
 
 		AddColumn(ColumnPtr);
 	}
diff --git a/Engine/Source/Developer/TraceInsights/Private/Insights/Table/ViewModels/Table.h b/Engine/Source/Developer/TraceInsights/Private/Insights/Table/ViewModels/Table.h
--- a/Engine/Source/Developer/TraceInsights/Private/Insights/Table/ViewModels/Table.h
+++ b/Engine/Source/Developer/TraceInsights/Private/Insights/Table/ViewModels/Table.h
@@ -52,16 +52,12 @@ public:
 	FText GetDisplayName() const { return FText::FromName(Name); }
 	const FText& GetDescription() const { return Description; }
 
-	TSharedPtr<Trace::IUntypedTable> GetSourceTable() const { return SourceTable; }
-	TSharedPtr<Trace::IUntypedTableReader> GetTableReader() const { return TableReader; }
-
-	void Reset();
-	void Init(TSharedPtr<Trace::IUntypedTable> InSourceTable);
-	void UpdateSourceTable(TSharedPtr<Trace::IUntypedTable> InSourceTable);
+	virtual void Reset();
 
 	bool IsValid() const { return Columns.Num() > 0; }
 
 	const TArray<TSharedPtr<FTableColumn>>& GetColumns() const { return Columns; }
+	void SetColumns(const TArray<TSharedPtr<Insights::FTableColumn>>& InColumns);
 
 	TSharedPtr<FTableColumn> FindColumnChecked(const FName& ColumnId) const
 	{
@@ -74,26 +70,32 @@ public:
 		return (ColumnPtrPtr != nullptr) ? *ColumnPtrPtr : nullptr;
 	}
 
+	int32 GetColumnPositionIndex(const FName& ColumnId) const;
+
+	TSharedPtr<Trace::IUntypedTable> GetSourceTable() const { return SourceTable; }
+	TSharedPtr<Trace::IUntypedTableReader> GetTableReader() const { return TableReader; }
+
+	/* Init this table to use the IUntypedTable source table. It will create new Columns array based on table layout. */
+	void Init(TSharedPtr<Trace::IUntypedTable> InSourceTable);
+	void UpdateSourceTable(TSharedPtr<Trace::IUntypedTable> InSourceTable);
+
 private:
-	void CreateHierarchyColumn(int32 ColumnIndex, const TCHAR* ColumnName);
 	void AddColumn(TSharedPtr<FTableColumn> ColumnPtr);
-	void CreateColumns();
+	void CreateHierarchyColumn(int32 ColumnIndex, const TCHAR* ColumnName);
+	void CreateColumnsFromTableLayout();
 
 private:
 	FName Name;
 	FText Description;
 
-	TSharedPtr<Trace::IUntypedTable> SourceTable;
-	TSharedPtr<Trace::IUntypedTableReader> TableReader;
-
 	/** All available columns. */
 	TArray<TSharedPtr<FTableColumn>> Columns;
 
-	/** The visible columns. */
-	TArray<TSharedPtr<FTableColumn>> VisibleColumns;
-
 	/** Mapping between column Ids and FTableColumn shared pointers. */
 	TMap<FName, TSharedPtr<FTableColumn>> ColumnIdToPtrMapping;
+
+	TSharedPtr<Trace::IUntypedTable> SourceTable;
+	TSharedPtr<Trace::IUntypedTableReader> TableReader;
 };
 
 ////////////////////////////////////////////////////////////////////////////////////////////////////
diff --git a/Engine/Source/Developer/TraceInsights/Private/Insights/Table/ViewModels/TableCellValue.cpp b/Engine/Source/Developer/TraceInsights/Private/Insights/Table/ViewModels/TableCellValue.cpp
new file mode 100644
--- /dev/null
+++ b/Engine/Source/Developer/TraceInsights/Private/Insights/Table/ViewModels/TableCellValue.cpp
@@ -0,0 +1,3 @@
+// Copyright 1998-2019 Epic Games, Inc. All Rights Reserved.
+
+#include "TableCellValue.h"
diff --git a/Engine/Source/Developer/TraceInsights/Private/Insights/Table/ViewModels/TableCellValue.h b/Engine/Source/Developer/TraceInsights/Private/Insights/Table/ViewModels/TableCellValue.h
new file mode 100644
--- /dev/null
+++ b/Engine/Source/Developer/TraceInsights/Private/Insights/Table/ViewModels/TableCellValue.h
@@ -0,0 +1,51 @@
+// Copyright 1998-2019 Epic Games, Inc. All Rights Reserved.
+
+#pragma once
+
+#include "CoreMinimal.h"
+
+////////////////////////////////////////////////////////////////////////////////////////////////////
+
+enum class ETableCellDataType : uint32
+{
+	Unknown,
+
+	Bool,
+	Int64,
+	Float,
+	Double,
+	CString,
+	Text,
+
+	/** Invalid enum type, may be used as a number of enumerations. */
+	InvalidOrMax,
+};
+
+////////////////////////////////////////////////////////////////////////////////////////////////////
+
+struct FTableCellValue
+{
+public:
+	FTableCellValue() : DataType(ETableCellDataType::Unknown) {}
+
+	FTableCellValue(bool Value) : DataType(ETableCellDataType::Bool), Bool(Value) {}
+	FTableCellValue(int64 Value) : DataType(ETableCellDataType::Int64), Int64(Value) {}
+	FTableCellValue(float Value) : DataType(ETableCellDataType::Float), Float(Value) {}
+	FTableCellValue(double Value) : DataType(ETableCellDataType::Double), Double(Value) {}
+	FTableCellValue(const TCHAR* Value) : DataType(ETableCellDataType::CString), CString(Value) {}
+	FTableCellValue(const FText& Value) : DataType(ETableCellDataType::Text), TextPtr(MakeShareable(new FText(Value))) {}
+
+	ETableCellDataType DataType;
+
+	union
+	{
+		bool Bool;
+		int64 Int64;
+		float Float;
+		double Double;
+		const TCHAR* CString;
+	};
+	TSharedPtr<FText> TextPtr;
+};
+
+////////////////////////////////////////////////////////////////////////////////////////////////////
diff --git a/Engine/Source/Developer/TraceInsights/Private/Insights/Table/ViewModels/TableCellValueFormatter.h b/Engine/Source/Developer/TraceInsights/Private/Insights/Table/ViewModels/TableCellValueFormatter.h
new file mode 100644
--- /dev/null
+++ b/Engine/Source/Developer/TraceInsights/Private/Insights/Table/ViewModels/TableCellValueFormatter.h
@@ -0,0 +1,149 @@
+// Copyright 1998-2019 Epic Games, Inc. All Rights Reserved.
+
+#pragma once
+
+#include "CoreMinimal.h"
+
+#include "Insights/Table/ViewModels/TableCellValue.h"
+
+namespace Insights
+{
+
+class FBaseTreeNode;
+class FTableColumn;
+
+////////////////////////////////////////////////////////////////////////////////////////////////////
+
+class ITableCellValueFormatter
+{
+public:
+	virtual FText FormatValue(const TOptional<FTableCellValue>& InValue) const = 0;
+	virtual FText FormatValueForTooltip(const TOptional<FTableCellValue>& InValue) const = 0;
+
+	virtual FText FormatValue(const FTableColumn& Column, const FBaseTreeNode& Node) const = 0;
+	virtual FText FormatValueForTooltip(const FTableColumn& Column, const FBaseTreeNode& Node) const = 0;
+};
+
+////////////////////////////////////////////////////////////////////////////////////////////////////
+
+class FTableCellValueFormatter : public ITableCellValueFormatter
+{
+public:
+	FTableCellValueFormatter() {}
+	virtual ~FTableCellValueFormatter() {}
+
+	virtual FText FormatValue(const TOptional<FTableCellValue>& InValue) const override { return FText::GetEmpty(); }
+	virtual FText FormatValueForTooltip(const TOptional<FTableCellValue>& InValue) const override { return FormatValue(InValue); }
+
+	virtual FText FormatValue(const FTableColumn& Column, const FBaseTreeNode& Node) const override; // { return FormatValue(Column.GetValue(Node)); }
+	virtual FText FormatValueForTooltip(const FTableColumn& Column, const FBaseTreeNode& Node) const override { return FormatValue(Column, Node); }
+};
+
+////////////////////////////////////////////////////////////////////////////////////////////////////
+
+class FTextValueFormatter : public FTableCellValueFormatter
+{
+public:
+	virtual FText FormatValue(const TOptional<FTableCellValue>& InValue) const override
+	{
+		if (InValue.IsSet())
+		{
+			if (InValue.GetValue().TextPtr.IsValid())
+			{
+				return *InValue.GetValue().TextPtr;
+			}
+		}
+		return FText::GetEmpty();
+	}
+};
+
+////////////////////////////////////////////////////////////////////////////////////////////////////
+
+class FBoolValueFormatterAsTrueFalse : public FTableCellValueFormatter
+{
+public:
+	virtual FText FormatValue(const TOptional<FTableCellValue>& InValue) const override;
+};
+
+////////////////////////////////////////////////////////////////////////////////////////////////////
+
+class FBoolValueFormatterAsOnOff : public FTableCellValueFormatter
+{
+public:
+	virtual FText FormatValue(const TOptional<FTableCellValue>& InValue) const override;
+};
+
+////////////////////////////////////////////////////////////////////////////////////////////////////
+
+class FInt64ValueFormatterAsNumber : public FTableCellValueFormatter
+{
+public:
+	virtual FText FormatValue(const TOptional<FTableCellValue>& InValue) const override
+	{
+		return InValue.IsSet() ? FText::AsNumber(InValue.GetValue().Int64) : FText::GetEmpty();
+	}
+};
+
+////////////////////////////////////////////////////////////////////////////////////////////////////
+
+class FInt64ValueFormatterAsMemory : public FTableCellValueFormatter
+{
+public:
+	virtual FText FormatValue(const TOptional<FTableCellValue>& InValue) const override;
+	virtual FText FormatValueForTooltip(const TOptional<FTableCellValue>& InValue) const override;
+};
+
+////////////////////////////////////////////////////////////////////////////////////////////////////
+
+class FFloatValueFormatterAsNumber : public FTableCellValueFormatter
+{
+public:
+	virtual FText FormatValue(const TOptional<FTableCellValue>& InValue) const override;
+};
+
+////////////////////////////////////////////////////////////////////////////////////////////////////
+
+class FFloatValueFormatterAsTimeAuto : public FTableCellValueFormatter
+{
+public:
+	virtual FText FormatValue(const TOptional<FTableCellValue>& InValue) const override;
+	virtual FText FormatValueForTooltip(const TOptional<FTableCellValue>& InValue) const override;
+};
+
+////////////////////////////////////////////////////////////////////////////////////////////////////
+
+class FDoubleValueFormatterAsNumber : public FTableCellValueFormatter
+{
+public:
+	virtual FText FormatValue(const TOptional<FTableCellValue>& InValue) const override;
+};
+
+////////////////////////////////////////////////////////////////////////////////////////////////////
+
+class FDoubleValueFormatterAsTimeAuto : public FTableCellValueFormatter
+{
+public:
+	virtual FText FormatValue(const TOptional<FTableCellValue>& InValue) const override;
+	virtual FText FormatValueForTooltip(const TOptional<FTableCellValue>& InValue) const override;
+};
+
+////////////////////////////////////////////////////////////////////////////////////////////////////
+
+class FDoubleValueFormatterAsTimeMs : public FTableCellValueFormatter
+{
+public:
+	virtual FText FormatValue(const TOptional<FTableCellValue>& InValue) const override;
+	virtual FText FormatValueForTooltip(const TOptional<FTableCellValue>& InValue) const override;
+};
+
+////////////////////////////////////////////////////////////////////////////////////////////////////
+
+class FCStringValueFormatterAsText : public FTableCellValueFormatter
+{
+public:
+	virtual FText FormatValue(const TOptional<FTableCellValue>& InValue) const override;
+};
+
+////////////////////////////////////////////////////////////////////////////////////////////////////
+
+} // namespace Insights
diff --git a/Engine/Source/Developer/TraceInsights/Private/Insights/Table/ViewModels/TableCellValueGetter.cpp b/Engine/Source/Developer/TraceInsights/Private/Insights/Table/ViewModels/TableCellValueGetter.cpp
new file mode 100644
--- /dev/null
+++ b/Engine/Source/Developer/TraceInsights/Private/Insights/Table/ViewModels/TableCellValueGetter.cpp
@@ -0,0 +1,30 @@
+// Copyright 1998-2019 Epic Games, Inc. All Rights Reserved.
+
+#include "TableCellValueGetter.h"
+
+#include "Insights/Table/ViewModels/BaseTreeNode.h"
+
+#define LOCTEXT_NAMESPACE "TableCellValueGetter"
+
+namespace Insights
+{
+
+////////////////////////////////////////////////////////////////////////////////////////////////////
+
+const TOptional<FTableCellValue> FNameValueGetter::GetValue(const FTableColumn& Column, const FBaseTreeNode& Node) const
+{
+	return TOptional<FTableCellValue>(FTableCellValue(FText::FromName(Node.GetName())));
+}
+
+////////////////////////////////////////////////////////////////////////////////////////////////////
+
+const TOptional<FTableCellValue> FDisplayNameValueGetter::GetValue(const FTableColumn& Column, const FBaseTreeNode& Node) const
+{
+	return TOptional<FTableCellValue>(FTableCellValue(Node.GetDisplayName()));
+}
+
+////////////////////////////////////////////////////////////////////////////////////////////////////
+
+} // namespace Insights
+
+#undef LOCTEXT_NAMESPACE
diff --git a/Engine/Source/Developer/TraceInsights/Private/Insights/Table/ViewModels/TableCellValueGetter.h b/Engine/Source/Developer/TraceInsights/Private/Insights/Table/ViewModels/TableCellValueGetter.h
new file mode 100644
--- /dev/null
+++ b/Engine/Source/Developer/TraceInsights/Private/Insights/Table/ViewModels/TableCellValueGetter.h
@@ -0,0 +1,52 @@
+// Copyright 1998-2019 Epic Games, Inc. All Rights Reserved.
+
+#pragma once
+
+#include "CoreMinimal.h"
+
+#include "Insights/Table/ViewModels/TableCellValue.h"
+
+namespace Insights
+{
+
+class FBaseTreeNode;
+class FTableColumn;
+
+////////////////////////////////////////////////////////////////////////////////////////////////////
+
+class ITableCellValueGetter
+{
+public:
+	virtual const TOptional<FTableCellValue> GetValue(const FTableColumn& Column, const FBaseTreeNode& Node) const = 0;
+};
+
+////////////////////////////////////////////////////////////////////////////////////////////////////
+
+class FTableCellValueGetter : public ITableCellValueGetter
+{
+public:
+	FTableCellValueGetter() {}
+	virtual ~FTableCellValueGetter() {}
+
+	virtual const TOptional<FTableCellValue> GetValue(const FTableColumn& Column, const FBaseTreeNode& Node) const override { return TOptional<FTableCellValue>(); }
+};
+
+////////////////////////////////////////////////////////////////////////////////////////////////////
+
+class FNameValueGetter : public FTableCellValueGetter
+{
+public:
+	virtual const TOptional<FTableCellValue> GetValue(const FTableColumn& Column, const FBaseTreeNode& Node) const override;
+};
+
+////////////////////////////////////////////////////////////////////////////////////////////////////
+
+class FDisplayNameValueGetter : public FTableCellValueGetter
+{
+public:
+	virtual const TOptional<FTableCellValue> GetValue(const FTableColumn& Column, const FBaseTreeNode& Node) const override;
+};
+
+////////////////////////////////////////////////////////////////////////////////////////////////////
+
+} // namespace Insights
diff --git a/Engine/Source/Developer/TraceInsights/Private/Insights/Table/ViewModels/TableCellValueSorter.cpp b/Engine/Source/Developer/TraceInsights/Private/Insights/Table/ViewModels/TableCellValueSorter.cpp
new file mode 100644
--- /dev/null
+++ b/Engine/Source/Developer/TraceInsights/Private/Insights/Table/ViewModels/TableCellValueSorter.cpp
@@ -0,0 +1,678 @@
+// Copyright 1998-2019 Epic Games, Inc. All Rights Reserved.
+
+#include "TableCellValueSorter.h"
+
+#include "Insights/Table/ViewModels/BaseTreeNode.h"
+
+#define LOCTEXT_NAMESPACE "TableCellValueSorter"
+
+// Sort by name (ascending). -- too slow :(
+//#define INSIGHTS_DEFAULT_SORTING_NODES(A, B) return A->GetName().LexicalLess(B->GetName());
+//#define INSIGHTS_DEFAULT_SORTING_ELEMENTS(A, B) return A.NodePtr->GetName().LexicalLess(B.NodePtr->GetName());
+
+// Sort by node id (ascending).
+#define INSIGHTS_DEFAULT_SORTING_NODES(A, B) return A->GetId() < B->GetId();
+#define INSIGHTS_DEFAULT_SORTING_ELEMENTS(A, B) return A.NodePtr->GetId() < B.NodePtr->GetId();
+
+namespace Insights
+{
+
+////////////////////////////////////////////////////////////////////////////////////////////////////
+// FTableCellValueSorter
+////////////////////////////////////////////////////////////////////////////////////////////////////
+
+FTableCellValueSorter::FTableCellValueSorter(const FName InName, const FText& InShortName, const FText& InTitleName, const FText& InDescription, TSharedRef<FTableColumn> InColumnRef)
+	: Name(InName)
+	, ShortName(InShortName)
+	, TitleName(InTitleName)
+	, Description(InDescription)
+	, ColumnRef(InColumnRef)
+	, AscendingIcon(nullptr)
+	, DescendingIcon(nullptr)
+	, AscendingCompareDelegate(nullptr)
+	, DescendingCompareDelegate(nullptr)
+{
+}
+
+////////////////////////////////////////////////////////////////////////////////////////////////////
+
+void FTableCellValueSorter::Sort(TArray<FBaseTreeNodePtr>& NodesToSort, ESortMode SortMode) const
+{
+	if (SortMode == ESortMode::Ascending)
+	{
+		if (AscendingCompareDelegate)
+		{
+			NodesToSort.Sort(AscendingCompareDelegate);
+		}
+	}
+	else
+	{
+		if (DescendingCompareDelegate)
+		{
+			NodesToSort.Sort(DescendingCompareDelegate);
+		}
+	}
+}
+
+////////////////////////////////////////////////////////////////////////////////////////////////////
+// FBaseTableColumnSorter
+////////////////////////////////////////////////////////////////////////////////////////////////////
+
+FBaseTableColumnSorter::FBaseTableColumnSorter(TSharedRef<FTableColumn> InColumnRef)
+	: FTableCellValueSorter(
+		InColumnRef->GetId(),
+		FText::Format(LOCTEXT("Sorter_ColumnValue_Name", "By {0}"), InColumnRef->GetShortName()),
+		FText::Format(LOCTEXT("Sorter_ColumnValue_Title", "Sort By {0}"), InColumnRef->GetTitleName()),
+		FText::Format(LOCTEXT("Sorter_ColumnValue_Desc", "Sort by {0} (ascending or descending), then by id (ascending)."), InColumnRef->GetShortName()),
+		InColumnRef)
+{
+}
+
+////////////////////////////////////////////////////////////////////////////////////////////////////
+// FSorterByName
+////////////////////////////////////////////////////////////////////////////////////////////////////
+
+FSorterByName::FSorterByName(TSharedRef<FTableColumn> InColumnRef)
+	: FTableCellValueSorter(
+		FName(TEXT("ByName")),
+		LOCTEXT("Sorter_ByName_Name", "By Name"),
+		LOCTEXT("Sorter_ByName_Title", "Sort By Name"),
+		LOCTEXT("Sorter_ByName_Desc", "Sort by name."),
+		InColumnRef)
+{
+	AscendingCompareDelegate = [](const FBaseTreeNodePtr& A, const FBaseTreeNodePtr& B) -> bool
+	{
+		// Sort by name (ascending).
+		return A->GetName().LexicalLess(B->GetName());
+	};
+
+	DescendingCompareDelegate = [](const FBaseTreeNodePtr& A, const FBaseTreeNodePtr& B) -> bool
+	{
+		// Sort by name (descending).
+		return B->GetName().LexicalLess(A->GetName());
+	};
+}
+
+////////////////////////////////////////////////////////////////////////////////////////////////////
+// FSorterByTypeName
+////////////////////////////////////////////////////////////////////////////////////////////////////
+
+FSorterByTypeName::FSorterByTypeName(TSharedRef<FTableColumn> InColumnRef)
+	: FTableCellValueSorter(
+		FName(TEXT("ByTypeName")),
+		LOCTEXT("Sorter_ByTypeName_Name", "By Type Name"),
+		LOCTEXT("Sorter_ByTypeName_Title", "Sort By Type Name"),
+		LOCTEXT("Sorter_ByTypeName_Desc", "Sort by type name (ascending or descending), then by id (ascending)."),
+		InColumnRef)
+{
+	AscendingCompareDelegate = [](const FBaseTreeNodePtr& A, const FBaseTreeNodePtr& B) -> bool
+	{
+		const FName& TypeA = A->GetTypeName();
+		const FName& TypeB = B->GetTypeName();
+
+		if (TypeA == TypeB)
+		{
+			INSIGHTS_DEFAULT_SORTING_NODES(A, B)
+		}
+		else
+		{
+			// Sort by type name (ascending).
+			return TypeA.FastLess(TypeB); // TypeA < TypeB;
+		}
+	};
+
+	DescendingCompareDelegate = [](const FBaseTreeNodePtr& A, const FBaseTreeNodePtr& B) -> bool
+	{
+		const FName& TypeA = A->GetTypeName();
+		const FName& TypeB = B->GetTypeName();
+
+		if (TypeA == TypeB)
+		{
+			INSIGHTS_DEFAULT_SORTING_NODES(A, B)
+		}
+		else
+		{
+			// Sort by type name (descending).
+			return TypeB.FastLess(TypeA); // TypeA > TypeB;
+		}
+	};
+}
+
+////////////////////////////////////////////////////////////////////////////////////////////////////
+// FSorterByBoolValue
+////////////////////////////////////////////////////////////////////////////////////////////////////
+
+FSorterByBoolValue::FSorterByBoolValue(TSharedRef<FTableColumn> InColumnRef)
+	: FBaseTableColumnSorter(InColumnRef)
+{
+	const FTableColumn& Column = *InColumnRef;
+
+	ensure(Column.GetDataType() == ETableCellDataType::Bool);
+
+	AscendingCompareDelegate = [&Column](const FBaseTreeNodePtr& A, const FBaseTreeNodePtr& B) -> bool
+	{
+		const TOptional<FTableCellValue> OptionalValueA = Column.GetValue(*A);
+		const TOptional<FTableCellValue> OptionalValueB = Column.GetValue(*B);
+
+		const bool ValueA = OptionalValueA.IsSet() ? OptionalValueA.GetValue().Bool : false;
+		const bool ValueB = OptionalValueB.IsSet() ? OptionalValueB.GetValue().Bool : false;
+
+		if (ValueA == ValueB)
+		{
+			INSIGHTS_DEFAULT_SORTING_NODES(A, B)
+		}
+		else
+		{
+			// Sort by value (ascending).
+			return ValueA < ValueB;
+		}
+	};
+
+	DescendingCompareDelegate = [&Column](const FBaseTreeNodePtr& A, const FBaseTreeNodePtr& B) -> bool
+	{
+		const TOptional<FTableCellValue> OptionalValueA = Column.GetValue(*A);
+		const TOptional<FTableCellValue> OptionalValueB = Column.GetValue(*B);
+
+		const bool ValueA = OptionalValueA.IsSet() ? OptionalValueA.GetValue().Bool : false;
+		const bool ValueB = OptionalValueB.IsSet() ? OptionalValueB.GetValue().Bool : false;
+
+		if (ValueA == ValueB)
+		{
+			INSIGHTS_DEFAULT_SORTING_NODES(A, B)
+		}
+		else
+		{
+			// Sort by value (descending).
+			return ValueB < ValueA;
+		}
+	};
+}
+
+////////////////////////////////////////////////////////////////////////////////////////////////////
+// FSorterByInt64Value
+////////////////////////////////////////////////////////////////////////////////////////////////////
+
+FSorterByInt64Value::FSorterByInt64Value(TSharedRef<FTableColumn> InColumnRef)
+	: FBaseTableColumnSorter(InColumnRef)
+{
+	const FTableColumn& Column = *InColumnRef;
+
+	ensure(Column.GetDataType() == ETableCellDataType::Int64);
+
+	AscendingCompareDelegate = [&Column](const FBaseTreeNodePtr& A, const FBaseTreeNodePtr& B) -> bool
+	{
+		const TOptional<FTableCellValue> OptionalValueA = Column.GetValue(*A);
+		const TOptional<FTableCellValue> OptionalValueB = Column.GetValue(*B);
+
+		const int64 ValueA = OptionalValueA.IsSet() ? OptionalValueA.GetValue().Int64 : 0;
+		const int64 ValueB = OptionalValueB.IsSet() ? OptionalValueB.GetValue().Int64 : 0;
+
+		if (ValueA == ValueB)
+		{
+			INSIGHTS_DEFAULT_SORTING_NODES(A, B)
+		}
+		else
+		{
+			// Sort by value (ascending).
+			return ValueA < ValueB;
+		}
+	};
+
+	DescendingCompareDelegate = [&Column](const FBaseTreeNodePtr& A, const FBaseTreeNodePtr& B) -> bool
+	{
+		const TOptional<FTableCellValue> OptionalValueA = Column.GetValue(*A);
+		const TOptional<FTableCellValue> OptionalValueB = Column.GetValue(*B);
+
+		const int64 ValueA = OptionalValueA.IsSet() ? OptionalValueA.GetValue().Int64 : 0;
+		const int64 ValueB = OptionalValueB.IsSet() ? OptionalValueB.GetValue().Int64 : 0;
+
+		if (ValueA == ValueB)
+		{
+			INSIGHTS_DEFAULT_SORTING_NODES(A, B)
+		}
+		else
+		{
+			// Sort by value (descending).
+			return ValueB < ValueA;
+		}
+	};
+}
+
+////////////////////////////////////////////////////////////////////////////////////////////////////
+
+void FSorterByInt64Value::Sort(TArray<FBaseTreeNodePtr>& NodesToSort, ESortMode SortMode) const
+{
+	// Note: This implementation is faster than using above delegates.
+	//       It caches the values before sorting, in this way minimizes the calls to Column.GetValue().
+
+	const FTableColumn& Column = *ColumnRef;
+
+	const int32 NumElements = NodesToSort.Num();
+
+	struct FSortElement
+	{
+		const FBaseTreeNodePtr NodePtr;
+		int64 Value;
+	};
+	TArray<FSortElement> ElementsToSort;
+
+	// Cache value for each node.
+	ElementsToSort.Reset(NumElements);
+	for (const FBaseTreeNodePtr& NodePtr : NodesToSort)
+	{
+		const TOptional<FTableCellValue> OptionalValue = Column.GetValue(*NodePtr);
+		const int64 Value = OptionalValue.IsSet() ? OptionalValue.GetValue().Int64 : 0;
+		ElementsToSort.Add({ NodePtr, Value });
+	}
+	ensure(ElementsToSort.Num() == NumElements);
+
+	if (SortMode == ESortMode::Ascending)
+	{
+		ElementsToSort.Sort([](const FSortElement& A, const FSortElement& B) -> bool
+		{
+			if (A.Value == B.Value)
+			{
+				INSIGHTS_DEFAULT_SORTING_ELEMENTS(A, B)
+			}
+			else
+			{
+				// Sort by value (ascending).
+				return A.Value < B.Value;
+			}
+		});
+	}
+	else // if (SortMode == ESortMode::Descending)
+	{
+		ElementsToSort.Sort([](const FSortElement& A, const FSortElement& B) -> bool
+		{
+			if (A.Value == B.Value)
+			{
+				INSIGHTS_DEFAULT_SORTING_ELEMENTS(A, B)
+			}
+			else
+			{
+				// Sort by value (descending).
+				return B.Value < A.Value;
+			}
+		});
+	}
+
+	for (int32 Index = 0; Index < NumElements; ++Index)
+	{
+		NodesToSort[Index] = ElementsToSort[Index].NodePtr;
+	}
+}
+
+////////////////////////////////////////////////////////////////////////////////////////////////////
+// FSorterByFloatValue
+////////////////////////////////////////////////////////////////////////////////////////////////////
+
+FSorterByFloatValue::FSorterByFloatValue(TSharedRef<FTableColumn> InColumnRef)
+	: FBaseTableColumnSorter(InColumnRef)
+{
+	const FTableColumn& Column = *InColumnRef;
+
+	ensure(Column.GetDataType() == ETableCellDataType::Float);
+
+	AscendingCompareDelegate = [&Column](const FBaseTreeNodePtr& A, const FBaseTreeNodePtr& B) -> bool
+	{
+		const TOptional<FTableCellValue> OptionalValueA = Column.GetValue(*A);
+		const TOptional<FTableCellValue> OptionalValueB = Column.GetValue(*B);
+
+		const float ValueA = OptionalValueA.IsSet() ? OptionalValueA.GetValue().Float : 0.0f;
+		const float ValueB = OptionalValueB.IsSet() ? OptionalValueB.GetValue().Float : 0.0f;
+
+		if (ValueA == ValueB)
+		{
+			INSIGHTS_DEFAULT_SORTING_NODES(A, B)
+		}
+		else
+		{
+			// Sort by value (ascending).
+			return ValueA < ValueB;
+		}
+	};
+
+	DescendingCompareDelegate = [&Column](const FBaseTreeNodePtr& A, const FBaseTreeNodePtr& B) -> bool
+	{
+		const TOptional<FTableCellValue> OptionalValueA = Column.GetValue(*A);
+		const TOptional<FTableCellValue> OptionalValueB = Column.GetValue(*B);
+
+		const float ValueA = OptionalValueA.IsSet() ? OptionalValueA.GetValue().Float : 0.0f;
+		const float ValueB = OptionalValueB.IsSet() ? OptionalValueB.GetValue().Float : 0.0f;
+
+		if (ValueA == ValueB)
+		{
+			INSIGHTS_DEFAULT_SORTING_NODES(A, B)
+		}
+		else
+		{
+			// Sort by value (descending).
+			return ValueB < ValueA;
+		}
+	};
+}
+
+////////////////////////////////////////////////////////////////////////////////////////////////////
+
+void FSorterByFloatValue::Sort(TArray<FBaseTreeNodePtr>& NodesToSort, ESortMode SortMode) const
+{
+	// Note: This implementation is faster than using above delegates.
+	//       It caches the values before sorting, in this way minimizes the calls to Column.GetValue().
+
+	const FTableColumn& Column = *ColumnRef;
+
+	const int32 NumElements = NodesToSort.Num();
+
+	struct FSortElement
+	{
+		const FBaseTreeNodePtr NodePtr;
+		float Value;
+	};
+	TArray<FSortElement> ElementsToSort;
+
+	// Cache value for each node.
+	ElementsToSort.Reset(NumElements);
+	for (const FBaseTreeNodePtr& NodePtr : NodesToSort)
+	{
+		const TOptional<FTableCellValue> OptionalValue = Column.GetValue(*NodePtr);
+		const float Value = OptionalValue.IsSet() ? OptionalValue.GetValue().Float : 0.0f;
+		ElementsToSort.Add({ NodePtr, Value });
+	}
+	ensure(ElementsToSort.Num() == NumElements);
+
+	if (SortMode == ESortMode::Ascending)
+	{
+		ElementsToSort.Sort([](const FSortElement& A, const FSortElement& B) -> bool
+		{
+			if (A.Value == B.Value)
+			{
+				INSIGHTS_DEFAULT_SORTING_ELEMENTS(A, B)
+			}
+			else
+			{
+				// Sort by value (ascending).
+				return A.Value < B.Value;
+			}
+		});
+	}
+	else // if (SortMode == ESortMode::Descending)
+	{
+		ElementsToSort.Sort([](const FSortElement& A, const FSortElement& B) -> bool
+		{
+			if (A.Value == B.Value)
+			{
+				INSIGHTS_DEFAULT_SORTING_ELEMENTS(A, B)
+			}
+			else
+			{
+				// Sort by value (descending).
+				return B.Value < A.Value;
+			}
+		});
+	}
+
+	for (int32 Index = 0; Index < NumElements; ++Index)
+	{
+		NodesToSort[Index] = ElementsToSort[Index].NodePtr;
+	}
+}
+
+////////////////////////////////////////////////////////////////////////////////////////////////////
+// FSorterByDoubleValue
+////////////////////////////////////////////////////////////////////////////////////////////////////
+
+FSorterByDoubleValue::FSorterByDoubleValue(TSharedRef<FTableColumn> InColumnRef)
+	: FBaseTableColumnSorter(InColumnRef)
+{
+	const FTableColumn& Column = *InColumnRef;
+
+	ensure(Column.GetDataType() == ETableCellDataType::Double);
+
+	AscendingCompareDelegate = [&Column](const FBaseTreeNodePtr& A, const FBaseTreeNodePtr& B) -> bool
+	{
+		const TOptional<FTableCellValue> OptionalValueA = Column.GetValue(*A);
+		const TOptional<FTableCellValue> OptionalValueB = Column.GetValue(*B);
+
+		const double ValueA = OptionalValueA.IsSet() ? OptionalValueA.GetValue().Double : 0.0;
+		const double ValueB = OptionalValueB.IsSet() ? OptionalValueB.GetValue().Double : 0.0;
+
+		if (ValueA == ValueB)
+		{
+			INSIGHTS_DEFAULT_SORTING_NODES(A, B)
+		}
+		else
+		{
+			// Sort by value (ascending).
+			return ValueA < ValueB;
+		}
+	};
+
+	DescendingCompareDelegate = [&Column](const FBaseTreeNodePtr& A, const FBaseTreeNodePtr& B) -> bool
+	{
+		const TOptional<FTableCellValue> OptionalValueA = Column.GetValue(*A);
+		const TOptional<FTableCellValue> OptionalValueB = Column.GetValue(*B);
+
+		const double ValueA = OptionalValueA.IsSet() ? OptionalValueA.GetValue().Double : 0.0;
+		const double ValueB = OptionalValueB.IsSet() ? OptionalValueB.GetValue().Double : 0.0;
+
+		if (ValueA == ValueB)
+		{
+			INSIGHTS_DEFAULT_SORTING_NODES(A, B)
+		}
+		else
+		{
+			// Sort by value (descending).
+			return ValueB < ValueA;
+		}
+	};
+}
+
+////////////////////////////////////////////////////////////////////////////////////////////////////
+
+void FSorterByDoubleValue::Sort(TArray<FBaseTreeNodePtr>& NodesToSort, ESortMode SortMode) const
+{
+	// Note: This implementation is faster than using above delegates.
+	//       It caches the values before sorting, in this way minimizes the calls to Column.GetValue().
+
+	const FTableColumn& Column = *ColumnRef;
+
+	const int32 NumElements = NodesToSort.Num();
+
+	struct FSortElement
+	{
+		const FBaseTreeNodePtr NodePtr;
+		double Value;
+	};
+	TArray<FSortElement> ElementsToSort;
+
+	// Cache value for each node.
+	ElementsToSort.Reset(NumElements);
+	for (const FBaseTreeNodePtr& NodePtr : NodesToSort)
+	{
+		const TOptional<FTableCellValue> OptionalValue = Column.GetValue(*NodePtr);
+		const double Value = OptionalValue.IsSet() ? OptionalValue.GetValue().Double : 0.0;
+		ElementsToSort.Add({ NodePtr, Value });
+	}
+	ensure(ElementsToSort.Num() == NumElements);
+
+	if (SortMode == ESortMode::Ascending)
+	{
+		ElementsToSort.Sort([](const FSortElement& A, const FSortElement& B) -> bool
+		{
+			if (A.Value == B.Value)
+			{
+				INSIGHTS_DEFAULT_SORTING_ELEMENTS(A, B)
+			}
+			else
+			{
+				// Sort by value (ascending).
+				return A.Value < B.Value;
+			}
+		});
+	}
+	else // if (SortMode == ESortMode::Descending)
+	{
+		ElementsToSort.Sort([](const FSortElement& A, const FSortElement& B) -> bool
+		{
+			if (A.Value == B.Value)
+			{
+				INSIGHTS_DEFAULT_SORTING_ELEMENTS(A, B)
+			}
+			else
+			{
+				// Sort by value (descending).
+				return B.Value < A.Value;
+			}
+		});
+	}
+
+	for (int32 Index = 0; Index < NumElements; ++Index)
+	{
+		NodesToSort[Index] = ElementsToSort[Index].NodePtr;
+	}
+}
+
+////////////////////////////////////////////////////////////////////////////////////////////////////
+// FSorterByCStringValue
+////////////////////////////////////////////////////////////////////////////////////////////////////
+
+FSorterByCStringValue::FSorterByCStringValue(TSharedRef<FTableColumn> InColumnRef)
+	: FBaseTableColumnSorter(InColumnRef)
+{
+	const FTableColumn& Column = *InColumnRef;
+
+	ensure(Column.GetDataType() == ETableCellDataType::CString);
+
+	AscendingCompareDelegate = [&Column](const FBaseTreeNodePtr& A, const FBaseTreeNodePtr& B) -> bool
+	{
+		const TOptional<FTableCellValue> OptionalValueA = Column.GetValue(*A);
+		const TOptional<FTableCellValue> OptionalValueB = Column.GetValue(*B);
+
+		const TCHAR* ValueA = OptionalValueA.IsSet() ? OptionalValueA.GetValue().CString : nullptr;
+		const TCHAR* ValueB = OptionalValueB.IsSet() ? OptionalValueB.GetValue().CString : nullptr;
+
+		if (ValueA == ValueB)
+		{
+			// Sort by name (ascending).
+			//return A->GetName().LexicalLess(B->GetName()); // so slow :(
+
+			// Sort by node id (ascending).
+			return A->GetId() < B->GetId();
+		}
+		else
+		{
+			// Sort by value (ascending).
+			return FCString::Strcmp(ValueA, ValueB) < 0;
+		}
+	};
+
+	DescendingCompareDelegate = [&Column](const FBaseTreeNodePtr& A, const FBaseTreeNodePtr& B) -> bool
+	{
+		const TOptional<FTableCellValue> OptionalValueA = Column.GetValue(*A);
+		const TOptional<FTableCellValue> OptionalValueB = Column.GetValue(*B);
+
+		const TCHAR* ValueA = OptionalValueA.IsSet() ? OptionalValueA.GetValue().CString : nullptr;
+		const TCHAR* ValueB = OptionalValueB.IsSet() ? OptionalValueB.GetValue().CString : nullptr;
+
+		if (ValueA == ValueB)
+		{
+			// Sort by name (ascending).
+			//return A->GetName().LexicalLess(B->GetName()); // so slow :(
+
+			// Sort by node id (ascending).
+			return A->GetId() < B->GetId();
+		}
+		else
+		{
+			// Sort by value (descending).
+			return FCString::Strcmp(ValueB, ValueA) < 0;
+		}
+	};
+}
+
+////////////////////////////////////////////////////////////////////////////////////////////////////
+// FSorterByTextValue
+////////////////////////////////////////////////////////////////////////////////////////////////////
+
+FSorterByTextValue::FSorterByTextValue(TSharedRef<FTableColumn> InColumnRef)
+	: FBaseTableColumnSorter(InColumnRef)
+{
+	const FTableColumn& Column = *InColumnRef;
+
+	ensure(Column.GetDataType() == ETableCellDataType::Text);
+
+	AscendingCompareDelegate = [&Column](const FBaseTreeNodePtr& A, const FBaseTreeNodePtr& B) -> bool
+	{
+		const TOptional<FTableCellValue> OptionalValueA = Column.GetValue(*A);
+		const TOptional<FTableCellValue> OptionalValueB = Column.GetValue(*B);
+
+		if (!OptionalValueB.IsSet() || !OptionalValueB.GetValue().TextPtr.IsValid())
+		{
+			return true;
+		}
+
+		if (!OptionalValueA.IsSet() || !OptionalValueA.GetValue().TextPtr.IsValid())
+		{
+			return false;
+		}
+
+		const FText& ValueA = *OptionalValueA.GetValue().TextPtr;
+		const FText& ValueB = *OptionalValueB.GetValue().TextPtr;
+
+		if (ValueA.EqualTo(ValueB))
+		{
+			// Sort by name (ascending).
+			//return A->GetName().LexicalLess(B->GetName()); // so slow :(
+
+			// Sort by node id (ascending).
+			return A->GetId() < B->GetId();
+		}
+		else
+		{
+			// Sort by value (ascending).
+			return ValueA.CompareTo(ValueB) < 0;
+		}
+	};
+
+	DescendingCompareDelegate = [&Column](const FBaseTreeNodePtr& A, const FBaseTreeNodePtr& B) -> bool
+	{
+		const TOptional<FTableCellValue> OptionalValueA = Column.GetValue(*A);
+		const TOptional<FTableCellValue> OptionalValueB = Column.GetValue(*B);
+
+		if (!OptionalValueB.IsSet() || !OptionalValueB.GetValue().TextPtr.IsValid())
+		{
+			return true;
+		}
+
+		if (!OptionalValueA.IsSet() || !OptionalValueA.GetValue().TextPtr.IsValid())
+		{
+			return false;
+		}
+
+		const FText& ValueA = *OptionalValueA.GetValue().TextPtr;
+		const FText& ValueB = *OptionalValueB.GetValue().TextPtr;
+
+		if (ValueA.EqualTo(ValueB))
+		{
+			// Sort by name (ascending).
+			//return A->GetName().LexicalLess(B->GetName()); // so slow :(
+
+			// Sort by node id (ascending).
+			return A->GetId() < B->GetId();
+		}
+		else
+		{
+			// Sort by value (descending).
+			return ValueB.CompareTo(ValueA) < 0;
+		}
+	};
+}
+
+////////////////////////////////////////////////////////////////////////////////////////////////////
+
+} // namespace Insights
+
+#undef INSIGHTS_DEFAULT_SORTING_NODES
+#undef INSIGHTS_DEFAULT_SORTING_ELEMENTS
+#undef LOCTEXT_NAMESPACE
diff --git a/Engine/Source/Developer/TraceInsights/Private/Insights/Table/ViewModels/TableColumn.h b/Engine/Source/Developer/TraceInsights/Private/Insights/Table/ViewModels/TableColumn.h
--- a/Engine/Source/Developer/TraceInsights/Private/Insights/Table/ViewModels/TableColumn.h
+++ b/Engine/Source/Developer/TraceInsights/Private/Insights/Table/ViewModels/TableColumn.h
@@ -6,12 +6,16 @@
 #include "Types/SlateEnums.h"
 #include "Misc/EnumClassFlags.h"
 
+#include "Insights/Table/ViewModels/TableCellValue.h"
+
 namespace Insights
 {
 
+class FBaseTreeNode;
 class FTable;
-class FTableColumn;
-struct FTableRowId;
+class ITableCellValueGetter;
+class ITableCellValueFormatter;
+class ITableCellValueSorter;
 
 ////////////////////////////////////////////////////////////////////////////////////////////////////
 
@@ -21,66 +25,13 @@ enum class ETableColumnFlags : uint32
 
 	ShouldBeVisible = (1 << 0),
 	CanBeHidden     = (1 << 1),
-	CanBeSorted     = (1 << 2),
-	CanBeFiltered   = (1 << 3),
-	IsHierarchy     = (1 << 4),
+	CanBeFiltered   = (1 << 2),
+	IsHierarchy     = (1 << 3),
 };
 ENUM_CLASS_FLAGS(ETableColumnFlags);
 
 ////////////////////////////////////////////////////////////////////////////////////////////////////
 
-enum class ETableCellDataType : uint32
-{
-	Unknown,
-
-	Bool,
-	Int64,
-	Float,
-	Double,
-	CString,
-
-	/** Invalid enum type, may be used as a number of enumerations. */
-	InvalidOrMax,
-};
-
-////////////////////////////////////////////////////////////////////////////////////////////////////
-
-struct FTableCellValue
-{
-	FTableCellValue() = default;
-	FTableCellValue(bool Value) : Bool(Value) {}
-	FTableCellValue(int64 Value) : Int64(Value) {}
-	FTableCellValue(float Value) : Float(Value) {}
-	FTableCellValue(double Value) : Double(Value) {}
-	FTableCellValue(const TCHAR* Value) : CString(Value) {}
-
-	union
-	{
-		bool Bool;
-		int64 Int64;
-		float Float;
-		double Double;
-		const TCHAR* CString;
-	};
-};
-
-////////////////////////////////////////////////////////////////////////////////////////////////////
-
-class FTableCellValueFormatter
-{
-public:
-	FTableCellValueFormatter() {}
-	virtual ~FTableCellValueFormatter() {}
-
-	virtual FText FormatValue(const TOptional<FTableCellValue>& InValue) { return FText::GetEmpty(); }
-	virtual FText FormatValueForTooltip(const TOptional<FTableCellValue>& InValue) { return FormatValue(InValue); }
-
-	virtual FText FormatValue(const FTable& Table, const FTableColumn& Column, const FTableRowId& RowId) { return FText::GetEmpty(); }
-	virtual FText FormatValueForTooltip(const FTable& Table, const FTableColumn& Column, const FTableRowId& RowId) { return FormatValue(Table, Column, RowId); }
-};
-
-////////////////////////////////////////////////////////////////////////////////////////////////////
-
 enum class ETableColumnAggregation : uint32
 {
 	None = 0,
@@ -98,64 +49,55 @@ enum class ETableColumnAggregation : uint32
  */
 class FTableColumn
 {
-public:
-	typedef TFunction<FText(const FTableCellValue& Value)> FGetCellValueAsTextFunction;
-	typedef TFunction<FText(const FTable& Table, const FTableColumn& Column, const FTableRowId& RowId)> FGetValueAsTextFunction;
-
 public:
 	/** No default constructor. */
 	FTableColumn() = delete;
 
 	/** Initialization constructor. */
-	FTableColumn
-	(
-		int32 InOrder,
-		int32 InIndex,
-		const FName InId,
-		FText InShortName,
-		FText InTitleName,
-		FText InDescription,
-		const ETableColumnFlags InFlags,
-		const ETableCellDataType InDataType,
-		const ETableColumnAggregation InAggregation,
-		const EHorizontalAlignment InHorizontalAlignment,
-		const float InInitialWidth,
-		const float InMinWidth,
-		const float InMaxWidth,
-		TSharedRef<FTableCellValueFormatter> InValueFormatter
-	)
-		: Order(InOrder)
-		, Index(InIndex)
-		, Id(InId)
-		, ShortName(MoveTemp(InShortName))
-		, TitleName(MoveTemp(InTitleName))
-		, Description(MoveTemp(InDescription))
+	FTableColumn(const FName InId)
+		: Id(InId)
+		, Order(0)
+		, Index(-1)
+		, ShortName()
+		, TitleName()
+		, Description()
 		, bIsVisible(false)
-		, Flags(InFlags)
-		, DataType(InDataType)
-		, Aggregation(InAggregation)
-		, HorizontalAlignment(InHorizontalAlignment)
-		, InitialWidth(InInitialWidth)
-		, MinWidth(InMinWidth)
-		, MaxWidth(InMaxWidth)
-		, ValueFormatter(InValueFormatter)
+		, Flags(ETableColumnFlags::None)
+		, HorizontalAlignment(HAlign_Left)
+		, InitialWidth(60.0f)
+		, MinWidth(0.0f)
+		, MaxWidth(FLT_MAX)
+		, DataType(ETableCellDataType::Unknown)
+		, Aggregation(ETableColumnAggregation::None)
+		, ValueGetter(GetDefaultValueGetter())
+		, ValueFormatter(GetDefaultValueFormatter())
+		, ValueSorter(nullptr)
 		, ParentTable(nullptr)
 	{
 	}
 
-	int32 GetOrder() const { return Order; }
-	int32 GetIndex() const { return Index; }
+	//////////////////////////////////////////////////
 
 	const FName& GetId() const { return Id; }
 
-	const FText& GetShortName() const { return ShortName; }
-	const FText& GetTitleName() const { return TitleName; }
-	const FText& GetDescription() const { return Description; }
+	int32 GetIndex() const { return Index; }
+	void SetIndex(int32 InIndex) { Index = InIndex; }
+
+	//////////////////////////////////////////////////
+	// Name and Description
 
+	const FText& GetShortName() const { return ShortName; }
 	void SetShortName(const FText& InShortName) { ShortName = InShortName; }
+
+	const FText& GetTitleName() const { return TitleName; }
 	void SetTitleName(const FText& InTitleName) { TitleName = InTitleName; }
+
+	const FText& GetDescription() const { return Description; }
 	void SetDescription(const FText& InDescription) { Description = InDescription; }
 
+	//////////////////////////////////////////////////
+	// Visibility and other Flags
+
 	bool IsVisible() const { return bIsVisible; }
 	void Show() { bIsVisible = true; OnVisibilityChanged(); }
 	void Hide() { bIsVisible = false; OnVisibilityChanged(); }
@@ -169,53 +111,92 @@ public:
 	/** Whether this column can be hidden. */
 	bool CanBeHidden() const { return EnumHasAnyFlags(Flags, ETableColumnFlags::CanBeHidden); }
 
-	/** Whether this column can be used for sorting. */
-	bool CanBeSorted() const { return EnumHasAnyFlags(Flags, ETableColumnFlags::CanBeSorted); }
-
 	/** Whether this column can be used for filtering displayed results. */
 	bool CanBeFiltered() const { return EnumHasAnyFlags(Flags, ETableColumnFlags::CanBeFiltered); }
 
 	/** Whether this column is the hierarcy (name) column, in a tree view. */
 	bool IsHierarchy() const { return EnumHasAnyFlags(Flags, ETableColumnFlags::IsHierarchy); }
 
-	ETableCellDataType GetDataType() const { return DataType; }
-	ETableColumnAggregation GetAggregation() const { return Aggregation; }
+	void SetFlags(ETableColumnFlags InFlags) { Flags = InFlags; }
+
+	//////////////////////////////////////////////////
+	// Alignment and Width
 
 	EHorizontalAlignment GetHorizontalAlignment() const { return HorizontalAlignment; }
+	void SetHorizontalAlignment(EHorizontalAlignment InHorizontalAlignment) { HorizontalAlignment = InHorizontalAlignment; }
 
 	float GetInitialWidth() const { return InitialWidth; }
+	void SetInitialWidth(float InWidth) { InitialWidth = InWidth; }
+
 	float GetMinWidth() const { return MinWidth; }
+	void SetMinWidth(float InWidth) { MinWidth = InWidth; }
+
 	float GetMaxWidth() const { return MaxWidth; }
+	void SetMaxWidth(float InWidth) { MaxWidth = InWidth; }
+
 	//float GetWidth() const { return Width; }
+	//void SetWidth(float InWidth) { Width = InWidth; }
 
 	/** If MinWidth == MaxWidth, this column has fixed width and cannot be resized. */
 	bool IsFixedWidth() const { return MinWidth == MaxWidth; }
 
-	TSharedRef<FTableCellValueFormatter> GetValueFormatter() const { return ValueFormatter; }
+	//////////////////////////////////////////////////
+	// Data Type
 
-	FText GetValueAsText(const FTableRowId& InRowId) const
-	{
-		return ValueFormatter->FormatValue(*ParentTable.Pin(), *this, InRowId);
-	}
+	ETableCellDataType GetDataType() const { return DataType; }
+	void SetDataType(ETableCellDataType InDataType) { DataType = InDataType; }
 
-	FText GetValueAsTooltipText(const FTableRowId& InRowId) const
-	{
-		return ValueFormatter->FormatValueForTooltip(*ParentTable.Pin(), *this, InRowId);
-	}
+	//////////////////////////////////////////////////
+	// Aggregation
+
+	ETableColumnAggregation GetAggregation() const { return Aggregation; }
+	void SetAggregation(ETableColumnAggregation InAggregation) { Aggregation = InAggregation; }
+
+	//////////////////////////////////////////////////
+	// Value Getter
+
+	static TSharedRef<ITableCellValueGetter> GetDefaultValueGetter();
+
+	TSharedRef<ITableCellValueGetter> GetValueGetter() const { return ValueGetter; }
+	void SetValueGetter(TSharedRef<ITableCellValueGetter> InValueGetter) { ValueGetter = InValueGetter; }
+
+	const TOptional<FTableCellValue> GetValue(const FBaseTreeNode& InNode) const;
+
+	//////////////////////////////////////////////////
+	// Value Formatter
+
+	static TSharedRef<ITableCellValueFormatter> GetDefaultValueFormatter();
+
+	TSharedRef<ITableCellValueFormatter> GetValueFormatter() const { return ValueFormatter; }
+	void SetValueFormatter(TSharedRef<ITableCellValueFormatter> InValueFormatter) { ValueFormatter = InValueFormatter; }
+
+	FText GetValueAsText(const FBaseTreeNode& InNode) const;
+	FText GetValueAsTooltipText(const FBaseTreeNode& InNode) const;
+
+	//////////////////////////////////////////////////
+	// Value Sorter (can be nullptr)
+
+	TSharedPtr<ITableCellValueSorter> GetValueSorter() const { return ValueSorter; }
+	void SetValueSorter(TSharedPtr<ITableCellValueSorter> InValueSorter) { ValueSorter = InValueSorter; }
+
+	/** Whether this column can be used for sorting. */
+	bool CanBeSorted() const { return ValueSorter.IsValid(); }
+
+	//////////////////////////////////////////////////
 
 	TWeakPtr<FTable> GetParentTable() const { return ParentTable; }
 	void SetParentTable(TWeakPtr<FTable> InParentTable) { ParentTable = InParentTable; }
 
 private:
+	/** Id of the column. */
+	FName Id;
+
 	/** Order value, to sort columns in the list/tree view. */
 	int32 Order;
 
 	/** Column index in source table. */
 	int32 Index;
 
-	/** Id of the column. */
-	FName Id;
-
 	/** Short name of the column, displayed in the column header. */
 	FText ShortName;
 
@@ -231,12 +212,6 @@ private:
 	/** Other on/off switches. */
 	ETableColumnFlags Flags;
 
-	/** Data type for this column. */
-	ETableCellDataType DataType;
-
-	/** Aggregation for values in this column, when grouped. */
-	ETableColumnAggregation Aggregation; // TODO: make this an object (Aggregator)
-
 	/** Horizontal alignment of the content in this column. */
 	EHorizontalAlignment HorizontalAlignment;
 
@@ -245,8 +220,20 @@ private:
 	float MaxWidth; /**< Maximum column width. */
 	//float Width; /**< Current column width. */
 
-	/** Custom formater for the value displayed by this column. */
-	TSharedRef<FTableCellValueFormatter> ValueFormatter;
+	/** Data type for this column. */
+	ETableCellDataType DataType;
+
+	/** Aggregation for values in this column, when grouped. */
+	ETableColumnAggregation Aggregation; // TODO: make this an object (Aggregator)
+
+	/** Custom getter for values identified by this column. */
+	TSharedRef<ITableCellValueGetter> ValueGetter;
+
+	/** Custom formater for values displayed by this column. */
+	TSharedRef<ITableCellValueFormatter> ValueFormatter;
+
+	/** Custom sorter for values displayed by this column. */
+	TSharedPtr<ITableCellValueSorter> ValueSorter;
 
 	/* Parent table. Only one table instance can own this column. */
 	TWeakPtr<FTable> ParentTable;
diff --git a/Engine/Source/Developer/TraceInsights/Private/Insights/Table/ViewModels/TableTreeNode.cpp b/Engine/Source/Developer/TraceInsights/Private/Insights/Table/ViewModels/TableTreeNode.cpp
--- a/Engine/Source/Developer/TraceInsights/Private/Insights/Table/ViewModels/TableTreeNode.cpp
+++ b/Engine/Source/Developer/TraceInsights/Private/Insights/Table/ViewModels/TableTreeNode.cpp
@@ -2,184 +2,12 @@
 
 #include "TableTreeNode.h"
 
-// Insights
-#include "Insights/Table/ViewModels/TableColumn.h"
-
 #define LOCTEXT_NAMESPACE "Insights_TableTreeNode"
 
 namespace Insights
 {
 
-////////////////////////////////////////////////////////////////////////////////////////////////////
-
-FTableCellValue FTableTreeNode::GetValue(const FTableColumn& Column) const
-{
-	switch (Column.GetDataType())
-	{
-	case ETableCellDataType::Bool:    return FTableCellValue(GetValueBool(Column));
-	case ETableCellDataType::Int64:   return FTableCellValue(GetValueInt64(Column));
-	case ETableCellDataType::Float:   return FTableCellValue(GetValueFloat(Column));
-	case ETableCellDataType::Double:  return FTableCellValue(GetValueDouble(Column));
-	case ETableCellDataType::CString: return FTableCellValue(GetValueCString(Column));
-	default:                          return FTableCellValue();
-	}
-}
-
-////////////////////////////////////////////////////////////////////////////////////////////////////
-
-bool FTableTreeNode::GetValueBool(const FTableColumn& Column) const
-{
-	bool Value = false;
-
-	if (!IsGroup())
-	{
-		TSharedPtr<FTable> Table = ParentTable.Pin();
-		if (Table.IsValid())
-		{
-			TSharedPtr<Trace::IUntypedTableReader> Reader = Table->GetTableReader();
-			if (Reader.IsValid() && RowId.HasValidIndex())
-			{
-				Reader->SetRowIndex(RowId.RowIndex);
-				Value = Reader->GetValueBool(Column.GetIndex());
-			}
-		}
-	}
-	else
-	{
-		const FTableCellValue* ValuePtr = AggregatedValues.Find(Column.GetId());
-		if (ValuePtr != nullptr)
-		{
-			Value = ValuePtr->Bool;
-		}
-	}
-
-	return Value;
-}
-
-////////////////////////////////////////////////////////////////////////////////////////////////////
-
-int64 FTableTreeNode::GetValueInt64(const FTableColumn& Column) const
-{
-	int64 Value = 0;
-
-	if (!IsGroup())
-	{
-		TSharedPtr<FTable> Table = ParentTable.Pin();
-		if (Table.IsValid())
-		{
-			TSharedPtr<Trace::IUntypedTableReader> Reader = Table->GetTableReader();
-			if (Reader.IsValid() && RowId.HasValidIndex())
-			{
-				Reader->SetRowIndex(RowId.RowIndex);
-				Value = Reader->GetValueInt(Column.GetIndex());
-			}
-		}
-	}
-	else
-	{
-		const FTableCellValue* ValuePtr = AggregatedValues.Find(Column.GetId());
-		if (ValuePtr != nullptr)
-		{
-			Value = ValuePtr->Int64;
-		}
-	}
-
-	return Value;
-}
-
-////////////////////////////////////////////////////////////////////////////////////////////////////
-
-float FTableTreeNode::GetValueFloat(const FTableColumn& Column) const
-{
-	float Value = 0.0f;
-
-	if (!IsGroup())
-	{
-		TSharedPtr<FTable> Table = ParentTable.Pin();
-		if (Table.IsValid())
-		{
-			TSharedPtr<Trace::IUntypedTableReader> Reader = Table->GetTableReader();
-			if (Reader.IsValid() && RowId.HasValidIndex())
-			{
-				Reader->SetRowIndex(RowId.RowIndex);
-				Value = Reader->GetValueFloat(Column.GetIndex());
-			}
-		}
-	}
-	else
-	{
-		const FTableCellValue* ValuePtr = AggregatedValues.Find(Column.GetId());
-		if (ValuePtr != nullptr)
-		{
-			Value = ValuePtr->Float;
-		}
-	}
-
-	return Value;
-}
-
-////////////////////////////////////////////////////////////////////////////////////////////////////
-
-double FTableTreeNode::GetValueDouble(const FTableColumn& Column) const
-{
-	double Value = 0.0;
-
-	if (!IsGroup())
-	{
-		TSharedPtr<FTable> Table = ParentTable.Pin();
-		if (Table.IsValid())
-		{
-			TSharedPtr<Trace::IUntypedTableReader> Reader = Table->GetTableReader();
-			if (Reader.IsValid() && RowId.HasValidIndex())
-			{
-				Reader->SetRowIndex(RowId.RowIndex);
-				Value = Reader->GetValueDouble(Column.GetIndex());
-			}
-		}
-	}
-	{
-		const FTableCellValue* ValuePtr = AggregatedValues.Find(Column.GetId());
-		if (ValuePtr != nullptr)
-		{
-			Value = ValuePtr->Double;
-		}
-	}
-
-	return Value;
-}
-
-////////////////////////////////////////////////////////////////////////////////////////////////////
-
-const TCHAR* FTableTreeNode::GetValueCString(const FTableColumn& Column) const
-{
-	const TCHAR* Value = nullptr;
-
-	if (!IsGroup())
-	{
-		TSharedPtr<FTable> Table = ParentTable.Pin();
-		if (Table.IsValid())
-		{
-			TSharedPtr<Trace::IUntypedTableReader> Reader = Table->GetTableReader();
-			if (Reader.IsValid() && RowId.HasValidIndex())
-			{
-				Reader->SetRowIndex(RowId.RowIndex);
-				Value = Reader->GetValueCString(Column.GetIndex());
-			}
-		}
-	}
-	else
-	{
-		const FTableCellValue* ValuePtr = AggregatedValues.Find(Column.GetId());
-		if (ValuePtr != nullptr)
-		{
-			Value = ValuePtr->CString;
-		}
-	}
-
-	return Value;
-}
-
-////////////////////////////////////////////////////////////////////////////////////////////////////
+const FName FTableTreeNode::TypeName(TEXT("FTableTreeNode"));
 
 } // namespace Insights
 
diff --git a/Engine/Source/Developer/TraceInsights/Private/Insights/Table/ViewModels/TableTreeNode.h b/Engine/Source/Developer/TraceInsights/Private/Insights/Table/ViewModels/TableTreeNode.h
--- a/Engine/Source/Developer/TraceInsights/Private/Insights/Table/ViewModels/TableTreeNode.h
+++ b/Engine/Source/Developer/TraceInsights/Private/Insights/Table/ViewModels/TableTreeNode.h
@@ -7,12 +7,15 @@
 
 // Insights
 #include "Insights/Table/ViewModels/BaseTreeNode.h"
-#include "Insights/Table/ViewModels/Table.h"
-#include "Insights/Table/ViewModels/TableColumn.h"
+#include "Insights/Table/ViewModels/Table.h" // for FTableRowId
+#include "Insights/Table/ViewModels/TableCellValue.h"
 
 namespace Insights
 {
 
+//class FTable;
+//class FTableColumn;
+
 ////////////////////////////////////////////////////////////////////////////////////////////////////
 
 class FTableTreeNode;
@@ -36,6 +39,9 @@ typedef TWeakPtr<class FTableTreeNode> FTableTreeNodeWeak;
  */
 class FTableTreeNode : public FBaseTreeNode
 {
+public:
+	static const FName TypeName;
+
 public:
 	/** Initialization constructor for a table record node. */
 	FTableTreeNode(uint64 InId, const FName InName, TWeakPtr<FTable> InParentTable, int32 InRowIndex)
@@ -53,21 +59,16 @@ public:
 	{
 	}
 
+	virtual const FName& GetTypeName() const override { return TypeName; }
+
 	TWeakPtr<FTable> GetParentTable() { return ParentTable; }
 	FTableRowId GetRowId() const { return RowId; }
 	int32 GetRowIndex() const { return RowId.RowIndex; }
 
-	FTableCellValue GetValue(const FTableColumn& Column) const;
-
-	bool GetValueBool(const FTableColumn& Column) const;
-	int64 GetValueInt64(const FTableColumn& Column) const;
-	float GetValueFloat(const FTableColumn& Column) const;
-	double GetValueDouble(const FTableColumn& Column) const;
-	const TCHAR* GetValueCString(const FTableColumn& Column) const;
-
 	void ResetAggregatedValues() { AggregatedValues.Reset(); }
 	bool HasAggregatedValue(const FName& ColumnId) const { return AggregatedValues.Contains(ColumnId); }
-	const FTableCellValue& GetAggregatedValue(const FName& ColumnId) { return AggregatedValues.FindChecked(ColumnId); }
+	const FTableCellValue* FindAggregatedValue(const FName& ColumnId) const { return AggregatedValues.Find(ColumnId); }
+	const FTableCellValue& GetAggregatedValue(const FName& ColumnId) const { return AggregatedValues.FindChecked(ColumnId); }
 	void AddAggregatedValue(const FName& ColumnId, const FTableCellValue& Value) { AggregatedValues.Add(ColumnId, Value); }
 	void SetAggregatedValue(const FName& ColumnId, const FTableCellValue& Value) { AggregatedValues[ColumnId] = Value; }
 
diff --git a/Engine/Source/Developer/TraceInsights/Private/Insights/Table/ViewModels/TreeNodeGrouping.cpp b/Engine/Source/Developer/TraceInsights/Private/Insights/Table/ViewModels/TreeNodeGrouping.cpp
--- a/Engine/Source/Developer/TraceInsights/Private/Insights/Table/ViewModels/TreeNodeGrouping.cpp
+++ b/Engine/Source/Developer/TraceInsights/Private/Insights/Table/ViewModels/TreeNodeGrouping.cpp
@@ -57,8 +57,7 @@ FTreeNodeGroupingByUniqueValue::FTreeNodeGroupingByUniqueValue(TSharedRef<FTable
 FTreeNodeGroupInfo FTreeNodeGroupingByUniqueValue::GetGroupForNode(const FBaseTreeNodePtr InNode) const
 {
 	FTableTreeNodePtr TableTreeNodePtr = StaticCastSharedPtr<FTableTreeNode>(InNode);
-	return { FName(*ColumnRef->GetValueFormatter()->FormatValue(TableTreeNodePtr->GetValue(*ColumnRef)).ToString()), false };
-	//return { FName(*ColumnRef->GetValueAsText(TableTreeNodePtr->GetRowId()).ToString()), false };
+	return { FName(*ColumnRef->GetValueAsText(*TableTreeNodePtr).ToString()), false };
 }
 
 ////////////////////////////////////////////////////////////////////////////////////////////////////
@@ -84,9 +83,9 @@ FTreeNodeGroupInfo FTreeNodeGroupingByNameFirstLetter::GetGroupForNode(const FBa
 
 FTreeNodeGroupingByType::FTreeNodeGroupingByType()
 	: FTreeNodeGrouping(
-		LOCTEXT("Grouping_ByType_ShortName", "Type"),
-		LOCTEXT("Grouping_ByType_TitleName", "Type"),
-		LOCTEXT("Grouping_ByType_Desc", "Creates a group for each type."),
+		LOCTEXT("Grouping_ByTypeName_ShortName", "TypeName"),
+		LOCTEXT("Grouping_ByTypeName_TitleName", "TypeName"),
+		LOCTEXT("Grouping_ByTypeName_Desc", "Creates a group for each node type."),
 		TEXT("Profiler.FiltersAndPresets.StatTypeIcon"), //TODO
 		nullptr)
 {
@@ -96,7 +95,7 @@ FTreeNodeGroupingByType::FTreeNodeGroupingByType()
 
 FTreeNodeGroupInfo FTreeNodeGroupingByType::GetGroupForNode(const FBaseTreeNodePtr InNode) const
 {
-	return { InNode->GetTypeId(), true };
+	return { InNode->GetTypeName(), true };
 }
 
 ////////////////////////////////////////////////////////////////////////////////////////////////////
diff --git a/Engine/Source/Developer/TraceInsights/Private/Insights/Table/ViewModels/TreeNodeSorting.cpp b/Engine/Source/Developer/TraceInsights/Private/Insights/Table/ViewModels/TreeNodeSorting.cpp
--- a/Engine/Source/Developer/TraceInsights/Private/Insights/Table/ViewModels/TreeNodeSorting.cpp
+++ b/Engine/Source/Developer/TraceInsights/Private/Insights/Table/ViewModels/TreeNodeSorting.cpp
@@ -1,362 +1,3 @@
 // Copyright 1998-2019 Epic Games, Inc. All Rights Reserved.
 
 #include "TreeNodeSorting.h"
-
-// Insights
-#include "Insights/Table/ViewModels/TableColumn.h"
-#include "Insights/Table/ViewModels/TableTreeNode.h"
-
-#define LOCTEXT_NAMESPACE "Insights_TreeNode"
-
-namespace Insights
-{
-
-////////////////////////////////////////////////////////////////////////////////////////////////////
-// FTreeNodeSorting
-////////////////////////////////////////////////////////////////////////////////////////////////////
-
-FTreeNodeSorting::FTreeNodeSorting(FName InName, FText InShortName, FText InTitleName, FText InDescription, FName InColumnId)
-	: Name(InName)
-	, ShortName(InShortName)
-	, TitleName(InTitleName)
-	, Description(InDescription)
-	, ColumnId(InColumnId)
-	, AscendingIcon(nullptr)
-	, DescendingIcon(nullptr)
-	, AscendingCompareDelegate(nullptr)
-	, DescendingCompareDelegate(nullptr)
-{
-}
-
-////////////////////////////////////////////////////////////////////////////////////////////////////
-
-void FTreeNodeSorting::SortAscending(TArray<FBaseTreeNodePtr>& NodesToSort) const
-{
-	if (AscendingCompareDelegate)
-	{
-		NodesToSort.Sort(AscendingCompareDelegate);
-	}
-}
-
-////////////////////////////////////////////////////////////////////////////////////////////////////
-
-void FTreeNodeSorting::SortDescending(TArray<FBaseTreeNodePtr>& NodesToSort) const
-{
-	if (DescendingCompareDelegate)
-	{
-		NodesToSort.Sort(DescendingCompareDelegate);
-	}
-}
-
-////////////////////////////////////////////////////////////////////////////////////////////////////
-// FTreeNodeSortingByName
-////////////////////////////////////////////////////////////////////////////////////////////////////
-
-FTreeNodeSortingByName::FTreeNodeSortingByName()
-	: FTreeNodeSorting(
-		FName(TEXT("ByName")),
-		LOCTEXT("Sorting_ByName_Name", "By Name"),
-		LOCTEXT("Sorting_ByName_Title", "Sort By Name"),
-		LOCTEXT("Sorting_ByName_Desc", "Sort by name."),
-		FName(TEXT("Name")))
-{
-	AscendingCompareDelegate = [](const FBaseTreeNodePtr& A, const FBaseTreeNodePtr& B) -> bool
-	{
-		// Sort by name (ascending).
-		return A->GetName().LexicalLess(B->GetName());
-	};
-
-	DescendingCompareDelegate = [](const FBaseTreeNodePtr& A, const FBaseTreeNodePtr& B) -> bool
-	{
-		// Sort by name (descending).
-		return B->GetName().LexicalLess(A->GetName());
-	};
-}
-
-////////////////////////////////////////////////////////////////////////////////////////////////////
-// FTreeNodeSortingByType
-////////////////////////////////////////////////////////////////////////////////////////////////////
-
-FTreeNodeSortingByType::FTreeNodeSortingByType()
-	: FTreeNodeSorting(
-		FName(TEXT("ByType")),
-		LOCTEXT("Sorting_ByName_Name", "By Type"),
-		LOCTEXT("Sorting_ByName_Title", "Sort By Type"),
-		LOCTEXT("Sorting_ByName_Desc", "Sort by type (ascending or descending), then by name (ascending)."),
-		FName(TEXT("Type")))
-{
-	AscendingCompareDelegate = [](const FBaseTreeNodePtr& A, const FBaseTreeNodePtr& B) -> bool
-	{
-		const FName& TypeA = A->GetTypeId();
-		const FName& TypeB = B->GetTypeId();
-
-		if (TypeA == TypeB)
-		{
-			// Sort by name (ascending).
-			return A->GetName().LexicalLess(B->GetName());
-		}
-		else
-		{
-			// Sort by type (ascending).
-			return TypeA.FastLess(TypeB); // TypeA < TypeB;
-		}
-	};
-
-	DescendingCompareDelegate = [](const FBaseTreeNodePtr& A, const FBaseTreeNodePtr& B) -> bool
-	{
-		const FName& TypeA = A->GetTypeId();
-		const FName& TypeB = B->GetTypeId();
-
-		if (TypeA == TypeB)
-		{
-			// Sort by name (ascending).
-			return A->GetName().LexicalLess(B->GetName());
-		}
-		else
-		{
-			// Sort by type (descending).
-			return TypeB.FastLess(TypeA); // TypeA > TypeB;
-		}
-	};
-}
-
-////////////////////////////////////////////////////////////////////////////////////////////////////
-// FTableTreeNodeSorting
-////////////////////////////////////////////////////////////////////////////////////////////////////
-
-FTableTreeNodeSorting::FTableTreeNodeSorting(TSharedRef<FTableColumn> InColumnRef)
-	: FTreeNodeSorting(
-		InColumnRef->GetId(),
-		FText::Format(LOCTEXT("Sorting_TableValue_Name", "By {0}"), InColumnRef->GetShortName()),
-		FText::Format(LOCTEXT("Sorting_TableValue_Title", "Sort By {0}"), InColumnRef->GetTitleName()),
-		FText::Format(LOCTEXT("Sorting_TableValue_Desc", "Sort by {0} (ascending or descending), then by name (ascending)."), InColumnRef->GetShortName()),
-		InColumnRef->GetId())
-	, ColumnRef(InColumnRef)
-{
-}
-
-////////////////////////////////////////////////////////////////////////////////////////////////////
-// FTreeNodeSortingByBoolValue
-////////////////////////////////////////////////////////////////////////////////////////////////////
-
-FTreeNodeSortingByBoolValue::FTreeNodeSortingByBoolValue(TSharedRef<FTableColumn> InColumnRef)
-	: FTableTreeNodeSorting(InColumnRef)
-{
-	const FTableColumn& Column = *InColumnRef;
-
-	AscendingCompareDelegate = [&Column](const FBaseTreeNodePtr& A, const FBaseTreeNodePtr& B) -> bool
-	{
-		const bool ValueA = StaticCastSharedPtr<FTableTreeNode>(A)->GetValueBool(Column);
-		const bool ValueB = StaticCastSharedPtr<FTableTreeNode>(B)->GetValueBool(Column);
-
-		if (ValueA == ValueB)
-		{
-			// Sort by name (ascending).
-			return A->GetName().LexicalLess(B->GetName());
-		}
-		else
-		{
-			// Sort by value (ascending).
-			return ValueA < ValueB;
-		}
-	};
-
-	DescendingCompareDelegate = [&Column](const FBaseTreeNodePtr& A, const FBaseTreeNodePtr& B) -> bool
-	{
-		const bool ValueA = StaticCastSharedPtr<FTableTreeNode>(A)->GetValueBool(Column);
-		const bool ValueB = StaticCastSharedPtr<FTableTreeNode>(B)->GetValueBool(Column);
-
-		if (ValueA == ValueB)
-		{
-			// Sort by name (ascending).
-			return A->GetName().LexicalLess(B->GetName());
-		}
-		else
-		{
-			// Sort by value (descending).
-			return ValueB < ValueA;
-		}
-	};
-}
-
-////////////////////////////////////////////////////////////////////////////////////////////////////
-// FTreeNodeSortingByIntValue
-////////////////////////////////////////////////////////////////////////////////////////////////////
-
-FTreeNodeSortingByIntValue::FTreeNodeSortingByIntValue(TSharedRef<FTableColumn> InColumnRef)
-	: FTableTreeNodeSorting(InColumnRef)
-{
-	const FTableColumn& Column = *InColumnRef;
-
-	AscendingCompareDelegate = [&Column](const FBaseTreeNodePtr& A, const FBaseTreeNodePtr& B) -> bool
-	{
-		const int64 ValueA = StaticCastSharedPtr<FTableTreeNode>(A)->GetValueInt64(Column);
-		const int64 ValueB = StaticCastSharedPtr<FTableTreeNode>(B)->GetValueInt64(Column);
-
-		if (ValueA == ValueB)
-		{
-			// Sort by name (ascending).
-			return A->GetName().LexicalLess(B->GetName());
-		}
-		else
-		{
-			// Sort by value (ascending).
-			return ValueA < ValueB;
-		}
-	};
-
-	DescendingCompareDelegate = [&Column](const FBaseTreeNodePtr& A, const FBaseTreeNodePtr& B) -> bool
-	{
-		const int64 ValueA = StaticCastSharedPtr<FTableTreeNode>(A)->GetValueInt64(Column);
-		const int64 ValueB = StaticCastSharedPtr<FTableTreeNode>(B)->GetValueInt64(Column);
-
-		if (ValueA == ValueB)
-		{
-			// Sort by name (ascending).
-			return A->GetName().LexicalLess(B->GetName());
-		}
-		else
-		{
-			// Sort by value (descending).
-			return ValueB < ValueA;
-		}
-	};
-}
-
-////////////////////////////////////////////////////////////////////////////////////////////////////
-// FTreeNodeSortingByFloatValue
-////////////////////////////////////////////////////////////////////////////////////////////////////
-
-FTreeNodeSortingByFloatValue::FTreeNodeSortingByFloatValue(TSharedRef<FTableColumn> InColumnRef)
-	: FTableTreeNodeSorting(InColumnRef)
-{
-	const FTableColumn& Column = *InColumnRef;
-
-	AscendingCompareDelegate = [&Column](const FBaseTreeNodePtr& A, const FBaseTreeNodePtr& B) -> bool
-	{
-		const float ValueA = StaticCastSharedPtr<FTableTreeNode>(A)->GetValueFloat(Column);
-		const float ValueB = StaticCastSharedPtr<FTableTreeNode>(B)->GetValueFloat(Column);
-
-		if (ValueA == ValueB)
-		{
-			// Sort by name (ascending).
-			return A->GetName().LexicalLess(B->GetName());
-		}
-		else
-		{
-			// Sort by value (ascending).
-			return ValueA < ValueB;
-		}
-	};
-
-	DescendingCompareDelegate = [&Column](const FBaseTreeNodePtr& A, const FBaseTreeNodePtr& B) -> bool
-	{
-		const float ValueA = StaticCastSharedPtr<FTableTreeNode>(A)->GetValueFloat(Column);
-		const float ValueB = StaticCastSharedPtr<FTableTreeNode>(B)->GetValueFloat(Column);
-
-		if (ValueA == ValueB)
-		{
-			// Sort by name (ascending).
-			return A->GetName().LexicalLess(B->GetName());
-		}
-		else
-		{
-			// Sort by value (descending).
-			return ValueB < ValueA;
-		}
-	};
-}
-
-////////////////////////////////////////////////////////////////////////////////////////////////////
-// FTreeNodeSortingByDoubleValue
-////////////////////////////////////////////////////////////////////////////////////////////////////
-
-FTreeNodeSortingByDoubleValue::FTreeNodeSortingByDoubleValue(TSharedRef<FTableColumn> InColumnRef)
-	: FTableTreeNodeSorting(InColumnRef)
-{
-	const FTableColumn& Column = *InColumnRef;
-
-	AscendingCompareDelegate = [&Column](const FBaseTreeNodePtr& A, const FBaseTreeNodePtr& B) -> bool
-	{
-		const double ValueA = StaticCastSharedPtr<FTableTreeNode>(A)->GetValueDouble(Column);
-		const double ValueB = StaticCastSharedPtr<FTableTreeNode>(B)->GetValueDouble(Column);
-
-		if (ValueA == ValueB)
-		{
-			// Sort by name (ascending).
-			return A->GetName().LexicalLess(B->GetName());
-		}
-		else
-		{
-			// Sort by value (ascending).
-			return ValueA < ValueB;
-		}
-	};
-
-	DescendingCompareDelegate = [&Column](const FBaseTreeNodePtr& A, const FBaseTreeNodePtr& B) -> bool
-	{
-		const double ValueA = StaticCastSharedPtr<FTableTreeNode>(A)->GetValueDouble(Column);
-		const double ValueB = StaticCastSharedPtr<FTableTreeNode>(B)->GetValueDouble(Column);
-
-		if (ValueA == ValueB)
-		{
-			// Sort by name (ascending).
-			return A->GetName().LexicalLess(B->GetName());
-		}
-		else
-		{
-			// Sort by value (descending).
-			return ValueB < ValueA;
-		}
-	};
-}
-
-////////////////////////////////////////////////////////////////////////////////////////////////////
-// FTreeNodeSortingByCStringValue
-////////////////////////////////////////////////////////////////////////////////////////////////////
-
-FTreeNodeSortingByCStringValue::FTreeNodeSortingByCStringValue(TSharedRef<FTableColumn> InColumnRef)
-	: FTableTreeNodeSorting(InColumnRef)
-{
-	const FTableColumn& Column = *InColumnRef;
-
-	AscendingCompareDelegate = [&Column](const FBaseTreeNodePtr& A, const FBaseTreeNodePtr& B) -> bool
-	{
-		const TCHAR* ValueA = StaticCastSharedPtr<FTableTreeNode>(A)->GetValueCString(Column);
-		const TCHAR* ValueB = StaticCastSharedPtr<FTableTreeNode>(B)->GetValueCString(Column);
-
-		if (ValueA == ValueB)
-		{
-			// Sort by name (ascending).
-			return A->GetName().LexicalLess(B->GetName());
-		}
-		else
-		{
-			// Sort by value (ascending).
-			return FCString::Strcmp(ValueA, ValueB) < 0;
-		}
-	};
-
-	DescendingCompareDelegate = [&Column](const FBaseTreeNodePtr& A, const FBaseTreeNodePtr& B) -> bool
-	{
-		const TCHAR* ValueA = StaticCastSharedPtr<FTableTreeNode>(A)->GetValueCString(Column);
-		const TCHAR* ValueB = StaticCastSharedPtr<FTableTreeNode>(B)->GetValueCString(Column);
-
-		if (ValueA == ValueB)
-		{
-			// Sort by name (ascending).
-			return A->GetName().LexicalLess(B->GetName());
-		}
-		else
-		{
-			// Sort by value (descending).
-			return FCString::Strcmp(ValueB, ValueA) < 0;
-		}
-	};
-}
-
-////////////////////////////////////////////////////////////////////////////////////////////////////
-
-} // namespace Insights
-
-#undef LOCTEXT_NAMESPACE
