[Insights]

diff --git a/Engine/Source/Developer/TraceInsights/Private/Insights/Table/ViewModels/TreeNodeSorting.h b/Engine/Source/Developer/TraceInsights/Private/Insights/Table/ViewModels/TreeNodeSorting.h
--- a/Engine/Source/Developer/TraceInsights/Private/Insights/Table/ViewModels/TreeNodeSorting.h
+++ b/Engine/Source/Developer/TraceInsights/Private/Insights/Table/ViewModels/TreeNodeSorting.h
@@ -3,156 +3,13 @@
 #pragma once
 
 #include "CoreMinimal.h"
-#include "Templates/Function.h"
-#include "Styling/SlateBrush.h"
 
 // Insights
-#include "Insights/Table/ViewModels/BaseTreeNode.h"
+#include "Insights/Table/ViewModels/TableCellValueSorter.h"
 
 namespace Insights
 {
 
-class FTableColumn;
-class FTableTreeNode;
-
-////////////////////////////////////////////////////////////////////////////////////////////////////
-
-typedef TFunction<bool(const FBaseTreeNodePtr& A, const FBaseTreeNodePtr& B)> FTreeNodeSortingCompareFn;
-
-class ITreeNodeSorting
-{
-public:
-	virtual FName GetName() const = 0;
-	virtual FText GetShortName() const = 0;
-	virtual FText GetTitleName() const = 0;
-	virtual FText GetDescription() const = 0;
-	virtual FName GetColumnId() const = 0;
-
-	virtual FSlateBrush* GetAscendingIcon() const = 0;
-	virtual FSlateBrush* GetDescendingIcon() const = 0;
-
-	virtual FTreeNodeSortingCompareFn GetAscendingCompareDelegate() const = 0;
-	virtual FTreeNodeSortingCompareFn GetDescendingCompareDelegate() const = 0;
-
-	virtual void SortAscending(TArray<FBaseTreeNodePtr>& NodesToSort) const = 0;
-	virtual void SortDescending(TArray<FBaseTreeNodePtr>& NodesToSort) const = 0;
-};
-
-////////////////////////////////////////////////////////////////////////////////////////////////////
-
-class FTreeNodeSorting : public ITreeNodeSorting
-{
-public:
-	typedef bool FCompareFn(const FBaseTreeNodePtr& A, const FBaseTreeNodePtr& B) const;
-
-public:
-	FTreeNodeSorting(FName InName, FText InShortName, FText InTitleName, FText InDescription, FName InColumnId);
-	virtual ~FTreeNodeSorting() {}
-
-	virtual FName GetName() const override { return Name; }
-	virtual FText GetShortName() const override { return ShortName; }
-	virtual FText GetTitleName() const override { return TitleName; }
-	virtual FText GetDescription() const override { return Description; }
-	virtual FName GetColumnId() const override { return ColumnId; }
-
-	virtual FSlateBrush* GetAscendingIcon() const override { return AscendingIcon; }
-	virtual FSlateBrush* GetDescendingIcon() const override { return DescendingIcon; }
-
-	virtual FTreeNodeSortingCompareFn GetAscendingCompareDelegate() const override { return AscendingCompareDelegate; }
-	virtual FTreeNodeSortingCompareFn GetDescendingCompareDelegate() const override { return DescendingCompareDelegate; }
-
-	virtual void SortAscending(TArray<FBaseTreeNodePtr>& NodesToSort) const override;
-	virtual void SortDescending(TArray<FBaseTreeNodePtr>& NodesToSort) const override;
-
-protected:
-	FName Name;
-	FText ShortName;
-	FText TitleName;
-	FText Description;
-	FName ColumnId;
-	
-	FSlateBrush* AscendingIcon;
-	FSlateBrush* DescendingIcon;
-
-	FTreeNodeSortingCompareFn AscendingCompareDelegate;
-	FTreeNodeSortingCompareFn DescendingCompareDelegate;
-};
-
-////////////////////////////////////////////////////////////////////////////////////////////////////
-
-class FTreeNodeSortingByName : public FTreeNodeSorting
-{
-public:
-	FTreeNodeSortingByName();
-	virtual ~FTreeNodeSortingByName() {}
-};
-
-////////////////////////////////////////////////////////////////////////////////////////////////////
-
-class FTreeNodeSortingByType : public FTreeNodeSorting
-{
-public:
-	FTreeNodeSortingByType();
-	virtual ~FTreeNodeSortingByType() {}
-};
-
-////////////////////////////////////////////////////////////////////////////////////////////////////
-
-class FTableTreeNodeSorting : public FTreeNodeSorting
-{
-public:
-	FTableTreeNodeSorting(TSharedRef<FTableColumn> InColumnRef);
-	virtual ~FTableTreeNodeSorting() {}
-
-protected:
-	TSharedRef<FTableColumn> ColumnRef;
-};
-
-////////////////////////////////////////////////////////////////////////////////////////////////////
-
-class FTreeNodeSortingByBoolValue : public FTableTreeNodeSorting
-{
-public:
-	FTreeNodeSortingByBoolValue(TSharedRef<FTableColumn> InColumnRef);
-	virtual ~FTreeNodeSortingByBoolValue() {}
-};
-
-////////////////////////////////////////////////////////////////////////////////////////////////////
-
-class FTreeNodeSortingByIntValue : public FTableTreeNodeSorting
-{
-public:
-	FTreeNodeSortingByIntValue(TSharedRef<FTableColumn> InColumnRef);
-	virtual ~FTreeNodeSortingByIntValue() {}
-};
-
-////////////////////////////////////////////////////////////////////////////////////////////////////
-
-class FTreeNodeSortingByFloatValue : public FTableTreeNodeSorting
-{
-public:
-	FTreeNodeSortingByFloatValue(TSharedRef<FTableColumn> InColumnRef);
-	virtual ~FTreeNodeSortingByFloatValue() {}
-};
-
-////////////////////////////////////////////////////////////////////////////////////////////////////
-
-class FTreeNodeSortingByDoubleValue : public FTableTreeNodeSorting
-{
-public:
-	FTreeNodeSortingByDoubleValue(TSharedRef<FTableColumn> InColumnRef);
-	virtual ~FTreeNodeSortingByDoubleValue() {}
-};
-
-////////////////////////////////////////////////////////////////////////////////////////////////////
-
-class FTreeNodeSortingByCStringValue : public FTableTreeNodeSorting
-{
-public:
-	FTreeNodeSortingByCStringValue(TSharedRef<FTableColumn> InColumnRef);
-	virtual ~FTreeNodeSortingByCStringValue() {}
-};
-
-////////////////////////////////////////////////////////////////////////////////////////////////////
+//class FTableTreeNodeSorterByRowIndex
 
 } // namespace Insights
diff --git a/Engine/Source/Developer/TraceInsights/Private/Insights/Table/Widgets/STableTreeView.cpp b/Engine/Source/Developer/TraceInsights/Private/Insights/Table/Widgets/STableTreeView.cpp
--- a/Engine/Source/Developer/TraceInsights/Private/Insights/Table/Widgets/STableTreeView.cpp
+++ b/Engine/Source/Developer/TraceInsights/Private/Insights/Table/Widgets/STableTreeView.cpp
@@ -32,10 +32,6 @@
 namespace Insights
 {
 
-////////////////////////////////////////////////////////////////////////////////////////////////////
-
-const FName STableTreeView::DefaultColumnBeingSorted;
-const EColumnSortMode::Type STableTreeView::DefaultColumnSortMode(EColumnSortMode::Descending);
 const FName STableTreeView::RootNodeName(TEXT("Root"));
 
 ////////////////////////////////////////////////////////////////////////////////////////////////////
@@ -44,7 +40,6 @@ STableTreeView::STableTreeView()
 	: Table()
 	, Session(FInsightsManager::Get()->GetSession())
 	, TreeView(nullptr)
-	, TreeViewHeaderColumnArgs()
 	, TreeViewHeaderRow(nullptr)
 	, ExternalScrollbar(nullptr)
 	, HoveredColumnId()
@@ -62,11 +57,10 @@ STableTreeView::STableTreeView()
 	, AvailableGroupings()
 	, CurrentGroupings()
 	, GroupingBreadcrumbTrail(nullptr)
-	, AvailableSortings()
-	, ColumnToSortingMap()
-	, CurrentSorting(nullptr)
-	, ColumnBeingSorted(DefaultColumnBeingSorted)
-	, ColumnSortMode(DefaultColumnSortMode)
+	, AvailableSorters()
+	, CurrentSorter(nullptr)
+	, ColumnBeingSorted(GetDefaultColumnBeingSorted())
+	, ColumnSortMode(GetDefaultColumnSortMode())
 	, StatsStartTime(0.0)
 	, StatsEndTime(0.0)
 {
@@ -247,7 +241,7 @@ TSharedPtr<SWidget> STableTreeView::TreeView_GetMenuContent()
 		if (HoveredColumnPtr != nullptr)
 		{
 			PropertyName = HoveredColumnPtr->GetShortName();
-			PropertyValue = HoveredColumnPtr->GetValueAsText(SelectedNode->GetRowId());
+			PropertyValue = HoveredColumnPtr->GetValueAsTooltipText(*SelectedNode);
 		}
 		SelectionStr = FText::FromName(SelectedNode->GetName());
 	}
@@ -446,8 +440,6 @@ void STableTreeView::InitializeAndShowHeaderColumns()
 {
 	for (const TSharedPtr<FTableColumn>& ColumnPtr : Table->GetColumns())
 	{
-		TreeViewHeaderRow_CreateColumnArgs(*ColumnPtr);
-
 		if (ColumnPtr->ShouldBeVisible())
 		{
 			ShowColumn(ColumnPtr->GetId());
@@ -457,47 +449,9 @@ void STableTreeView::InitializeAndShowHeaderColumns()
 
 ////////////////////////////////////////////////////////////////////////////////////////////////////
 
-void STableTreeView::TreeViewHeaderRow_CreateColumnArgs(const FTableColumn& Column)
+FText STableTreeView::GetColumnHeaderText(const FName ColumnId) const
 {
-	SHeaderRow::FColumn::FArguments ColumnArgs;
-
-	ColumnArgs
-		.ColumnId(Column.GetId())
-		.DefaultLabel(Column.GetShortName())
-		.SortMode(EColumnSortMode::None)
-		.HAlignHeader(HAlign_Fill)
-		.VAlignHeader(VAlign_Fill)
-		.HeaderContentPadding(FMargin(2.0f))
-		.HAlignCell(HAlign_Fill)
-		.VAlignCell(VAlign_Fill)
-		.SortMode(this, &STableTreeView::GetSortModeForColumn, Column.GetId())
-		.OnSort(this, &STableTreeView::OnSortModeChanged)
-		.ManualWidth(Column.GetInitialWidth())
-		.FixedWidth(Column.IsFixedWidth() ? Column.GetInitialWidth() : TOptional<float>())
-		.HeaderContent()
-		[
-			SNew(SBox)
-			.ToolTip(STableTreeViewTooltip::GetColumnTooltip(Column))
-			.HAlign(Column.GetHorizontalAlignment())
-			.VAlign(VAlign_Center)
-			[
-				SNew(STextBlock)
-				.Text(this, &STableTreeView::OnGetHeaderText, Column.GetId())
-			]
-		]
-		.MenuContent()
-		[
-			TreeViewHeaderRow_GenerateColumnMenu(Column)
-		];
-
-	TreeViewHeaderColumnArgs.Add(Column.GetId(), ColumnArgs);
-}
-
-////////////////////////////////////////////////////////////////////////////////////////////////////
-
-FText STableTreeView::OnGetHeaderText(const FName ColumnId) const
-{
-	FTableColumn& Column = *Table->FindColumnChecked(ColumnId);
+	const FTableColumn& Column = *Table->FindColumnChecked(ColumnId);
 	return Column.GetShortName();
 }
 
@@ -565,12 +519,12 @@ TSharedRef<SWidget> STableTreeView::TreeViewHeaderRow_GenerateColumnMenu(const F
 			MenuBuilder.EndSection();
 		}
 
-		if (Column.CanBeFiltered())
-		{
-			MenuBuilder.BeginSection("FilterMode", LOCTEXT("ContextMenu_Header_Misc_Filter_FilterMode", "Filter Mode"));
-			bIsMenuVisible = true;
-			MenuBuilder.EndSection();
-		}
+		//if (Column.CanBeFiltered())
+		//{
+		//	MenuBuilder.BeginSection("FilterMode", LOCTEXT("ContextMenu_Header_Misc_Filter_FilterMode", "Filter Mode"));
+		//	bIsMenuVisible = true;
+		//	MenuBuilder.EndSection();
+		//}
 	}
 
 	return bIsMenuVisible ? MenuBuilder.MakeWidget() : (TSharedRef<SWidget>)SNullWidget::NullWidget;
@@ -724,18 +678,6 @@ void STableTreeView::TreeView_Refresh()
 
 void STableTreeView::TreeView_OnSelectionChanged(FTableTreeNodePtr SelectedItem, ESelectInfo::Type SelectInfo)
 {
-	if (SelectInfo != ESelectInfo::Direct)
-	{
-		TArray<FTableTreeNodePtr> SelectedItems = TreeView->GetSelectedItems();
-		if (SelectedItems.Num() == 1)
-		{
-			//HighlightedNodeName = SelectedItems[0]->GetName();
-		}
-		else
-		{
-			//HighlightedNodeName = NAME_None;
-		}
-	}
 }
 
 ////////////////////////////////////////////////////////////////////////////////////////////////////
@@ -751,26 +693,12 @@ void STableTreeView::TreeView_OnGetChildren(FTableTreeNodePtr InParent, TArray<F
 
 ////////////////////////////////////////////////////////////////////////////////////////////////////
 
-void STableTreeView::TreeView_OnMouseButtonDoubleClick(FTableTreeNodePtr TimerNode)
+void STableTreeView::TreeView_OnMouseButtonDoubleClick(FTableTreeNodePtr NodePtr)
 {
-	if (!TimerNode->IsGroup())
-	{
-		//im:TODO: const bool bIsStatTracked = FTimingProfilerManager::Get()->IsStatTracked(TimerNode->GetId());
-	//	if (!bIsStatTracked)
-	//	{
-	//		// Add a new graph.
-	//		FTimingProfilerManager::Get()->TrackStat(TimerNode->GetId());
-	//	}
-	//	else
-	//	{
-	//		// Remove a graph
-	//		FTimingProfilerManager::Get()->UntrackStat(TimerNode->GetId());
-	//	}
-	}
-	else
+	if (NodePtr->IsGroup())
 	{
-		const bool bIsGroupExpanded = TreeView->IsItemExpanded(TimerNode);
-		TreeView->SetItemExpansion(TimerNode, !bIsGroupExpanded);
+		const bool bIsGroupExpanded = TreeView->IsItemExpanded(NodePtr);
+		TreeView->SetItemExpansion(NodePtr, !bIsGroupExpanded);
 	}
 }
 
@@ -796,17 +724,15 @@ TSharedRef<ITableRow> STableTreeView::TreeView_OnGenerateRow(FTableTreeNodePtr N
 
 ////////////////////////////////////////////////////////////////////////////////////////////////////
 
-void STableTreeView::TableRow_SetHoveredCell(TSharedPtr<FTable> TablePtr, TSharedPtr<FTableColumn> ColumnPtr, const FTableTreeNodePtr NodePtr)
+void STableTreeView::TableRow_SetHoveredCell(TSharedPtr<FTable> InTablePtr, TSharedPtr<FTableColumn> InColumnPtr, const FTableTreeNodePtr InNodePtr)
 {
-	HoveredColumnId = ColumnPtr ? ColumnPtr->GetId() : FName();
+	HoveredColumnId = InColumnPtr ? InColumnPtr->GetId() : FName();
 
 	const bool bIsAnyMenusVisible = FSlateApplication::Get().AnyMenusVisible();
 	if (!HasMouseCapture() && !bIsAnyMenusVisible)
 	{
-		HoveredNodePtr = NodePtr;
+		HoveredNodePtr = InNodePtr;
 	}
-
-	//UE_LOG(TimingProfiler, Log, TEXT("%s -> %s"), *HoveredColumnId.GetPlainNameString(), NodePtr.IsValid() ? *NodePtr->GetName().GetPlainNameString() : TEXT("nullptr"));
 }
 
 ////////////////////////////////////////////////////////////////////////////////////////////////////
@@ -908,7 +834,7 @@ void STableTreeView::CreateGroups()
 void STableTreeView::GroupNodesRec(const TArray<FTableTreeNodePtr>& Nodes, FTableTreeNode& ParentGroup, int32 GroupingDepth)
 {
 	ensure(CurrentGroupings.Num() > 0);
-	
+
 	FTreeNodeGrouping& Grouping = *CurrentGroupings[GroupingDepth];
 
 	TMap<FName, FTableTreeNodePtr> GroupMap;
@@ -993,8 +919,12 @@ void STableTreeView::UpdateInt64SumAggregationRec(FTableColumn& Column, FTableTr
 			UpdateInt64SumAggregationRec(Column, TableNode);
 		}
 
-		FTableCellValue Value = TableNode.GetValue(Column);
-		AggregatedValue += Value.Int64;
+		const TOptional<FTableCellValue> OptionalValue = Column.GetValue(TableNode);
+		if (OptionalValue.IsSet())
+		{
+			ensure(OptionalValue.GetValue().DataType == ETableCellDataType::Int64);
+			AggregatedValue += OptionalValue.GetValue().Int64;
+		}
 	}
 
 	GroupNode.AddAggregatedValue(Column.GetId(), FTableCellValue(AggregatedValue));
@@ -1014,8 +944,12 @@ void STableTreeView::UpdateFloatSumAggregationRec(FTableColumn& Column, FTableTr
 			UpdateFloatSumAggregationRec(Column, TableNode);
 		}
 
-		FTableCellValue Value = TableNode.GetValue(Column);
-		AggregatedValue += Value.Float;
+		const TOptional<FTableCellValue> OptionalValue = Column.GetValue(TableNode);
+		if (OptionalValue.IsSet())
+		{
+			ensure(OptionalValue.GetValue().DataType == ETableCellDataType::Float);
+			AggregatedValue += OptionalValue.GetValue().Float;
+		}
 	}
 
 	GroupNode.AddAggregatedValue(Column.GetId(), FTableCellValue(AggregatedValue));
@@ -1034,9 +968,13 @@ void STableTreeView::UpdateDoubleSumAggregationRec(FTableColumn& Column, FTableT
 		{
 			UpdateDoubleSumAggregationRec(Column, TableNode);
 		}
-							
-		FTableCellValue Value = TableNode.GetValue(Column);
-		AggregatedValue += Value.Double;
+
+		const TOptional<FTableCellValue> OptionalValue = Column.GetValue(TableNode);
+		if (OptionalValue.IsSet())
+		{
+			ensure(OptionalValue.GetValue().DataType == ETableCellDataType::Double);
+			AggregatedValue += OptionalValue.GetValue().Double;
+		}
 	}
 
 	GroupNode.AddAggregatedValue(Column.GetId(), FTableCellValue(AggregatedValue));
@@ -1517,56 +1455,35 @@ bool STableTreeView::GroupingCrumbMenu_Add_CanExecute(const TSharedPtr<FTreeNode
 // Sorting
 ////////////////////////////////////////////////////////////////////////////////////////////////////
 
-void STableTreeView::CreateSortings()
+////////////////////////////////////////////////////////////////////////////////////////////////////
+
+const EColumnSortMode::Type STableTreeView::GetDefaultColumnSortMode()
 {
-	AvailableSortings.Reset();
-	ColumnToSortingMap.Reset();
-	CurrentSorting = nullptr;
+	return EColumnSortMode::Descending;
+}
+
+////////////////////////////////////////////////////////////////////////////////////////////////////
 
-	TSharedPtr<FTreeNodeSorting> SortingByName = MakeShareable(new FTreeNodeSortingByName());
-	AvailableSortings.Add(SortingByName);
+const FName STableTreeView::GetDefaultColumnBeingSorted()
+{
+	return NAME_None;
+}
 
-	TSharedPtr<FTreeNodeSorting> SortingByType = MakeShareable(new FTreeNodeSortingByType());
-	AvailableSortings.Add(SortingByType);
+////////////////////////////////////////////////////////////////////////////////////////////////////
 
-	const Trace::ITableLayout& TableLayout = Table->GetSourceTable()->GetLayout();
+void STableTreeView::CreateSortings()
+{
+	AvailableSorters.Reset();
+	CurrentSorter = nullptr;
 
 	for (const TSharedPtr<FTableColumn> ColumnPtr : Table->GetColumns())
 	{
 		if (ColumnPtr->CanBeSorted())
 		{
-			if (ColumnPtr->IsHierarchy())
+			TSharedPtr<Insights::ITableCellValueSorter> SorterPtr = ColumnPtr->GetValueSorter();
+			if (ensure(SorterPtr.IsValid()))
 			{
-				ColumnToSortingMap.Add(ColumnPtr->GetId(), SortingByName);
-			}
-			else
-			{
-				TSharedPtr<FTreeNodeSorting> Sorting = nullptr;
-
-				switch (TableLayout.GetColumnType(ColumnPtr->GetIndex()))
-				{
-				case Trace::TableColumnType_Bool:
-					Sorting = MakeShareable(new FTreeNodeSortingByBoolValue(ColumnPtr.ToSharedRef()));
-					break;
-				case Trace::TableColumnType_Int:
-					Sorting = MakeShareable(new FTreeNodeSortingByIntValue(ColumnPtr.ToSharedRef()));
-					break;
-				case Trace::TableColumnType_Float:
-					Sorting = MakeShareable(new FTreeNodeSortingByFloatValue(ColumnPtr.ToSharedRef()));
-					break;
-				case Trace::TableColumnType_Double:
-					Sorting = MakeShareable(new FTreeNodeSortingByDoubleValue(ColumnPtr.ToSharedRef()));
-					break;
-				case Trace::TableColumnType_CString:
-					Sorting = MakeShareable(new FTreeNodeSortingByCStringValue(ColumnPtr.ToSharedRef()));
-					break;
-				}
-
-				if (Sorting)
-				{
-					AvailableSortings.Add(Sorting);
-					ColumnToSortingMap.Add(ColumnPtr->GetId(), Sorting);
-				}
+				AvailableSorters.Add(SorterPtr);
 			}
 		}
 	}
@@ -1576,30 +1493,29 @@ void STableTreeView::CreateSortings()
 
 void STableTreeView::UpdateCurrentSortingByColumn()
 {
-	TSharedPtr<FTreeNodeSorting>* SorterPtrPtr = ColumnToSortingMap.Find(ColumnBeingSorted);
-	CurrentSorting = SorterPtrPtr ? *SorterPtrPtr : nullptr;
+	TSharedPtr<Insights::FTableColumn> ColumnPtr = Table->FindColumn(ColumnBeingSorted);
+	CurrentSorter = ColumnPtr.IsValid() ? ColumnPtr->GetValueSorter() : nullptr;
 }
 
 ////////////////////////////////////////////////////////////////////////////////////////////////////
 
 void STableTreeView::SortTreeNodes()
 {
-	if (CurrentSorting.IsValid())
+	if (CurrentSorter.IsValid())
 	{
-		const Insights::ITreeNodeSorting* Sorter = CurrentSorting.Get();
-		SortTreeNodesRec(*Root, Sorter);
+		SortTreeNodesRec(*Root, *CurrentSorter);
 	}
 }
 
 ////////////////////////////////////////////////////////////////////////////////////////////////////
 
-void STableTreeView::SortTreeNodesRec(FTableTreeNode& GroupNode, const Insights::ITreeNodeSorting* Sorter)
+void STableTreeView::SortTreeNodesRec(FTableTreeNode& GroupNode, const ITableCellValueSorter& Sorter)
 {
 	if (ColumnSortMode == EColumnSortMode::Type::Descending)
 	{
 		GroupNode.SortChildrenDescending(Sorter);
 	}
-	else
+	else // if (ColumnSortMode == EColumnSortMode::Type::Ascending)
 	{
 		GroupNode.SortChildrenAscending(Sorter);
 	}
@@ -1658,13 +1574,8 @@ bool STableTreeView::HeaderMenu_SortMode_IsChecked(const FName ColumnId, const E
 
 bool STableTreeView::HeaderMenu_SortMode_CanExecute(const FName ColumnId, const EColumnSortMode::Type InSortMode) const
 {
-	FTableColumn& Column = *Table->FindColumnChecked(ColumnId);
-	const bool bIsValid = Column.CanBeSorted();
-
-	bool bCanExecute = ColumnBeingSorted != ColumnId ? true : ColumnSortMode != InSortMode;
-	bCanExecute = bCanExecute && bIsValid;
-
-	return bCanExecute;
+	const FTableColumn& Column = *Table->FindColumnChecked(ColumnId);
+	return Column.CanBeSorted();
 }
 
 ////////////////////////////////////////////////////////////////////////////////////////////////////
@@ -1739,15 +1650,43 @@ void STableTreeView::ShowColumn(const FName ColumnId)
 	FTableColumn& Column = *Table->FindColumnChecked(ColumnId);
 	Column.Show();
 
-	SHeaderRow::FColumn::FArguments& ColumnArgs = TreeViewHeaderColumnArgs.FindChecked(ColumnId);
+	SHeaderRow::FColumn::FArguments ColumnArgs;
+	ColumnArgs
+		.ColumnId(Column.GetId())
+		.DefaultLabel(Column.GetShortName())
+		.HAlignHeader(HAlign_Fill)
+		.VAlignHeader(VAlign_Fill)
+		.HeaderContentPadding(FMargin(2.0f))
+		.HAlignCell(HAlign_Fill)
+		.VAlignCell(VAlign_Fill)
+		.SortMode(this, &STableTreeView::GetSortModeForColumn, Column.GetId())
+		.OnSort(this, &STableTreeView::OnSortModeChanged)
+		.ManualWidth(Column.GetInitialWidth())
+		.FixedWidth(Column.IsFixedWidth() ? Column.GetInitialWidth() : TOptional<float>())
+		.HeaderContent()
+		[
+			SNew(SBox)
+			.ToolTip(STableTreeViewTooltip::GetColumnTooltip(Column))
+			.HAlign(Column.GetHorizontalAlignment())
+			.VAlign(VAlign_Center)
+			[
+				SNew(STextBlock)
+				.Text(this, &STableTreeView::GetColumnHeaderText, Column.GetId())
+			]
+		]
+		.MenuContent()
+		[
+			TreeViewHeaderRow_GenerateColumnMenu(Column)
+		];
 
-	const int32 NumColumns = TreeViewHeaderRow->GetColumns().Num();
 	int32 ColumnIndex = 0;
+	const int32 NewColumnPosition = Table->GetColumnPositionIndex(ColumnId);
+	const int32 NumColumns = TreeViewHeaderRow->GetColumns().Num();
 	for (; ColumnIndex < NumColumns; ColumnIndex++)
 	{
 		const SHeaderRow::FColumn& CurrentColumn = TreeViewHeaderRow->GetColumns()[ColumnIndex];
-		const FTableColumn& CurrentTableColumn = *Table->FindColumnChecked(CurrentColumn.ColumnId);
-		if (Column.GetOrder() < CurrentTableColumn.GetOrder())
+		const int32 CurrentColumnPosition = Table->GetColumnPositionIndex(CurrentColumn.ColumnId);
+		if (NewColumnPosition < CurrentColumnPosition)
 		{
 			break;
 		}
@@ -1762,7 +1701,7 @@ void STableTreeView::ShowColumn(const FName ColumnId)
 
 bool STableTreeView::CanHideColumn(const FName ColumnId) const
 {
-	FTableColumn& Column = *Table->FindColumnChecked(ColumnId);
+	const FTableColumn& Column = *Table->FindColumnChecked(ColumnId);
 	return Column.CanBeHidden();
 }
 
@@ -1782,7 +1721,7 @@ void STableTreeView::HideColumn(const FName ColumnId)
 
 bool STableTreeView::IsColumnVisible(const FName ColumnId)
 {
-	FTableColumn& Column = *Table->FindColumnChecked(ColumnId);
+	const FTableColumn& Column = *Table->FindColumnChecked(ColumnId);
 	return Column.IsVisible();
 }
 
@@ -1790,7 +1729,7 @@ bool STableTreeView::IsColumnVisible(const FName ColumnId)
 
 bool STableTreeView::CanToggleColumnVisibility(const FName ColumnId) const
 {
-	FTableColumn& Column = *Table->FindColumnChecked(ColumnId);
+	const FTableColumn& Column = *Table->FindColumnChecked(ColumnId);
 	return !Column.IsVisible() || Column.CanBeHidden();
 }
 
@@ -1798,7 +1737,7 @@ bool STableTreeView::CanToggleColumnVisibility(const FName ColumnId) const
 
 void STableTreeView::ToggleColumnVisibility(const FName ColumnId)
 {
-	FTableColumn& Column = *Table->FindColumnChecked(ColumnId);
+	const FTableColumn& Column = *Table->FindColumnChecked(ColumnId);
 	if (Column.IsVisible())
 	{
 		HideColumn(ColumnId);
@@ -1822,8 +1761,8 @@ bool STableTreeView::ContextMenu_ShowAllColumns_CanExecute() const
 
 void STableTreeView::ContextMenu_ShowAllColumns_Execute()
 {
-	ColumnBeingSorted = DefaultColumnBeingSorted;
-	ColumnSortMode = DefaultColumnSortMode;
+	ColumnBeingSorted = GetDefaultColumnBeingSorted();
+	ColumnSortMode = GetDefaultColumnSortMode();
 	UpdateCurrentSortingByColumn();
 
 	for (const TSharedPtr<FTableColumn>& ColumnPtr : Table->GetColumns())
@@ -1850,8 +1789,8 @@ bool STableTreeView::ContextMenu_ResetColumns_CanExecute() const
 
 void STableTreeView::ContextMenu_ResetColumns_Execute()
 {
-	ColumnBeingSorted = DefaultColumnBeingSorted;
-	ColumnSortMode = DefaultColumnSortMode;
+	ColumnBeingSorted = GetDefaultColumnBeingSorted();
+	ColumnSortMode = GetDefaultColumnSortMode();
 	UpdateCurrentSortingByColumn();
 
 	for (const TSharedPtr<FTableColumn>& ColumnPtr : Table->GetColumns())
@@ -1879,7 +1818,6 @@ void STableTreeView::RebuildTree(bool bResync)
 	{
 		const int32 PreviousNodeCount = TableTreeNodes.Num();
 		TableTreeNodes.Empty(PreviousNodeCount);
-		//TableTreeNodesMap.Empty(PreviousNodeCount);
 		TableTreeNodesIdMap.Empty(PreviousNodeCount);
 		bListHasChanged = true;
 	}
@@ -1900,52 +1838,9 @@ void STableTreeView::RebuildTree(bool bResync)
 		{
 			const int32 PreviousNodeCount = TableTreeNodes.Num();
 			TableTreeNodes.Empty(PreviousNodeCount);
-			//TableTreeNodesMap.Empty(PreviousNodeCount);
 			TableTreeNodesIdMap.Empty(PreviousNodeCount);
 			bListHasChanged = true;
 
-			/*
-			const Trace::ITableLayout& TableLayout = SourceTable->GetLayout();
-			const int32 ColumnCount = TableLayout.GetColumnCount();
-
-			for (; TableReader->IsValid(); TableReader->NextRow())
-			{
-				FString Line;
-				for (int32 ColumnIndex = 0; ColumnIndex < ColumnCount; ++ColumnIndex)
-				{
-					switch (TableLayout.GetColumnType(ColumnIndex))
-					{
-					case Trace::TableColumnType_Bool:
-						Line += TableReader->GetValueBool(ColumnIndex) ? "true" : "false";
-						break;
-					case Trace::TableColumnType_Int:
-						Line += FString::Printf(TEXT("%lld"), TableReader->GetValueInt(ColumnIndex));
-						break;
-					case Trace::TableColumnType_Float:
-						Line += FString::Printf(TEXT("%f"), TableReader->GetValueFloat(ColumnIndex));
-						break;
-					case Trace::TableColumnType_Double:
-						Line += FString::Printf(TEXT("%f"), TableReader->GetValueDouble(ColumnIndex));
-						break;
-					case Trace::TableColumnType_CString:
-						FString ValueString = TableReader->GetValueCString(ColumnIndex);
-						ValueString.ReplaceInline(TEXT(","), TEXT(" "));
-						Line += ValueString;
-						break;
-					}
-					if (ColumnIndex < ColumnCount - 1)
-					{
-						Line += TEXT(",");
-					}
-					else
-					{
-						Line += TEXT("\n");
-					}
-				}
-				//UE_LOG(TimingProfiler, Log, TEXT("%s"), *Line.ToString());
-			}
-			*/
-
 			for (int32 RowIndex = 0; RowIndex < TotalRowCount; ++RowIndex)
 			{
 				TableReader->SetRowIndex(RowIndex);
@@ -1953,7 +1848,6 @@ void STableTreeView::RebuildTree(bool bResync)
 				FName NodeName(*FString::Printf(TEXT("row %d"), RowIndex));
 				FTableTreeNodePtr NodePtr = MakeShareable(new FTableTreeNode(NodeId, NodeName, Table, RowIndex));
 				TableTreeNodes.Add(NodePtr);
-				//TableTreeNodesMap.Add(Name, NodePtr);
 				TableTreeNodesIdMap.Add(NodeId, NodePtr);
 			}
 		}
@@ -1974,7 +1868,6 @@ void STableTreeView::SelectNodeByNodeId(uint64 Id)
 	{
 		FTableTreeNodePtr NodePtr = *NodePtrPtr;
 
-		//UE_LOG(TimingProfiler, Log, TEXT("Select and RequestScrollIntoView %s"), *(*TimerNodePtrPtr)->GetName().ToString());
 		TreeView->SetSelection(NodePtr);
 		TreeView->RequestScrollIntoView(NodePtr);
 	}
@@ -1988,7 +1881,6 @@ void STableTreeView::SelectNodeByTableRowIndex(int32 RowIndex)
 	{
 		FTableTreeNodePtr NodePtr = TableTreeNodes[RowIndex];
 
-		//UE_LOG(TimingProfiler, Log, TEXT("Select and RequestScrollIntoView %s"), *(*TimerNodePtrPtr)->GetName().ToString());
 		TreeView->SetSelection(NodePtr);
 		TreeView->RequestScrollIntoView(NodePtr);
 	}
diff --git a/Engine/Source/Developer/TraceInsights/Private/Insights/Table/Widgets/STableTreeView.h b/Engine/Source/Developer/TraceInsights/Private/Insights/Table/Widgets/STableTreeView.h
--- a/Engine/Source/Developer/TraceInsights/Private/Insights/Table/Widgets/STableTreeView.h
+++ b/Engine/Source/Developer/TraceInsights/Private/Insights/Table/Widgets/STableTreeView.h
@@ -16,8 +16,6 @@
 #include "Widgets/Views/STreeView.h"
 
 // Insights
-#include "Insights/Table/ViewModels/Table.h"
-#include "Insights/Table/ViewModels/TableColumn.h"
 #include "Insights/Table/ViewModels/TableTreeNode.h"
 
 class FMenuBuilder;
@@ -31,8 +29,10 @@ namespace Trace
 namespace Insights
 {
 
-class FTreeNodeSorting;
+class FTable;
+class FTableColumn;
 class FTreeNodeGrouping;
+class ITableCellValueSorter;
 
 ////////////////////////////////////////////////////////////////////////////////////////////////////
 
@@ -103,9 +103,8 @@ protected:
 	// Tree View - Columns' Header
 
 	void InitializeAndShowHeaderColumns();
-	void TreeViewHeaderRow_CreateColumnArgs(const FTableColumn& Column);
 
-	FText OnGetHeaderText(const FName ColumnId) const;
+	FText GetColumnHeaderText(const FName ColumnId) const;
 
 	TSharedRef<SWidget> TreeViewHeaderRow_GenerateColumnMenu(const FTableColumn& Column);
 
@@ -187,11 +186,14 @@ protected:
 	////////////////////////////////////////////////////////////////////////////////////////////////////
 	// Sorting
 
+	static const EColumnSortMode::Type GetDefaultColumnSortMode();
+	static const FName GetDefaultColumnBeingSorted();
+
 	void CreateSortings();
 
 	void UpdateCurrentSortingByColumn();
 	void SortTreeNodes();
-	void SortTreeNodesRec(FTableTreeNode& GroupNode, const Insights::ITreeNodeSorting* Sorter);
+	void SortTreeNodesRec(FTableTreeNode& GroupNode, const ITableCellValueSorter& Sorter);
 
 	EColumnSortMode::Type GetSortModeForColumn(const FName ColumnId) const;
 	void SetSortModeForColumn(const FName& ColumnId, EColumnSortMode::Type SortMode);
@@ -254,9 +256,6 @@ protected:
 	/** The child STreeView widget. */
 	TSharedPtr<STreeView<FTableTreeNodePtr>> TreeView;
 
-	/** Column arguments used to initialize a new header column in the tree view, stored as column name to column arguments mapping. */
-	TMap<FName, SHeaderRow::FColumn::FArguments> TreeViewHeaderColumnArgs;
-
 	/** Holds the tree view header row widget which display all columns in the tree view. */
 	TSharedPtr<SHeaderRow> TreeViewHeaderRow;
 
@@ -323,14 +322,11 @@ protected:
 	//////////////////////////////////////////////////
 	// Sorting
 
-	/** All available sortings. */
-	TArray<TSharedPtr<FTreeNodeSorting>> AvailableSortings;
+	/** All available sorters. */
+	TArray<TSharedPtr<ITableCellValueSorter>> AvailableSorters;
 
-	/** Sortings associated with a column. */
-	TMap<FName, TSharedPtr<FTreeNodeSorting>> ColumnToSortingMap;
-
-	/** Current sorting. It is nullptr if sorting is disabled. */
-	TSharedPtr<FTreeNodeSorting> CurrentSorting;
+	/** Current sorter. It is nullptr if sorting is disabled. */
+	TSharedPtr<ITableCellValueSorter> CurrentSorter;
 
 	/** Name of the column currently being sorted. Can be NAME_None if sorting is disabled (CurrentSorting == nullptr) or if a complex sorting is used (CurrentSorting != nullptr). */
 	FName ColumnBeingSorted;
@@ -338,9 +334,6 @@ protected:
 	/** How we sort the nodes? Ascending or Descending. */
 	EColumnSortMode::Type ColumnSortMode;
 
-	static const FName DefaultColumnBeingSorted;
-	static const EColumnSortMode::Type DefaultColumnSortMode;
-
 	//////////////////////////////////////////////////
 
 	double StatsStartTime;
diff --git a/Engine/Source/Developer/TraceInsights/Private/Insights/Table/Widgets/STableTreeViewCell.cpp b/Engine/Source/Developer/TraceInsights/Private/Insights/Table/Widgets/STableTreeViewCell.cpp
--- a/Engine/Source/Developer/TraceInsights/Private/Insights/Table/Widgets/STableTreeViewCell.cpp
+++ b/Engine/Source/Developer/TraceInsights/Private/Insights/Table/Widgets/STableTreeViewCell.cpp
@@ -31,6 +31,10 @@ void STableTreeViewCell::Construct(const FArguments& InArgs, const TSharedRef<IT
 	ColumnPtr = InArgs._ColumnPtr;
 	TableTreeNodePtr = InArgs._TableTreeNodePtr;
 
+	ensure(TablePtr.IsValid());
+	ensure(ColumnPtr.IsValid());
+	ensure(TableTreeNodePtr.IsValid());
+
 	SetHoveredCellDelegate = InArgs._OnSetHoveredCell;
 
 	ChildSlot
@@ -101,18 +105,7 @@ TSharedRef<SWidget> STableTreeViewCell::GenerateWidgetForNameColumn(const FArgum
 
 TSharedRef<SWidget> STableTreeViewCell::GenerateWidgetForTableColumn(const FArguments& InArgs, const TSharedRef<ITableRow>& TableRow)
 {
-	FText FormattedValue = FText::GetEmpty();
-	if (TableTreeNodePtr->IsGroup())
-	{
-		if (ColumnPtr->GetAggregation() != ETableColumnAggregation::None)
-		{
-			FormattedValue = ColumnPtr->GetValueFormatter()->FormatValue(TableTreeNodePtr->GetValue(*ColumnPtr));
-		}
-	}
-	else
-	{
-		FormattedValue = ColumnPtr->GetValueAsText(TableTreeNodePtr->GetRowId());
-	}
+	const FText CellText = ColumnPtr->GetValueAsText(*TableTreeNodePtr);
 
 	return
 		SNew(SHorizontalBox)
@@ -125,7 +118,7 @@ TSharedRef<SWidget> STableTreeViewCell::GenerateWidgetForTableColumn(const FArgu
 		.Padding(FMargin(2.0f, 0.0f))
 		[
 			SNew(STextBlock)
-			.Text(FormattedValue)
+			.Text(CellText)
 			.TextStyle(FEditorStyle::Get(), TEXT("Profiler.Tooltip"))
 			.ColorAndOpacity(this, &STableTreeViewCell::GetStatsColorAndOpacity)
 			.ShadowColorAndOpacity(this, &STableTreeViewCell::GetShadowColorAndOpacity)
diff --git a/Engine/Source/Developer/TraceInsights/Private/Insights/Table/Widgets/STableTreeViewCell.h b/Engine/Source/Developer/TraceInsights/Private/Insights/Table/Widgets/STableTreeViewCell.h
--- a/Engine/Source/Developer/TraceInsights/Private/Insights/Table/Widgets/STableTreeViewCell.h
+++ b/Engine/Source/Developer/TraceInsights/Private/Insights/Table/Widgets/STableTreeViewCell.h
@@ -17,7 +17,7 @@ namespace Insights
 class FTable;
 class FTableColumn;
 
-DECLARE_DELEGATE_ThreeParams(FSetHoveredTableTreeViewCell, TSharedPtr<FTable> /*TablePtr*/, TSharedPtr<FTableColumn> /*ColumnPtr*/, const FTableTreeNodePtr /*NodePtr*/);
+DECLARE_DELEGATE_ThreeParams(FSetHoveredTableTreeViewCell, TSharedPtr<FTable> /*TablePtr*/, TSharedPtr<FTableColumn> /*ColumnPtr*/, const FTableTreeNodePtr /*TableTreeNodePtr*/);
 
 class STableTreeViewCell : public SCompoundWidget
 {
@@ -119,9 +119,9 @@ protected:
 		const FLinearColor TextColor =
 			TableTreeNodePtr->IsGroup() ?
 				FLinearColor(0.7f, 0.7f, 0.7f, 1.0f) :
-				TableTreeNodePtr->IsFiltered() ?
-					FLinearColor(1.0f, 1.0f, 1.0f, 0.5f) :
-					FLinearColor(1.0f, 1.0f, 1.0f, 1.0f);
+			TableTreeNodePtr->IsFiltered() ?
+				FLinearColor(1.0f, 1.0f, 1.0f, 0.5f) :
+				FLinearColor(1.0f, 1.0f, 1.0f, 1.0f);
 		return TextColor;
 	}
 
diff --git a/Engine/Source/Developer/TraceInsights/Private/Insights/Table/Widgets/STableTreeViewTooltip.cpp b/Engine/Source/Developer/TraceInsights/Private/Insights/Table/Widgets/STableTreeViewTooltip.cpp
--- a/Engine/Source/Developer/TraceInsights/Private/Insights/Table/Widgets/STableTreeViewTooltip.cpp
+++ b/Engine/Source/Developer/TraceInsights/Private/Insights/Table/Widgets/STableTreeViewTooltip.cpp
@@ -13,9 +13,9 @@
 
 // Insights
 #include "Insights/Common/TimeUtils.h"
-#include "Insights/Table/ViewModels/TableTreeNode.h"
 #include "Insights/Table/ViewModels/Table.h"
 #include "Insights/Table/ViewModels/TableColumn.h"
+#include "Insights/Table/ViewModels/TableTreeNode.h"
 
 #define LOCTEXT_NAMESPACE "STableTreeView"
 
@@ -90,13 +90,6 @@ TSharedPtr<SToolTip> STableTreeViewTooltip::GetColumnTooltip(const FTableColumn&
 
 TSharedPtr<SToolTip> STableTreeViewTooltip::GetCellTooltip(const TSharedPtr<FTableTreeNode> TreeNodePtr, const TSharedPtr<FTableColumn> ColumnPtr)
 {
-	const FLinearColor TextColor(0.8f, 0.8f, 0.8f, 1.0f);
-	const FLinearColor ValueColor(1.0f, 1.0f, 1.0f, 1.0f);
-
-	//const int32 NumDigits = 5;
-	//const FText TimeText = FText::FromString(TimeUtils::FormatTimeAuto(TotalTime));
-	//const FText TimeText = FText::FromString(TimeUtils::FormatTimeMs(Time, NumDigits, true));
-
 	TSharedPtr<SGridPanel> GridPanel;
 	TSharedPtr<SHorizontalBox> HBox;
 
@@ -105,12 +98,12 @@ TSharedPtr<SToolTip> STableTreeViewTooltip::GetCellTooltip(const TSharedPtr<FTab
 		[
 			SAssignNew(HBox, SHorizontalBox)
 
-			+SHorizontalBox::Slot()
+			+ SHorizontalBox::Slot()
 			.AutoWidth()
 			[
 				SNew(SVerticalBox)
 
-				+SVerticalBox::Slot()
+				+ SVerticalBox::Slot()
 				.AutoHeight()
 				.Padding(2.0f)
 				[
@@ -118,24 +111,22 @@ TSharedPtr<SToolTip> STableTreeViewTooltip::GetCellTooltip(const TSharedPtr<FTab
 					.Orientation(Orient_Horizontal)
 				]
 
-				+SVerticalBox::Slot()
+				+ SVerticalBox::Slot()
 				.AutoHeight()
 				.Padding(2.0f)
 				[
 					SNew(SGridPanel)
 
 					// Id: [Id]
-					+SGridPanel::Slot(0, 0)
+					+ SGridPanel::Slot(0, 0)
 					.Padding(2.0f)
-					//.HAlign(HAlign_Right)
 					[
 						SNew(STextBlock)
 						.Text(LOCTEXT("TT_Id", "Id:"))
 						.TextStyle(FEditorStyle::Get(), TEXT("Profiler.TooltipBold"))
 					]
-					+SGridPanel::Slot(1, 0)
+					+ SGridPanel::Slot(1, 0)
 					.Padding(2.0f)
-					.HAlign(HAlign_Left)
 					[
 						SNew(STextBlock)
 						.Text(FText::AsNumber(TreeNodePtr->GetId()))
@@ -143,60 +134,40 @@ TSharedPtr<SToolTip> STableTreeViewTooltip::GetCellTooltip(const TSharedPtr<FTab
 					]
 
 					// Name: [Name]
-					+SGridPanel::Slot(0, 1)
+					+ SGridPanel::Slot(0, 1)
 					.Padding(2.0f)
-					//.HAlign(HAlign_Right)
 					[
 						SNew(STextBlock)
 						.Text(LOCTEXT("TT_Name", "Name:"))
 						.TextStyle(FEditorStyle::Get(), TEXT("Profiler.TooltipBold"))
 					]
-					+SGridPanel::Slot(1, 1)
+					+ SGridPanel::Slot(1, 1)
 					.Padding(2.0f)
-					.HAlign(HAlign_Left)
 					[
 						SNew(STextBlock)
 						.Text(FText::FromName(TreeNodePtr->GetName()))
 						.TextStyle(FEditorStyle::Get(), TEXT("Profiler.Tooltip"))
 					]
-
-					// Type: [Type]
-					+ SGridPanel::Slot(0, 3)
-					.Padding(2.0f)
-					//.HAlign(HAlign_Right)
-					[
-						SNew(STextBlock)
-						.Text(LOCTEXT("TT_Type", "Type:"))
-						.TextStyle(FEditorStyle::Get(), TEXT("Profiler.TooltipBold"))
-					]
-					+ SGridPanel::Slot(1, 3)
-					.Padding(2.0f)
-					.HAlign(HAlign_Left)
-					[
-						SNew(STextBlock)
-						.Text(FText::FromName(TreeNodePtr->GetTypeId()))
-						.TextStyle(FEditorStyle::Get(), TEXT("Profiler.Tooltip"))
-					]
 				]
 
 				+ SVerticalBox::Slot()
-					.AutoHeight()
-					.Padding(2.0f)
-					[
-						SNew(SSeparator)
-						.Orientation(Orient_Horizontal)
-					]
+				.AutoHeight()
+				.Padding(2.0f)
+				[
+					SNew(SSeparator)
+					.Orientation(Orient_Horizontal)
+				]
 
 				+ SVerticalBox::Slot()
-					.AutoHeight()
-					.Padding(2.0f)
-					[
-						SAssignNew(GridPanel, SGridPanel)
+				.AutoHeight()
+				.Padding(2.0f)
+				[
+					SAssignNew(GridPanel, SGridPanel)
 
-						// Values for each table column are added here.
-					]
+					// Values for each table column are added here.
+				]
 
-				+SVerticalBox::Slot()
+				+ SVerticalBox::Slot()
 				.AutoHeight()
 				.Padding(2.0f)
 				[
@@ -212,10 +183,10 @@ TSharedPtr<SToolTip> STableTreeViewTooltip::GetCellTooltip(const TSharedPtr<FTab
 		int32 Row = 0;
 		for (const TSharedPtr<FTableColumn>& Column : Table->GetColumns())
 		{
-			if (!Column->IsHierarchy() && TreeNodePtr->GetRowId().HasValidIndex())
+			if (!Column->IsHierarchy())
 			{
 				FText Name = FText::Format(LOCTEXT("TooltipValueFormat", "{0}:"), Column->GetTitleName());
-				AddGridRow(GridPanel, Row, Name, Column->GetValueAsTooltipText(TreeNodePtr->GetRowId()));
+				AddGridRow(GridPanel, Row, Name, Column->GetValueAsTooltipText(*TreeNodePtr));
 			}
 		}
 	}
@@ -229,7 +200,6 @@ void STableTreeViewTooltip::AddGridRow(TSharedPtr<SGridPanel> Grid, int32& Row,
 {
 	Grid->AddSlot(0, Row)
 		.Padding(2.0f)
-		//.HAlign(HAlign_Right)
 		[
 			SNew(STextBlock)
 			.Text(Name)
@@ -238,7 +208,6 @@ void STableTreeViewTooltip::AddGridRow(TSharedPtr<SGridPanel> Grid, int32& Row,
 
 	Grid->AddSlot(1, Row)
 		.Padding(2.0f)
-		.HAlign(HAlign_Left)
 		[
 			SNew(STextBlock)
 			.Text(Value)
diff --git a/Engine/Source/Developer/TraceInsights/Private/Insights/TimingProfilerManager.cpp b/Engine/Source/Developer/TraceInsights/Private/Insights/TimingProfilerManager.cpp
--- a/Engine/Source/Developer/TraceInsights/Private/Insights/TimingProfilerManager.cpp
+++ b/Engine/Source/Developer/TraceInsights/Private/Insights/TimingProfilerManager.cpp
@@ -411,7 +411,7 @@ void FTimingProfilerManager::UpdateCallersAndCallees()
 					CallersTreeView->SetTree(Callers);
 				}
 
-				if (CallersTreeView.IsValid())
+				if (CalleesTreeView.IsValid())
 				{
 					const Trace::FTimingProfilerButterflyNode& Callees = TimingProfilerButterfly->GenerateCalleesTree(TimerId);
 					CalleesTreeView->SetTree(Callees);
diff --git a/Engine/Source/Developer/TraceInsights/Private/Insights/ViewModels/StatsNode.cpp b/Engine/Source/Developer/TraceInsights/Private/Insights/ViewModels/StatsNode.cpp
--- a/Engine/Source/Developer/TraceInsights/Private/Insights/ViewModels/StatsNode.cpp
+++ b/Engine/Source/Developer/TraceInsights/Private/Insights/ViewModels/StatsNode.cpp
@@ -7,6 +7,8 @@
 
 #define LOCTEXT_NAMESPACE "StatsNode"
 
+const FName FStatsNode::TypeName(TEXT("FStatsNode"));
+
 ////////////////////////////////////////////////////////////////////////////////////////////////////
 
 void FStatsNode::ResetAggregatedStats()
diff --git a/Engine/Source/Developer/TraceInsights/Private/Insights/ViewModels/StatsNode.h b/Engine/Source/Developer/TraceInsights/Private/Insights/ViewModels/StatsNode.h
--- a/Engine/Source/Developer/TraceInsights/Private/Insights/ViewModels/StatsNode.h
+++ b/Engine/Source/Developer/TraceInsights/Private/Insights/ViewModels/StatsNode.h
@@ -90,6 +90,7 @@ typedef TWeakPtr<class FStatsNode> FStatsNodeWeak;
 class FStatsNode : public Insights::FBaseTreeNode
 {
 public:
+	static const FName TypeName;
 	static const uint64 InvalidId = -1;
 
 public:
@@ -119,6 +120,8 @@ public:
 		ResetAggregatedStats();
 	}
 
+	virtual const FName& GetTypeName() const override { return TypeName; }
+
 	/**
 	 * @return a name of the meta group that this stats node belongs to, taken from the metadata.
 	 */
diff --git a/Engine/Source/Developer/TraceInsights/Private/Insights/ViewModels/TimerGroupingAndSorting.cpp b/Engine/Source/Developer/TraceInsights/Private/Insights/ViewModels/TimerGroupingAndSorting.cpp
--- a/Engine/Source/Developer/TraceInsights/Private/Insights/ViewModels/TimerGroupingAndSorting.cpp
+++ b/Engine/Source/Developer/TraceInsights/Private/Insights/ViewModels/TimerGroupingAndSorting.cpp
@@ -2,8 +2,303 @@
 
 #include "TimerGroupingAndSorting.h"
 
+#include "Insights/Table/ViewModels/TableColumn.h"
+#include "Insights/ViewModels/TimerNodeHelper.h"
 
 #define LOCTEXT_NAMESPACE "TimerNode"
 
+// Sort by name (ascending).
+#define INSIGHTS_DEFAULT_SORTING_NODES(A, B) return A->GetName().LexicalLess(B->GetName());
+
+////////////////////////////////////////////////////////////////////////////////////////////////////
+// Sorting by Timer Type
+////////////////////////////////////////////////////////////////////////////////////////////////////
+
+FTimerNodeSortingByTimerType::FTimerNodeSortingByTimerType(TSharedRef<Insights::FTableColumn> InColumnRef)
+	: Insights::FTableCellValueSorter(
+		FName(TEXT("ByTimerType")),
+		LOCTEXT("Sorting_ByTimerType_Name", "By Type"),
+		LOCTEXT("Sorting_ByTimerType_Title", "Sort By Type"),
+		LOCTEXT("Sorting_ByTimerType_Desc", "Sort by timer type (ascending or descending), then by name (ascending)."),
+		InColumnRef)
+{
+}
+
+////////////////////////////////////////////////////////////////////////////////////////////////////
+
+void FTimerNodeSortingByTimerType::Sort(TArray<Insights::FBaseTreeNodePtr>& NodesToSort, Insights::ESortMode SortMode) const
+{
+	if (SortMode == Insights::ESortMode::Ascending)
+	{
+		NodesToSort.Sort([](const Insights::FBaseTreeNodePtr& A, const Insights::FBaseTreeNodePtr& B) -> bool
+		{
+			ensure(A.IsValid() && A->GetTypeName() == FTimerNode::TypeName);
+			const FTimerNodePtr TimerNodeA = StaticCastSharedPtr<FTimerNode, Insights::FBaseTreeNode>(A);
+
+			ensure(B.IsValid() && B->GetTypeName() == FTimerNode::TypeName);
+			const FTimerNodePtr TimerNodeB = StaticCastSharedPtr<FTimerNode, Insights::FBaseTreeNode>(B);
+
+			if (TimerNodeA->GetType() == TimerNodeB->GetType())
+			{
+				INSIGHTS_DEFAULT_SORTING_NODES(A, B)
+			}
+			else
+			{
+				// Sort by timer type (ascending).
+				return TimerNodeA->GetType() < TimerNodeB->GetType();
+			}
+		});
+	}
+	else // if (SortMode == ESortMode::Descending)
+	{
+		NodesToSort.Sort([](const Insights::FBaseTreeNodePtr& A, const Insights::FBaseTreeNodePtr& B) -> bool
+		{
+			ensure(A.IsValid() && A->GetTypeName() == FTimerNode::TypeName);
+			const FTimerNodePtr TimerNodeA = StaticCastSharedPtr<FTimerNode, Insights::FBaseTreeNode>(A);
+
+			ensure(B.IsValid() && B->GetTypeName() == FTimerNode::TypeName);
+			const FTimerNodePtr TimerNodeB = StaticCastSharedPtr<FTimerNode, Insights::FBaseTreeNode>(B);
+
+			if (TimerNodeA->GetType() == TimerNodeB->GetType())
+			{
+				INSIGHTS_DEFAULT_SORTING_NODES(A, B)
+			}
+			else
+			{
+				// Sort by timer type (descending).
+				return TimerNodeB->GetType() < TimerNodeA->GetType();
+			}
+		});
+	}
+}
+
+////////////////////////////////////////////////////////////////////////////////////////////////////
+// Sort by Instance Count
+////////////////////////////////////////////////////////////////////////////////////////////////////
+
+FTimerNodeSortingByInstanceCount::FTimerNodeSortingByInstanceCount(TSharedRef<Insights::FTableColumn> InColumnRef)
+	: Insights::FTableCellValueSorter(
+		FName(TEXT("ByInstanceCount")),
+		LOCTEXT("Sorting_ByInstanceCount_Name", "By Instance Count"),
+		LOCTEXT("Sorting_ByInstanceCount_Title", "Sort By Instance Count"),
+		LOCTEXT("Sorting_ByInstanceCount_Desc", "Sort by aggregated instance count (ascending or descending), then by name (ascending)."),
+		InColumnRef)
+{
+}
+
+////////////////////////////////////////////////////////////////////////////////////////////////////
+
+void FTimerNodeSortingByInstanceCount::Sort(TArray<Insights::FBaseTreeNodePtr>& NodesToSort, Insights::ESortMode SortMode) const
+{
+	if (SortMode == Insights::ESortMode::Ascending)
+	{
+		NodesToSort.Sort([](const Insights::FBaseTreeNodePtr& A, const Insights::FBaseTreeNodePtr& B) -> bool
+		{
+			ensure(A.IsValid() && A->GetTypeName() == FTimerNode::TypeName);
+			const FTimerNodePtr TimerNodeA = StaticCastSharedPtr<FTimerNode, Insights::FBaseTreeNode>(A);
+			const uint64 ValueA = TimerNodeA->GetAggregatedStats().InstanceCount;
+
+			ensure(B.IsValid() && B->GetTypeName() == FTimerNode::TypeName);
+			const FTimerNodePtr TimerNodeB = StaticCastSharedPtr<FTimerNode, Insights::FBaseTreeNode>(B);
+			const uint64 ValueB = TimerNodeB->GetAggregatedStats().InstanceCount;
+
+			if (ValueA == ValueB)
+			{
+				INSIGHTS_DEFAULT_SORTING_NODES(A, B)
+			}
+			else
+			{
+				// Sort by instance count (ascending).
+				return ValueA < ValueB;
+			}
+		});
+	}
+	else // if (SortMode == ESortMode::Descending)
+	{
+		NodesToSort.Sort([](const Insights::FBaseTreeNodePtr& A, const Insights::FBaseTreeNodePtr& B) -> bool
+		{
+			ensure(A.IsValid() && A->GetTypeName() == FTimerNode::TypeName);
+			const FTimerNodePtr TimerNodeA = StaticCastSharedPtr<FTimerNode, Insights::FBaseTreeNode>(A);
+			const uint64 ValueA = TimerNodeA->GetAggregatedStats().InstanceCount;
+
+			ensure(B.IsValid() && B->GetTypeName() == FTimerNode::TypeName);
+			const FTimerNodePtr TimerNodeB = StaticCastSharedPtr<FTimerNode, Insights::FBaseTreeNode>(B);
+			const uint64 ValueB = TimerNodeB->GetAggregatedStats().InstanceCount;
+
+			if (ValueA == ValueB)
+			{
+				INSIGHTS_DEFAULT_SORTING_NODES(A, B)
+			}
+			else
+			{
+				// Sort by instance count (descending).
+				return ValueB < ValueA;
+			}
+		});
+	}
+}
+
+////////////////////////////////////////////////////////////////////////////////////////////////////
+// Sort by Total Inclusive Time
+////////////////////////////////////////////////////////////////////////////////////////////////////
+
+FTimerNodeSortingByTotalInclusiveTime::FTimerNodeSortingByTotalInclusiveTime(TSharedRef<Insights::FTableColumn> InColumnRef)
+	: Insights::FTableCellValueSorter(
+		FName(TEXT("ByTotalInclusiveTime")),
+		LOCTEXT("Sorting_ByTotalInclusiveTime_Name", "By Total Inclusive Time"),
+		LOCTEXT("Sorting_ByTotalInclusiveTime_Title", "Sort By Total Inxclusive Time"),
+		LOCTEXT("Sorting_ByTotalInclusiveTime_Desc", "Sort by aggregated total inclusive time (ascending or descending), then by name (ascending)."),
+		InColumnRef)
+{
+}
+
+////////////////////////////////////////////////////////////////////////////////////////////////////
+
+void FTimerNodeSortingByTotalInclusiveTime::Sort(TArray<Insights::FBaseTreeNodePtr>& NodesToSort, Insights::ESortMode SortMode) const
+{
+	if (SortMode == Insights::ESortMode::Ascending)
+	{
+		NodesToSort.Sort([](const Insights::FBaseTreeNodePtr& A, const Insights::FBaseTreeNodePtr& B) -> bool
+		{
+			ensure(A.IsValid() && A->GetTypeName() == FTimerNode::TypeName);
+			const FTimerNodePtr TimerNodeA = StaticCastSharedPtr<FTimerNode, Insights::FBaseTreeNode>(A);
+			const double ValueA = TimerNodeA->GetAggregatedStats().TotalInclusiveTime;
+
+			ensure(B.IsValid() && B->GetTypeName() == FTimerNode::TypeName);
+			const FTimerNodePtr TimerNodeB = StaticCastSharedPtr<FTimerNode, Insights::FBaseTreeNode>(B);
+			const double ValueB = TimerNodeB->GetAggregatedStats().TotalInclusiveTime;
+
+			if (ValueA == ValueB)
+			{
+				INSIGHTS_DEFAULT_SORTING_NODES(A, B)
+			}
+			else
+			{
+				// Sort by total inclusive time (ascending).
+				return ValueA < ValueB;
+			}
+		});
+	}
+	else // if (SortMode == ESortMode::Descending)
+	{
+		NodesToSort.Sort([](const Insights::FBaseTreeNodePtr& A, const Insights::FBaseTreeNodePtr& B) -> bool
+		{
+			ensure(A.IsValid() && A->GetTypeName() == FTimerNode::TypeName);
+			const FTimerNodePtr TimerNodeA = StaticCastSharedPtr<FTimerNode, Insights::FBaseTreeNode>(A);
+			const double ValueA = TimerNodeA->GetAggregatedStats().TotalInclusiveTime;
+
+			ensure(B.IsValid() && B->GetTypeName() == FTimerNode::TypeName);
+			const FTimerNodePtr TimerNodeB = StaticCastSharedPtr<FTimerNode, Insights::FBaseTreeNode>(B);
+			const double ValueB = TimerNodeB->GetAggregatedStats().TotalInclusiveTime;
+
+			if (ValueA == ValueB)
+			{
+				INSIGHTS_DEFAULT_SORTING_NODES(A, B)
+			}
+			else
+			{
+				// Sort by total inclusive time (descending).
+				return ValueB < ValueA;
+			}
+		});
+	}
+}
+
+////////////////////////////////////////////////////////////////////////////////////////////////////
+// Sort by Total Exclusive Time
+////////////////////////////////////////////////////////////////////////////////////////////////////
+
+FTimerNodeSortingByTotalExclusiveTime::FTimerNodeSortingByTotalExclusiveTime(TSharedRef<Insights::FTableColumn> InColumnRef)
+	: Insights::FTableCellValueSorter(
+		FName(TEXT("ByTotalExclusiveTime")),
+		LOCTEXT("Sorting_ByTotalExclusiveTime_Name", "By Total Exclusive Time"),
+		LOCTEXT("Sorting_ByTotalExclusiveTime_Title", "Sort By Total Exclusive Time"),
+		LOCTEXT("Sorting_ByTotalExclusiveTime_Desc", "Sort by aggregated total exclusive time (ascending or descending), then by name (ascending)."),
+		InColumnRef)
+{
+}
+
+////////////////////////////////////////////////////////////////////////////////////////////////////
+
+void FTimerNodeSortingByTotalExclusiveTime::Sort(TArray<Insights::FBaseTreeNodePtr>& NodesToSort, Insights::ESortMode SortMode) const
+{
+	if (SortMode == Insights::ESortMode::Ascending)
+	{
+		NodesToSort.Sort([](const Insights::FBaseTreeNodePtr& A, const Insights::FBaseTreeNodePtr& B) -> bool
+		{
+			ensure(A.IsValid() && A->GetTypeName() == FTimerNode::TypeName);
+			const FTimerNodePtr TimerNodeA = StaticCastSharedPtr<FTimerNode, Insights::FBaseTreeNode>(A);
+			const double ValueA = TimerNodeA->GetAggregatedStats().TotalExclusiveTime;
+
+			ensure(B.IsValid() && B->GetTypeName() == FTimerNode::TypeName);
+			const FTimerNodePtr TimerNodeB = StaticCastSharedPtr<FTimerNode, Insights::FBaseTreeNode>(B);
+			const double ValueB = TimerNodeB->GetAggregatedStats().TotalExclusiveTime;
+
+			if (ValueA == ValueB)
+			{
+				INSIGHTS_DEFAULT_SORTING_NODES(A, B)
+			}
+			else
+			{
+				// Sort by total exclusive time (ascending).
+				return ValueA < ValueB;
+			}
+		});
+	}
+	else // if (SortMode == ESortMode::Descending)
+	{
+		NodesToSort.Sort([](const Insights::FBaseTreeNodePtr& A, const Insights::FBaseTreeNodePtr& B) -> bool
+		{
+			ensure(A.IsValid() && A->GetTypeName() == FTimerNode::TypeName);
+			const FTimerNodePtr TimerNodeA = StaticCastSharedPtr<FTimerNode, Insights::FBaseTreeNode>(A);
+			const double ValueA = TimerNodeA->GetAggregatedStats().TotalExclusiveTime;
+
+			ensure(B.IsValid() && B->GetTypeName() == FTimerNode::TypeName);
+			const FTimerNodePtr TimerNodeB = StaticCastSharedPtr<FTimerNode, Insights::FBaseTreeNode>(B);
+			const double ValueB = TimerNodeB->GetAggregatedStats().TotalExclusiveTime;
+
+			if (ValueA == ValueB)
+			{
+				INSIGHTS_DEFAULT_SORTING_NODES(A, B)
+			}
+			else
+			{
+				// Sort by total exclusive time (descending).
+				return ValueB < ValueA;
+			}
+		});
+	}
+}
+
+////////////////////////////////////////////////////////////////////////////////////////////////////
+// Grouping by Timer Type
+////////////////////////////////////////////////////////////////////////////////////////////////////
+
+FTimerNodeGroupingByTimerType::FTimerNodeGroupingByTimerType()
+	: FTreeNodeGrouping(
+		LOCTEXT("Grouping_ByTimerType_ShortName", "Type"),
+		LOCTEXT("Grouping_ByTimerType_TitleName", "Type"),
+		LOCTEXT("Grouping_ByTimerType_Desc", "Creates a group for each timer type."),
+		TEXT("Profiler.FiltersAndPresets.StatTypeIcon"), //TODO
+		nullptr)
+{
+}
+
+////////////////////////////////////////////////////////////////////////////////////////////////////
+
+Insights::FTreeNodeGroupInfo FTimerNodeGroupingByTimerType::GetGroupForNode(const Insights::FBaseTreeNodePtr InNodePtr) const
+{
+	if (InNodePtr->GetTypeName() == FTimerNode::TypeName)
+	{
+		const FTimerNodePtr& TimerNodePtr = StaticCastSharedPtr<FTimerNode>(InNodePtr);
+		return { FName(*TimerNodeTypeHelper::ToText(TimerNodePtr->GetType()).ToString()), true };
+	}
+	else
+	{
+		return { FName("Unknown"), true };
+	}
+}
+
+////////////////////////////////////////////////////////////////////////////////////////////////////
 
 #undef LOCTEXT_NAMESPACE
diff --git a/Engine/Source/Developer/TraceInsights/Private/Insights/ViewModels/TimerGroupingAndSorting.h b/Engine/Source/Developer/TraceInsights/Private/Insights/ViewModels/TimerGroupingAndSorting.h
--- a/Engine/Source/Developer/TraceInsights/Private/Insights/ViewModels/TimerGroupingAndSorting.h
+++ b/Engine/Source/Developer/TraceInsights/Private/Insights/ViewModels/TimerGroupingAndSorting.h
@@ -3,10 +3,58 @@
 #pragma once
 
 #include "CoreMinimal.h"
+
+// Insights
 #include "Insights/ViewModels/TimerNode.h"
+#include "Insights/Table/ViewModels/TableColumn.h"
+#include "Insights/Table/ViewModels/TableCellValueSorter.h"
+#include "Insights/Table/ViewModels/TreeNodeGrouping.h"
 
 ////////////////////////////////////////////////////////////////////////////////////////////////////
-// Grouping
+// Sorters
+////////////////////////////////////////////////////////////////////////////////////////////////////
+
+class FTimerNodeSortingByTimerType: public Insights::FTableCellValueSorter
+{
+public:
+	FTimerNodeSortingByTimerType(TSharedRef<Insights::FTableColumn> InColumnRef);
+
+	virtual void Sort(TArray<Insights::FBaseTreeNodePtr>& NodesToSort, Insights::ESortMode SortMode) const override;
+};
+
+////////////////////////////////////////////////////////////////////////////////////////////////////
+
+class FTimerNodeSortingByInstanceCount : public Insights::FTableCellValueSorter
+{
+public:
+	FTimerNodeSortingByInstanceCount(TSharedRef<Insights::FTableColumn> InColumnRef);
+
+	virtual void Sort(TArray<Insights::FBaseTreeNodePtr>& NodesToSort, Insights::ESortMode SortMode) const override;
+};
+
+////////////////////////////////////////////////////////////////////////////////////////////////////
+
+class FTimerNodeSortingByTotalInclusiveTime : public Insights::FTableCellValueSorter
+{
+public:
+	FTimerNodeSortingByTotalInclusiveTime(TSharedRef<Insights::FTableColumn> InColumnRef);
+
+	virtual void Sort(TArray<Insights::FBaseTreeNodePtr>& NodesToSort, Insights::ESortMode SortMode) const override;
+};
+
+////////////////////////////////////////////////////////////////////////////////////////////////////
+
+class FTimerNodeSortingByTotalExclusiveTime : public Insights::FTableCellValueSorter
+{
+public:
+	FTimerNodeSortingByTotalExclusiveTime(TSharedRef<Insights::FTableColumn> InColumnRef);
+
+	virtual void Sort(TArray<Insights::FBaseTreeNodePtr>& NodesToSort, Insights::ESortMode SortMode) const override;
+};
+
+////////////////////////////////////////////////////////////////////////////////////////////////////
+// Organizers
+////////////////////////////////////////////////////////////////////////////////////////////////////
 
 /** Enumerates types of grouping or sorting for the timer nodes. */
 enum class ETimerGroupingMode
@@ -37,3 +85,15 @@ enum class ETimerGroupingMode
 typedef TSharedPtr<ETimerGroupingMode> ETimerGroupingModePtr;
 
 ////////////////////////////////////////////////////////////////////////////////////////////////////
+
+/** Creates a group for each timer type. */
+class FTimerNodeGroupingByTimerType : public Insights::FTreeNodeGrouping
+{
+public:
+	FTimerNodeGroupingByTimerType();
+	virtual ~FTimerNodeGroupingByTimerType() {}
+
+	virtual Insights::FTreeNodeGroupInfo GetGroupForNode(const Insights::FBaseTreeNodePtr InNode) const override;
+};
+
+////////////////////////////////////////////////////////////////////////////////////////////////////
diff --git a/Engine/Source/Developer/TraceInsights/Private/Insights/ViewModels/TimerNode.cpp b/Engine/Source/Developer/TraceInsights/Private/Insights/ViewModels/TimerNode.cpp
--- a/Engine/Source/Developer/TraceInsights/Private/Insights/ViewModels/TimerNode.cpp
+++ b/Engine/Source/Developer/TraceInsights/Private/Insights/ViewModels/TimerNode.cpp
@@ -2,8 +2,14 @@
 
 #include "TimerNode.h"
 
+// Insights
+#include "Insights/Table/ViewModels/TableColumn.h"
+#include "Insights/ViewModels/TimersViewColumnFactory.h"
+
 #define LOCTEXT_NAMESPACE "TimerNode"
 
+const FName FTimerNode::TypeName(TEXT("FTimerNode"));
+
 ////////////////////////////////////////////////////////////////////////////////////////////////////
 
 void FTimerNode::ResetAggregatedStats()
diff --git a/Engine/Source/Developer/TraceInsights/Private/Insights/ViewModels/TimerNode.h b/Engine/Source/Developer/TraceInsights/Private/Insights/ViewModels/TimerNode.h
--- a/Engine/Source/Developer/TraceInsights/Private/Insights/ViewModels/TimerNode.h
+++ b/Engine/Source/Developer/TraceInsights/Private/Insights/ViewModels/TimerNode.h
@@ -51,7 +51,8 @@ typedef TWeakPtr<class FTimerNode> FTimerNodeWeak;
 class FTimerNode : public Insights::FBaseTreeNode
 {
 public:
-	static const uint64 InvalidId = -1;
+	static const FName TypeName;
+	//static constexpr uint64 InvalidTimerId = -1;
 
 public:
 	/** Initialization constructor for the timer node. */
@@ -59,6 +60,7 @@ public:
 		: FBaseTreeNode(InId, InName, InType == ETimerNodeType::Group)
 		, MetaGroupName(InMetaGroupName)
 		, Type(InType)
+		, bIsHotPath(false)
 	{
 		ResetAggregatedStats();
 	}
@@ -67,10 +69,13 @@ public:
 	FTimerNode(const FName InGroupName)
 		: FBaseTreeNode(0, InGroupName, true)
 		, Type(ETimerNodeType::Group)
+		, bIsHotPath(false)
 	{
 		ResetAggregatedStats();
 	}
 
+	virtual const FName& GetTypeName() const override { return TypeName; }
+
 	/**
 	 * @return a name of the meta group that this timer node belongs to, taken from the metadata.
 	 */
@@ -110,6 +115,9 @@ public:
 		Algo::SortBy(GetChildrenMutable(), Projection, Instance);
 	}
 
+	bool IsHotPath() const { return bIsHotPath; }
+	void SetIsHotPath(bool bOnOff) { bIsHotPath = bOnOff; }
+
 private:
 	/** The name of the meta group that this timer belongs to, based on the timer's metadata; only valid for timer nodes. */
 	const FName MetaGroupName;
@@ -119,6 +127,9 @@ private:
 
 	/** Aggregated stats. */
 	Trace::FAggregatedTimingStats AggregatedStats;
+
+	/** True if this tree node is on the hot path. */
+	bool bIsHotPath;
 };
 
 ////////////////////////////////////////////////////////////////////////////////////////////////////
diff --git a/Engine/Source/Developer/TraceInsights/Private/Insights/ViewModels/TimerNodeHelper.cpp b/Engine/Source/Developer/TraceInsights/Private/Insights/ViewModels/TimerNodeHelper.cpp
--- a/Engine/Source/Developer/TraceInsights/Private/Insights/ViewModels/TimerNodeHelper.cpp
+++ b/Engine/Source/Developer/TraceInsights/Private/Insights/ViewModels/TimerNodeHelper.cpp
@@ -9,7 +9,7 @@
 // TimerNode Type Helper
 ////////////////////////////////////////////////////////////////////////////////////////////////////
 
-FText TimerNodeTypeHelper::ToName(const ETimerNodeType NodeType)
+FText TimerNodeTypeHelper::ToText(const ETimerNodeType NodeType)
 {
 	static_assert(static_cast<int>(ETimerNodeType::InvalidOrMax) == 4, "Not all cases are handled in switch below!?");
 	switch (NodeType)
@@ -78,7 +78,7 @@ const FSlateBrush* TimerNodeTypeHelper::GetIconForTimerNodeType(const ETimerNode
 // TimerNode Grouping Helper
 ////////////////////////////////////////////////////////////////////////////////////////////////////
 
-FText TimerNodeGroupingHelper::ToName(const ETimerGroupingMode GroupingMode)
+FText TimerNodeGroupingHelper::ToText(const ETimerGroupingMode GroupingMode)
 {
 	static_assert(static_cast<int>(ETimerGroupingMode::InvalidOrMax) == 7, "Not all cases are handled in switch below!?");
 	switch (GroupingMode)
diff --git a/Engine/Source/Developer/TraceInsights/Private/Insights/ViewModels/TimerNodeHelper.h b/Engine/Source/Developer/TraceInsights/Private/Insights/ViewModels/TimerNodeHelper.h
--- a/Engine/Source/Developer/TraceInsights/Private/Insights/ViewModels/TimerNodeHelper.h
+++ b/Engine/Source/Developer/TraceInsights/Private/Insights/ViewModels/TimerNodeHelper.h
@@ -19,7 +19,7 @@ struct TimerNodeTypeHelper
 	 *
 	 * @return text representation of the specified ETimerNodeType value.
 	 */
-	static FText ToName(const ETimerNodeType Type);
+	static FText ToText(const ETimerNodeType Type);
 
 	/**
 	 * @param Type - The value to get the text for.
@@ -49,7 +49,7 @@ struct TimerNodeGroupingHelper
 	 *
 	 * @return text representation of the specified ETimerGroupingMode value.
 	 */
-	static FText ToName(const ETimerGroupingMode TimerGroupingMode);
+	static FText ToText(const ETimerGroupingMode TimerGroupingMode);
 
 	/**
 	 * @param TimerGroupingMode - The value to get the text for.
diff --git a/Engine/Source/Developer/TraceInsights/Private/Insights/ViewModels/TimersViewColumn.h b/Engine/Source/Developer/TraceInsights/Private/Insights/ViewModels/TimersViewColumn.h
--- a/Engine/Source/Developer/TraceInsights/Private/Insights/ViewModels/TimersViewColumn.h
+++ b/Engine/Source/Developer/TraceInsights/Private/Insights/ViewModels/TimersViewColumn.h
@@ -1,118 +1,3 @@
 // Copyright 1998-2019 Epic Games, Inc. All Rights Reserved.
 
 #pragma once
-
-#include "CoreMinimal.h"
-#include "Types/SlateEnums.h"
-#include "Misc/EnumClassFlags.h"
-
-class FTimerNode; // TODO: IInsightsTreeNode; generic formatter; see also StatsViewColumn.h
-
-enum class ETimersViewColumnFlags : uint32
-{
-	None = 0,
-
-	CanBeHidden = (1 << 0),
-	CanBeSorted = (1 << 1),
-	CanBeFiltered = (1 << 2),
-};
-ENUM_CLASS_FLAGS(ETimersViewColumnFlags);
-
-/** Holds information about a column in the timer view widget. */
-class FTimersViewColumn
-{
-	friend struct FTimersViewColumnFactory;
-
-public:
-	typedef TFunction<FText(const FTimersViewColumn& Column, const FTimerNode& TimerNode)> FGetFormattedValueFn;
-
-public:
-	/** Whether this column can be hidden. */
-	bool bCanBeHidden() const { return EnumHasAnyFlags(Flags, ETimersViewColumnFlags::CanBeHidden); }
-
-	/** Whether this column cab be used for sorting. */
-	bool bCanBeSorted() const { return EnumHasAnyFlags(Flags, ETimersViewColumnFlags::CanBeSorted); }
-
-	/** Where this column can be used to filtering displayed results. */
-	bool bCanBeFiltered() const { return EnumHasAnyFlags(Flags, ETimersViewColumnFlags::CanBeFiltered); }
-
-	/** If MinColumnWidth == MaxColumnWidth, this column has fixed width and cannot be resized. */
-	bool bIsFixedColumnWidth() const { return MinColumnWidth == MaxColumnWidth; }
-
-	FText GetFormattedValue(const FTimerNode& TimerNode) const
-	{
-		return GetFormattedValueFn(*this, TimerNode);
-	}
-
-protected:
-	/** No default constructor. */
-	FTimersViewColumn() = delete;
-
-	/** Initialization constructor, only used in FStatsViewColumnFactory. */
-	FTimersViewColumn
-	(
-		int32 InOrder,
-		const FName InId,
-		const FName InSearchId,
-		FText InShortName,
-		FText InTitleName,
-		FText InDescription,
-		bool bInIsVisible,
-		const ETimersViewColumnFlags InFlags,
-		const EHorizontalAlignment InHorizontalAlignment,
-		const float InInitialColumnWidth,
-		const float InMinColumnWidth,
-		const float InMaxColumnWidth,
-		const FGetFormattedValueFn InGetFormattedValueFn
-	)
-		: Order(InOrder)
-		, Id(InId)
-		, SearchId(InSearchId)
-		, ShortName(MoveTemp(InShortName))
-		, TitleName(MoveTemp(InTitleName))
-		, Description(MoveTemp(InDescription))
-		, bIsVisible(bInIsVisible)
-		, Flags(InFlags)
-		, HorizontalAlignment(InHorizontalAlignment)
-		, InitialColumnWidth(InInitialColumnWidth)
-		, MinColumnWidth(InMinColumnWidth)
-		, MaxColumnWidth(InMaxColumnWidth)
-		, GetFormattedValueFn(InGetFormattedValueFn)
-	{
-	}
-
-public:
-	/** Order value, to sort columns in the tree view. */
-	int32 Order;
-
-	/** Name of the column, name of the property. */
-	FName Id;
-
-	/** Name of the column used by the searching system. */
-	FName SearchId;
-
-	/** Short name of the column, displayed in the column header. */
-	FText ShortName;
-
-	/** Title name of the column, displayed as title in the column tooltip. */
-	FText TitleName;
-
-	/** Long name of the column, displayed in the column tooltip. */
-	FText Description;
-
-	/** Is this column visible? */
-	bool bIsVisible;
-
-	/** On/off switches. */
-	ETimersViewColumnFlags Flags;
-
-	/** Horizontal alignment of the content in this column. */
-	EHorizontalAlignment HorizontalAlignment;
-
-	float InitialColumnWidth; /**< Initial column width. */
-	float MinColumnWidth; /**< Minimum column width. */
-	float MaxColumnWidth; /**< Maximum column width. */
-
-	/** Custom function used to format (as an FText) the value of the timer node to be displayed by this column. */
-	FGetFormattedValueFn GetFormattedValueFn;
-};
