[Insights]

diff --git a/Engine/Source/Developer/TraceInsights/Private/Insights/ViewModels/TimersViewColumnFactory.cpp b/Engine/Source/Developer/TraceInsights/Private/Insights/ViewModels/TimersViewColumnFactory.cpp
--- a/Engine/Source/Developer/TraceInsights/Private/Insights/ViewModels/TimersViewColumnFactory.cpp
+++ b/Engine/Source/Developer/TraceInsights/Private/Insights/ViewModels/TimersViewColumnFactory.cpp
@@ -4,6 +4,10 @@
 
 // Insights
 #include "Insights/Common/TimeUtils.h"
+#include "Insights/Table/ViewModels/TableCellValueFormatter.h"
+#include "Insights/Table/ViewModels/TableCellValueGetter.h"
+#include "Insights/Table/ViewModels/TableCellValueSorter.h"
+#include "Insights/ViewModels/TimerGroupingAndSorting.h"
 #include "Insights/ViewModels/TimerNodeHelper.h"
 
 #define LOCTEXT_NAMESPACE "STimerView"
@@ -11,7 +15,7 @@
 ////////////////////////////////////////////////////////////////////////////////////////////////////
 // Column identifiers
 
-const FName FTimersViewColumns::NameColumnID(TEXT("Name"));
+const FName FTimersViewColumns::NameColumnID(TEXT("Name")); // TEXT("_Hierarchy")
 const FName FTimersViewColumns::MetaGroupNameColumnID(TEXT("MetaGroupName"));
 const FName FTimersViewColumns::TypeColumnID(TEXT("Type"));
 const FName FTimersViewColumns::InstanceCountColumnID(TEXT("Count"));
@@ -36,319 +40,669 @@ const FName FTimersViewColumns::MinExclusiveTimeColumnID(TEXT("MinExclTime"));
 
 ////////////////////////////////////////////////////////////////////////////////////////////////////
 
-FTimersViewColumnFactory::FTimersViewColumnFactory()
+void FTimersViewColumnFactory::CreateTimersViewColumns(TArray<TSharedPtr<Insights::FTableColumn>>& Columns)
 {
-	uint32 Index = 0;
-
-	// FTimerViewColumn
-	// (
-	//     Order,
-	//     Id,
-	//     SearchId,
-	//     ShortName,
-	//     TitleName,
-	//     Description,
-	//     bIsVisible,
-	//     Flags,
-	//     HorizontalAlignment,
-	//     InitialWidth, MinWidth, MaxWidth,
-	//     GetFormattedValueFn
-	// )
-
-	//////////////////////////////////////////////////
-	// Name
-
-	const ETimersViewColumnFlags NameColumnFlags = ETimersViewColumnFlags::CanBeSorted |
-												   ETimersViewColumnFlags::CanBeFiltered;
-
-	Collection.Add(new FTimersViewColumn
-	(
-		Index++,
-		FTimersViewColumns::NameColumnID,
-		TEXT("name"),
-		LOCTEXT("TimerNameColumnName", "Name"),
-		LOCTEXT("TimerNameColumnTitle", "Timer or Group Name"),
-		LOCTEXT("TimerNameColumnDesc", "Name of the timer or group"),
-		true,
-		NameColumnFlags,
-		HAlign_Left,
-		246.0f, 10.0f, FLT_MAX,
-		[](const FTimersViewColumn& Column, const FTimerNode& TimerNode) -> FText
-		{ return FText::FromName(TimerNode.GetName()); }
-	));
-
-	//////////////////////////////////////////////////
-	// Meta Group Name
-
-	const ETimersViewColumnFlags MetaGroupNameColumnFlags = ETimersViewColumnFlags::CanBeHidden |
-															ETimersViewColumnFlags::CanBeSorted |
-															ETimersViewColumnFlags::CanBeFiltered;
-
-	Collection.Add(new FTimersViewColumn
-	(
-		Index++,
-		FTimersViewColumns::MetaGroupNameColumnID,
-		TEXT("metagroupname"),
-		LOCTEXT("TimerMetaGroupNameColumnName", "Meta Group"),
-		LOCTEXT("TimerMetaGroupNameColumnTitle", "Meta Group Name"),
-		LOCTEXT("TimerMetaGroupNameColumnDesc", "Name of the meta group"),
-		false,
-		MetaGroupNameColumnFlags,
-		HAlign_Left,
-		100.0f, 10.0f, FLT_MAX,
-		[](const FTimersViewColumn& Column, const FTimerNode& TimerNode) -> FText
-		{ return FText::FromName(TimerNode.GetMetaGroupName()); }
-	));
-
-	//////////////////////////////////////////////////
-	// Type
-
-	const ETimersViewColumnFlags TypeColumnFlags = ETimersViewColumnFlags::CanBeHidden |
-												   ETimersViewColumnFlags::CanBeSorted |
-												   ETimersViewColumnFlags::CanBeFiltered;
-
-	Collection.Add(new FTimersViewColumn
-	(
-		Index++,
-		FTimersViewColumns::TypeColumnID,
-		TEXT("type"),
-		LOCTEXT("TimerTypeColumnName", "Type"),
-		LOCTEXT("TimerTypeColumnTitle", "Type"),
-		LOCTEXT("TimerTypeColumnDesc", "Type of timer or group"),
-		false,
-		TypeColumnFlags,
-		HAlign_Left,
-		60.0f, 10.0f, FLT_MAX,
-		[](const FTimersViewColumn& Column, const FTimerNode& TimerNode) -> FText
-		{ return TimerNodeTypeHelper::ToName(TimerNode.GetType()); }
-	));
-
-	//////////////////////////////////////////////////
-
-	const ETimersViewColumnFlags AggregatedStatsColumnFlags = ETimersViewColumnFlags::CanBeHidden |
-															  ETimersViewColumnFlags::CanBeSorted |
-															  ETimersViewColumnFlags::CanBeFiltered;
-
-	//////////////////////////////////////////////////
-	// Instance Count
-
-	Collection.Add(new FTimersViewColumn
-	(
-		Index++,
-		FTimersViewColumns::InstanceCountColumnID,
-		TEXT("count"),
-		LOCTEXT("InstanceCountName", "Count"),
-		LOCTEXT("InstanceCountTitle", "Instance Count"),
-		LOCTEXT("InstanceCountDesc", "Number of timer's instances"),
-		true,
-		AggregatedStatsColumnFlags,
-		HAlign_Right,
-		60.0f, 0.0f, FLT_MAX,
-		[](const FTimersViewColumn& Column, const FTimerNode& TimerNode) -> FText
-		{ return FText::AsNumber(TimerNode.GetAggregatedStats().InstanceCount); }
-	));
-
-	//////////////////////////////////////////////////
-
-	constexpr float TotalTimeColumnInitialWidth = 60.0f;
-	constexpr float TotalTimeColumnMinWidth = 4.0f;
-	constexpr float TotalTimeColumnMaxWidth = FLT_MAX;
-
-	constexpr float TimeMsColumnInitialWidth = 50.0f;
-	constexpr float TimeMsColumnMinWidth = 4.0f;
-	constexpr float TimeMsColumnMaxWidth = FLT_MAX;
-
-	//////////////////////////////////////////////////
-	// Inclusive Time stats
-
-	Collection.Add(new FTimersViewColumn
-	(
-		Index++,
-		FTimersViewColumns::TotalInclusiveTimeColumnID,
-		TEXT("inc"),
-		LOCTEXT("TotalInclusiveTimeName", "Incl"),
-		LOCTEXT("TotalInclusiveTimeTitle", "Total Inclusive Time"),
-		LOCTEXT("TotalInclusiveTimeDesc", "Total inclusive duration of selected timer's instances"),
-		true,
-		AggregatedStatsColumnFlags,
-		HAlign_Right,
-		TotalTimeColumnInitialWidth, TotalTimeColumnMinWidth, TotalTimeColumnMaxWidth,
-		[](const FTimersViewColumn& Column, const FTimerNode& TimerNode) -> FText
-		{ return FText::FromString(TimeUtils::FormatTimeAuto(TimerNode.GetAggregatedStats().TotalInclusiveTime)); }
-	));
-
-	Collection.Add(new FTimersViewColumn
-	(
-		Index++,
-		FTimersViewColumns::MaxInclusiveTimeColumnID,
-		TEXT("maxinc"),
-		LOCTEXT("MaxInclusiveTimeName", "I.Max"),
-		LOCTEXT("MaxInclusiveTimeTitle", "Max Inclusive Time (ms)"),
-		LOCTEXT("MaxInclusiveTimeDesc", "Maximum inclusive duration of selected timer's instances, in milliseconds"),
-		false,
-		AggregatedStatsColumnFlags,
-		HAlign_Right,
-		TimeMsColumnInitialWidth, TimeMsColumnMinWidth, TimeMsColumnMaxWidth,
-		[](const FTimersViewColumn& Column, const FTimerNode& TimerNode) -> FText
-		{ return FText::FromString(TimeUtils::FormatTimeMs(TimerNode.GetAggregatedStats().MaxInclusiveTime)); }
-	));
-
-	Collection.Add(new FTimersViewColumn
-	(
-		Index++,
-		FTimersViewColumns::AverageInclusiveTimeColumnID,
-		TEXT("avginc"),
-		LOCTEXT("AvgInclusiveTimeName", "I.Avg"),
-		LOCTEXT("AvgInclusiveTimeTitle", "Average Inclusive Time (ms)"),
-		LOCTEXT("AvgInclusiveTimeDesc", "Average inclusive duration of selected timer's instances, in milliseconds"),
-		false,
-		AggregatedStatsColumnFlags,
-		HAlign_Right,
-		TimeMsColumnInitialWidth, TimeMsColumnMinWidth, TimeMsColumnMaxWidth,
-		[](const FTimersViewColumn& Column, const FTimerNode& TimerNode) -> FText
-		{ return FText::FromString(TimeUtils::FormatTimeMs(TimerNode.GetAggregatedStats().AverageInclusiveTime)); }
-	));
-
-	Collection.Add(new FTimersViewColumn
-	(
-		Index++,
-		FTimersViewColumns::MedianInclusiveTimeColumnID,
-		TEXT("medinc"),
-		LOCTEXT("MedInclusiveTimeName", "I.Med"),
-		LOCTEXT("MedInclusiveTimeTitle", "Median Inclusive Time (ms)"),
-		LOCTEXT("MedInclusiveTimeDesc", "Median inclusive duration of selected timer's instances, in milliseconds"),
-		false,
-		AggregatedStatsColumnFlags,
-		HAlign_Right,
-		TimeMsColumnInitialWidth, TimeMsColumnMinWidth, TimeMsColumnMaxWidth,
-		[](const FTimersViewColumn& Column, const FTimerNode& TimerNode) -> FText
-		{ return FText::FromString(TimeUtils::FormatTimeMs(TimerNode.GetAggregatedStats().MedianInclusiveTime)); }
-	));
-
-	Collection.Add(new FTimersViewColumn
-	(
-		Index++,
-		FTimersViewColumns::MinInclusiveTimeColumnID,
-		TEXT("mininc"),
-		LOCTEXT("MinInclusiveTimeName", "I.Min"),
-		LOCTEXT("MinInclusiveTimeTitle", "Min Inclusive Time (ms)"),
-		LOCTEXT("MinInclusiveTimeDesc", "Minimum inclusive duration of selected timer's instances, in milliseconds"),
-		false,
-		AggregatedStatsColumnFlags,
-		HAlign_Right,
-		TimeMsColumnInitialWidth, TimeMsColumnMinWidth, TimeMsColumnMaxWidth,
-		[](const FTimersViewColumn& Column, const FTimerNode& TimerNode) -> FText
-		{ return FText::FromString(TimeUtils::FormatTimeMs(TimerNode.GetAggregatedStats().MinInclusiveTime)); }
-	));
-
-	//////////////////////////////////////////////////
-	// Exclusive Time stats
-
-	Collection.Add(new FTimersViewColumn
-	(
-		Index++,
-		FTimersViewColumns::TotalExclusiveTimeColumnID,
-		TEXT("exc"),
-		LOCTEXT("TotalExclusiveTimeName", "Excl"),
-		LOCTEXT("TotalExclusiveTimeTitle", "Total Exclusive Time"),
-		LOCTEXT("TotalExclusiveTimeDesc", "Total exclusive duration of selected timer's instances"),
-		true,
-		AggregatedStatsColumnFlags,
-		HAlign_Right,
-		TotalTimeColumnInitialWidth, TotalTimeColumnMinWidth, TotalTimeColumnMaxWidth,
-		[](const FTimersViewColumn& Column, const FTimerNode& TimerNode) -> FText
-		{ return FText::FromString(TimeUtils::FormatTimeAuto(TimerNode.GetAggregatedStats().TotalExclusiveTime)); }
-	));
-
-	Collection.Add(new FTimersViewColumn
-	(
-		Index++,
-		FTimersViewColumns::MaxExclusiveTimeColumnID,
-		TEXT("maxexc"),
-		LOCTEXT("MaxExclusiveTimeName", "E.Max"),
-		LOCTEXT("MaxExclusiveTimeTitle", "Max Exclusive Time (ms)"),
-		LOCTEXT("MaxExclusiveTimeDesc", "Maximum exclusive duration of selected timer's instances, in milliseconds"),
-		false,
-		AggregatedStatsColumnFlags,
-		HAlign_Right,
-		TimeMsColumnInitialWidth, TimeMsColumnMinWidth, TimeMsColumnMaxWidth,
-		[](const FTimersViewColumn& Column, const FTimerNode& TimerNode) -> FText
-		{ return FText::FromString(TimeUtils::FormatTimeMs(TimerNode.GetAggregatedStats().MaxExclusiveTime)); }
-	));
-
-	Collection.Add(new FTimersViewColumn
-	(
-		Index++,
-		FTimersViewColumns::AverageExclusiveTimeColumnID,
-		TEXT("avgexc"),
-		LOCTEXT("AvgExclusiveTimeName", "E.Avg"),
-		LOCTEXT("AvgExclusiveTimeTitle", "Average Exclusive Time (ms)"),
-		LOCTEXT("AvgExclusiveTimeDesc", "Average exclusive duration of selected timer's instances, in milliseconds"),
-		false,
-		AggregatedStatsColumnFlags,
-		HAlign_Right,
-		TimeMsColumnInitialWidth, TimeMsColumnMinWidth, TimeMsColumnMaxWidth,
-		[](const FTimersViewColumn& Column, const FTimerNode& TimerNode) -> FText
-		{ return FText::FromString(TimeUtils::FormatTimeMs(TimerNode.GetAggregatedStats().AverageExclusiveTime)); }
-	));
-
-	Collection.Add(new FTimersViewColumn
-	(
-		Index++,
-		FTimersViewColumns::MedianExclusiveTimeColumnID,
-		TEXT("medexc"),
-		LOCTEXT("MedExclusiveTimeName", "E.Med"),
-		LOCTEXT("MedExclusiveTimeTitle", "Median Exclusive Time (ms)"),
-		LOCTEXT("MedExclusiveTimeDesc", "Median exclusive duration of timer's instances, in milliseconds"),
-		false,
-		AggregatedStatsColumnFlags,
-		HAlign_Right,
-		TimeMsColumnInitialWidth, TimeMsColumnMinWidth, TimeMsColumnMaxWidth,
-		[](const FTimersViewColumn& Column, const FTimerNode& TimerNode) -> FText
-		{ return FText::FromString(TimeUtils::FormatTimeMs(TimerNode.GetAggregatedStats().MedianExclusiveTime)); }
-	));
-
-	Collection.Add(new FTimersViewColumn
-	(
-		Index++,
-		FTimersViewColumns::MinExclusiveTimeColumnID,
-		TEXT("minexc"),
-		LOCTEXT("MinExclusiveTimeName", "E.Min"),
-		LOCTEXT("MinExclusiveTimeTitle", "Min Exclusive Time (ms)"),
-		LOCTEXT("MinExclusiveTimeDesc", "Minimum exclusive duration of selected timer's instances, in milliseconds"),
-		false,
-		AggregatedStatsColumnFlags,
-		HAlign_Right,
-		TimeMsColumnInitialWidth, TimeMsColumnMinWidth, TimeMsColumnMaxWidth,
-		[](const FTimersViewColumn& Column, const FTimerNode& TimerNode) -> FText
-		{ return FText::FromString(TimeUtils::FormatTimeMs(TimerNode.GetAggregatedStats().MinExclusiveTime)); }
-	));
-
-	//////////////////////////////////////////////////
-
-	for (int32 MapIndex = 0; MapIndex < Collection.Num(); MapIndex++)
+	Columns.Reset();
+
+	Columns.Add(CreateNameColumn());
+	Columns.Add(CreateMetaGroupNameColumn());
+	Columns.Add(CreateTypeColumn());
+	Columns.Add(CreateInstanceCountColumn());
+
+	Columns.Add(CreateTotalInclusiveTimeColumn());
+	Columns.Add(CreateMaxInclusiveTimeColumn());
+	Columns.Add(CreateAverageInclusiveTimeColumn());
+	Columns.Add(CreateMedianInclusiveTimeColumn());
+	Columns.Add(CreateMinInclusiveTimeColumn());
+
+	Columns.Add(CreateTotalExclusiveTimeColumn());
+	Columns.Add(CreateMaxExclusiveTimeColumn());
+	Columns.Add(CreateAverageExclusiveTimeColumn());
+	Columns.Add(CreateMedianExclusiveTimeColumn());
+	Columns.Add(CreateMinExclusiveTimeColumn());
+}
+
+////////////////////////////////////////////////////////////////////////////////////////////////////
+
+void FTimersViewColumnFactory::CreateTimerTreeViewColumns(TArray<TSharedPtr<Insights::FTableColumn>>& Columns)
+{
+	Columns.Reset();
+
+	Columns.Add(CreateNameColumn());
+	Columns.Add(CreateTypeColumn());
+	Columns.Add(CreateInstanceCountColumn());
+	Columns.Add(CreateTotalInclusiveTimeColumn());
+	Columns.Add(CreateTotalExclusiveTimeColumn());
+}
+
+////////////////////////////////////////////////////////////////////////////////////////////////////
+
+TSharedRef<Insights::FTableColumn> FTimersViewColumnFactory::CreateNameColumn()
+{
+	using namespace Insights;
+
+	TSharedPtr<FTableColumn> ColumnPtr = MakeShareable(new FTableColumn(FTimersViewColumns::NameColumnID));
+	FTableColumn& Column = *ColumnPtr;
+
+	Column.SetShortName(LOCTEXT("TimerNameColumnName", "Name"));
+	Column.SetTitleName(LOCTEXT("TimerNameColumnTitle", "Timer or Group Name"));
+	Column.SetDescription(LOCTEXT("TimerNameColumnDesc", "Name of the timer or group"));
+
+	Column.SetFlags(ETableColumnFlags::ShouldBeVisible |
+					ETableColumnFlags::CanBeFiltered |
+					ETableColumnFlags::IsHierarchy);
+
+	Column.SetHorizontalAlignment(HAlign_Left);
+	Column.SetInitialWidth(246.0f);
+	Column.SetMinWidth(42.0f);
+
+	Column.SetDataType(ETableCellDataType::Text);
+
+	TSharedPtr<ITableCellValueGetter> GetterPtr = MakeShareable(new FDisplayNameValueGetter());
+	Column.SetValueGetter(GetterPtr.ToSharedRef());
+
+	TSharedPtr<ITableCellValueFormatter> FormatterPtr = MakeShareable(new FTextValueFormatter());
+	Column.SetValueFormatter(FormatterPtr.ToSharedRef());
+
+	TSharedPtr<ITableCellValueSorter> SorterPtr = MakeShareable(new FSorterByName(ColumnPtr.ToSharedRef()));
+	Column.SetValueSorter(SorterPtr);
+
+	return ColumnPtr.ToSharedRef();
+}
+
+////////////////////////////////////////////////////////////////////////////////////////////////////
+
+TSharedRef<Insights::FTableColumn> FTimersViewColumnFactory::CreateMetaGroupNameColumn()
+{
+	using namespace Insights;
+
+	TSharedPtr<FTableColumn> ColumnPtr = MakeShareable(new FTableColumn(FTimersViewColumns::MetaGroupNameColumnID));
+	FTableColumn& Column = *ColumnPtr;
+
+	Column.SetShortName(LOCTEXT("TimerMetaGroupNameColumnName", "Meta Group"));
+	Column.SetTitleName(LOCTEXT("TimerMetaGroupNameColumnTitle", "Meta Group Name"));
+	Column.SetDescription(LOCTEXT("TimerMetaGroupNameColumnDesc", "Name of the meta group"));
+
+	Column.SetFlags(ETableColumnFlags::CanBeHidden |
+					//ETableColumnFlags::ShouldBeVisible |
+					ETableColumnFlags::CanBeFiltered);
+
+	Column.SetHorizontalAlignment(HAlign_Left);
+	Column.SetInitialWidth(100.0f);
+
+	Column.SetDataType(ETableCellDataType::Text);
+
+	class FMetaGroupNameValueGetter : public FTableCellValueGetter
+	{
+	public:
+		virtual const TOptional<FTableCellValue> GetValue(const FTableColumn& Column, const FBaseTreeNode& Node) const
+		{
+			ensure(Column.GetId() == FTimersViewColumns::MetaGroupNameColumnID);
+			const FTimerNode& TimerNode = static_cast<const FTimerNode&>(Node);
+			return TOptional<FTableCellValue>(FTableCellValue(FText::FromName(TimerNode.GetMetaGroupName())));
+		}
+	};
+
+	TSharedPtr<FTableCellValueGetter> GetterPtr = MakeShareable(new FMetaGroupNameValueGetter());
+	Column.SetValueGetter(GetterPtr.ToSharedRef());
+
+	TSharedPtr<FTableCellValueFormatter> FormatterPtr = MakeShareable(new FTextValueFormatter());
+	Column.SetValueFormatter(FormatterPtr.ToSharedRef());
+
+	TSharedPtr<ITableCellValueSorter> SorterPtr = MakeShareable(new FSorterByTextValue(ColumnPtr.ToSharedRef()));
+	Column.SetValueSorter(SorterPtr);
+
+	return ColumnPtr.ToSharedRef();
+}
+
+////////////////////////////////////////////////////////////////////////////////////////////////////
+
+TSharedRef<Insights::FTableColumn> FTimersViewColumnFactory::CreateTypeColumn()
+{
+	using namespace Insights;
+
+	TSharedPtr<FTableColumn> ColumnPtr = MakeShareable(new FTableColumn(FTimersViewColumns::TypeColumnID));
+	FTableColumn& Column = *ColumnPtr;
+
+	Column.SetShortName(LOCTEXT("TimerTypeColumnName", "Type"));
+	Column.SetTitleName(LOCTEXT("TimerTypeColumnTitle", "Type"));
+	Column.SetDescription(LOCTEXT("TimerTypeColumnDesc", "Type of timer or group"));
+
+	Column.SetFlags(ETableColumnFlags::CanBeHidden |
+					//ETableColumnFlags::ShouldBeVisible |
+					ETableColumnFlags::CanBeFiltered);
+
+	Column.SetHorizontalAlignment(HAlign_Left);
+	Column.SetInitialWidth(60.0f);
+
+	Column.SetDataType(ETableCellDataType::Text);
+
+	class FTimerTypeValueGetter : public FTableCellValueGetter
+	{
+	public:
+		virtual const TOptional<FTableCellValue> GetValue(const FTableColumn& Column, const FBaseTreeNode& Node) const
+		{
+			ensure(Column.GetId() == FTimersViewColumns::TypeColumnID);
+			const FTimerNode& TimerNode = static_cast<const FTimerNode&>(Node);
+			return TOptional<FTableCellValue>(FTableCellValue(TimerNodeTypeHelper::ToText(TimerNode.GetType())));
+		}
+	};
+
+	TSharedPtr<FTableCellValueGetter> GetterPtr = MakeShareable(new FTimerTypeValueGetter());
+	Column.SetValueGetter(GetterPtr.ToSharedRef());
+
+	TSharedPtr<FTableCellValueFormatter> FormatterPtr = MakeShareable(new FTextValueFormatter());
+	Column.SetValueFormatter(FormatterPtr.ToSharedRef());
+
+	//TSharedPtr<ITableCellValueSorter> SorterPtr = MakeShareable(new FSorterByTextValue(ColumnPtr.ToSharedRef()));
+	TSharedPtr<ITableCellValueSorter> SorterPtr = MakeShareable(new FTimerNodeSortingByTimerType(ColumnPtr.ToSharedRef()));
+	Column.SetValueSorter(SorterPtr);
+
+	return ColumnPtr.ToSharedRef();
+}
+
+////////////////////////////////////////////////////////////////////////////////////////////////////
+
+TSharedRef<Insights::FTableColumn> FTimersViewColumnFactory::CreateInstanceCountColumn()
+{
+	using namespace Insights;
+
+	TSharedPtr<FTableColumn> ColumnPtr = MakeShareable(new FTableColumn(FTimersViewColumns::InstanceCountColumnID));
+	FTableColumn& Column = *ColumnPtr;
+
+	Column.SetShortName(LOCTEXT("InstanceCountName", "Count"));
+	Column.SetTitleName(LOCTEXT("InstanceCountTitle", "Instance Count"));
+	Column.SetDescription(LOCTEXT("InstanceCountDesc", "Number of timer's instances"));
+
+	Column.SetFlags(ETableColumnFlags::CanBeHidden |
+					ETableColumnFlags::ShouldBeVisible |
+					ETableColumnFlags::CanBeFiltered);
+
+	Column.SetHorizontalAlignment(HAlign_Right);
+	Column.SetInitialWidth(60.0f);
+
+	Column.SetDataType(ETableCellDataType::Int64);
+
+	class FInstanceCountValueGetter : public FTableCellValueGetter
+	{
+	public:
+		virtual const TOptional<FTableCellValue> GetValue(const FTableColumn& Column, const FBaseTreeNode& Node) const
+		{
+			ensure(Column.GetId() == FTimersViewColumns::InstanceCountColumnID);
+			const FTimerNode& TimerNode = static_cast<const FTimerNode&>(Node);
+			return TOptional<FTableCellValue>(FTableCellValue(static_cast<int64>(TimerNode.GetAggregatedStats().InstanceCount)));
+		}
+	};
+
+	TSharedPtr<FTableCellValueGetter> GetterPtr = MakeShareable(new FInstanceCountValueGetter());
+	Column.SetValueGetter(GetterPtr.ToSharedRef());
+
+	TSharedPtr<FTableCellValueFormatter> FormatterPtr = MakeShareable(new FInt64ValueFormatterAsNumber());
+	Column.SetValueFormatter(FormatterPtr.ToSharedRef());
+
+	//TSharedPtr<ITableCellValueSorter> SorterPtr = MakeShareable(new FSorterByInt64Value(ColumnPtr.ToSharedRef()));
+	TSharedPtr<ITableCellValueSorter> SorterPtr = MakeShareable(new FTimerNodeSortingByInstanceCount(ColumnPtr.ToSharedRef()));
+	Column.SetValueSorter(SorterPtr);
+
+	return ColumnPtr.ToSharedRef();
+}
+
+////////////////////////////////////////////////////////////////////////////////////////////////////
+// Inclusive Time Columns
+////////////////////////////////////////////////////////////////////////////////////////////////////
+
+TSharedRef<Insights::FTableColumn> FTimersViewColumnFactory::CreateTotalInclusiveTimeColumn()
+{
+	using namespace Insights;
+
+	TSharedPtr<FTableColumn> ColumnPtr = MakeShareable(new FTableColumn(FTimersViewColumns::TotalInclusiveTimeColumnID));
+	FTableColumn& Column = *ColumnPtr;
+
+	Column.SetShortName(LOCTEXT("TotalInclusiveTimeName", "Incl"));
+	Column.SetTitleName(LOCTEXT("TotalInclusiveTimeTitle", "Total Inclusive Time"));
+	Column.SetDescription(LOCTEXT("TotalInclusiveTimeDesc", "Total inclusive duration of selected timer's instances"));
+
+	Column.SetFlags(ETableColumnFlags::CanBeHidden |
+					ETableColumnFlags::ShouldBeVisible |
+					ETableColumnFlags::CanBeFiltered);
+
+	Column.SetHorizontalAlignment(HAlign_Right);
+	Column.SetInitialWidth(TotalTimeColumnInitialWidth);
+
+	Column.SetDataType(ETableCellDataType::Double);
+	Column.SetAggregation(ETableColumnAggregation::Sum);
+
+	class FTotalInclusiveTimeValueGetter : public FTableCellValueGetter
+	{
+	public:
+		virtual const TOptional<FTableCellValue> GetValue(const FTableColumn& Column, const FBaseTreeNode& Node) const
+		{
+			ensure(Column.GetId() == FTimersViewColumns::TotalInclusiveTimeColumnID);
+			const FTimerNode& TimerNode = static_cast<const FTimerNode&>(Node);
+			return TOptional<FTableCellValue>(FTableCellValue(TimerNode.GetAggregatedStats().TotalInclusiveTime));
+		}
+	};
+
+	TSharedPtr<FTableCellValueGetter> GetterPtr = MakeShareable(new FTotalInclusiveTimeValueGetter());
+	Column.SetValueGetter(GetterPtr.ToSharedRef());
+
+	TSharedPtr<FTableCellValueFormatter> FormatterPtr = MakeShareable(new FDoubleValueFormatterAsTimeAuto());
+	Column.SetValueFormatter(FormatterPtr.ToSharedRef());
+
+	//TSharedPtr<ITableCellValueSorter> SorterPtr = MakeShareable(new FSorterByDoubleValue(ColumnPtr.ToSharedRef()));
+	TSharedPtr<ITableCellValueSorter> SorterPtr = MakeShareable(new FTimerNodeSortingByTotalInclusiveTime(ColumnPtr.ToSharedRef()));
+	Column.SetValueSorter(SorterPtr);
+
+	return ColumnPtr.ToSharedRef();
+}
+
+////////////////////////////////////////////////////////////////////////////////////////////////////
+
+TSharedRef<Insights::FTableColumn> FTimersViewColumnFactory::CreateMaxInclusiveTimeColumn()
+{
+	using namespace Insights;
+
+	TSharedPtr<FTableColumn> ColumnPtr = MakeShareable(new FTableColumn(FTimersViewColumns::MaxInclusiveTimeColumnID));
+	FTableColumn& Column = *ColumnPtr;
+
+	Column.SetShortName(LOCTEXT("MaxInclusiveTimeName", "I.Max"));
+	Column.SetTitleName(LOCTEXT("MaxInclusiveTimeTitle", "Max Inclusive Time (ms)"));
+	Column.SetDescription(LOCTEXT("MaxInclusiveTimeDesc", "Maximum inclusive duration of selected timer's instances, in milliseconds"));
+
+	Column.SetFlags(ETableColumnFlags::CanBeHidden |
+					//ETableColumnFlags::ShouldBeVisible |
+					ETableColumnFlags::CanBeFiltered);
+
+	Column.SetHorizontalAlignment(HAlign_Right);
+	Column.SetInitialWidth(TimeMsColumnInitialWidth);
+
+	Column.SetDataType(ETableCellDataType::Double);
+
+	class FMaxInclusiveTimeValueGetter : public FTableCellValueGetter
+	{
+	public:
+		virtual const TOptional<FTableCellValue> GetValue(const FTableColumn& Column, const FBaseTreeNode& Node) const
+		{
+			ensure(Column.GetId() == FTimersViewColumns::MaxInclusiveTimeColumnID);
+			const FTimerNode& TimerNode = static_cast<const FTimerNode&>(Node);
+			return TOptional<FTableCellValue>(FTableCellValue(TimerNode.GetAggregatedStats().MaxInclusiveTime));
+		}
+	};
+
+	TSharedPtr<FTableCellValueGetter> GetterPtr = MakeShareable(new FMaxInclusiveTimeValueGetter());
+	Column.SetValueGetter(GetterPtr.ToSharedRef());
+
+	TSharedPtr<FTableCellValueFormatter> FormatterPtr = MakeShareable(new FDoubleValueFormatterAsTimeMs());
+	Column.SetValueFormatter(FormatterPtr.ToSharedRef());
+
+	TSharedPtr<ITableCellValueSorter> SorterPtr = MakeShareable(new FSorterByDoubleValue(ColumnPtr.ToSharedRef()));
+	Column.SetValueSorter(SorterPtr);
+
+	return ColumnPtr.ToSharedRef();
+}
+
+////////////////////////////////////////////////////////////////////////////////////////////////////
+
+TSharedRef<Insights::FTableColumn> FTimersViewColumnFactory::CreateAverageInclusiveTimeColumn()
+{
+	using namespace Insights;
+
+	TSharedPtr<FTableColumn> ColumnPtr = MakeShareable(new FTableColumn(FTimersViewColumns::AverageInclusiveTimeColumnID));
+	FTableColumn& Column = *ColumnPtr;
+
+	Column.SetShortName(LOCTEXT("AvgInclusiveTimeName", "I.Avg"));
+	Column.SetTitleName(LOCTEXT("AvgInclusiveTimeTitle", "Average Inclusive Time (ms)"));
+	Column.SetDescription(LOCTEXT("AvgInclusiveTimeDesc", "Average inclusive duration of selected timer's instances, in milliseconds"));
+
+	Column.SetFlags(ETableColumnFlags::CanBeHidden |
+					//ETableColumnFlags::ShouldBeVisible |
+					ETableColumnFlags::CanBeFiltered);
+
+	Column.SetHorizontalAlignment(HAlign_Right);
+	Column.SetInitialWidth(TimeMsColumnInitialWidth);
+
+	Column.SetDataType(ETableCellDataType::Double);
+
+	class FAverageInclusiveTimeValueGetter : public FTableCellValueGetter
+	{
+	public:
+		virtual const TOptional<FTableCellValue> GetValue(const FTableColumn& Column, const FBaseTreeNode& Node) const
+		{
+			ensure(Column.GetId() == FTimersViewColumns::AverageInclusiveTimeColumnID);
+			const FTimerNode& TimerNode = static_cast<const FTimerNode&>(Node);
+			return TOptional<FTableCellValue>(FTableCellValue(TimerNode.GetAggregatedStats().AverageInclusiveTime));
+		}
+	};
+
+	TSharedPtr<FTableCellValueGetter> GetterPtr = MakeShareable(new FAverageInclusiveTimeValueGetter());
+	Column.SetValueGetter(GetterPtr.ToSharedRef());
+
+	TSharedPtr<FTableCellValueFormatter> FormatterPtr = MakeShareable(new FDoubleValueFormatterAsTimeMs());
+	Column.SetValueFormatter(FormatterPtr.ToSharedRef());
+
+	TSharedPtr<ITableCellValueSorter> SorterPtr = MakeShareable(new FSorterByDoubleValue(ColumnPtr.ToSharedRef()));
+	Column.SetValueSorter(SorterPtr);
+
+	return ColumnPtr.ToSharedRef();
+}
+
+////////////////////////////////////////////////////////////////////////////////////////////////////
+
+TSharedRef<Insights::FTableColumn> FTimersViewColumnFactory::CreateMedianInclusiveTimeColumn()
+{
+	using namespace Insights;
+
+	TSharedPtr<FTableColumn> ColumnPtr = MakeShareable(new FTableColumn(FTimersViewColumns::MedianInclusiveTimeColumnID));
+	FTableColumn& Column = *ColumnPtr;
+
+	Column.SetShortName(LOCTEXT("MedInclusiveTimeName", "I.Med"));
+	Column.SetTitleName(LOCTEXT("MedInclusiveTimeTitle", "Median Inclusive Time (ms)"));
+	Column.SetDescription(LOCTEXT("MedInclusiveTimeDesc", "Median inclusive duration of selected timer's instances, in milliseconds"));
+
+	Column.SetFlags(ETableColumnFlags::CanBeHidden |
+					//ETableColumnFlags::ShouldBeVisible |
+					ETableColumnFlags::CanBeFiltered);
+
+	Column.SetHorizontalAlignment(HAlign_Right);
+	Column.SetInitialWidth(TimeMsColumnInitialWidth);
+
+	Column.SetDataType(ETableCellDataType::Double);
+
+	class FMedianInclusiveTimeValueGetter : public FTableCellValueGetter
 	{
-		ColumnIdToPtrMapping.Add(Collection[MapIndex]->Id, Collection[MapIndex]);
-	}
+	public:
+		virtual const TOptional<FTableCellValue> GetValue(const FTableColumn& Column, const FBaseTreeNode& Node) const
+		{
+			ensure(Column.GetId() == FTimersViewColumns::MedianInclusiveTimeColumnID);
+			const FTimerNode& TimerNode = static_cast<const FTimerNode&>(Node);
+			return TOptional<FTableCellValue>(FTableCellValue(TimerNode.GetAggregatedStats().MedianInclusiveTime));
+		}
+	};
+
+	TSharedPtr<FTableCellValueGetter> GetterPtr = MakeShareable(new FMedianInclusiveTimeValueGetter());
+	Column.SetValueGetter(GetterPtr.ToSharedRef());
+
+	TSharedPtr<FTableCellValueFormatter> FormatterPtr = MakeShareable(new FDoubleValueFormatterAsTimeMs());
+	Column.SetValueFormatter(FormatterPtr.ToSharedRef());
+
+	TSharedPtr<ITableCellValueSorter> SorterPtr = MakeShareable(new FSorterByDoubleValue(ColumnPtr.ToSharedRef()));
+	Column.SetValueSorter(SorterPtr);
+
+	return ColumnPtr.ToSharedRef();
 }
 
 ////////////////////////////////////////////////////////////////////////////////////////////////////
 
-FTimersViewColumnFactory::~FTimersViewColumnFactory()
+TSharedRef<Insights::FTableColumn> FTimersViewColumnFactory::CreateMinInclusiveTimeColumn()
 {
-	for (const FTimersViewColumn* Column : Collection)
+	using namespace Insights;
+
+	TSharedPtr<FTableColumn> ColumnPtr = MakeShareable(new FTableColumn(FTimersViewColumns::MinInclusiveTimeColumnID));
+	FTableColumn& Column = *ColumnPtr;
+
+	Column.SetShortName(LOCTEXT("MinInclusiveTimeName", "I.Min"));
+	Column.SetTitleName(LOCTEXT("MinInclusiveTimeTitle", "Min Inclusive Time (ms)"));
+	Column.SetDescription(LOCTEXT("MinInclusiveTimeDesc", "Minimum inclusive duration of selected timer's instances, in milliseconds"));
+
+	Column.SetFlags(ETableColumnFlags::CanBeHidden |
+					//ETableColumnFlags::ShouldBeVisible |
+					ETableColumnFlags::CanBeFiltered);
+
+	Column.SetHorizontalAlignment(HAlign_Right);
+	Column.SetInitialWidth(TimeMsColumnInitialWidth);
+
+	Column.SetDataType(ETableCellDataType::Double);
+
+	class FMinInclusiveTimeValueGetter : public FTableCellValueGetter
 	{
-		delete Column;
-	}
-	Collection.Reset();
+	public:
+		virtual const TOptional<FTableCellValue> GetValue(const FTableColumn& Column, const FBaseTreeNode& Node) const
+		{
+			ensure(Column.GetId() == FTimersViewColumns::MinInclusiveTimeColumnID);
+			const FTimerNode& TimerNode = static_cast<const FTimerNode&>(Node);
+			return TOptional<FTableCellValue>(FTableCellValue(TimerNode.GetAggregatedStats().MinInclusiveTime));
+		}
+	};
+
+	TSharedPtr<FTableCellValueGetter> GetterPtr = MakeShareable(new FMinInclusiveTimeValueGetter());
+	Column.SetValueGetter(GetterPtr.ToSharedRef());
+
+	TSharedPtr<FTableCellValueFormatter> FormatterPtr = MakeShareable(new FDoubleValueFormatterAsTimeMs());
+	Column.SetValueFormatter(FormatterPtr.ToSharedRef());
+
+	TSharedPtr<ITableCellValueSorter> SorterPtr = MakeShareable(new FSorterByDoubleValue(ColumnPtr.ToSharedRef()));
+	Column.SetValueSorter(SorterPtr);
+
+	return ColumnPtr.ToSharedRef();
 }
 
+////////////////////////////////////////////////////////////////////////////////////////////////////
+// Exclusive Time Columns
 ////////////////////////////////////////////////////////////////////////////////////////////////////
 
-const FTimersViewColumnFactory& FTimersViewColumnFactory::Get()
+TSharedRef<Insights::FTableColumn> FTimersViewColumnFactory::CreateTotalExclusiveTimeColumn()
 {
-	static FTimersViewColumnFactory Instance;
-	return Instance;
+	using namespace Insights;
+
+	TSharedPtr<FTableColumn> ColumnPtr = MakeShareable(new FTableColumn(FTimersViewColumns::TotalExclusiveTimeColumnID));
+	FTableColumn& Column = *ColumnPtr;
+
+	Column.SetShortName(LOCTEXT("TotalExclusiveTimeName", "Excl"));
+	Column.SetTitleName(LOCTEXT("TotalExclusiveTimeTitle", "Total Exclusive Time"));
+	Column.SetDescription(LOCTEXT("TotalExclusiveTimeDesc", "Total exclusive duration of selected timer's instances"));
+
+	Column.SetFlags(ETableColumnFlags::CanBeHidden |
+					ETableColumnFlags::ShouldBeVisible |
+					ETableColumnFlags::CanBeFiltered);
+
+	Column.SetHorizontalAlignment(HAlign_Right);
+	Column.SetInitialWidth(TotalTimeColumnInitialWidth);
+
+	Column.SetDataType(ETableCellDataType::Double);
+	Column.SetAggregation(ETableColumnAggregation::Sum);
+
+	class FTotalExclusiveTimeValueGetter : public FTableCellValueGetter
+	{
+	public:
+		virtual const TOptional<FTableCellValue> GetValue(const FTableColumn& Column, const FBaseTreeNode& Node) const
+		{
+			ensure(Column.GetId() == FTimersViewColumns::TotalExclusiveTimeColumnID);
+			const FTimerNode& TimerNode = static_cast<const FTimerNode&>(Node);
+			return TOptional<FTableCellValue>(FTableCellValue(TimerNode.GetAggregatedStats().TotalExclusiveTime));
+		}
+	};
+
+	TSharedPtr<FTableCellValueGetter> GetterPtr = MakeShareable(new FTotalExclusiveTimeValueGetter());
+	Column.SetValueGetter(GetterPtr.ToSharedRef());
+
+	TSharedPtr<FTableCellValueFormatter> FormatterPtr = MakeShareable(new FDoubleValueFormatterAsTimeAuto());
+	Column.SetValueFormatter(FormatterPtr.ToSharedRef());
+
+	//TSharedPtr<ITableCellValueSorter> SorterPtr = MakeShareable(new FSorterByDoubleValue(ColumnPtr.ToSharedRef()));
+	TSharedPtr<ITableCellValueSorter> SorterPtr = MakeShareable(new FTimerNodeSortingByTotalExclusiveTime(ColumnPtr.ToSharedRef()));
+	Column.SetValueSorter(SorterPtr);
+
+	return ColumnPtr.ToSharedRef();
+}
+
+////////////////////////////////////////////////////////////////////////////////////////////////////
+
+TSharedRef<Insights::FTableColumn> FTimersViewColumnFactory::CreateMaxExclusiveTimeColumn()
+{
+	using namespace Insights;
+
+	TSharedPtr<FTableColumn> ColumnPtr = MakeShareable(new FTableColumn(FTimersViewColumns::MaxExclusiveTimeColumnID));
+	FTableColumn& Column = *ColumnPtr;
+
+	Column.SetShortName(LOCTEXT("MaxExclusiveTimeName", "E.Max"));
+	Column.SetTitleName(LOCTEXT("MaxExclusiveTimeTitle", "Max Exclusive Time (ms)"));
+	Column.SetDescription(LOCTEXT("MaxExclusiveTimeDesc", "Maximum exclusive duration of selected timer's instances, in milliseconds"));
+
+	Column.SetFlags(ETableColumnFlags::CanBeHidden |
+					//ETableColumnFlags::ShouldBeVisible |
+					ETableColumnFlags::CanBeFiltered);
+
+	Column.SetHorizontalAlignment(HAlign_Right);
+	Column.SetInitialWidth(TimeMsColumnInitialWidth);
+
+	Column.SetDataType(ETableCellDataType::Double);
+
+	class FMaxExclusiveTimeValueGetter : public FTableCellValueGetter
+	{
+	public:
+		virtual const TOptional<FTableCellValue> GetValue(const FTableColumn& Column, const FBaseTreeNode& Node) const
+		{
+			ensure(Column.GetId() == FTimersViewColumns::MaxExclusiveTimeColumnID);
+			const FTimerNode& TimerNode = static_cast<const FTimerNode&>(Node);
+			return TOptional<FTableCellValue>(FTableCellValue(TimerNode.GetAggregatedStats().MaxExclusiveTime));
+		}
+	};
+
+	TSharedPtr<FTableCellValueGetter> GetterPtr = MakeShareable(new FMaxExclusiveTimeValueGetter());
+	Column.SetValueGetter(GetterPtr.ToSharedRef());
+
+	TSharedPtr<FTableCellValueFormatter> FormatterPtr = MakeShareable(new FDoubleValueFormatterAsTimeMs());
+	Column.SetValueFormatter(FormatterPtr.ToSharedRef());
+
+	TSharedPtr<ITableCellValueSorter> SorterPtr = MakeShareable(new FSorterByDoubleValue(ColumnPtr.ToSharedRef()));
+	Column.SetValueSorter(SorterPtr);
+
+	return ColumnPtr.ToSharedRef();
+}
+
+////////////////////////////////////////////////////////////////////////////////////////////////////
+
+TSharedRef<Insights::FTableColumn> FTimersViewColumnFactory::CreateAverageExclusiveTimeColumn()
+{
+	using namespace Insights;
+
+	TSharedPtr<FTableColumn> ColumnPtr = MakeShareable(new FTableColumn(FTimersViewColumns::AverageExclusiveTimeColumnID));
+	FTableColumn& Column = *ColumnPtr;
+
+	Column.SetShortName(LOCTEXT("AvgExclusiveTimeName", "E.Avg"));
+	Column.SetTitleName(LOCTEXT("AvgExclusiveTimeTitle", "Average Exclusive Time (ms)"));
+	Column.SetDescription(LOCTEXT("AvgExclusiveTimeDesc", "Average exclusive duration of selected timer's instances, in milliseconds"));
+
+	Column.SetFlags(ETableColumnFlags::CanBeHidden |
+					//ETableColumnFlags::ShouldBeVisible |
+					ETableColumnFlags::CanBeFiltered);
+
+	Column.SetHorizontalAlignment(HAlign_Right);
+	Column.SetInitialWidth(TimeMsColumnInitialWidth);
+
+	Column.SetDataType(ETableCellDataType::Double);
+
+	class FAverageExclusiveTimeValueGetter : public FTableCellValueGetter
+	{
+	public:
+		virtual const TOptional<FTableCellValue> GetValue(const FTableColumn& Column, const FBaseTreeNode& Node) const
+		{
+			ensure(Column.GetId() == FTimersViewColumns::AverageExclusiveTimeColumnID);
+			const FTimerNode& TimerNode = static_cast<const FTimerNode&>(Node);
+			return TOptional<FTableCellValue>(FTableCellValue(TimerNode.GetAggregatedStats().AverageExclusiveTime));
+		}
+	};
+
+	TSharedPtr<FTableCellValueGetter> GetterPtr = MakeShareable(new FAverageExclusiveTimeValueGetter());
+	Column.SetValueGetter(GetterPtr.ToSharedRef());
+
+	TSharedPtr<FTableCellValueFormatter> FormatterPtr = MakeShareable(new FDoubleValueFormatterAsTimeMs());
+	Column.SetValueFormatter(FormatterPtr.ToSharedRef());
+
+	TSharedPtr<ITableCellValueSorter> SorterPtr = MakeShareable(new FSorterByDoubleValue(ColumnPtr.ToSharedRef()));
+	Column.SetValueSorter(SorterPtr);
+
+	return ColumnPtr.ToSharedRef();
+}
+
+////////////////////////////////////////////////////////////////////////////////////////////////////
+
+TSharedRef<Insights::FTableColumn> FTimersViewColumnFactory::CreateMedianExclusiveTimeColumn()
+{
+	using namespace Insights;
+
+	TSharedPtr<FTableColumn> ColumnPtr = MakeShareable(new FTableColumn(FTimersViewColumns::MedianExclusiveTimeColumnID));
+	FTableColumn& Column = *ColumnPtr;
+
+	Column.SetShortName(LOCTEXT("MedExclusiveTimeName", "E.Med"));
+	Column.SetTitleName(LOCTEXT("MedExclusiveTimeTitle", "Median Exclusive Time (ms)"));
+	Column.SetDescription(LOCTEXT("MedExclusiveTimeDesc", "Median exclusive duration of selected timer's instances, in milliseconds"));
+
+	Column.SetFlags(ETableColumnFlags::CanBeHidden |
+					//ETableColumnFlags::ShouldBeVisible |
+					ETableColumnFlags::CanBeFiltered);
+
+	Column.SetHorizontalAlignment(HAlign_Right);
+	Column.SetInitialWidth(TimeMsColumnInitialWidth);
+
+	Column.SetDataType(ETableCellDataType::Double);
+
+	class FMedianExclusiveTimeValueGetter : public FTableCellValueGetter
+	{
+	public:
+		virtual const TOptional<FTableCellValue> GetValue(const FTableColumn& Column, const FBaseTreeNode& Node) const
+		{
+			ensure(Column.GetId() == FTimersViewColumns::MedianExclusiveTimeColumnID);
+			const FTimerNode& TimerNode = static_cast<const FTimerNode&>(Node);
+			return TOptional<FTableCellValue>(FTableCellValue(TimerNode.GetAggregatedStats().MedianExclusiveTime));
+		}
+	};
+
+	TSharedPtr<FTableCellValueGetter> GetterPtr = MakeShareable(new FMedianExclusiveTimeValueGetter());
+	Column.SetValueGetter(GetterPtr.ToSharedRef());
+
+	TSharedPtr<FTableCellValueFormatter> FormatterPtr = MakeShareable(new FDoubleValueFormatterAsTimeMs());
+	Column.SetValueFormatter(FormatterPtr.ToSharedRef());
+
+	TSharedPtr<ITableCellValueSorter> SorterPtr = MakeShareable(new FSorterByDoubleValue(ColumnPtr.ToSharedRef()));
+	Column.SetValueSorter(SorterPtr);
+
+	return ColumnPtr.ToSharedRef();
+}
+
+////////////////////////////////////////////////////////////////////////////////////////////////////
+
+TSharedRef<Insights::FTableColumn> FTimersViewColumnFactory::CreateMinExclusiveTimeColumn()
+{
+	using namespace Insights;
+
+	TSharedPtr<FTableColumn> ColumnPtr = MakeShareable(new FTableColumn(FTimersViewColumns::MinExclusiveTimeColumnID));
+	FTableColumn& Column = *ColumnPtr;
+
+	Column.SetShortName(LOCTEXT("MinExclusiveTimeName", "E.Min"));
+	Column.SetTitleName(LOCTEXT("MinExclusiveTimeTitle", "Min Exclusive Time (ms)"));
+	Column.SetDescription(LOCTEXT("MinExclusiveTimeDesc", "Minimum exclusive duration of selected timer's instances, in milliseconds"));
+
+	Column.SetFlags(ETableColumnFlags::CanBeHidden |
+					//ETableColumnFlags::ShouldBeVisible |
+					ETableColumnFlags::CanBeFiltered);
+
+	Column.SetHorizontalAlignment(HAlign_Right);
+	Column.SetInitialWidth(TimeMsColumnInitialWidth);
+
+	Column.SetDataType(ETableCellDataType::Double);
+
+	class FMinExclusiveTimeValueGetter : public FTableCellValueGetter
+	{
+	public:
+		virtual const TOptional<FTableCellValue> GetValue(const FTableColumn& Column, const FBaseTreeNode& Node) const
+		{
+			ensure(Column.GetId() == FTimersViewColumns::MinExclusiveTimeColumnID);
+			const FTimerNode& TimerNode = static_cast<const FTimerNode&>(Node);
+			return TOptional<FTableCellValue>(FTableCellValue(TimerNode.GetAggregatedStats().MinExclusiveTime));
+		}
+	};
+
+	TSharedPtr<FTableCellValueGetter> GetterPtr = MakeShareable(new FMinExclusiveTimeValueGetter());
+	Column.SetValueGetter(GetterPtr.ToSharedRef());
+
+	TSharedPtr<FTableCellValueFormatter> FormatterPtr = MakeShareable(new FDoubleValueFormatterAsTimeMs());
+	Column.SetValueFormatter(FormatterPtr.ToSharedRef());
+
+	TSharedPtr<ITableCellValueSorter> SorterPtr = MakeShareable(new FSorterByDoubleValue(ColumnPtr.ToSharedRef()));
+	Column.SetValueSorter(SorterPtr);
+
+	return ColumnPtr.ToSharedRef();
 }
 
 ////////////////////////////////////////////////////////////////////////////////////////////////////
diff --git a/Engine/Source/Developer/TraceInsights/Private/Insights/ViewModels/TimersViewColumnFactory.h b/Engine/Source/Developer/TraceInsights/Private/Insights/ViewModels/TimersViewColumnFactory.h
--- a/Engine/Source/Developer/TraceInsights/Private/Insights/ViewModels/TimersViewColumnFactory.h
+++ b/Engine/Source/Developer/TraceInsights/Private/Insights/ViewModels/TimersViewColumnFactory.h
@@ -4,8 +4,6 @@
 
 #include "CoreMinimal.h"
 
-#include "Insights/ViewModels/TimersViewColumn.h"
-
 ////////////////////////////////////////////////////////////////////////////////////////////////////
 
 struct FTimersViewColumns
@@ -39,24 +37,39 @@ struct FTimersViewColumns
 
 ////////////////////////////////////////////////////////////////////////////////////////////////////
 
-struct FTimersViewColumnFactory
+namespace Insights
 {
-private:
-	/** Default constructor. */
-	FTimersViewColumnFactory();
+	class FTableColumn;
+}
 
-	/** Destructor. */
-	~FTimersViewColumnFactory();
+////////////////////////////////////////////////////////////////////////////////////////////////////
 
+struct FTimersViewColumnFactory
+{
 public:
-	/** Contains basic information about columns used in the Timers view widget. Names should be localized. */
-	TArray<FTimersViewColumn*> Collection;
+	static void CreateTimersViewColumns(TArray<TSharedPtr<Insights::FTableColumn>>& Columns);
+	static void CreateTimerTreeViewColumns(TArray<TSharedPtr<Insights::FTableColumn>>& Columns);
+
+	static TSharedRef<Insights::FTableColumn> CreateNameColumn();
+	static TSharedRef<Insights::FTableColumn> CreateMetaGroupNameColumn();
+	static TSharedRef<Insights::FTableColumn> CreateTypeColumn();
+	static TSharedRef<Insights::FTableColumn> CreateInstanceCountColumn();
 
-	/** Mapping between column IDs and FTimersViewColumn pointers. */
-	TMap<FName, const FTimersViewColumn*> ColumnIdToPtrMapping;
+	static TSharedRef<Insights::FTableColumn> CreateTotalInclusiveTimeColumn();
+	static TSharedRef<Insights::FTableColumn> CreateMaxInclusiveTimeColumn();
+	static TSharedRef<Insights::FTableColumn> CreateAverageInclusiveTimeColumn();
+	static TSharedRef<Insights::FTableColumn> CreateMedianInclusiveTimeColumn();
+	static TSharedRef<Insights::FTableColumn> CreateMinInclusiveTimeColumn();
 
-	/** Singleton */
-	static const FTimersViewColumnFactory& Get();
+	static TSharedRef<Insights::FTableColumn> CreateTotalExclusiveTimeColumn();
+	static TSharedRef<Insights::FTableColumn> CreateMaxExclusiveTimeColumn();
+	static TSharedRef<Insights::FTableColumn> CreateAverageExclusiveTimeColumn();
+	static TSharedRef<Insights::FTableColumn> CreateMedianExclusiveTimeColumn();
+	static TSharedRef<Insights::FTableColumn> CreateMinExclusiveTimeColumn();
+
+private:
+	static constexpr float TotalTimeColumnInitialWidth = 60.0f;
+	static constexpr float TimeMsColumnInitialWidth = 50.0f;
 };
 
 ////////////////////////////////////////////////////////////////////////////////////////////////////
diff --git a/Engine/Source/Developer/TraceInsights/Private/Insights/Widgets/SStatsTableCell.cpp b/Engine/Source/Developer/TraceInsights/Private/Insights/Widgets/SStatsTableCell.cpp
--- a/Engine/Source/Developer/TraceInsights/Private/Insights/Widgets/SStatsTableCell.cpp
+++ b/Engine/Source/Developer/TraceInsights/Private/Insights/Widgets/SStatsTableCell.cpp
@@ -121,7 +121,8 @@ TSharedRef<SWidget> SStatsTableCell::GenerateWidgetForNameColumn(const FArgument
 TSharedRef<SWidget> SStatsTableCell::GenerateWidgetForStatsColumn(const FArguments& InArgs, const TSharedRef<class ITableRow>& TableRow)
 {
 	const FStatsViewColumn& Column = *FStatsViewColumnFactory::Get().ColumnIdToPtrMapping.FindChecked(ColumnId);
-	const FText FormattedValue = Column.GetFormattedValue(*StatsNodePtr);
+	const FText CellText = Column.GetFormattedValue(*StatsNodePtr);
+	//const FText CellText = ColumnPtr->GetValueAsText(*StatsNodePtr);
 
 	return
 		SNew(SHorizontalBox)
@@ -134,7 +135,7 @@ TSharedRef<SWidget> SStatsTableCell::GenerateWidgetForStatsColumn(const FArgumen
 		.Padding(FMargin(2.0f, 0.0f))
 		[
 			SNew(STextBlock)
-			.Text(FormattedValue)
+			.Text(CellText)
 			.TextStyle(FEditorStyle::Get(), TEXT("Profiler.Tooltip"))
 			.ColorAndOpacity(this, &SStatsTableCell::GetStatsColorAndOpacity)
 			.ShadowColorAndOpacity(this, &SStatsTableCell::GetShadowColorAndOpacity)
diff --git a/Engine/Source/Developer/TraceInsights/Private/Insights/Widgets/STimerTableCell.cpp b/Engine/Source/Developer/TraceInsights/Private/Insights/Widgets/STimerTableCell.cpp
--- a/Engine/Source/Developer/TraceInsights/Private/Insights/Widgets/STimerTableCell.cpp
+++ b/Engine/Source/Developer/TraceInsights/Private/Insights/Widgets/STimerTableCell.cpp
@@ -7,13 +7,14 @@
 #include "Widgets/Images/SImage.h"
 #include "Widgets/Input/SButton.h"
 #include "Widgets/SBoxPanel.h"
+#include "Widgets/SOverlay.h"
 #include "Widgets/SToolTip.h"
 #include "Widgets/Text/STextBlock.h"
 #include "Widgets/Views/SExpanderArrow.h"
 
 // Insights
-#include "Insights/ViewModels/TimersViewColumn.h"
-#include "Insights/ViewModels/TimersViewColumnFactory.h"
+#include "Insights/Table/ViewModels/Table.h"
+#include "Insights/Table/ViewModels/TableColumn.h"
 #include "Insights/Widgets/STimersViewTooltip.h"
 
 #define LOCTEXT_NAMESPACE "STimersView"
@@ -24,10 +25,15 @@ BEGIN_SLATE_FUNCTION_BUILD_OPTIMIZATION
 
 void STimerTableCell::Construct(const FArguments& InArgs, const TSharedRef<class ITableRow>& TableRow)
 {
-	SetHoveredTableCellDelegate = InArgs._OnSetHoveredTableCell;
+	TablePtr = InArgs._TablePtr;
+	ColumnPtr = InArgs._ColumnPtr;
 	TimerNodePtr = InArgs._TimerNodePtr;
-	ColumnId = InArgs._ColumnId;
-	//HighlightText = InArgs._HighlightText;
+
+	ensure(TablePtr.IsValid());
+	ensure(ColumnPtr.IsValid());
+	ensure(TimerNodePtr.IsValid());
+
+	SetHoveredCellDelegate = InArgs._OnSetHoveredCell;
 
 	ChildSlot
 		[
@@ -53,8 +59,6 @@ TSharedRef<SWidget> STimerTableCell::GenerateWidgetForColumn(const FArguments& I
 
 TSharedRef<SWidget> STimerTableCell::GenerateWidgetForNameColumn(const FArguments& InArgs, const TSharedRef<class ITableRow>& TableRow)
 {
-	const FTimersViewColumn& Column = *FTimersViewColumnFactory::Get().ColumnIdToPtrMapping.FindChecked(ColumnId);
-
 	return
 		SNew(SHorizontalBox)
 
@@ -66,23 +70,36 @@ TSharedRef<SWidget> STimerTableCell::GenerateWidgetForNameColumn(const FArgument
 			SNew(SExpanderArrow, TableRow)
 		]
 
-		// Info icon + tooltip
 		+ SHorizontalBox::Slot()
 		.AutoWidth()
 		.HAlign(HAlign_Center)
 		.VAlign(VAlign_Center)
 		[
-			SNew(SImage)
-			.Visibility(this, &STimerTableCell::GetHintIconVisibility)
-			.Image(FEditorStyle::GetBrush("Profiler.Tooltip.HintIcon10"))
-			.ToolTip(STimersViewTooltip::GetTableCellTooltip(TimerNodePtr))
+			SNew(SOverlay)
+
+			// Hot path icon
+			+ SOverlay::Slot()
+			[
+				SNew(SImage)
+				.Visibility(this, &STimerTableCell::GetHotPathIconVisibility)
+				.Image(FEditorStyle::GetBrush(TEXT("Profiler.EventGraph.HotPathSmall")))
+			]
+
+			// Info icon + tooltip
+			+ SOverlay::Slot()
+			[
+				SNew(SImage)
+				.Visibility(this, &STimerTableCell::GetHintIconVisibility)
+				.Image(FEditorStyle::GetBrush("Profiler.Tooltip.HintIcon10"))
+				.ToolTip(STimersViewTooltip::GetCellTooltip(TimerNodePtr, ColumnPtr))
+			]
 		]
 
 		// Name
 		+ SHorizontalBox::Slot()
 		.AutoWidth()
 		.VAlign(VAlign_Center)
-		.HAlign(Column.HorizontalAlignment)
+		.HAlign(ColumnPtr->GetHorizontalAlignment())
 		.Padding(FMargin(2.0f, 0.0f))
 		[
 			SNew(STextBlock)
@@ -99,8 +116,7 @@ TSharedRef<SWidget> STimerTableCell::GenerateWidgetForNameColumn(const FArgument
 
 TSharedRef<SWidget> STimerTableCell::GenerateWidgetForStatsColumn(const FArguments& InArgs, const TSharedRef<class ITableRow>& TableRow)
 {
-	const FTimersViewColumn& Column = *FTimersViewColumnFactory::Get().ColumnIdToPtrMapping.FindChecked(ColumnId);
-	const FText FormattedValue = Column.GetFormattedValue(*TimerNodePtr);
+	const FText CellText = ColumnPtr->GetValueAsText(*TimerNodePtr);
 
 	return
 		SNew(SHorizontalBox)
@@ -109,11 +125,11 @@ TSharedRef<SWidget> STimerTableCell::GenerateWidgetForStatsColumn(const FArgumen
 		+ SHorizontalBox::Slot()
 		.FillWidth(1.0f)
 		.VAlign(VAlign_Center)
-		.HAlign(Column.HorizontalAlignment)
+		.HAlign(ColumnPtr->GetHorizontalAlignment())
 		.Padding(FMargin(2.0f, 0.0f))
 		[
 			SNew(STextBlock)
-			.Text(FormattedValue)
+			.Text(CellText)
 			.TextStyle(FEditorStyle::Get(), TEXT("Profiler.Tooltip"))
 			.ColorAndOpacity(this, &STimerTableCell::GetStatsColorAndOpacity)
 			.ShadowColorAndOpacity(this, &STimerTableCell::GetShadowColorAndOpacity)
diff --git a/Engine/Source/Developer/TraceInsights/Private/Insights/Widgets/STimerTableCell.h b/Engine/Source/Developer/TraceInsights/Private/Insights/Widgets/STimerTableCell.h
--- a/Engine/Source/Developer/TraceInsights/Private/Insights/Widgets/STimerTableCell.h
+++ b/Engine/Source/Developer/TraceInsights/Private/Insights/Widgets/STimerTableCell.h
@@ -9,17 +9,23 @@
 // Insights
 #include "Insights/ViewModels/TimerNode.h"
 
-DECLARE_DELEGATE_TwoParams(FSetHoveredTimerTableCell, const FName /*ColumnId*/, const FTimerNodePtr /*SamplePtr*/);
-DECLARE_DELEGATE_RetVal_OneParam(EHorizontalAlignment, FGetColumnOutlineHAlignmentDelegate, const FName /*ColumnId*/);
+namespace Insights
+{
+	class FTable;
+	class FTableColumn;
+}
+
+DECLARE_DELEGATE_ThreeParams(FSetHoveredTimerTableCell, TSharedPtr<Insights::FTable> /*TablePtr*/, TSharedPtr<Insights::FTableColumn> /*ColumnPtr*/, const FTimerNodePtr /*TimerNodePtr*/);
 
 class STimerTableCell : public SCompoundWidget
 {
 public:
 	SLATE_BEGIN_ARGS(STimerTableCell) {}
-		SLATE_EVENT(FSetHoveredTimerTableCell, OnSetHoveredTableCell)
+		SLATE_EVENT(FSetHoveredTimerTableCell, OnSetHoveredCell)
 		SLATE_ATTRIBUTE(FText, HighlightText)
+		SLATE_ARGUMENT(TSharedPtr<Insights::FTable>, TablePtr)
+		SLATE_ARGUMENT(TSharedPtr<Insights::FTableColumn>, ColumnPtr)
 		SLATE_ARGUMENT(FTimerNodePtr, TimerNodePtr)
-		SLATE_ARGUMENT(FName, ColumnId)
 		SLATE_ARGUMENT(bool, IsNameColumn)
 	SLATE_END_ARGS()
 
@@ -39,7 +45,7 @@ protected:
 	virtual void OnMouseEnter(const FGeometry& MyGeometry, const FPointerEvent& MouseEvent) override
 	{
 		SCompoundWidget::OnMouseEnter(MyGeometry, MouseEvent);
-		SetHoveredTableCellDelegate.ExecuteIfBound(ColumnId, TimerNodePtr);
+		SetHoveredCellDelegate.ExecuteIfBound(TablePtr, ColumnPtr, TimerNodePtr);
 	}
 
 	/**
@@ -50,7 +56,7 @@ protected:
 	virtual void OnMouseLeave(const FPointerEvent& MouseEvent) override
 	{
 		SCompoundWidget::OnMouseLeave(MouseEvent);
-		SetHoveredTableCellDelegate.ExecuteIfBound(NAME_None, nullptr);
+		SetHoveredCellDelegate.ExecuteIfBound(nullptr, nullptr, nullptr);
 	}
 
 	/**
@@ -73,7 +79,7 @@ protected:
 	virtual void OnDragEnter(const FGeometry& MyGeometry, const FDragDropEvent& DragDropEvent) override
 	{
 		SCompoundWidget::OnDragEnter(MyGeometry, DragDropEvent);
-		SetHoveredTableCellDelegate.ExecuteIfBound(ColumnId, TimerNodePtr);
+		SetHoveredCellDelegate.ExecuteIfBound(TablePtr, ColumnPtr, TimerNodePtr);
 	}
 
 	/**
@@ -84,7 +90,12 @@ protected:
 	virtual void OnDragLeave(const FDragDropEvent& DragDropEvent)  override
 	{
 		SCompoundWidget::OnDragLeave(DragDropEvent);
-		SetHoveredTableCellDelegate.ExecuteIfBound(NAME_None, nullptr);
+		SetHoveredCellDelegate.ExecuteIfBound(nullptr, nullptr, nullptr);
+	}
+
+	EVisibility GetHotPathIconVisibility() const
+	{
+		return TimerNodePtr->IsHotPath() && !IsHovered() ? EVisibility::Visible : EVisibility::Collapsed;
 	}
 
 	EVisibility GetHintIconVisibility() const
@@ -100,34 +111,41 @@ protected:
 	FSlateColor GetColorAndOpacity() const
 	{
 		const FLinearColor TextColor =
-			TimerNodePtr->IsFiltered() ? FLinearColor(1.0f, 1.0f, 1.0f, 0.5f) :
-			FLinearColor::White;
+			TimerNodePtr->IsFiltered() ?
+				FLinearColor(1.0f, 1.0f, 1.0f, 0.5f) :
+				FLinearColor(1.0f, 1.0f, 1.0f, 1.0f);
 		return TextColor;
 	}
 
 	FSlateColor GetStatsColorAndOpacity() const
 	{
 		const FLinearColor TextColor =
-			TimerNodePtr->IsFiltered() ? FLinearColor(1.0f, 1.0f, 1.0f, 0.5f) :
-			TimerNodePtr->GetAggregatedStats().InstanceCount == 0 ? FLinearColor(1.0f, 1.0f, 1.0f, 0.6f) :
-			FLinearColor::White;
+			TimerNodePtr->IsFiltered() ?
+				FLinearColor(1.0f, 1.0f, 1.0f, 0.5f) :
+			TimerNodePtr->GetAggregatedStats().InstanceCount == 0 ?
+				FLinearColor(1.0f, 1.0f, 1.0f, 0.6f) :
+				FLinearColor(1.0f, 1.0f, 1.0f, 1.0f);
 		return TextColor;
 	}
 
 	FLinearColor GetShadowColorAndOpacity() const
 	{
 		const FLinearColor ShadowColor =
-			TimerNodePtr->IsFiltered() ? FLinearColor(0.f, 0.f, 0.f, 0.25f) :
-			FLinearColor(0.0f, 0.0f, 0.0f, 0.5f);
+			TimerNodePtr->IsFiltered() ?
+				FLinearColor(0.f, 0.f, 0.f, 0.25f) :
+				FLinearColor(0.0f, 0.0f, 0.0f, 0.5f);
 		return ShadowColor;
 	}
 
 protected:
-	/** A shared pointer to the timer node. */
-	FTimerNodePtr TimerNodePtr;
+	/** A shared pointer to the table view model. */
+	TSharedPtr<Insights::FTable> TablePtr; // TODO: TSharedRef
 
-	/** The Id of the column where this timer belongs. */
-	FName ColumnId;
+	/** A shared pointer to the table column view model. */
+	TSharedPtr<Insights::FTableColumn> ColumnPtr; // TODO: TSharedRef
+
+	/** A shared pointer to the timer node. */
+	FTimerNodePtr TimerNodePtr; // TODO: TSharedRef
 
-	FSetHoveredTimerTableCell SetHoveredTableCellDelegate;
+	FSetHoveredTimerTableCell SetHoveredCellDelegate;
 };
diff --git a/Engine/Source/Developer/TraceInsights/Private/Insights/Widgets/STimerTableRow.cpp b/Engine/Source/Developer/TraceInsights/Private/Insights/Widgets/STimerTableRow.cpp
--- a/Engine/Source/Developer/TraceInsights/Private/Insights/Widgets/STimerTableRow.cpp
+++ b/Engine/Source/Developer/TraceInsights/Private/Insights/Widgets/STimerTableRow.cpp
@@ -9,8 +9,8 @@
 
 // Insights
 #include "Insights/Common/TimeUtils.h"
-#include "Insights/ViewModels/TimersViewColumnFactory.h"
-#include "Insights/Widgets/STimersViewTooltip.h"
+#include "Insights/Table/ViewModels/Table.h"
+#include "Insights/Table/ViewModels/TableColumn.h"
 #include "Insights/Widgets/STimerTableCell.h"
 
 #define LOCTEXT_NAMESPACE "STimersView"
@@ -23,12 +23,13 @@ void STimerTableRow::Construct(const FArguments& InArgs, const TSharedRef<STable
 {
 	OnShouldBeEnabled = InArgs._OnShouldBeEnabled;
 	IsColumnVisibleDelegate = InArgs._OnIsColumnVisible;
-	SetHoveredTableCellDelegate = InArgs._OnSetHoveredTableCell;
+	SetHoveredCellDelegate = InArgs._OnSetHoveredCell;
 	GetColumnOutlineHAlignmentDelegate = InArgs._OnGetColumnOutlineHAlignmentDelegate;
 
 	HighlightText = InArgs._HighlightText;
 	HighlightedNodeName = InArgs._HighlightedNodeName;
 
+	TablePtr = InArgs._TablePtr;
 	TimerNodePtr = InArgs._TimerNodePtr;
 
 	SetEnabled(TAttribute<bool>(this, &STimerTableRow::HandleShouldBeEnabled));
@@ -40,6 +41,8 @@ void STimerTableRow::Construct(const FArguments& InArgs, const TSharedRef<STable
 
 TSharedRef<SWidget> STimerTableRow::GenerateWidgetForColumn(const FName& ColumnId)
 {
+	TSharedPtr<Insights::FTableColumn> ColumnPtr = TablePtr->FindColumnChecked(ColumnId);
+
 	return
 		SNew(SOverlay)
 		.Visibility(EVisibility::SelfHitTestInvisible)
@@ -64,11 +67,12 @@ TSharedRef<SWidget> STimerTableRow::GenerateWidgetForColumn(const FName& ColumnI
 		[
 			SNew(STimerTableCell, SharedThis(this))
 			.Visibility(this, &STimerTableRow::IsColumnVisible, ColumnId)
+			.TablePtr(TablePtr)
+			.ColumnPtr(ColumnPtr)
 			.TimerNodePtr(TimerNodePtr)
-			.ColumnId(ColumnId)
 			.HighlightText(HighlightText)
-			.IsNameColumn(ColumnId == FTimersViewColumnFactory::Get().Collection[0]->Id) // name column
-			.OnSetHoveredTableCell(this, &STimerTableRow::OnSetHoveredTableCell)
+			.IsNameColumn(ColumnPtr->IsHierarchy())
+			.OnSetHoveredCell(this, &STimerTableRow::OnSetHoveredCell)
 		];
 }
 
@@ -197,9 +201,9 @@ EVisibility STimerTableRow::IsColumnVisible(const FName ColumnId) const
 
 ////////////////////////////////////////////////////////////////////////////////////////////////////
 
-void STimerTableRow::OnSetHoveredTableCell(const FName InColumnId, const FTimerNodePtr InSamplePtr)
+void STimerTableRow::OnSetHoveredCell(TSharedPtr<Insights::FTable> InTablePtr, TSharedPtr<Insights::FTableColumn> InColumnPtr, const FTimerNodePtr InTimerNodePtr)
 {
-	SetHoveredTableCellDelegate.ExecuteIfBound(InColumnId, InSamplePtr);
+	SetHoveredCellDelegate.ExecuteIfBound(InTablePtr, InColumnPtr, InTimerNodePtr);
 }
 
 ////////////////////////////////////////////////////////////////////////////////////////////////////
diff --git a/Engine/Source/Developer/TraceInsights/Private/Insights/Widgets/STimerTableRow.h b/Engine/Source/Developer/TraceInsights/Private/Insights/Widgets/STimerTableRow.h
--- a/Engine/Source/Developer/TraceInsights/Private/Insights/Widgets/STimerTableRow.h
+++ b/Engine/Source/Developer/TraceInsights/Private/Insights/Widgets/STimerTableRow.h
@@ -12,22 +12,29 @@
 // Insights
 #include "Insights/ViewModels/TimerNodeHelper.h"
 
+namespace Insights
+{
+	class FTable;
+	class FTableColumn;
+}
+
 DECLARE_DELEGATE_RetVal_OneParam(bool, FShouldBeEnabledDelegate, const uint32 /*TimerId*/);
 DECLARE_DELEGATE_RetVal_OneParam(bool, FIsColumnVisibleDelegate, const FName /*ColumnId*/);
-DECLARE_DELEGATE_TwoParams(FSetHoveredTimerTableCell, const FName /*ColumnId*/, const FTimerNodePtr /*TimerNodePtr*/);
 DECLARE_DELEGATE_RetVal_OneParam(EHorizontalAlignment, FGetColumnOutlineHAlignmentDelegate, const FName /*ColumnId*/);
+DECLARE_DELEGATE_ThreeParams(FSetHoveredTimerTableCell, TSharedPtr<Insights::FTable> /*TablePtr*/, TSharedPtr<Insights::FTableColumn> /*ColumnPtr*/, const FTimerNodePtr /*TimerNodePtr*/);
 
-/** Widget that represents a table row in the timers' tree control. Generates widgets for each column on demand. */
+/** Widget that represents a table row in a timer tree control. Generates widgets for each column on demand. */
 class STimerTableRow : public SMultiColumnTableRow<FTimerNodePtr>
 {
 public:
 	SLATE_BEGIN_ARGS(STimerTableRow) {}
 		SLATE_EVENT(FShouldBeEnabledDelegate, OnShouldBeEnabled)
 		SLATE_EVENT(FIsColumnVisibleDelegate, OnIsColumnVisible)
-		SLATE_EVENT(FSetHoveredTimerTableCell, OnSetHoveredTableCell)
 		SLATE_EVENT(FGetColumnOutlineHAlignmentDelegate, OnGetColumnOutlineHAlignmentDelegate)
+		SLATE_EVENT(FSetHoveredTimerTableCell, OnSetHoveredCell)
 		SLATE_ATTRIBUTE(FText, HighlightText)
 		SLATE_ATTRIBUTE(FName, HighlightedNodeName)
+		SLATE_ARGUMENT(TSharedPtr<Insights::FTable>, TablePtr)
 		SLATE_ARGUMENT(FTimerNodePtr, TimerNodePtr)
 	SLATE_END_ARGS()
 
@@ -59,15 +66,18 @@ protected:
 	const FSlateBrush* GetOutlineBrush(const FName ColumnId) const;
 	bool HandleShouldBeEnabled() const;
 	EVisibility IsColumnVisible(const FName ColumnId) const;
-	void OnSetHoveredTableCell(const FName ColumnId, const FTimerNodePtr SamplePtr);
+	void OnSetHoveredCell(TSharedPtr<Insights::FTable> InTablePtr, TSharedPtr<Insights::FTableColumn> InColumnPtr, const FTimerNodePtr InTimerNodePtr);
 
 protected:
+	/** A shared pointer to the table view model. */
+	TSharedPtr<Insights::FTable> TablePtr;
+
 	/** Data context for this table row. */
 	FTimerNodePtr TimerNodePtr;
 
 	FShouldBeEnabledDelegate OnShouldBeEnabled;
 	FIsColumnVisibleDelegate IsColumnVisibleDelegate;
-	FSetHoveredTimerTableCell SetHoveredTableCellDelegate;
+	FSetHoveredTimerTableCell SetHoveredCellDelegate;
 	FGetColumnOutlineHAlignmentDelegate GetColumnOutlineHAlignmentDelegate;
 
 	/** Text to be highlighted on timer name. */
diff --git a/Engine/Source/Developer/TraceInsights/Private/Insights/Widgets/STimerTreeView.cpp b/Engine/Source/Developer/TraceInsights/Private/Insights/Widgets/STimerTreeView.cpp
--- a/Engine/Source/Developer/TraceInsights/Private/Insights/Widgets/STimerTreeView.cpp
+++ b/Engine/Source/Developer/TraceInsights/Private/Insights/Widgets/STimerTreeView.cpp
@@ -7,9 +7,13 @@
 #include "Framework/MultiBox/MultiBoxBuilder.h"
 #include "SlateOptMacros.h"
 #include "Widgets/Layout/SScrollBox.h"
-//#include "Widgets/Views/STableViewBase.h"
+#include "Widgets/SToolTip.h"
 
 // Insights
+#include "Insights/Table/ViewModels/Table.h"
+#include "Insights/Table/ViewModels/TableColumn.h"
+#include "Insights/Table/ViewModels/TreeNodeSorting.h"
+#include "Insights/ViewModels/TimersViewColumnFactory.h"
 #include "Insights/Widgets/STimersViewTooltip.h"
 #include "Insights/Widgets/STimerTableRow.h"
 
@@ -17,14 +21,23 @@
 
 ////////////////////////////////////////////////////////////////////////////////////////////////////
 
-const FName STimerTreeView::TimerNameColumnId(TEXT("Name")); // FTimersViewColumns::TotalInclusiveTimeColumnID
-const FName STimerTreeView::CountColumnId(TEXT("Count")); // FTimersViewColumns::TotalInclusiveTimeColumnID
-const FName STimerTreeView::InclusiveTimeColumnId(TEXT("TotalInclTime")); // FTimersViewColumns::TotalInclusiveTimeColumnID
-const FName STimerTreeView::ExclusiveTimeColumnId(TEXT("TotalExclTime")); // FTimersViewColumns::TotalExclusiveTimeColumnID
-
-////////////////////////////////////////////////////////////////////////////////////////////////////
-
 STimerTreeView::STimerTreeView()
+	: Table(MakeShareable(new Insights::FTable()))
+	, ViewName()
+	, TreeView(nullptr)
+	, TreeViewHeaderRow(nullptr)
+	, ExternalScrollbar(nullptr)
+	, HoveredColumnId()
+	, HoveredNodePtr(nullptr)
+	, HighlightedNodeName()
+	, TreeNodes()
+	, AvailableSorters()
+	, CurrentSorter(nullptr)
+	, ColumnBeingSorted(GetDefaultColumnBeingSorted())
+	, ColumnSortMode(GetDefaultColumnSortMode())
+	, StatsStartTime(0.0)
+	, StatsEndTime(0.0)
+	, StatsTimerId(0)
 {
 }
 
@@ -69,7 +82,7 @@ void STimerTreeView::Construct(const FArguments& InArgs, const FText& InViewName
 					.OnGenerateRow(this, &STimerTreeView::TreeView_OnGenerateRow)
 					//.OnSelectionChanged(this, &STimerTreeView::TreeView_OnSelectionChanged)
 					//.OnMouseButtonDoubleClick(this, &STimerTreeView::TreeView_OnMouseButtonDoubleClick)
-					//.OnContextMenuOpening(FOnContextMenuOpening::CreateSP(this, &STimerTreeView::TreeView_GetMenuContent))
+					.OnContextMenuOpening(FOnContextMenuOpening::CreateSP(this, &STimerTreeView::TreeView_GetMenuContent))
 					.ItemHeight(12.0f)
 					.HeaderRow
 					(
@@ -93,6 +106,8 @@ void STimerTreeView::Construct(const FArguments& InArgs, const FText& InViewName
 	];
 
 	InitializeAndShowHeaderColumns();
+
+	CreateSortings();
 }
 END_SLATE_FUNCTION_BUILD_OPTIMIZATION
 
@@ -100,18 +115,28 @@ END_SLATE_FUNCTION_BUILD_OPTIMIZATION
 
 TSharedPtr<SWidget> STimerTreeView::TreeView_GetMenuContent()
 {
-	const TArray<FTimerNodePtr> SelectedTimerNodes = TreeView->GetSelectedItems();
-	const int32 NumSelectedTimerNodes = SelectedTimerNodes.Num();
-	FTimerNodePtr SelectedTimerNode = NumSelectedTimerNodes ? SelectedTimerNodes[0] : nullptr;
+	const TArray<FTimerNodePtr> SelectedNodes = TreeView->GetSelectedItems();
+	const int32 NumSelectedNodes = SelectedNodes.Num();
+	FTimerNodePtr SelectedNode = NumSelectedNodes ? SelectedNodes[0] : nullptr;
+
+	const TSharedPtr<Insights::FTableColumn> HoveredColumnPtr = Table->FindColumn(HoveredColumnId);
 
 	FText SelectionStr;
-	if (NumSelectedTimerNodes == 0)
+	FText PropertyName;
+	FText PropertyValue;
+
+	if (NumSelectedNodes == 0)
 	{
 		SelectionStr = LOCTEXT("NothingSelected", "Nothing selected");
 	}
-	else if (NumSelectedTimerNodes == 1)
+	else if (NumSelectedNodes == 1)
 	{
-		SelectionStr = FText::FromName(SelectedTimerNode->GetName());
+		if (HoveredColumnPtr != nullptr)
+		{
+			PropertyName = HoveredColumnPtr->GetShortName();
+			PropertyValue = HoveredColumnPtr->GetValueAsTooltipText(*SelectedNode);
+		}
+		SelectionStr = FText::FromName(SelectedNode->GetName());
 	}
 	else
 	{
@@ -143,173 +168,244 @@ TSharedPtr<SWidget> STimerTreeView::TreeView_GetMenuContent()
 	}
 	MenuBuilder.EndSection();
 
+	MenuBuilder.BeginSection("Misc", LOCTEXT("ContextMenu_Header_Misc", "Miscellaneous"));
+	{
+		MenuBuilder.AddSubMenu
+		(
+			LOCTEXT("ContextMenu_Header_Misc_Sort", "Sort By"),
+			LOCTEXT("ContextMenu_Header_Misc_Sort_Desc", "Sort by column"),
+			FNewMenuDelegate::CreateSP(this, &STimerTreeView::TreeView_BuildSortByMenu),
+			false,
+			FSlateIcon(FEditorStyle::GetStyleSetName(), "Profiler.Misc.SortBy")
+		);
+	}
+	MenuBuilder.EndSection();
+
+	MenuBuilder.BeginSection("Columns", LOCTEXT("ContextMenu_Header_Columns", "Columns"));
+	{
+		MenuBuilder.AddSubMenu
+		(
+			LOCTEXT("ContextMenu_Header_Columns_View", "View Column"),
+			LOCTEXT("ContextMenu_Header_Columns_View_Desc", "Hides or shows columns"),
+			FNewMenuDelegate::CreateSP(this, &STimerTreeView::TreeView_BuildViewColumnMenu),
+			false,
+			FSlateIcon(FEditorStyle::GetStyleSetName(), "Profiler.EventGraph.ViewColumn")
+		);
+
+		FUIAction Action_ShowAllColumns
+		(
+			FExecuteAction::CreateSP(this, &STimerTreeView::ContextMenu_ShowAllColumns_Execute),
+			FCanExecuteAction::CreateSP(this, &STimerTreeView::ContextMenu_ShowAllColumns_CanExecute)
+		);
+		MenuBuilder.AddMenuEntry
+		(
+			LOCTEXT("ContextMenu_Header_Columns_ShowAllColumns", "Show All Columns"),
+			LOCTEXT("ContextMenu_Header_Columns_ShowAllColumns_Desc", "Resets tree view to show all columns"),
+			FSlateIcon(FEditorStyle::GetStyleSetName(), "Profiler.EventGraph.ResetColumn"), Action_ShowAllColumns, NAME_None, EUserInterfaceActionType::Button
+		);
+
+		FUIAction Action_ResetColumns
+		(
+			FExecuteAction::CreateSP(this, &STimerTreeView::ContextMenu_ResetColumns_Execute),
+			FCanExecuteAction::CreateSP(this, &STimerTreeView::ContextMenu_ResetColumns_CanExecute)
+		);
+		MenuBuilder.AddMenuEntry
+		(
+			LOCTEXT("ContextMenu_Header_Columns_ResetColumns", "Reset Columns to Default"),
+			LOCTEXT("ContextMenu_Header_Columns_ResetColumns_Desc", "Resets columns to default"),
+			FSlateIcon(FEditorStyle::GetStyleSetName(), "Profiler.EventGraph.ResetColumn"), Action_ResetColumns, NAME_None, EUserInterfaceActionType::Button
+		);
+	}
+	MenuBuilder.EndSection();
+
 	return MenuBuilder.MakeWidget();
 }
 
 ////////////////////////////////////////////////////////////////////////////////////////////////////
 
-void STimerTreeView::InitializeAndShowHeaderColumns()
+void STimerTreeView::TreeView_BuildSortByMenu(FMenuBuilder& MenuBuilder)
 {
-	TreeViewHeaderRow_CreateTimerNameColumnArgs();
-	TreeViewHeaderRow_CreateCountColumnArgs();
-	TreeViewHeaderRow_CreateInclusiveTimeColumnArgs();
-	TreeViewHeaderRow_CreateExclusiveTimeColumnArgs();
+	MenuBuilder.BeginSection("ColumnName", LOCTEXT("ContextMenu_Header_Misc_ColumnName", "Column Name"));
 
-	TreeViewHeaderRow->AddColumn(TreeViewHeaderColumnArgs[TimerNameColumnId]);
-	TreeViewHeaderRow->AddColumn(TreeViewHeaderColumnArgs[CountColumnId]);
-	TreeViewHeaderRow->AddColumn(TreeViewHeaderColumnArgs[InclusiveTimeColumnId]);
-	TreeViewHeaderRow->AddColumn(TreeViewHeaderColumnArgs[ExclusiveTimeColumnId]);
-}
+	for (const TSharedPtr<Insights::FTableColumn>& ColumnPtr : Table->GetColumns())
+	{
+		const Insights::FTableColumn& Column = *ColumnPtr;
 
-////////////////////////////////////////////////////////////////////////////////////////////////////
+		if (Column.IsVisible() && Column.CanBeSorted())
+		{
+			FUIAction Action_SortByColumn
+			(
+				FExecuteAction::CreateSP(this, &STimerTreeView::ContextMenu_SortByColumn_Execute, Column.GetId()),
+				FCanExecuteAction::CreateSP(this, &STimerTreeView::ContextMenu_SortByColumn_CanExecute, Column.GetId()),
+				FIsActionChecked::CreateSP(this, &STimerTreeView::ContextMenu_SortByColumn_IsChecked, Column.GetId())
+			);
+			MenuBuilder.AddMenuEntry
+			(
+				Column.GetTitleName(),
+				Column.GetDescription(),
+				FSlateIcon(), Action_SortByColumn, NAME_None, EUserInterfaceActionType::RadioButton
+			);
+		}
+	}
 
-void STimerTreeView::TreeViewHeaderRow_CreateTimerNameColumnArgs()
-{
-	SHeaderRow::FColumn::FArguments ColumnArgs;
+	MenuBuilder.EndSection();
 
-	const FName& ColumnId = TimerNameColumnId;
+	//-----------------------------------------------------------------------------
 
-	ColumnArgs
-		.ColumnId(ColumnId)
-		.DefaultLabel(ViewName)
-		.HAlignHeader(HAlign_Fill)
-		.VAlignHeader(VAlign_Fill)
-		.HeaderContentPadding(FMargin(2.0f))
-		.HAlignCell(HAlign_Fill)
-		.VAlignCell(VAlign_Fill)
-		.ManualWidth(246.0f)
-		.HeaderContent()
-		[
-			SNew(SBox)
-			//.ToolTip(STimerTreeView::GetColumnTooltip(ColumnId))
-			.HAlign(HAlign_Left)
-			.VAlign(VAlign_Center)
-			[
-				SNew(STextBlock)
-				.Text(ViewName)
-			]
-		]
-		//.MenuContent()
-		//[
-		//	TreeViewHeaderRow_GenerateColumnMenu(ColumnId)
-		//]
-	;
+	MenuBuilder.BeginSection("SortMode", LOCTEXT("ContextMenu_Header_Misc_Sort_SortMode", "Sort Mode"));
+	{
+		FUIAction Action_SortAscending
+		(
+			FExecuteAction::CreateSP(this, &STimerTreeView::ContextMenu_SortMode_Execute, EColumnSortMode::Ascending),
+			FCanExecuteAction::CreateSP(this, &STimerTreeView::ContextMenu_SortMode_CanExecute, EColumnSortMode::Ascending),
+			FIsActionChecked::CreateSP(this, &STimerTreeView::ContextMenu_SortMode_IsChecked, EColumnSortMode::Ascending)
+		);
+		MenuBuilder.AddMenuEntry
+		(
+			LOCTEXT("ContextMenu_Header_Misc_Sort_SortAscending", "Sort Ascending"),
+			LOCTEXT("ContextMenu_Header_Misc_Sort_SortAscending_Desc", "Sorts ascending"),
+			FSlateIcon(FEditorStyle::GetStyleSetName(), "Profiler.Misc.SortAscending"), Action_SortAscending, NAME_None, EUserInterfaceActionType::RadioButton
+		);
 
-	TreeViewHeaderColumnArgs.Add(ColumnId, ColumnArgs);
+		FUIAction Action_SortDescending
+		(
+			FExecuteAction::CreateSP(this, &STimerTreeView::ContextMenu_SortMode_Execute, EColumnSortMode::Descending),
+			FCanExecuteAction::CreateSP(this, &STimerTreeView::ContextMenu_SortMode_CanExecute, EColumnSortMode::Descending),
+			FIsActionChecked::CreateSP(this, &STimerTreeView::ContextMenu_SortMode_IsChecked, EColumnSortMode::Descending)
+		);
+		MenuBuilder.AddMenuEntry
+		(
+			LOCTEXT("ContextMenu_Header_Misc_Sort_SortDescending", "Sort Descending"),
+			LOCTEXT("ContextMenu_Header_Misc_Sort_SortDescending_Desc", "Sorts descending"),
+			FSlateIcon(FEditorStyle::GetStyleSetName(), "Profiler.Misc.SortDescending"), Action_SortDescending, NAME_None, EUserInterfaceActionType::RadioButton
+		);
+	}
+	MenuBuilder.EndSection();
 }
 
 ////////////////////////////////////////////////////////////////////////////////////////////////////
 
-void STimerTreeView::TreeViewHeaderRow_CreateCountColumnArgs()
+void STimerTreeView::TreeView_BuildViewColumnMenu(FMenuBuilder& MenuBuilder)
 {
-	SHeaderRow::FColumn::FArguments ColumnArgs;
+	MenuBuilder.BeginSection("ViewColumn", LOCTEXT("ContextMenu_Header_Columns_View", "View Column"));
 
-	const FName& ColumnId = CountColumnId;
-	const FText ColumnName = LOCTEXT("CountColumnName", "Count");
+	for (const TSharedPtr<Insights::FTableColumn>& ColumnPtr : Table->GetColumns())
+	{
+		const Insights::FTableColumn& Column = *ColumnPtr;
 
-	ColumnArgs
-		.ColumnId(ColumnId)
-		.DefaultLabel(ColumnName)
-		.HAlignHeader(HAlign_Fill)
-		.VAlignHeader(VAlign_Fill)
-		.HeaderContentPadding(FMargin(2.0f))
-		.HAlignCell(HAlign_Fill)
-		.VAlignCell(VAlign_Fill)
-		.ManualWidth(60.0f)
-		.HeaderContent()
-		[
-			SNew(SBox)
-			//.ToolTip(STimerTreeView::GetColumnTooltip(ColumnId))
-			.HAlign(HAlign_Right)
-			.VAlign(VAlign_Center)
-			[
-				SNew(STextBlock)
-				.Text(ColumnName)
-			]
-		]
-		//.MenuContent()
-		//[
-		//	TreeViewHeaderRow_GenerateColumnMenu(ColumnId)
-		//]
-	;
+		FUIAction Action_ToggleColumn
+		(
+			FExecuteAction::CreateSP(this, &STimerTreeView::ToggleColumnVisibility, Column.GetId()),
+			FCanExecuteAction::CreateSP(this, &STimerTreeView::CanToggleColumnVisibility, Column.GetId()),
+			FIsActionChecked::CreateSP(this, &STimerTreeView::IsColumnVisible, Column.GetId())
+		);
+		MenuBuilder.AddMenuEntry
+		(
+			Column.GetTitleName(),
+			Column.GetDescription(),
+			FSlateIcon(), Action_ToggleColumn, NAME_None, EUserInterfaceActionType::ToggleButton
+		);
+	}
 
-	TreeViewHeaderColumnArgs.Add(ColumnId, ColumnArgs);
+	MenuBuilder.EndSection();
 }
 
 ////////////////////////////////////////////////////////////////////////////////////////////////////
 
-void STimerTreeView::TreeViewHeaderRow_CreateInclusiveTimeColumnArgs()
+void STimerTreeView::InitializeAndShowHeaderColumns()
 {
-	SHeaderRow::FColumn::FArguments ColumnArgs;
-
-	const FName& ColumnId = InclusiveTimeColumnId;
-	const FText ColumnName = LOCTEXT("InclusiveTimeColumnName", "Incl");
+	// Create columns.
+	TArray<TSharedPtr<Insights::FTableColumn>> Columns;
+	FTimersViewColumnFactory::CreateTimerTreeViewColumns(Columns);
+	ensure(Columns.Num() > 0 && Columns[0]->IsHierarchy());
+	Columns[0]->SetShortName(ViewName);
+	Columns[0]->SetTitleName(ViewName);
+	Table->SetColumns(Columns);
+
+	// Show columns.
+	for (const TSharedPtr<Insights::FTableColumn>& ColumnPtr : Table->GetColumns())
+	{
+		if (ColumnPtr->ShouldBeVisible())
+		{
+			ShowColumn(ColumnPtr->GetId());
+		}
+	}
+}
 
-	ColumnArgs
-		.ColumnId(ColumnId)
-		.DefaultLabel(ColumnName)
-		.HAlignHeader(HAlign_Fill)
-		.VAlignHeader(VAlign_Fill)
-		.HeaderContentPadding(FMargin(2.0f))
-		.HAlignCell(HAlign_Fill)
-		.VAlignCell(VAlign_Fill)
-		.ManualWidth(60.0f)
-		.HeaderContent()
-		[
-			SNew(SBox)
-			//.ToolTip(STimerTreeView::GetColumnTooltip(ColumnId))
-			.HAlign(HAlign_Right)
-			.VAlign(VAlign_Center)
-			[
-				SNew(STextBlock)
-				.Text(ColumnName)
-			]
-		]
-		//.MenuContent()
-		//[
-		//	TreeViewHeaderRow_GenerateColumnMenu(ColumnId)
-		//]
-	;
+////////////////////////////////////////////////////////////////////////////////////////////////////
 
-	TreeViewHeaderColumnArgs.Add(ColumnId, ColumnArgs);
+FText STimerTreeView::GetColumnHeaderText(const FName ColumnId) const
+{
+	const Insights::FTableColumn& Column = *Table->FindColumnChecked(ColumnId);
+	return Column.GetShortName();
 }
 
 ////////////////////////////////////////////////////////////////////////////////////////////////////
 
-void STimerTreeView::TreeViewHeaderRow_CreateExclusiveTimeColumnArgs()
+TSharedRef<SWidget> STimerTreeView::TreeViewHeaderRow_GenerateColumnMenu(const Insights::FTableColumn& Column)
 {
-	SHeaderRow::FColumn::FArguments ColumnArgs;
+	bool bIsMenuVisible = false;
 
-	const FName& ColumnId = ExclusiveTimeColumnId;
-	const FText ColumnName = LOCTEXT("ExclusiveTimeColumnName", "Excl");
+	const bool bShouldCloseWindowAfterMenuSelection = true;
+	FMenuBuilder MenuBuilder(bShouldCloseWindowAfterMenuSelection, NULL);
+	{
+		if (Column.CanBeHidden())
+		{
+			MenuBuilder.BeginSection("Column", LOCTEXT("TreeViewHeaderRow_Header_Column", "Column"));
+
+			FUIAction Action_HideColumn
+			(
+				FExecuteAction::CreateSP(this, &STimerTreeView::HideColumn, Column.GetId()),
+				FCanExecuteAction::CreateSP(this, &STimerTreeView::CanHideColumn, Column.GetId())
+			);
+
+			MenuBuilder.AddMenuEntry
+			(
+				LOCTEXT("TreeViewHeaderRow_HideColumn", "Hide"),
+				LOCTEXT("TreeViewHeaderRow_HideColumn_Desc", "Hides the selected column"),
+				FSlateIcon(), Action_HideColumn, NAME_None, EUserInterfaceActionType::Button
+			);
+			bIsMenuVisible = true;
+
+			MenuBuilder.EndSection();
+		}
 
-	ColumnArgs
-		.ColumnId(ColumnId)
-		.DefaultLabel(ColumnName)
-		.HAlignHeader(HAlign_Fill)
-		.VAlignHeader(VAlign_Fill)
-		.HeaderContentPadding(FMargin(2.0f))
-		.HAlignCell(HAlign_Fill)
-		.VAlignCell(VAlign_Fill)
-		.ManualWidth(60.0f)
-		.HeaderContent()
-		[
-			SNew(SBox)
-			//.ToolTip(STimerTreeView::GetColumnTooltip(ColumnId))
-			.HAlign(HAlign_Right)
-			.VAlign(VAlign_Center)
-			[
-				SNew(STextBlock)
-				.Text(ColumnName)
-			]
-		]
-		//.MenuContent()
-		//[
-		//	TreeViewHeaderRow_GenerateColumnMenu(ColumnId)
-		//]
-	;
+		if (Column.CanBeSorted())
+		{
+			MenuBuilder.BeginSection("SortMode", LOCTEXT("ContextMenu_Header_Misc_Sort_SortMode", "Sort Mode"));
+
+			FUIAction Action_SortAscending
+			(
+				FExecuteAction::CreateSP(this, &STimerTreeView::HeaderMenu_SortMode_Execute, Column.GetId(), EColumnSortMode::Ascending),
+				FCanExecuteAction::CreateSP(this, &STimerTreeView::HeaderMenu_SortMode_CanExecute, Column.GetId(), EColumnSortMode::Ascending),
+				FIsActionChecked::CreateSP(this, &STimerTreeView::HeaderMenu_SortMode_IsChecked, Column.GetId(), EColumnSortMode::Ascending)
+			);
+			MenuBuilder.AddMenuEntry
+			(
+				LOCTEXT("ContextMenu_Header_Misc_Sort_SortAscending", "Sort Ascending"),
+				LOCTEXT("ContextMenu_Header_Misc_Sort_SortAscending_Desc", "Sorts ascending"),
+				FSlateIcon(FEditorStyle::GetStyleSetName(), "Profiler.Misc.SortAscending"), Action_SortAscending, NAME_None, EUserInterfaceActionType::RadioButton
+			);
+
+			FUIAction Action_SortDescending
+			(
+				FExecuteAction::CreateSP(this, &STimerTreeView::HeaderMenu_SortMode_Execute, Column.GetId(), EColumnSortMode::Descending),
+				FCanExecuteAction::CreateSP(this, &STimerTreeView::HeaderMenu_SortMode_CanExecute, Column.GetId(), EColumnSortMode::Descending),
+				FIsActionChecked::CreateSP(this, &STimerTreeView::HeaderMenu_SortMode_IsChecked, Column.GetId(), EColumnSortMode::Descending)
+			);
+			MenuBuilder.AddMenuEntry
+			(
+				LOCTEXT("ContextMenu_Header_Misc_Sort_SortDescending", "Sort Descending"),
+				LOCTEXT("ContextMenu_Header_Misc_Sort_SortDescending_Desc", "Sorts descending"),
+				FSlateIcon(FEditorStyle::GetStyleSetName(), "Profiler.Misc.SortDescending"), Action_SortDescending, NAME_None, EUserInterfaceActionType::RadioButton
+			);
+			bIsMenuVisible = true;
+
+			MenuBuilder.EndSection();
+		}
+	}
 
-	TreeViewHeaderColumnArgs.Add(ColumnId, ColumnArgs);
+	return bIsMenuVisible ? MenuBuilder.MakeWidget() : (TSharedRef<SWidget>)SNullWidget::NullWidget;
 }
 
 ////////////////////////////////////////////////////////////////////////////////////////////////////
@@ -357,12 +453,13 @@ void STimerTreeView::TreeView_OnGetChildren(FTimerNodePtr InParent, TArray<FTime
 
 ////////////////////////////////////////////////////////////////////////////////////////////////////
 
-void STimerTreeView::TreeView_OnMouseButtonDoubleClick(FTimerNodePtr TimerNode)
+void STimerTreeView::TreeView_OnMouseButtonDoubleClick(FTimerNodePtr TimerNodePtr)
 {
-	if (TimerNode->IsGroup())
+	//if (TimerNodePtr->IsGroup())
+	if (TimerNodePtr->GetChildren().Num() > 0)
 	{
-		const bool bIsGroupExpanded = TreeView->IsItemExpanded(TimerNode);
-		TreeView->SetItemExpansion(TimerNode, !bIsGroupExpanded);
+		const bool bIsGroupExpanded = TreeView->IsItemExpanded(TimerNodePtr);
+		TreeView->SetItemExpansion(TimerNodePtr, !bIsGroupExpanded);
 	}
 }
 
@@ -375,11 +472,12 @@ TSharedRef<ITableRow> STimerTreeView::TreeView_OnGenerateRow(FTimerNodePtr Timer
 	TSharedRef<ITableRow> TableRow =
 		SNew(STimerTableRow, OwnerTable)
 		.OnShouldBeEnabled(this, &STimerTreeView::TableRow_ShouldBeEnabled)
-		.OnIsColumnVisible(this, &STimerTreeView::TableRow_IsColumnVisible)
-		.OnSetHoveredTableCell(this, &STimerTreeView::TableRow_SetHoveredTableCell)
+		.OnIsColumnVisible(this, &STimerTreeView::IsColumnVisible)
+		.OnSetHoveredCell(this, &STimerTreeView::TableRow_SetHoveredCell)
 		.OnGetColumnOutlineHAlignmentDelegate(this, &STimerTreeView::TableRow_GetColumnOutlineHAlignment)
 		.HighlightText(this, &STimerTreeView::TableRow_GetHighlightText)
 		.HighlightedNodeName(this, &STimerTreeView::TableRow_GetHighlightedNodeName)
+		.TablePtr(Table)
 		.TimerNodePtr(TimerNodePtr);
 
 	return TableRow;
@@ -394,21 +492,14 @@ bool STimerTreeView::TableRow_ShouldBeEnabled(const uint32 TimerId) const
 
 ////////////////////////////////////////////////////////////////////////////////////////////////////
 
-bool STimerTreeView::TableRow_IsColumnVisible(const FName ColumnId) const
-{
-	return true;
-}
-
-////////////////////////////////////////////////////////////////////////////////////////////////////
-
-void STimerTreeView::TableRow_SetHoveredTableCell(const FName ColumnId, const FTimerNodePtr TimerNodePtr)
+void STimerTreeView::TableRow_SetHoveredCell(TSharedPtr<Insights::FTable> InTablePtr, TSharedPtr<Insights::FTableColumn> InColumnPtr, const FTimerNodePtr InNodePtr)
 {
-	HoveredColumnId = ColumnId;
+	HoveredColumnId = InColumnPtr ? InColumnPtr->GetId() : FName();
 
 	const bool bIsAnyMenusVisible = FSlateApplication::Get().AnyMenusVisible();
 	if (!HasMouseCapture() && !bIsAnyMenusVisible)
 	{
-		HoveredTimerNodePtr = TimerNodePtr;
+		HoveredNodePtr = InNodePtr;
 	}
 }
 
@@ -416,13 +507,16 @@ void STimerTreeView::TableRow_SetHoveredTableCell(const FName ColumnId, const FT
 
 EHorizontalAlignment STimerTreeView::TableRow_GetColumnOutlineHAlignment(const FName ColumnId) const
 {
+	const TIndirectArray<SHeaderRow::FColumn>& Columns = TreeViewHeaderRow->GetColumns();
+	const int32 LastColumnIdx = Columns.Num() - 1;
+
 	// First column
-	if (ColumnId == TimerNameColumnId)
+	if (Columns[0].ColumnId == ColumnId)
 	{
 		return HAlign_Left;
 	}
 	// Last column
-	else if (ColumnId == ExclusiveTimeColumnId)
+	else if (Columns[LastColumnIdx].ColumnId == ColumnId)
 	{
 		return HAlign_Right;
 	}
@@ -446,11 +540,364 @@ FName STimerTreeView::TableRow_GetHighlightedNodeName() const
 	return HighlightedNodeName;
 }
 
+////////////////////////////////////////////////////////////////////////////////////////////////////
+// Sorting
+////////////////////////////////////////////////////////////////////////////////////////////////////
+
+const EColumnSortMode::Type STimerTreeView::GetDefaultColumnSortMode()
+{
+	return EColumnSortMode::Descending;
+}
+
 ////////////////////////////////////////////////////////////////////////////////////////////////////
 
-//TSharedPtr<SWidget> STimerTreeView::TreeView_GetMenuContent()
-//{
-//}
+const FName STimerTreeView::GetDefaultColumnBeingSorted()
+{
+	return FTimersViewColumns::TotalInclusiveTimeColumnID;
+}
+
+////////////////////////////////////////////////////////////////////////////////////////////////////
+
+void STimerTreeView::CreateSortings()
+{
+	AvailableSorters.Reset();
+	CurrentSorter = nullptr;
+
+	for (const TSharedPtr<Insights::FTableColumn> ColumnPtr : Table->GetColumns())
+	{
+		if (ColumnPtr->CanBeSorted())
+		{
+			TSharedPtr<Insights::ITableCellValueSorter> SorterPtr = ColumnPtr->GetValueSorter();
+			if (ensure(SorterPtr.IsValid()))
+			{
+				AvailableSorters.Add(SorterPtr);
+			}
+		}
+	}
+}
+
+////////////////////////////////////////////////////////////////////////////////////////////////////
+
+void STimerTreeView::UpdateCurrentSortingByColumn()
+{
+	TSharedPtr<Insights::FTableColumn> ColumnPtr = Table->FindColumn(ColumnBeingSorted);
+	CurrentSorter = ColumnPtr.IsValid() ? ColumnPtr->GetValueSorter() : nullptr;
+}
+
+////////////////////////////////////////////////////////////////////////////////////////////////////
+
+void STimerTreeView::SortTreeNodes()
+{
+	if (CurrentSorter.IsValid())
+	{
+		for (FTimerNodePtr& Root : TreeNodes)
+		{
+			SortTreeNodesRec(*Root, *CurrentSorter);
+		}
+	}
+}
+
+////////////////////////////////////////////////////////////////////////////////////////////////////
+
+void STimerTreeView::SortTreeNodesRec(FTimerNode& Node, const Insights::ITableCellValueSorter& Sorter)
+{
+	if (ColumnSortMode == EColumnSortMode::Type::Descending)
+	{
+		Node.SortChildrenDescending(Sorter);
+	}
+	else // if (ColumnSortMode == EColumnSortMode::Type::Ascending)
+	{
+		Node.SortChildrenAscending(Sorter);
+	}
+
+	for (Insights::FBaseTreeNodePtr ChildPtr : Node.GetChildren())
+	{
+		//if (ChildPtr->IsGroup())
+		if (ChildPtr->GetChildren().Num() > 0)
+		{
+			SortTreeNodesRec(*StaticCastSharedPtr<FTimerNode>(ChildPtr), Sorter);
+		}
+	}
+}
+
+////////////////////////////////////////////////////////////////////////////////////////////////////
+
+EColumnSortMode::Type STimerTreeView::GetSortModeForColumn(const FName ColumnId) const
+{
+	if (ColumnBeingSorted != ColumnId)
+	{
+		return EColumnSortMode::None;
+	}
+
+	return ColumnSortMode;
+}
+
+////////////////////////////////////////////////////////////////////////////////////////////////////
+
+void STimerTreeView::SetSortModeForColumn(const FName& ColumnId, const EColumnSortMode::Type SortMode)
+{
+	ColumnBeingSorted = ColumnId;
+	ColumnSortMode = SortMode;
+	UpdateCurrentSortingByColumn();
+
+	SortTreeNodes();
+	//ApplyFiltering();
+}
+
+////////////////////////////////////////////////////////////////////////////////////////////////////
+
+void STimerTreeView::OnSortModeChanged(const EColumnSortPriority::Type SortPriority, const FName& ColumnId, const EColumnSortMode::Type SortMode)
+{
+	SetSortModeForColumn(ColumnId, SortMode);
+	TreeView_Refresh();
+}
+
+////////////////////////////////////////////////////////////////////////////////////////////////////
+// SortMode action (HeaderMenu)
+////////////////////////////////////////////////////////////////////////////////////////////////////
+
+bool STimerTreeView::HeaderMenu_SortMode_IsChecked(const FName ColumnId, const EColumnSortMode::Type InSortMode)
+{
+	return ColumnBeingSorted == ColumnId && ColumnSortMode == InSortMode;
+}
+
+////////////////////////////////////////////////////////////////////////////////////////////////////
+
+bool STimerTreeView::HeaderMenu_SortMode_CanExecute(const FName ColumnId, const EColumnSortMode::Type InSortMode) const
+{
+	const Insights::FTableColumn& Column = *Table->FindColumnChecked(ColumnId);
+	return Column.CanBeSorted();
+}
+
+////////////////////////////////////////////////////////////////////////////////////////////////////
+
+void STimerTreeView::HeaderMenu_SortMode_Execute(const FName ColumnId, const EColumnSortMode::Type InSortMode)
+{
+	SetSortModeForColumn(ColumnId, InSortMode);
+	TreeView_Refresh();
+}
+
+////////////////////////////////////////////////////////////////////////////////////////////////////
+// SortMode action (ContextMenu)
+////////////////////////////////////////////////////////////////////////////////////////////////////
+
+bool STimerTreeView::ContextMenu_SortMode_IsChecked(const EColumnSortMode::Type InSortMode)
+{
+	return ColumnSortMode == InSortMode;
+}
+
+////////////////////////////////////////////////////////////////////////////////////////////////////
+
+bool STimerTreeView::ContextMenu_SortMode_CanExecute(const EColumnSortMode::Type InSortMode) const
+{
+	return true; //ColumnSortMode != InSortMode;
+}
+
+////////////////////////////////////////////////////////////////////////////////////////////////////
+
+void STimerTreeView::ContextMenu_SortMode_Execute(const EColumnSortMode::Type InSortMode)
+{
+	SetSortModeForColumn(ColumnBeingSorted, InSortMode);
+	TreeView_Refresh();
+}
+
+////////////////////////////////////////////////////////////////////////////////////////////////////
+// SortByColumn action (ContextMenu)
+////////////////////////////////////////////////////////////////////////////////////////////////////
+
+bool STimerTreeView::ContextMenu_SortByColumn_IsChecked(const FName ColumnId)
+{
+	return ColumnId == ColumnBeingSorted;
+}
+
+////////////////////////////////////////////////////////////////////////////////////////////////////
+
+bool STimerTreeView::ContextMenu_SortByColumn_CanExecute(const FName ColumnId) const
+{
+	return true; //ColumnId != ColumnBeingSorted;
+}
+
+////////////////////////////////////////////////////////////////////////////////////////////////////
+
+void STimerTreeView::ContextMenu_SortByColumn_Execute(const FName ColumnId)
+{
+	SetSortModeForColumn(ColumnId, EColumnSortMode::Descending);
+	TreeView_Refresh();
+}
+
+////////////////////////////////////////////////////////////////////////////////////////////////////
+// ShowColumn action
+////////////////////////////////////////////////////////////////////////////////////////////////////
+
+bool STimerTreeView::CanShowColumn(const FName ColumnId) const
+{
+	return true;
+}
+
+////////////////////////////////////////////////////////////////////////////////////////////////////
+
+void STimerTreeView::ShowColumn(const FName ColumnId)
+{
+	Insights::FTableColumn& Column = *Table->FindColumnChecked(ColumnId);
+	Column.Show();
+
+	SHeaderRow::FColumn::FArguments ColumnArgs;
+	ColumnArgs
+		.ColumnId(Column.GetId())
+		.DefaultLabel(Column.GetShortName())
+		.HAlignHeader(HAlign_Fill)
+		.VAlignHeader(VAlign_Fill)
+		.HeaderContentPadding(FMargin(2.0f))
+		.HAlignCell(HAlign_Fill)
+		.VAlignCell(VAlign_Fill)
+		.SortMode(this, &STimerTreeView::GetSortModeForColumn, Column.GetId())
+		.OnSort(this, &STimerTreeView::OnSortModeChanged)
+		.ManualWidth(Column.GetInitialWidth())
+		.FixedWidth(Column.IsFixedWidth() ? Column.GetInitialWidth() : TOptional<float>())
+		.HeaderContent()
+		[
+			SNew(SBox)
+			.ToolTip(STimersViewTooltip::GetColumnTooltip(Column))
+			.HAlign(Column.GetHorizontalAlignment())
+			.VAlign(VAlign_Center)
+			[
+				SNew(STextBlock)
+				.Text(this, &STimerTreeView::GetColumnHeaderText, Column.GetId())
+			]
+		]
+		.MenuContent()
+		[
+			TreeViewHeaderRow_GenerateColumnMenu(Column)
+		];
+
+	int32 ColumnIndex = 0;
+	const int32 NewColumnPosition = Table->GetColumnPositionIndex(ColumnId);
+	const int32 NumColumns = TreeViewHeaderRow->GetColumns().Num();
+	for (; ColumnIndex < NumColumns; ColumnIndex++)
+	{
+		const SHeaderRow::FColumn& CurrentColumn = TreeViewHeaderRow->GetColumns()[ColumnIndex];
+		const int32 CurrentColumnPosition = Table->GetColumnPositionIndex(CurrentColumn.ColumnId);
+		if (NewColumnPosition < CurrentColumnPosition)
+		{
+			break;
+		}
+	}
+
+	TreeViewHeaderRow->InsertColumn(ColumnArgs, ColumnIndex);
+}
+
+////////////////////////////////////////////////////////////////////////////////////////////////////
+// HideColumn action
+////////////////////////////////////////////////////////////////////////////////////////////////////
+
+bool STimerTreeView::CanHideColumn(const FName ColumnId) const
+{
+	const Insights::FTableColumn& Column = *Table->FindColumnChecked(ColumnId);
+	return Column.CanBeHidden();
+}
+
+////////////////////////////////////////////////////////////////////////////////////////////////////
+
+void STimerTreeView::HideColumn(const FName ColumnId)
+{
+	Insights::FTableColumn& Column = *Table->FindColumnChecked(ColumnId);
+	Column.Hide();
+
+	TreeViewHeaderRow->RemoveColumn(ColumnId);
+}
+
+////////////////////////////////////////////////////////////////////////////////////////////////////
+// ToggleColumn action
+////////////////////////////////////////////////////////////////////////////////////////////////////
+
+bool STimerTreeView::IsColumnVisible(const FName ColumnId)
+{
+	const Insights::FTableColumn& Column = *Table->FindColumnChecked(ColumnId);
+	return Column.IsVisible();
+}
+
+////////////////////////////////////////////////////////////////////////////////////////////////////
+
+bool STimerTreeView::CanToggleColumnVisibility(const FName ColumnId) const
+{
+	const Insights::FTableColumn& Column = *Table->FindColumnChecked(ColumnId);
+	return !Column.IsVisible() || Column.CanBeHidden();
+}
+
+////////////////////////////////////////////////////////////////////////////////////////////////////
+
+void STimerTreeView::ToggleColumnVisibility(const FName ColumnId)
+{
+	const Insights::FTableColumn& Column = *Table->FindColumnChecked(ColumnId);
+	if (Column.IsVisible())
+	{
+		HideColumn(ColumnId);
+	}
+	else
+	{
+		ShowColumn(ColumnId);
+	}
+}
+
+////////////////////////////////////////////////////////////////////////////////////////////////////
+// "Show All Columns" action (ContextMenu)
+////////////////////////////////////////////////////////////////////////////////////////////////////
+
+bool STimerTreeView::ContextMenu_ShowAllColumns_CanExecute() const
+{
+	return true;
+}
+
+////////////////////////////////////////////////////////////////////////////////////////////////////
+
+void STimerTreeView::ContextMenu_ShowAllColumns_Execute()
+{
+	ColumnBeingSorted = GetDefaultColumnBeingSorted();
+	ColumnSortMode = GetDefaultColumnSortMode();
+	UpdateCurrentSortingByColumn();
+
+	for (const TSharedPtr<Insights::FTableColumn>& ColumnPtr : Table->GetColumns())
+	{
+		const Insights::FTableColumn& Column = *ColumnPtr;
+
+		if (!Column.IsVisible())
+		{
+			ShowColumn(Column.GetId());
+		}
+	}
+}
+
+////////////////////////////////////////////////////////////////////////////////////////////////////
+// ResetColumns action (ContextMenu)
+////////////////////////////////////////////////////////////////////////////////////////////////////
+
+bool STimerTreeView::ContextMenu_ResetColumns_CanExecute() const
+{
+	return true;
+}
+
+////////////////////////////////////////////////////////////////////////////////////////////////////
+
+void STimerTreeView::ContextMenu_ResetColumns_Execute()
+{
+	ColumnBeingSorted = GetDefaultColumnBeingSorted();
+	ColumnSortMode = GetDefaultColumnSortMode();
+	UpdateCurrentSortingByColumn();
+
+	for (const TSharedPtr<Insights::FTableColumn>& ColumnPtr : Table->GetColumns())
+	{
+		const Insights::FTableColumn& Column = *ColumnPtr;
+
+		if (Column.ShouldBeVisible() && !Column.IsVisible())
+		{
+			ShowColumn(Column.GetId());
+		}
+		else if (!Column.ShouldBeVisible() && Column.IsVisible())
+		{
+			HideColumn(Column.GetId());
+		}
+	}
+}
 
 ////////////////////////////////////////////////////////////////////////////////////////////////////
 
@@ -469,9 +916,20 @@ void STimerTreeView::SetTree(const Trace::FTimingProfilerButterflyNode& Root)
 	FTimerNodePtr RootTimerNodePtr = CreateTimerNodeRec(Root);
 	if (RootTimerNodePtr)
 	{
+		// Mark the hot path. The child nodes are already sorted by InclTime (descending), so we just follow the first child.
+		FTimerNodePtr TimerNodePtr = RootTimerNodePtr;
+		while (TimerNodePtr.IsValid())
+		{
+			TimerNodePtr->SetIsHotPath(true);
+			const TArray<Insights::FBaseTreeNodePtr>& Children = TimerNodePtr->GetChildren();
+			TimerNodePtr = Children.Num() > 0 ? StaticCastSharedPtr<FTimerNode, Insights::FBaseTreeNode>(Children[0]) : nullptr;
+		}
+
 		TreeNodes.Add(RootTimerNodePtr);
 	}
 
+	SortTreeNodes();
+
 	TreeView_Refresh();
 
 	if (RootTimerNodePtr)
@@ -487,6 +945,7 @@ FTimerNodePtr STimerTreeView::CreateTimerNodeRec(const Trace::FTimingProfilerBut
 	if (Node.Timer == nullptr)
 	{
 		return nullptr;
+		//return MakeShareable(new FTimerNode(0, FName(TEXT("!!!!!")), FName(TEXT("[Group]")), ETimerNodeType::InvalidOrMax));
 	}
 
 	const uint64 Id = Node.Timer->Id;
@@ -532,7 +991,7 @@ void STimerTreeView::ExpandNodesRec(FTimerNodePtr NodePtr, int32 Depth)
 {
 	//constexpr int32 MaxDepth = 3;
 
-	TreeView->SetItemExpansion(NodePtr, true);
+	TreeView->SetItemExpansion(NodePtr, NodePtr->IsHotPath()); // expand only the hot path
 
 	//if (Depth < MaxDepth)
 	{
