[Insights]

diff --git a/Engine/Source/Developer/TraceInsights/Private/Insights/Widgets/STimerTreeView.h b/Engine/Source/Developer/TraceInsights/Private/Insights/Widgets/STimerTreeView.h
--- a/Engine/Source/Developer/TraceInsights/Private/Insights/Widgets/STimerTreeView.h
+++ b/Engine/Source/Developer/TraceInsights/Private/Insights/Widgets/STimerTreeView.h
@@ -14,12 +14,20 @@
 #include "Widgets/Views/STreeView.h"
 
 // Insights
+#include "Insights/ViewModels/TimerGroupingAndSorting.h"
 #include "Insights/ViewModels/TimerNode.h"
+#include "Insights/ViewModels/TimersViewColumn.h"
 
 ////////////////////////////////////////////////////////////////////////////////////////////////////
 
 class FMenuBuilder;
 
+namespace Insights
+{
+	class FTable;
+	class FTableColumn;
+	class ITableCellValueSorter;
+}
 
 ////////////////////////////////////////////////////////////////////////////////////////////////////
 
@@ -55,15 +63,17 @@ private:
 	// Tree View - Context Menu
 
 	TSharedPtr<SWidget> TreeView_GetMenuContent();
+	void TreeView_BuildSortByMenu(FMenuBuilder& MenuBuilder);
+	void TreeView_BuildViewColumnMenu(FMenuBuilder& MenuBuilder);
 
 	////////////////////////////////////////////////////////////////////////////////////////////////////
 	// Tree View - Columns' Header
 
 	void InitializeAndShowHeaderColumns();
-	void TreeViewHeaderRow_CreateTimerNameColumnArgs();
-	void TreeViewHeaderRow_CreateCountColumnArgs();
-	void TreeViewHeaderRow_CreateInclusiveTimeColumnArgs();
-	void TreeViewHeaderRow_CreateExclusiveTimeColumnArgs();
+
+	FText GetColumnHeaderText(const FName ColumnId) const;
+
+	TSharedRef<SWidget> TreeViewHeaderRow_GenerateColumnMenu(const Insights::FTableColumn& Column);
 
 	////////////////////////////////////////////////////////////////////////////////////////////////////
 	// Tree View - Misc
@@ -89,8 +99,7 @@ private:
 	/** Called by STreeView to generate a table row for the specified item. */
 	TSharedRef<ITableRow> TreeView_OnGenerateRow(FTimerNodePtr TreeNode, const TSharedRef<STableViewBase>& OwnerTable);
 
-	bool TableRow_IsColumnVisible(const FName ColumnId) const;
-	void TableRow_SetHoveredTableCell(const FName ColumnId, const FTimerNodePtr TimerNodePtr);
+	void TableRow_SetHoveredCell(TSharedPtr<Insights::FTable> TablePtr, TSharedPtr<Insights::FTableColumn> ColumnPtr, const FTimerNodePtr NodePtr);
 	EHorizontalAlignment TableRow_GetColumnOutlineHAlignment(const FName ColumnId) const;
 
 	FText TableRow_GetHighlightText() const;
@@ -99,8 +108,68 @@ private:
 	bool TableRow_ShouldBeEnabled(const uint32 TimerId) const;
 
 	////////////////////////////////////////////////////////////////////////////////////////////////////
+	// Sorting
+
+	static const EColumnSortMode::Type GetDefaultColumnSortMode();
+	static const FName GetDefaultColumnBeingSorted();
+
+	void CreateSortings();
+
+	void UpdateCurrentSortingByColumn();
+	void SortTreeNodes();
+	void SortTreeNodesRec(FTimerNode& Node, const Insights::ITableCellValueSorter& Sorter);
+
+	EColumnSortMode::Type GetSortModeForColumn(const FName ColumnId) const;
+	void SetSortModeForColumn(const FName& ColumnId, EColumnSortMode::Type SortMode);
+	void OnSortModeChanged(const EColumnSortPriority::Type SortPriority, const FName& ColumnId, const EColumnSortMode::Type SortMode);
+
+	////////////////////////////////////////////////////////////////////////////////////////////////////
+	// Sorting actions
+
+	// SortMode (HeaderMenu)
+	bool HeaderMenu_SortMode_IsChecked(const FName ColumnId, const EColumnSortMode::Type InSortMode);
+	bool HeaderMenu_SortMode_CanExecute(const FName ColumnId, const EColumnSortMode::Type InSortMode) const;
+	void HeaderMenu_SortMode_Execute(const FName ColumnId, const EColumnSortMode::Type InSortMode);
+
+	// SortMode (ContextMenu)
+	bool ContextMenu_SortMode_IsChecked(const EColumnSortMode::Type InSortMode);
+	bool ContextMenu_SortMode_CanExecute(const EColumnSortMode::Type InSortMode) const;
+	void ContextMenu_SortMode_Execute(const EColumnSortMode::Type InSortMode);
+
+	// SortByColumn (ContextMenu)
+	bool ContextMenu_SortByColumn_IsChecked(const FName ColumnId);
+	bool ContextMenu_SortByColumn_CanExecute(const FName ColumnId) const;
+	void ContextMenu_SortByColumn_Execute(const FName ColumnId);
+
+	////////////////////////////////////////////////////////////////////////////////////////////////////
+	// Column visibility actions
+
+	// ShowColumn
+	bool CanShowColumn(const FName ColumnId) const;
+	void ShowColumn(const FName ColumnId);
+
+	// HideColumn
+	bool CanHideColumn(const FName ColumnId) const;
+	void HideColumn(const FName ColumnId);
+
+	// ToggleColumnVisibility
+	bool IsColumnVisible(const FName ColumnId);
+	bool CanToggleColumnVisibility(const FName ColumnId) const;
+	void ToggleColumnVisibility(const FName ColumnId);
+
+	// ShowAllColumns (ContextMenu)
+	bool ContextMenu_ShowAllColumns_CanExecute() const;
+	void ContextMenu_ShowAllColumns_Execute();
+
+	// ResetColumns (ContextMenu)
+	bool ContextMenu_ResetColumns_CanExecute() const;
+	void ContextMenu_ResetColumns_Execute();
+
+	////////////////////////////////////////////////////////////////////////////////////////////////////
 
 private:
+	/** Table view model. */
+	TSharedPtr<Insights::FTable> Table;
 
 	/** The view name (ex.: "Callers" or "Callees"). */
 	FText ViewName;
@@ -111,9 +180,6 @@ private:
 	/** The tree widget which holds the list of groups and timers corresponding with each group. */
 	TSharedPtr<STreeView<FTimerNodePtr>> TreeView;
 
-	/** Column arguments used to initialize a new header column in the tree view, stored as column name to column arguments mapping. */
-	TMap<FName, SHeaderRow::FColumn::FArguments> TreeViewHeaderColumnArgs;
-
 	/** Holds the tree view header row widget which display all columns in the tree view. */
 	TSharedPtr<SHeaderRow> TreeViewHeaderRow;
 
@@ -127,7 +193,7 @@ private:
 	FName HoveredColumnId;
 
 	/** A shared pointer to the timer node currently being hovered by the mouse. */
-	FTimerNodePtr HoveredTimerNodePtr;
+	FTimerNodePtr HoveredNodePtr;
 
 	/** Name of the timer that should be drawn as highlighted. */
 	FName HighlightedNodeName;
@@ -138,14 +204,26 @@ private:
 	/** The root node(s) of the tree. */
 	TArray<FTimerNodePtr> TreeNodes;
 
+	//////////////////////////////////////////////////
+	// Sorting
+
+	/** All available sorters. */
+	TArray<TSharedPtr<Insights::ITableCellValueSorter>> AvailableSorters;
+
+	/** Current sorter. It is nullptr if sorting is disabled. */
+	TSharedPtr<Insights::ITableCellValueSorter> CurrentSorter;
+
+	/** Name of the column currently being sorted. Can be NAME_None if sorting is disabled (CurrentSorting == nullptr) or if a complex sorting is used (CurrentSorting != nullptr). */
+	FName ColumnBeingSorted;
+
+	/** How we sort the nodes? Ascending or Descending. */
+	EColumnSortMode::Type ColumnSortMode;
+
 	//////////////////////////////////////////////////
 
 	double StatsStartTime;
 	double StatsEndTime;
 	uint32 StatsTimerId;
-
-	static const FName TimerNameColumnId;
-	static const FName CountColumnId;
-	static const FName InclusiveTimeColumnId;
-	static const FName ExclusiveTimeColumnId;
 };
+
+////////////////////////////////////////////////////////////////////////////////////////////////////
diff --git a/Engine/Source/Developer/TraceInsights/Private/Insights/Widgets/STimersView.cpp b/Engine/Source/Developer/TraceInsights/Private/Insights/Widgets/STimersView.cpp
--- a/Engine/Source/Developer/TraceInsights/Private/Insights/Widgets/STimersView.cpp
+++ b/Engine/Source/Developer/TraceInsights/Private/Insights/Widgets/STimersView.cpp
@@ -17,6 +17,8 @@
 #include "Widgets/Views/STableViewBase.h"
 
 // Insights
+#include "Insights/Table/ViewModels/Table.h"
+#include "Insights/Table/ViewModels/TableColumn.h"
 #include "Insights/TimingProfilerCommon.h"
 #include "Insights/TimingProfilerManager.h"
 #include "Insights/ViewModels/TimerNodeHelper.h"
@@ -45,14 +47,15 @@ const FName STimersView::GetDefaultColumnBeingSorted()
 ////////////////////////////////////////////////////////////////////////////////////////////////////
 
 STimersView::STimersView()
-	: bExpansionSaved(false)
+	: Table(MakeShareable(new Insights::FTable()))
+	, bExpansionSaved(false)
 	, GroupingMode(ETimerGroupingMode::ByType)
 	, ColumnSortMode(GetDefaultColumnSortMode())
 	, ColumnBeingSorted(GetDefaultColumnBeingSorted())
 	, StatsStartTime(0.0)
 	, StatsEndTime(0.0)
 {
-	FMemory::Memset(bTimerNodeIsVisible, 1);
+	FMemory::Memset(bTimerTypeIsVisible, 1);
 }
 
 ////////////////////////////////////////////////////////////////////////////////////////////////////
@@ -239,29 +242,28 @@ END_SLATE_FUNCTION_BUILD_OPTIMIZATION
 
 TSharedPtr<SWidget> STimersView::TreeView_GetMenuContent()
 {
-	const TArray<FTimerNodePtr> SelectedTimerNodes = TreeView->GetSelectedItems();
-	const int32 NumSelectedTimerNodes = SelectedTimerNodes.Num();
-	FTimerNodePtr SelectedTimerNode = NumSelectedTimerNodes ? SelectedTimerNodes[0] : nullptr;
+	const TArray<FTimerNodePtr> SelectedNodes = TreeView->GetSelectedItems();
+	const int32 NumSelectedNodes = SelectedNodes.Num();
+	FTimerNodePtr SelectedNode = NumSelectedNodes ? SelectedNodes[0] : nullptr;
 
-	const FTimersViewColumn* const * ColumnPtrPtr = FTimersViewColumnFactory::Get().ColumnIdToPtrMapping.Find(HoveredColumnId);
-	const FTimersViewColumn* const ColumnPtr = (ColumnPtrPtr != nullptr) ? *ColumnPtrPtr : nullptr;
+	const TSharedPtr<Insights::FTableColumn> HoveredColumnPtr = Table->FindColumn(HoveredColumnId);
 
 	FText SelectionStr;
 	FText PropertyName;
 	FText PropertyValue;
 
-	if (NumSelectedTimerNodes == 0)
+	if (NumSelectedNodes == 0)
 	{
 		SelectionStr = LOCTEXT("NothingSelected", "Nothing selected");
 	}
-	else if (NumSelectedTimerNodes == 1)
+	else if (NumSelectedNodes == 1)
 	{
-		if (ColumnPtr != nullptr)
+		if (HoveredColumnPtr != nullptr)
 		{
-			PropertyName = ColumnPtr->ShortName;
-			PropertyValue = ColumnPtr->GetFormattedValue(*SelectedTimerNode);
+			PropertyName = HoveredColumnPtr->GetShortName();
+			PropertyValue = HoveredColumnPtr->GetValueAsTooltipText(*SelectedNode);
 		}
-		SelectionStr = FText::FromName(SelectedTimerNode->GetName());
+		SelectionStr = FText::FromName(SelectedNode->GetName());
 	}
 	else
 	{
@@ -380,22 +382,22 @@ void STimersView::TreeView_BuildSortByMenu(FMenuBuilder& MenuBuilder)
 
 	MenuBuilder.BeginSection("ColumnName", LOCTEXT("ContextMenu_Header_Misc_ColumnName", "Column Name"));
 
-	for (auto It = TreeViewHeaderColumns.CreateConstIterator(); It; ++It)
+	for (const TSharedPtr<Insights::FTableColumn>& ColumnPtr : Table->GetColumns())
 	{
-		const FTimersViewColumn& Column = It.Value();
+		const Insights::FTableColumn& Column = *ColumnPtr;
 
-		if (Column.bIsVisible && Column.bCanBeSorted())
+		if (Column.IsVisible() && Column.CanBeSorted())
 		{
 			FUIAction Action_SortByColumn
 			(
-				FExecuteAction::CreateSP(this, &STimersView::ContextMenu_SortByColumn_Execute, Column.Id),
-				FCanExecuteAction::CreateSP(this, &STimersView::ContextMenu_SortByColumn_CanExecute, Column.Id),
-				FIsActionChecked::CreateSP(this, &STimersView::ContextMenu_SortByColumn_IsChecked, Column.Id)
+				FExecuteAction::CreateSP(this, &STimersView::ContextMenu_SortByColumn_Execute, Column.GetId()),
+				FCanExecuteAction::CreateSP(this, &STimersView::ContextMenu_SortByColumn_CanExecute, Column.GetId()),
+				FIsActionChecked::CreateSP(this, &STimersView::ContextMenu_SortByColumn_IsChecked, Column.GetId())
 			);
 			MenuBuilder.AddMenuEntry
 			(
-				Column.TitleName,
-				Column.Description,
+				Column.GetTitleName(),
+				Column.GetDescription(),
 				FSlateIcon(), Action_SortByColumn, NAME_None, EUserInterfaceActionType::RadioButton
 			);
 		}
@@ -442,20 +444,20 @@ void STimersView::TreeView_BuildViewColumnMenu(FMenuBuilder& MenuBuilder)
 {
 	MenuBuilder.BeginSection("ViewColumn", LOCTEXT("ContextMenu_Header_Columns_View", "View Column"));
 
-	for (auto It = TreeViewHeaderColumns.CreateConstIterator(); It; ++It)
+	for (const TSharedPtr<Insights::FTableColumn>& ColumnPtr : Table->GetColumns())
 	{
-		const FTimersViewColumn& Column = It.Value();
+		const Insights::FTableColumn& Column = *ColumnPtr;
 
 		FUIAction Action_ToggleColumn
 		(
-			FExecuteAction::CreateSP(this, &STimersView::ContextMenu_ToggleColumn_Execute, Column.Id),
-			FCanExecuteAction::CreateSP(this, &STimersView::ContextMenu_ToggleColumn_CanExecute, Column.Id),
-			FIsActionChecked::CreateSP(this, &STimersView::ContextMenu_ToggleColumn_IsChecked, Column.Id)
+			FExecuteAction::CreateSP(this, &STimersView::ToggleColumnVisibility, Column.GetId()),
+			FCanExecuteAction::CreateSP(this, &STimersView::CanToggleColumnVisibility, Column.GetId()),
+			FIsActionChecked::CreateSP(this, &STimersView::IsColumnVisible, Column.GetId())
 		);
 		MenuBuilder.AddMenuEntry
 		(
-			Column.TitleName,
-			Column.Description,
+			Column.GetTitleName(),
+			Column.GetDescription(),
 			FSlateIcon(), Action_ToggleColumn, NAME_None, EUserInterfaceActionType::ToggleButton
 		);
 	}
@@ -467,100 +469,46 @@ void STimersView::TreeView_BuildViewColumnMenu(FMenuBuilder& MenuBuilder)
 
 void STimersView::InitializeAndShowHeaderColumns()
 {
-	const int32 NumColumns = FTimersViewColumnFactory::Get().Collection.Num();
-	for (int32 ColumnIndex = 0; ColumnIndex < NumColumns; ColumnIndex++)
-	{
-		TreeViewHeaderRow_CreateColumnArgs(ColumnIndex);
-	}
+	// Create columns.
+	TArray<TSharedPtr<Insights::FTableColumn>> Columns;
+	FTimersViewColumnFactory::CreateTimersViewColumns(Columns);
+	Table->SetColumns(Columns);
 
-	for (auto It = TreeViewHeaderColumns.CreateConstIterator(); It; ++It)
+	// Show columns.
+	for (const TSharedPtr<Insights::FTableColumn>& ColumnPtr : Table->GetColumns())
 	{
-		const FTimersViewColumn& Column = It.Value();
-
-		if (Column.bIsVisible)
+		if (ColumnPtr->ShouldBeVisible())
 		{
-			TreeViewHeaderRow_ShowColumn(Column.Id);
+			ShowColumn(ColumnPtr->GetId());
 		}
 	}
 }
 
 ////////////////////////////////////////////////////////////////////////////////////////////////////
 
-void STimersView::TreeViewHeaderRow_CreateColumnArgs(const int32 ColumnIndex)
+FText STimersView::GetColumnHeaderText(const FName ColumnId) const
 {
-	const FTimersViewColumn& Column = *FTimersViewColumnFactory::Get().Collection[ColumnIndex];
-	SHeaderRow::FColumn::FArguments ColumnArgs;
-
-	ColumnArgs
-		.ColumnId(Column.Id)
-		.DefaultLabel(Column.ShortName)
-		.HAlignHeader(HAlign_Fill)
-		.VAlignHeader(VAlign_Fill)
-		.HeaderContentPadding(FMargin(2.0f))
-		.HAlignCell(HAlign_Fill)
-		.VAlignCell(VAlign_Fill)
-		.SortMode(this, &STimersView::GetSortModeForColumn, Column.Id)
-		.OnSort(this, &STimersView::OnSortModeChanged)
-		.ManualWidth(Column.InitialColumnWidth)
-		.FixedWidth(Column.bIsFixedColumnWidth() ? Column.InitialColumnWidth : TOptional<float>())
-		.HeaderContent()
-		[
-			SNew(SBox)
-			.ToolTip(STimersViewTooltip::GetColumnTooltip(Column))
-			.HAlign(Column.HorizontalAlignment)
-			.VAlign(VAlign_Center)
-			[
-				SNew(STextBlock)
-				.Text(Column.ShortName)
-			]
-		]
-		.MenuContent()
-		[
-			TreeViewHeaderRow_GenerateColumnMenu(Column)
-		];
-
-	TreeViewHeaderColumnArgs.Add(Column.Id, ColumnArgs);
-	TreeViewHeaderColumns.Add(Column.Id, Column);
+	const Insights::FTableColumn& Column = *Table->FindColumnChecked(ColumnId);
+	return Column.GetShortName();
 }
 
 ////////////////////////////////////////////////////////////////////////////////////////////////////
 
-void STimersView::TreeViewHeaderRow_ShowColumn(const FName ColumnId)
-{
-	FTimersViewColumn& Column = TreeViewHeaderColumns.FindChecked(ColumnId);
-	Column.bIsVisible = true;
-	SHeaderRow::FColumn::FArguments& ColumnArgs = TreeViewHeaderColumnArgs.FindChecked(ColumnId);
-
-	const int32 NumColumns = TreeViewHeaderRow->GetColumns().Num();
-	int32 ColumnIndex = 0;
-	for (; ColumnIndex < NumColumns; ColumnIndex++)
-	{
-		const SHeaderRow::FColumn& CurrentColumn = TreeViewHeaderRow->GetColumns()[ColumnIndex];
-		const FTimersViewColumn& CurrentTimersViewColumn = TreeViewHeaderColumns.FindChecked(CurrentColumn.ColumnId);
-		if (Column.Order < CurrentTimersViewColumn.Order)
-			break;
-	}
-
-	TreeViewHeaderRow->InsertColumn(ColumnArgs, ColumnIndex);
-}
-
-////////////////////////////////////////////////////////////////////////////////////////////////////
-
-TSharedRef<SWidget> STimersView::TreeViewHeaderRow_GenerateColumnMenu(const FTimersViewColumn& Column)
+TSharedRef<SWidget> STimersView::TreeViewHeaderRow_GenerateColumnMenu(const Insights::FTableColumn& Column)
 {
 	bool bIsMenuVisible = false;
 
 	const bool bShouldCloseWindowAfterMenuSelection = true;
 	FMenuBuilder MenuBuilder(bShouldCloseWindowAfterMenuSelection, NULL);
 	{
-		if (Column.bCanBeHidden())
+		if (Column.CanBeHidden())
 		{
 			MenuBuilder.BeginSection("Column", LOCTEXT("TreeViewHeaderRow_Header_Column", "Column"));
 
 			FUIAction Action_HideColumn
 			(
-				FExecuteAction::CreateSP(this, &STimersView::HeaderMenu_HideColumn_Execute, Column.Id),
-				FCanExecuteAction::CreateSP(this, &STimersView::HeaderMenu_HideColumn_CanExecute, Column.Id)
+				FExecuteAction::CreateSP(this, &STimersView::HideColumn, Column.GetId()),
+				FCanExecuteAction::CreateSP(this, &STimersView::CanHideColumn, Column.GetId())
 			);
 
 			MenuBuilder.AddMenuEntry
@@ -574,15 +522,15 @@ TSharedRef<SWidget> STimersView::TreeViewHeaderRow_GenerateColumnMenu(const FTim
 			MenuBuilder.EndSection();
 		}
 
-		if (Column.bCanBeSorted())
+		if (Column.CanBeSorted())
 		{
 			MenuBuilder.BeginSection("SortMode", LOCTEXT("ContextMenu_Header_Misc_Sort_SortMode", "Sort Mode"));
 
 			FUIAction Action_SortAscending
 			(
-				FExecuteAction::CreateSP(this, &STimersView::HeaderMenu_SortMode_Execute, Column.Id, EColumnSortMode::Ascending),
-				FCanExecuteAction::CreateSP(this, &STimersView::HeaderMenu_SortMode_CanExecute, Column.Id, EColumnSortMode::Ascending),
-				FIsActionChecked::CreateSP(this, &STimersView::HeaderMenu_SortMode_IsChecked, Column.Id, EColumnSortMode::Ascending)
+				FExecuteAction::CreateSP(this, &STimersView::HeaderMenu_SortMode_Execute, Column.GetId(), EColumnSortMode::Ascending),
+				FCanExecuteAction::CreateSP(this, &STimersView::HeaderMenu_SortMode_CanExecute, Column.GetId(), EColumnSortMode::Ascending),
+				FIsActionChecked::CreateSP(this, &STimersView::HeaderMenu_SortMode_IsChecked, Column.GetId(), EColumnSortMode::Ascending)
 			);
 			MenuBuilder.AddMenuEntry
 			(
@@ -593,9 +541,9 @@ TSharedRef<SWidget> STimersView::TreeViewHeaderRow_GenerateColumnMenu(const FTim
 
 			FUIAction Action_SortDescending
 			(
-				FExecuteAction::CreateSP(this, &STimersView::HeaderMenu_SortMode_Execute, Column.Id, EColumnSortMode::Descending),
-				FCanExecuteAction::CreateSP(this, &STimersView::HeaderMenu_SortMode_CanExecute, Column.Id, EColumnSortMode::Descending),
-				FIsActionChecked::CreateSP(this, &STimersView::HeaderMenu_SortMode_IsChecked, Column.Id, EColumnSortMode::Descending)
+				FExecuteAction::CreateSP(this, &STimersView::HeaderMenu_SortMode_Execute, Column.GetId(), EColumnSortMode::Descending),
+				FCanExecuteAction::CreateSP(this, &STimersView::HeaderMenu_SortMode_CanExecute, Column.GetId(), EColumnSortMode::Descending),
+				FIsActionChecked::CreateSP(this, &STimersView::HeaderMenu_SortMode_IsChecked, Column.GetId(), EColumnSortMode::Descending)
 			);
 			MenuBuilder.AddMenuEntry
 			(
@@ -675,7 +623,7 @@ void STimersView::ApplyFiltering()
 		{
 			// Add a child.
 			const FTimerNodePtr& NodePtr = StaticCastSharedPtr<FTimerNode, Insights::FBaseTreeNode>(GroupChildren[Cx]);
-			const bool bIsChildVisible = Filters->PassesAllFilters(NodePtr) && bTimerNodeIsVisible[static_cast<int>(NodePtr->GetType())];
+			const bool bIsChildVisible = Filters->PassesAllFilters(NodePtr) && bTimerTypeIsVisible[static_cast<int>(NodePtr->GetType())];
 			if (bIsChildVisible)
 			{
 				GroupPtr->AddFilteredChild(NodePtr);
@@ -764,7 +712,7 @@ TSharedRef<SWidget> STimersView::GetToggleButtonForTimerType(const ETimerNodeTyp
 				.VAlign(VAlign_Center)
 				[
 					SNew(STextBlock)
-						.Text(TimerNodeTypeHelper::ToName(NodeType))
+						.Text(TimerNodeTypeHelper::ToText(NodeType))
 						.TextStyle(FEditorStyle::Get(), TEXT("Profiler.Caption"))
 				]
 		];
@@ -774,7 +722,7 @@ TSharedRef<SWidget> STimersView::GetToggleButtonForTimerType(const ETimerNodeTyp
 
 void STimersView::FilterByTimerType_OnCheckStateChanged(ECheckBoxState NewRadioState, const ETimerNodeType InStatType)
 {
-	bTimerNodeIsVisible[static_cast<int>(InStatType)] = (NewRadioState == ECheckBoxState::Checked);
+	bTimerTypeIsVisible[static_cast<int>(InStatType)] = (NewRadioState == ECheckBoxState::Checked);
 	ApplyFiltering();
 }
 
@@ -782,7 +730,7 @@ void STimersView::FilterByTimerType_OnCheckStateChanged(ECheckBoxState NewRadioS
 
 ECheckBoxState STimersView::FilterByTimerType_IsChecked(const ETimerNodeType InStatType) const
 {
-	return bTimerNodeIsVisible[static_cast<int>(InStatType)] ? ECheckBoxState::Checked : ECheckBoxState::Unchecked;
+	return bTimerTypeIsVisible[static_cast<int>(InStatType)] ? ECheckBoxState::Checked : ECheckBoxState::Unchecked;
 }
 
 ////////////////////////////////////////////////////////////////////////////////////////////////////
@@ -804,7 +752,7 @@ void STimersView::TreeView_OnSelectionChanged(FTimerNodePtr SelectedItem, ESelec
 	if (SelectInfo != ESelectInfo::Direct)
 	{
 		TArray<FTimerNodePtr> SelectedItems = TreeView->GetSelectedItems();
-		if (SelectedItems.Num() == 1)
+		if (SelectedItems.Num() == 1 && !SelectedItems[0]->IsGroup())
 		{
 			FTimingProfilerManager::Get()->SetSelectedTimer(SelectedItems[0]->GetId());
 		}
@@ -857,11 +805,12 @@ TSharedRef<ITableRow> STimersView::TreeView_OnGenerateRow(FTimerNodePtr TimerNod
 	TSharedRef<ITableRow> TableRow =
 		SNew(STimerTableRow, OwnerTable)
 		.OnShouldBeEnabled(this, &STimersView::TableRow_ShouldBeEnabled)
-		.OnIsColumnVisible(this, &STimersView::TableRow_IsColumnVisible)
-		.OnSetHoveredTableCell(this, &STimersView::TableRow_SetHoveredTableCell)
+		.OnIsColumnVisible(this, &STimersView::IsColumnVisible)
+		.OnSetHoveredCell(this, &STimersView::TableRow_SetHoveredCell)
 		.OnGetColumnOutlineHAlignmentDelegate(this, &STimersView::TableRow_GetColumnOutlineHAlignment)
 		.HighlightText(this, &STimersView::TableRow_GetHighlightText)
 		.HighlightedNodeName(this, &STimersView::TableRow_GetHighlightedNodeName)
+		.TablePtr(Table)
 		.TimerNodePtr(TimerNodePtr);
 
 	return TableRow;
@@ -876,23 +825,14 @@ bool STimersView::TableRow_ShouldBeEnabled(const uint32 TimerId) const
 
 ////////////////////////////////////////////////////////////////////////////////////////////////////
 
-bool STimersView::TableRow_IsColumnVisible(const FName ColumnId) const
+void STimersView::TableRow_SetHoveredCell(TSharedPtr<Insights::FTable> InTablePtr, TSharedPtr<Insights::FTableColumn> InColumnPtr, const FTimerNodePtr InNodePtr)
 {
-	bool bResult = false;
-	const FTimersViewColumn& ColumnPtr = TreeViewHeaderColumns.FindChecked(ColumnId);
-	return ColumnPtr.bIsVisible;
-}
-
-////////////////////////////////////////////////////////////////////////////////////////////////////
-
-void STimersView::TableRow_SetHoveredTableCell(const FName ColumnId, const FTimerNodePtr TimerNodePtr)
-{
-	HoveredColumnId = ColumnId;
+	HoveredColumnId = InColumnPtr ? InColumnPtr->GetId() : FName();
 
 	const bool bIsAnyMenusVisible = FSlateApplication::Get().AnyMenusVisible();
 	if (!HasMouseCapture() && !bIsAnyMenusVisible)
 	{
-		HoveredTimerNodePtr = TimerNodePtr;
+		HoveredTimerNodePtr = InNodePtr;
 	}
 }
 
@@ -997,7 +937,7 @@ void STimersView::CreateGroups()
 	{
 		for (const FTimerNodePtr& TimerNodePtr : TimerNodes)
 		{
-			const FName GroupName = *TimerNodeTypeHelper::ToName(TimerNodePtr->GetType()).ToString();
+			const FName GroupName = *TimerNodeTypeHelper::ToText(TimerNodePtr->GetType()).ToString();
 
 			FTimerNodePtr* GroupPtr = GroupNodeSet.Find(GroupName);
 			if (!GroupPtr)
@@ -1087,7 +1027,7 @@ void STimersView::GroupBy_OnSelectionChanged(TSharedPtr<ETimerGroupingMode> NewG
 TSharedRef<SWidget> STimersView::GroupBy_OnGenerateWidget(TSharedPtr<ETimerGroupingMode> InGroupingMode) const
 {
 	return SNew(STextBlock)
-		.Text(TimerNodeGroupingHelper::ToName(*InGroupingMode))
+		.Text(TimerNodeGroupingHelper::ToText(*InGroupingMode))
 		.ToolTipText(TimerNodeGroupingHelper::ToDescription(*InGroupingMode));
 }
 
@@ -1095,7 +1035,7 @@ TSharedRef<SWidget> STimersView::GroupBy_OnGenerateWidget(TSharedPtr<ETimerGroup
 
 FText STimersView::GroupBy_GetSelectedText() const
 {
-	return TimerNodeGroupingHelper::ToName(GroupingMode);
+	return TimerNodeGroupingHelper::ToText(GroupingMode);
 }
 
 ////////////////////////////////////////////////////////////////////////////////////////////////////
@@ -1196,8 +1136,8 @@ bool STimersView::HeaderMenu_SortMode_IsChecked(const FName ColumnId, const ECol
 
 bool STimersView::HeaderMenu_SortMode_CanExecute(const FName ColumnId, const EColumnSortMode::Type InSortMode) const
 {
-	const FTimersViewColumn& Column = TreeViewHeaderColumns.FindChecked(ColumnId);
-	return Column.bCanBeSorted();
+	const Insights::FTableColumn& Column = *Table->FindColumnChecked(ColumnId);
+	return Column.CanBeSorted();
 }
 
 ////////////////////////////////////////////////////////////////////////////////////////////////////
@@ -1257,54 +1197,116 @@ void STimersView::ContextMenu_SortByColumn_Execute(const FName ColumnId)
 }
 
 ////////////////////////////////////////////////////////////////////////////////////////////////////
-// HideColumn action (HeaderMenu)
+// ShowColumn action
+////////////////////////////////////////////////////////////////////////////////////////////////////
+
+bool STimersView::CanShowColumn(const FName ColumnId) const
+{
+	return true;
+}
+
+////////////////////////////////////////////////////////////////////////////////////////////////////
+
+void STimersView::ShowColumn(const FName ColumnId)
+{
+	Insights::FTableColumn& Column = *Table->FindColumnChecked(ColumnId);
+	Column.Show();
+
+	SHeaderRow::FColumn::FArguments ColumnArgs;
+	ColumnArgs
+		.ColumnId(Column.GetId())
+		.DefaultLabel(Column.GetShortName())
+		.HAlignHeader(HAlign_Fill)
+		.VAlignHeader(VAlign_Fill)
+		.HeaderContentPadding(FMargin(2.0f))
+		.HAlignCell(HAlign_Fill)
+		.VAlignCell(VAlign_Fill)
+		.SortMode(this, &STimersView::GetSortModeForColumn, Column.GetId())
+		.OnSort(this, &STimersView::OnSortModeChanged)
+		.ManualWidth(Column.GetInitialWidth())
+		.FixedWidth(Column.IsFixedWidth() ? Column.GetInitialWidth() : TOptional<float>())
+		.HeaderContent()
+		[
+			SNew(SBox)
+			.ToolTip(STimersViewTooltip::GetColumnTooltip(Column))
+			.HAlign(Column.GetHorizontalAlignment())
+			.VAlign(VAlign_Center)
+			[
+				SNew(STextBlock)
+				.Text(this, &STimersView::GetColumnHeaderText, Column.GetId())
+			]
+		]
+		.MenuContent()
+		[
+			TreeViewHeaderRow_GenerateColumnMenu(Column)
+		];
+
+	int32 ColumnIndex = 0;
+	const int32 NewColumnPosition = Table->GetColumnPositionIndex(ColumnId);
+	const int32 NumColumns = TreeViewHeaderRow->GetColumns().Num();
+	for (; ColumnIndex < NumColumns; ColumnIndex++)
+	{
+		const SHeaderRow::FColumn& CurrentColumn = TreeViewHeaderRow->GetColumns()[ColumnIndex];
+		const int32 CurrentColumnPosition = Table->GetColumnPositionIndex(CurrentColumn.ColumnId);
+		if (NewColumnPosition < CurrentColumnPosition)
+		{
+			break;
+		}
+	}
+
+	TreeViewHeaderRow->InsertColumn(ColumnArgs, ColumnIndex);
+}
+
+////////////////////////////////////////////////////////////////////////////////////////////////////
+// HideColumn action
 ////////////////////////////////////////////////////////////////////////////////////////////////////
 
-bool STimersView::HeaderMenu_HideColumn_CanExecute(const FName ColumnId) const
+bool STimersView::CanHideColumn(const FName ColumnId) const
 {
-	const FTimersViewColumn& Column = TreeViewHeaderColumns.FindChecked(ColumnId);
-	return Column.bCanBeHidden();
+	const Insights::FTableColumn& Column = *Table->FindColumnChecked(ColumnId);
+	return Column.CanBeHidden();
 }
 
 ////////////////////////////////////////////////////////////////////////////////////////////////////
 
-void STimersView::HeaderMenu_HideColumn_Execute(const FName ColumnId)
+void STimersView::HideColumn(const FName ColumnId)
 {
-	FTimersViewColumn& Column = TreeViewHeaderColumns.FindChecked(ColumnId);
-	Column.bIsVisible = false;
+	Insights::FTableColumn& Column = *Table->FindColumnChecked(ColumnId);
+	Column.Hide();
+
 	TreeViewHeaderRow->RemoveColumn(ColumnId);
 }
 
 ////////////////////////////////////////////////////////////////////////////////////////////////////
-// ToggleColumn action (ContextMenu)
+// ToggleColumn action
 ////////////////////////////////////////////////////////////////////////////////////////////////////
 
-bool STimersView::ContextMenu_ToggleColumn_IsChecked(const FName ColumnId)
+bool STimersView::IsColumnVisible(const FName ColumnId)
 {
-	const FTimersViewColumn& Column = TreeViewHeaderColumns.FindChecked(ColumnId);
-	return Column.bIsVisible;
+	const Insights::FTableColumn& Column = *Table->FindColumnChecked(ColumnId);
+	return Column.IsVisible();
 }
 
 ////////////////////////////////////////////////////////////////////////////////////////////////////
 
-bool STimersView::ContextMenu_ToggleColumn_CanExecute(const FName ColumnId) const
+bool STimersView::CanToggleColumnVisibility(const FName ColumnId) const
 {
-	const FTimersViewColumn& Column = TreeViewHeaderColumns.FindChecked(ColumnId);
-	return Column.bCanBeHidden();
+	const Insights::FTableColumn& Column = *Table->FindColumnChecked(ColumnId);
+	return !Column.IsVisible() || Column.CanBeHidden();
 }
 
 ////////////////////////////////////////////////////////////////////////////////////////////////////
 
-void STimersView::ContextMenu_ToggleColumn_Execute(const FName ColumnId)
+void STimersView::ToggleColumnVisibility(const FName ColumnId)
 {
-	FTimersViewColumn& Column = TreeViewHeaderColumns.FindChecked(ColumnId);
-	if (Column.bIsVisible)
+	const Insights::FTableColumn& Column = *Table->FindColumnChecked(ColumnId);
+	if (Column.IsVisible())
 	{
-		HeaderMenu_HideColumn_Execute(ColumnId);
+		HideColumn(ColumnId);
 	}
 	else
 	{
-		TreeViewHeaderRow_ShowColumn(ColumnId);
+		ShowColumn(ColumnId);
 	}
 }
 
@@ -1324,15 +1326,13 @@ void STimersView::ContextMenu_ShowAllColumns_Execute()
 	ColumnSortMode = GetDefaultColumnSortMode();
 	ColumnBeingSorted = GetDefaultColumnBeingSorted();
 
-	const int32 NumColumns = FTimersViewColumnFactory::Get().Collection.Num();
-	for (int32 ColumnIndex = 0; ColumnIndex < NumColumns; ColumnIndex++)
+	for (const TSharedPtr<Insights::FTableColumn>& ColumnPtr : Table->GetColumns())
 	{
-		const FTimersViewColumn& DefaultColumn = *FTimersViewColumnFactory::Get().Collection[ColumnIndex];
-		const FTimersViewColumn& CurrentColumn = TreeViewHeaderColumns.FindChecked(DefaultColumn.Id);
+		const Insights::FTableColumn& Column = *ColumnPtr;
 
-		if (!CurrentColumn.bIsVisible)
+		if (!Column.IsVisible())
 		{
-			TreeViewHeaderRow_ShowColumn(DefaultColumn.Id);
+			ShowColumn(Column.GetId());
 		}
 	}
 }
@@ -1371,20 +1371,19 @@ void STimersView::ContextMenu_ShowMinMaxMedColumns_Execute()
 	ColumnSortMode = EColumnSortMode::Descending;
 	ColumnBeingSorted = FTimersViewColumns::TotalInclusiveTimeColumnID;
 
-	const int32 NumColumns = FTimersViewColumnFactory::Get().Collection.Num();
-	for (int32 ColumnIndex = 0; ColumnIndex < NumColumns; ColumnIndex++)
+	for (const TSharedPtr<Insights::FTableColumn>& ColumnPtr : Table->GetColumns())
 	{
-		const FTimersViewColumn& DefaultColumn = *FTimersViewColumnFactory::Get().Collection[ColumnIndex];
-		const FTimersViewColumn& CurrentColumn = TreeViewHeaderColumns.FindChecked(DefaultColumn.Id);
+		const Insights::FTableColumn& Column = *ColumnPtr;
+
+		const bool bShouldBeVisible = Preset.Contains(Column.GetId());
 
-		bool bIsVisible = Preset.Contains(DefaultColumn.Id);
-		if (bIsVisible && !CurrentColumn.bIsVisible)
+		if (bShouldBeVisible && !Column.IsVisible())
 		{
-			TreeViewHeaderRow_ShowColumn(DefaultColumn.Id);
+			ShowColumn(Column.GetId());
 		}
-		else if (!bIsVisible && CurrentColumn.bIsVisible)
+		else if (!bShouldBeVisible && Column.IsVisible())
 		{
-			HeaderMenu_HideColumn_Execute(DefaultColumn.Id);
+			HideColumn(Column.GetId());
 		}
 	}
 }
@@ -1405,19 +1404,17 @@ void STimersView::ContextMenu_ResetColumns_Execute()
 	ColumnSortMode = GetDefaultColumnSortMode();
 	ColumnBeingSorted = GetDefaultColumnBeingSorted();
 
-	const int32 NumColumns = FTimersViewColumnFactory::Get().Collection.Num();
-	for (int32 ColumnIndex = 0; ColumnIndex < NumColumns; ColumnIndex++)
+	for (const TSharedPtr<Insights::FTableColumn>& ColumnPtr : Table->GetColumns())
 	{
-		const FTimersViewColumn& DefaultColumn = *FTimersViewColumnFactory::Get().Collection[ColumnIndex];
-		const FTimersViewColumn& CurrentColumn = TreeViewHeaderColumns.FindChecked(DefaultColumn.Id);
+		const Insights::FTableColumn& Column = *ColumnPtr;
 
-		if (DefaultColumn.bIsVisible && !CurrentColumn.bIsVisible)
+		if (Column.ShouldBeVisible() && !Column.IsVisible())
 		{
-			TreeViewHeaderRow_ShowColumn(DefaultColumn.Id);
+			ShowColumn(Column.GetId());
 		}
-		else if (!DefaultColumn.bIsVisible && CurrentColumn.bIsVisible)
+		else if (!Column.ShouldBeVisible() && Column.IsVisible())
 		{
-			HeaderMenu_HideColumn_Execute(DefaultColumn.Id);
+			HideColumn(Column.GetId());
 		}
 	}
 }
diff --git a/Engine/Source/Developer/TraceInsights/Private/Insights/Widgets/STimersView.h b/Engine/Source/Developer/TraceInsights/Private/Insights/Widgets/STimersView.h
--- a/Engine/Source/Developer/TraceInsights/Private/Insights/Widgets/STimersView.h
+++ b/Engine/Source/Developer/TraceInsights/Private/Insights/Widgets/STimersView.h
@@ -27,7 +27,12 @@ class FMenuBuilder;
 namespace Trace
 {
 	class IAnalysisSession;
-	struct FTimelineEvent;
+}
+
+namespace Insights
+{
+	class FTable;
+	class ITableCellValueSorter;
 }
 
 ////////////////////////////////////////////////////////////////////////////////////////////////////
@@ -60,6 +65,8 @@ public:
 	 */
 	void Construct(const FArguments& InArgs);
 
+	TSharedPtr<Insights::FTable> GetTable() const { return Table; }
+
 	/**
 	 * Rebuilds the tree (if necessary).
 	 * @param bResync - If true, it forces a resync with list of timers from Analysis, even if the list did not changed since last sync.
@@ -74,7 +81,6 @@ public:
 	const FTimerNodePtr* GetTimerNode(uint64 Id) const { return TimerNodesIdMap.Find(Id); }
 
 protected:
-
 	void UpdateTree();
 
 	/** Called when the analysis session has changed. */
@@ -100,9 +106,10 @@ protected:
 	// Tree View - Columns' Header
 
 	void InitializeAndShowHeaderColumns();
-	void TreeViewHeaderRow_CreateColumnArgs(const int32 ColumnIndex);
-	void TreeViewHeaderRow_ShowColumn(const FName ColumnId);
-	TSharedRef<SWidget> TreeViewHeaderRow_GenerateColumnMenu(const FTimersViewColumn& Column);
+
+	FText GetColumnHeaderText(const FName ColumnId) const;
+
+	TSharedRef<SWidget> TreeViewHeaderRow_GenerateColumnMenu(const Insights::FTableColumn& Column);
 
 	////////////////////////////////////////////////////////////////////////////////////////////////////
 	// Tree View - Misc
@@ -128,8 +135,7 @@ protected:
 	/** Called by STreeView to generate a table row for the specified item. */
 	TSharedRef<ITableRow> TreeView_OnGenerateRow(FTimerNodePtr TreeNode, const TSharedRef<STableViewBase>& OwnerTable);
 
-	bool TableRow_IsColumnVisible(const FName ColumnId) const;
-	void TableRow_SetHoveredTableCell(const FName ColumnId, const FTimerNodePtr TimerNodePtr);
+	void TableRow_SetHoveredCell(TSharedPtr<Insights::FTable> TablePtr, TSharedPtr<Insights::FTableColumn> ColumnPtr, const FTimerNodePtr NodePtr);
 	EHorizontalAlignment TableRow_GetColumnOutlineHAlignment(const FName ColumnId) const;
 
 	FText TableRow_GetHighlightText() const;
@@ -151,7 +157,7 @@ protected:
 	void SearchBox_OnTextChanged(const FText& InFilterText);
 
 	////////////////////////////////////////////////////////////////////////////////////////////////////
-	// GroupBy
+	// Grouping
 
 	void CreateGroups();
 	void CreateGroupByOptionsSources();
@@ -197,14 +203,18 @@ protected:
 	////////////////////////////////////////////////////////////////////////////////////////////////////
 	// Column visibility actions
 
-	// HideColumn (HeaderMenu)
-	bool HeaderMenu_HideColumn_CanExecute(const FName ColumnId) const;
-	void HeaderMenu_HideColumn_Execute(const FName ColumnId);
+	// ShowColumn
+	bool CanShowColumn(const FName ColumnId) const;
+	void ShowColumn(const FName ColumnId);
+
+	// HideColumn
+	bool CanHideColumn(const FName ColumnId) const;
+	void HideColumn(const FName ColumnId);
 
-	// ToggleColumn (ContextMenu)
-	bool ContextMenu_ToggleColumn_IsChecked(const FName ColumnId);
-	bool ContextMenu_ToggleColumn_CanExecute(const FName ColumnId) const;
-	void ContextMenu_ToggleColumn_Execute(const FName ColumnId);
+	// ToggleColumnVisibility
+	bool IsColumnVisible(const FName ColumnId);
+	bool CanToggleColumnVisibility(const FName ColumnId) const;
+	void ToggleColumnVisibility(const FName ColumnId);
 
 	// ShowAllColumns (ContextMenu)
 	bool ContextMenu_ShowAllColumns_CanExecute() const;
@@ -221,6 +231,9 @@ protected:
 	////////////////////////////////////////////////////////////////////////////////////////////////////
 
 protected:
+	/** Table view model. */
+	TSharedPtr<Insights::FTable> Table;
+
 	/** A weak pointer to the profiler session used to populate this widget. */
 	TSharedPtr<const Trace::IAnalysisSession>/*Weak*/ Session;
 
@@ -230,12 +243,6 @@ protected:
 	/** The tree widget which holds the list of groups and timers corresponding with each group. */
 	TSharedPtr<STreeView<FTimerNodePtr>> TreeView;
 
-	/** Column metadata used to initialize column arguments, stored as PropertyName -> FTimersViewColumn. */
-	TMap<FName, FTimersViewColumn> TreeViewHeaderColumns;
-
-	/** Column arguments used to initialize a new header column in the tree view, stored as column name to column arguments mapping. */
-	TMap<FName, SHeaderRow::FColumn::FArguments> TreeViewHeaderColumnArgs;
-
 	/** Holds the tree view header row widget which display all columns in the tree view. */
 	TSharedPtr<SHeaderRow> TreeViewHeaderRow;
 
@@ -293,7 +300,7 @@ protected:
 	TSharedPtr<FTimerNodeFilterCollection> Filters;
 
 	/** Holds the visibility of each timer type. */
-	bool bTimerNodeIsVisible[static_cast<int>(ETimerNodeType::InvalidOrMax)];
+	bool bTimerTypeIsVisible[static_cast<int>(ETimerNodeType::InvalidOrMax)];
 
 	//////////////////////////////////////////////////
 	// Grouping
diff --git a/Engine/Source/Developer/TraceInsights/Private/Insights/Widgets/STimersViewTooltip.cpp b/Engine/Source/Developer/TraceInsights/Private/Insights/Widgets/STimersViewTooltip.cpp
--- a/Engine/Source/Developer/TraceInsights/Private/Insights/Widgets/STimersViewTooltip.cpp
+++ b/Engine/Source/Developer/TraceInsights/Private/Insights/Widgets/STimersViewTooltip.cpp
@@ -13,19 +13,20 @@
 
 // Insights
 #include "Insights/Common/TimeUtils.h"
+#include "Insights/Table/ViewModels/Table.h"
+#include "Insights/Table/ViewModels/TableColumn.h"
 #include "Insights/ViewModels/TimerNode.h"
-#include "Insights/ViewModels/TimersViewColumn.h"
 #include "Insights/ViewModels/TimerNodeHelper.h"
 
 ////////////////////////////////////////////////////////////////////////////////////////////////////
 
-#define LOCTEXT_NAMESPACE "STimersViewTooltip"
+#define LOCTEXT_NAMESPACE "STimersView"
 
 ////////////////////////////////////////////////////////////////////////////////////////////////////
 
 BEGIN_SLATE_FUNCTION_BUILD_OPTIMIZATION
 
-TSharedPtr<SToolTip> STimersViewTooltip::GetColumnTooltip(const FTimersViewColumn& Column)
+TSharedPtr<SToolTip> STimersViewTooltip::GetTableTooltip(const Insights::FTable& Table)
 {
 	TSharedPtr<SToolTip> ColumnTooltip =
 		SNew(SToolTip)
@@ -37,7 +38,7 @@ TSharedPtr<SToolTip> STimersViewTooltip::GetColumnTooltip(const FTimersViewColum
 			.Padding(2.0f)
 			[
 				SNew(STextBlock)
-				.Text(Column.TitleName)
+				.Text(Table.GetDisplayName())
 				.TextStyle(FEditorStyle::Get(), TEXT("Profiler.TooltipBold"))
 			]
 
@@ -46,7 +47,7 @@ TSharedPtr<SToolTip> STimersViewTooltip::GetColumnTooltip(const FTimersViewColum
 			.Padding(2.0f)
 			[
 				SNew(STextBlock)
-				.Text(Column.Description)
+				.Text(Table.GetDescription())
 				.TextStyle(FEditorStyle::Get(), TEXT("Profiler.Tooltip"))
 			]
 		];
@@ -56,17 +57,39 @@ TSharedPtr<SToolTip> STimersViewTooltip::GetColumnTooltip(const FTimersViewColum
 
 ////////////////////////////////////////////////////////////////////////////////////////////////////
 
-TSharedPtr<SToolTip> STimersViewTooltip::GetTableCellTooltip(const TSharedPtr<FTimerNode> TimerNodePtr)
+TSharedPtr<SToolTip> STimersViewTooltip::GetColumnTooltip(const Insights::FTableColumn& Column)
 {
-	//const FSlateFontInfo TitleFont = FCoreStyle::GetDefaultFontStyle("Bold", 10);
-	//const FSlateFontInfo DescriptionFont = FCoreStyle::GetDefaultFontStyle("Regular", 8);
-	//const FSlateFontInfo DescriptionFontB = FCoreStyle::GetDefaultFontStyle("Bold", 8);
+	TSharedPtr<SToolTip> ColumnTooltip =
+		SNew(SToolTip)
+		[
+			SNew(SVerticalBox)
 
-	const FLinearColor DefaultColor(1.0f,1.0f,1.0f,1.0f);
-	const FLinearColor ThreadColor(5.0f, 0.0f, 0.0f, 1.0f);
-	const float Alpha = 0.0f;//TimerNodePtr->_FramePct * 0.01f;
-	const FLinearColor ColorAndOpacity = FMath::Lerp(DefaultColor, ThreadColor,Alpha);
+			+ SVerticalBox::Slot()
+			.AutoHeight()
+			.Padding(2.0f)
+			[
+				SNew(STextBlock)
+				.Text(Column.GetTitleName())
+				.TextStyle(FEditorStyle::Get(), TEXT("Profiler.TooltipBold"))
+			]
 
+			+ SVerticalBox::Slot()
+			.AutoHeight()
+			.Padding(2.0f)
+			[
+				SNew(STextBlock)
+				.Text(Column.GetDescription())
+				.TextStyle(FEditorStyle::Get(), TEXT("Profiler.Tooltip"))
+			]
+		];
+
+	return ColumnTooltip;
+}
+
+////////////////////////////////////////////////////////////////////////////////////////////////////
+
+TSharedPtr<SToolTip> STimersViewTooltip::GetCellTooltip(const TSharedPtr<FTimerNode> TimerNodePtr, const TSharedPtr<Insights::FTableColumn> ColumnPtr)
+{
 	const Trace::FAggregatedTimingStats& Stats = TimerNodePtr->GetAggregatedStats();
 
 	const int32 NumDigits = 5;
@@ -85,7 +108,6 @@ TSharedPtr<SToolTip> STimersViewTooltip::GetTableCellTooltip(const TSharedPtr<FT
 
 	const FText InstanceCountText = FText::AsNumber(Stats.InstanceCount);
 
-	//TSharedPtr<SHorizontalBox> HBoxCaption;
 	TSharedPtr<SGridPanel> GridPanel;
 	TSharedPtr<SHorizontalBox> HBox;
 
@@ -94,19 +116,12 @@ TSharedPtr<SToolTip> STimersViewTooltip::GetTableCellTooltip(const TSharedPtr<FT
 		[
 			SAssignNew(HBox, SHorizontalBox)
 
-			+SHorizontalBox::Slot()
+			+ SHorizontalBox::Slot()
 			.AutoWidth()
 			[
 				SNew(SVerticalBox)
 
-				//+SVerticalBox::Slot()
-				//.AutoHeight()
-				//.Padding(2.0f)
-				//[
-				//	SAssignNew(HBoxCaption, SHorizontalBox)
-				//]
-
-				+SVerticalBox::Slot()
+				+ SVerticalBox::Slot()
 				.AutoHeight()
 				.Padding(2.0f)
 				[
@@ -114,21 +129,21 @@ TSharedPtr<SToolTip> STimersViewTooltip::GetTableCellTooltip(const TSharedPtr<FT
 					.Orientation(Orient_Horizontal)
 				]
 
-				+SVerticalBox::Slot()
+				+ SVerticalBox::Slot()
 				.AutoHeight()
 				.Padding(2.0f)
 				[
 					SNew(SGridPanel)
 
 					// Id: [Id]
-					+SGridPanel::Slot(0, 0)
+					+ SGridPanel::Slot(0, 0)
 					.Padding(2.0f)
 					[
 						SNew(STextBlock)
 						.Text(LOCTEXT("TT_Id", "Id:"))
 						.TextStyle(FEditorStyle::Get(), TEXT("Profiler.TooltipBold"))
 					]
-					+SGridPanel::Slot(1, 0)
+					+ SGridPanel::Slot(1, 0)
 					.Padding(2.0f)
 					[
 						SNew(STextBlock)
@@ -137,14 +152,14 @@ TSharedPtr<SToolTip> STimersViewTooltip::GetTableCellTooltip(const TSharedPtr<FT
 					]
 
 					// Name: [Name]
-					+SGridPanel::Slot(0, 1)
+					+ SGridPanel::Slot(0, 1)
 					.Padding(2.0f)
 					[
 						SNew(STextBlock)
 						.Text(LOCTEXT("TT_Name", "Name:"))
 						.TextStyle(FEditorStyle::Get(), TEXT("Profiler.TooltipBold"))
 					]
-					+SGridPanel::Slot(1, 1)
+					+ SGridPanel::Slot(1, 1)
 					.Padding(2.0f)
 					[
 						SNew(STextBlock)
@@ -153,14 +168,14 @@ TSharedPtr<SToolTip> STimersViewTooltip::GetTableCellTooltip(const TSharedPtr<FT
 					]
 
 					//// Group: [MetaGroupName]
-					//+SGridPanel::Slot(0, 2)
+					//+ SGridPanel::Slot(0, 2)
 					//.Padding(2.0f)
 					//[
 					//	SNew(STextBlock)
 					//	.Text(LOCTEXT("TT_Group", "Group:"))
 					//	.TextStyle(FEditorStyle::Get(), TEXT("Profiler.TooltipBold"))
 					//]
-					//+SGridPanel::Slot(1, 2)
+					//+ SGridPanel::Slot(1, 2)
 					//.Padding(2.0f)
 					//[
 					//	SNew(STextBlock)
@@ -168,7 +183,7 @@ TSharedPtr<SToolTip> STimersViewTooltip::GetTableCellTooltip(const TSharedPtr<FT
 					//	.TextStyle(FEditorStyle::Get(), TEXT("Profiler.Tooltip"))
 					//]
 
-					// Type: [Type]
+					// Timer Type: [Type]
 					+ SGridPanel::Slot(0, 3)
 					.Padding(2.0f)
 					[
@@ -180,12 +195,12 @@ TSharedPtr<SToolTip> STimersViewTooltip::GetTableCellTooltip(const TSharedPtr<FT
 					.Padding(2.0f)
 					[
 						SNew(STextBlock)
-						.Text(TimerNodeTypeHelper::ToName(TimerNodePtr->GetType()))
+						.Text(TimerNodeTypeHelper::ToText(TimerNodePtr->GetType()))
 						.TextStyle(FEditorStyle::Get(), TEXT("Profiler.Tooltip"))
 					]
 				]
 
-				+SVerticalBox::Slot()
+				+ SVerticalBox::Slot()
 				.AutoHeight()
 				.Padding(2.0f)
 				[
@@ -215,7 +230,7 @@ TSharedPtr<SToolTip> STimersViewTooltip::GetTableCellTooltip(const TSharedPtr<FT
 					]
 				]
 
-				+SVerticalBox::Slot()
+				+ SVerticalBox::Slot()
 				.AutoHeight()
 				.Padding(2.0f)
 				[
@@ -249,7 +264,7 @@ TSharedPtr<SToolTip> STimersViewTooltip::GetTableCellTooltip(const TSharedPtr<FT
 					// Stats are added here.
 				]
 
-				+SVerticalBox::Slot()
+				+ SVerticalBox::Slot()
 				.AutoHeight()
 				.Padding(2.0f)
 				[
@@ -266,260 +281,11 @@ TSharedPtr<SToolTip> STimersViewTooltip::GetTableCellTooltip(const TSharedPtr<FT
 	AddStatsRow(GridPanel, Row, LOCTEXT("TT_MedianTime",  "Median Time:"),  MedInclusiveTimeText,   MedExclusiveTimeText);
 	AddStatsRow(GridPanel, Row, LOCTEXT("TT_MinTime",     "Min Time:"),     MinInclusiveTimeText,   MinExclusiveTimeText);
 
-	/*
-	//TODO: We need Stats hierarchy (not the grouping hierarchy)!
-	const bool bHasParent = TimerNodePtr->GetStats()->GetParent().IsValid();
-	const bool bHasChildren = TimerNodePtr->GetStats()->GetChildren().Num() > 0;
-
-	if (bHasParent)
-	{
-		const FText ParentName = FText::FromName(TimerNodePtr->GetGroupPtr()->GetName());
-		HBoxCaption->AddSlot()
-			.AutoWidth()
-			[
-				SNew(STextBlock)
-				.Text(ParentName)
-				.TextStyle(FEditorStyle::Get(), TEXT("Profiler.Caption"))
-			];
-
-		HBoxCaption->AddSlot()
-			.AutoWidth()
-			[
-				SNew(SImage)
-				.Image(FEditorStyle::GetBrush("BreadcrumbTrail.Delimiter"))
-			];
-	}
-
-	const FText TimerName = FText::FromName(TimerNodePtr->GetName());
-	HBoxCaption->AddSlot()
-		.AutoWidth()
-		[
-			SNew(STextBlock)
-			.Text(TimerName)
-			.TextStyle(FEditorStyle::Get(), TEXT("Profiler.CaptionBold"))
-		];
-
-	if (bHasChildren)
-	{
-		typedef TKeyValuePair<float, FName> FEventNameAndPct;
-		TArray<FEventNameAndPct> MinimalChildren;
-
-		const TArray<FTimerNodePtr>& Children = TimerNodePtr->GetChildren();
-		for(int32 ChildIndex = 0; ChildIndex < Children.Num(); ChildIndex++)
-		{
-			const FTimerNodePtr Child = Children[ChildIndex];
-			float Percent = static_cast<float>(Child->GetStats().TotalInclusiveTime / TimerNodePtr->GetStats().TotalInclusiveTime * 100.0);
-			MinimalChildren.Add(FEventNameAndPct(Child->GetStats().TotalInclusiveTime, Child->GetName()));
-		}
-
-		struct FCompareByFloatDescending
-		{
-			FORCEINLINE bool operator()(const FEventNameAndPct& A, const FEventNameAndPct& B) const
-			{
-				return A.Key > B.Key;
-			}
-		};
-		MinimalChildren.Sort(FCompareByFloatDescending());
-
-		FString ChildrenNames;
-		const int32 NumChildrenToDisplay = FMath::Min(MinimalChildren.Num(), 3);
-		for(int32 SortedChildIndex = 0; SortedChildIndex < NumChildrenToDisplay; SortedChildIndex++)
-		{
-			const FEventNameAndPct& MinimalChild = MinimalChildren[SortedChildIndex];
-			ChildrenNames += FString::Printf(TEXT("%s (%.1f %%)"), *MinimalChild.Value.ToString(), MinimalChild.Key);
-
-			const bool bAddDelimiter = SortedChildIndex < NumChildrenToDisplay - 1;
-			if (bAddDelimiter)
-			{
-				ChildrenNames += TEXT(", ");
-			}
-		}
-
-		HBoxCaption->AddSlot()
-			.AutoWidth()
-			[
-				SNew(SImage)
-				.Image(FEditorStyle::GetBrush("BreadcrumbTrail.Delimiter"))
-			];
-
-		HBoxCaption->AddSlot()
-			.AutoWidth()
-			[
-				SNew(STextBlock)
-				.Text(FText::FromString(ChildrenNames))
-				.TextStyle(FEditorStyle::Get(), TEXT("Profiler.Caption"))
-			];
-	}
-	*/
-
 	return TableCellTooltip;
 }
 
 ////////////////////////////////////////////////////////////////////////////////////////////////////
 
-TSharedRef<SToolTip> STimersViewTooltip::GetTooltip()
-{
-	if (Session.IsValid())
-	{
-		const TSharedRef<SGridPanel> ToolTipGrid = SNew(SGridPanel);
-		int32 CurrentRowPos = 0;
-
-		AddHeader(ToolTipGrid, CurrentRowPos);
-		AddDescription(ToolTipGrid, CurrentRowPos);
-
-		AddNoDataInformation(ToolTipGrid, CurrentRowPos);
-
-		return SNew(SToolTip)
-			[
-				ToolTipGrid
-			];
-	}
-	else
-	{
-		return SNew(SToolTip)
-			.Text(LOCTEXT("NotImplemented", "Tooltip for multiple profiler instances has not been implemented yet"));
-	}
-}
-
-////////////////////////////////////////////////////////////////////////////////////////////////////
-
-void STimersViewTooltip::AddNoDataInformation(const TSharedRef<SGridPanel>& Grid, int32& RowPos)
-{
-	Grid->AddSlot(0, RowPos)
-		.Padding(2.0f)
-		.ColumnSpan(3)
-		[
-			SNew(STextBlock)
-			.TextStyle(FEditorStyle::Get(), TEXT("Profiler.TooltipBold"))
-			.Text(LOCTEXT("NoDataAvailable", "N/A"))
-		];
-	RowPos++;
-}
-
-////////////////////////////////////////////////////////////////////////////////////////////////////
-
-void STimersViewTooltip::AddHeader(const TSharedRef<SGridPanel>& Grid, int32& RowPos)
-{
-	const FString InstanceName = Session->GetName();
-
-	Grid->AddSlot(0, RowPos++)
-		.Padding(2.0f)
-		.ColumnSpan(3)
-		[
-			SNew(STextBlock)
-			.TextStyle(FEditorStyle::Get(), TEXT("Profiler.TooltipBold"))
-			.Text(LOCTEXT("StatInstance", "Stat information for profiler instance"))
-		];
-
-	Grid->AddSlot(0, RowPos++)
-		.Padding(2.0f)
-		.ColumnSpan(3)
-		[
-			SNew(STextBlock)
-			.TextStyle(FEditorStyle::Get(), TEXT("Profiler.Tooltip"))
-			.Text(FText::FromString(InstanceName))
-		];
-
-	AddSeparator(Grid, RowPos);
-}
-
-////////////////////////////////////////////////////////////////////////////////////////////////////
-
-void STimersViewTooltip::AddDescription(const TSharedRef<SGridPanel>& Grid, int32& RowPos)
-{
-	/*
-	const FProfilerStat& ProfilerStat = Session->GetMetaData()->GetTimerById(TimerId);
-	const ETimerNodeType SampleType = Session->GetMetaData()->GetSampleTypeForTimerId(TimerId);
-	const FSlateBrush* const StatIcon = STimersViewHelper::GetIconForStatType(SampleType);
-
-	Grid->AddSlot(0, RowPos)
-	.Padding(2.0f)
-	[
-		SNew(STextBlock)
-		.TextStyle(FEditorStyle::Get(), TEXT("Profiler.TooltipBold"))
-		.Text(LOCTEXT("GroupDesc","Group:"))
-	];
-
-	Grid->AddSlot(1, RowPos)
-	.Padding(2.0f)
-	.ColumnSpan(2)
-	[
-		SNew(STextBlock)
-		.TextStyle(FEditorStyle::Get(), TEXT("Profiler.Tooltip"))
-		.Text(FText::FromName(ProfilerStat.OwningGroup().Name()))
-	];
-	RowPos++;
-
-	Grid->AddSlot(0, RowPos)
-	.Padding(2.0f)
-	[
-		SNew(STextBlock)
-		.TextStyle(FEditorStyle::Get(), TEXT("Profiler.TooltipBold"))
-		.Text(LOCTEXT("NameDesc","Name:"))
-	];
-
-	Grid->AddSlot(1, RowPos)
-	.Padding(2.0f)
-	.ColumnSpan(2)
-	[
-		SNew(STextBlock)
-		.TextStyle(FEditorStyle::Get(), TEXT("Profiler.Tooltip"))
-		.Text(FText::FromName(ProfilerStat.Name()))
-	];
-	RowPos++;
-
-	Grid->AddSlot(0, RowPos)
-	.Padding(2.0f)
-	[
-		SNew(STextBlock)
-		.TextStyle(FEditorStyle::Get(), TEXT("Profiler.TooltipBold"))
-		.Text(LOCTEXT("TypeDesc","Type:"))
-	];
-
-	Grid->AddSlot(1, RowPos)
-	.Padding(2.0f)
-	.ColumnSpan(2)
-	[
-		SNew(SHorizontalBox)
-
-		+SHorizontalBox::Slot()
-		.AutoWidth()
-		[
-			SNew(SImage)
-			.Image(StatIcon)
-		]
-
-		+SHorizontalBox::Slot()
-		.AutoWidth()
-		.HAlign(HAlign_Left)
-		.VAlign(VAlign_Center)
-		[
-			SNew(STextBlock)
-			.Text(FText::FromString(ETimerNodeType::ToDescription(SampleType)))
-			.TextStyle(FEditorStyle::Get(), TEXT("Profiler.Tooltip"))
-		]
-	];
-	RowPos++;
-
-	AddSeparator(Grid, RowPos);
-	*/
-}
-
-////////////////////////////////////////////////////////////////////////////////////////////////////
-
-void STimersViewTooltip::AddSeparator(const TSharedRef<SGridPanel>& Grid, int32& RowPos)
-{
-	Grid->AddSlot(0, RowPos++)
-		.Padding(2.0f)
-		.ColumnSpan(3)
-		[
-			SNew(SSeparator)
-			.Orientation(Orient_Horizontal)
-		];
-}
-
-////////////////////////////////////////////////////////////////////////////////////////////////////
-
 void STimersViewTooltip::AddStatsRow(TSharedPtr<SGridPanel> Grid, int32& Row, const FText& Name, const FText& Value1, const FText& Value2)
 {
 	Grid->AddSlot(0, Row)
diff --git a/Engine/Source/Developer/TraceInsights/Private/Insights/Widgets/STimersViewTooltip.h b/Engine/Source/Developer/TraceInsights/Private/Insights/Widgets/STimersViewTooltip.h
--- a/Engine/Source/Developer/TraceInsights/Private/Insights/Widgets/STimersViewTooltip.h
+++ b/Engine/Source/Developer/TraceInsights/Private/Insights/Widgets/STimersViewTooltip.h
@@ -10,40 +10,27 @@
 class SToolTip;
 class SGridPanel;
 
-namespace Trace { class IAnalysisSession; }
+namespace Insights
+{
+	class FTable;
+	class FTableColumn;
+}
 
-// Insights
 class FTimerNode;
-class FTimersViewColumn;
 
 #define LOCTEXT_NAMESPACE "STimersView"
 
 /** Timers View Tooltip */
 class STimersViewTooltip
 {
-	const uint32 TimerId;
-	TSharedPtr<const Trace::IAnalysisSession> Session;
-
 public:
-	STimersViewTooltip(const uint32 InTimerId)
-		: TimerId(InTimerId)
-	{
-		Session = FInsightsManager::Get()->GetSession();
-	}
+	STimersViewTooltip() = delete;
 
-	TSharedRef<SToolTip> GetTooltip();
-
-protected:
-	void AddNoDataInformation(const TSharedRef<SGridPanel>& Grid, int32& RowPos);
-	void AddHeader(const TSharedRef<SGridPanel>& Grid, int32& RowPos);
-	void AddDescription(const TSharedRef<SGridPanel>& Grid, int32& RowPos);
-	void AddSeparator(const TSharedRef<SGridPanel>& Grid, int32& RowPos);
-
-public:
-	static TSharedPtr<SToolTip> GetColumnTooltip(const FTimersViewColumn& Column);
-	static TSharedPtr<SToolTip> GetTableCellTooltip(const TSharedPtr<FTimerNode> TimerNodePtr);
+	static TSharedPtr<SToolTip> GetTableTooltip(const Insights::FTable& Table);
+	static TSharedPtr<SToolTip> GetColumnTooltip(const Insights::FTableColumn& Column);
+	static TSharedPtr<SToolTip> GetCellTooltip(const TSharedPtr<FTimerNode> TreeNodePtr, const TSharedPtr<Insights::FTableColumn> ColumnPtr);
 
-protected:
+private:
 	static void AddStatsRow(TSharedPtr<SGridPanel> Grid, int32& Row, const FText& Name, const FText& Value1, const FText& Value2);
 };
 
