[Insights]
  - Log View: Fixed clicking on a message to focus the Timing view on time of respective message even if the message was already selected. Fixes UE-80261.
  - Timers view: Changed tooltip to show total inclusive/exclusive time values with increased precission.

diff --git a/Engine/Source/Developer/TraceInsights/Private/Insights/Common/TimeUtils.cpp b/Engine/Source/Developer/TraceInsights/Private/Insights/Common/TimeUtils.cpp
--- a/Engine/Source/Developer/TraceInsights/Private/Insights/Common/TimeUtils.cpp
+++ b/Engine/Source/Developer/TraceInsights/Private/Insights/Common/TimeUtils.cpp
@@ -12,7 +12,33 @@ namespace TimeUtils
 
 ////////////////////////////////////////////////////////////////////////////////////////////////////
 
-FString FormatTimeAuto(const double InDuration)
+FString FormatTimeValue(const double Duration, const int32 NumDigits)
+{
+#if !PLATFORM_USE_GENERIC_STRING_IMPLEMENTATION
+	FString Str = FString::Printf(TEXT("%.*f"), NumDigits, Duration);
+#else
+	// proper resolution is tracked as UE-79534
+	TCHAR FormatString[32];
+	FCString::Snprintf(FormatString, sizeof(FormatString), TEXT("%%.%df"), NumDigits);
+	FString Str = FString::Printf(FormatString, Duration);
+#endif
+
+	if (NumDigits == 1)
+	{
+		Str.RemoveFromEnd(TEXT(".0"));
+	}
+	else
+	{
+		while (Str.RemoveFromEnd(TEXT("0"))) { /* keep removing the ending 0 */ }
+		Str.RemoveFromEnd(TEXT("."));
+	}
+
+	return Str;
+}
+
+////////////////////////////////////////////////////////////////////////////////////////////////////
+
+FString FormatTimeAuto(const double InDuration, const int32 NumDigits)
 {
 	//TestTimeAutoFormatting();
 
@@ -94,9 +120,7 @@ FString FormatTimeAuto(const double InDuration)
 		}
 		else
 		{
-			FString S = FString::Printf(TEXT("%.1f"), Microseconds);
-			S.RemoveFromEnd(TEXT(".0"));
-			Str += S;
+			Str += FormatTimeValue(Microseconds, NumDigits);
 			Str += TEXT(" µs");
 		}
 	}
@@ -110,9 +134,7 @@ FString FormatTimeAuto(const double InDuration)
 		}
 		else
 		{
-			FString S = FString::Printf(TEXT("%.1f"), Miliseconds);
-			S.RemoveFromEnd(TEXT(".0"));
-			Str += S;
+			Str += FormatTimeValue(Miliseconds, NumDigits);
 			Str += TEXT(" ms");
 		}
 	}
@@ -125,9 +147,7 @@ FString FormatTimeAuto(const double InDuration)
 		}
 		else
 		{
-			FString S = FString::Printf(TEXT("%.1f"), Duration);
-			S.RemoveFromEnd(TEXT(".0"));
-			Str += S;
+			Str += FormatTimeValue(Duration, NumDigits);
 			Str += TEXT("s");
 		}
 	}
@@ -137,10 +157,19 @@ FString FormatTimeAuto(const double InDuration)
 		const double Minutes = FMath::FloorToDouble(Duration / TimeUtils::Minute);
 		Str += FString::Printf(TEXT("%dm"), static_cast<int32>(Minutes));
 		Duration -= Minutes * TimeUtils::Minute;
-		const double Seconds = FMath::FloorToDouble(Duration / TimeUtils::Second);
-		if (Seconds > 0.5)
+		if (NumDigits <= 1)
 		{
-			Str += FString::Printf(TEXT(" %ds"), static_cast<int32>(Seconds));
+			const double Seconds = FMath::FloorToDouble(Duration / TimeUtils::Second);
+			if (Seconds > 0.5)
+			{
+				Str += FString::Printf(TEXT(" %ds"), static_cast<int32>(Seconds));
+			}
+		}
+		else
+		{
+			Str += TEXT(" ");
+			Str += FormatTimeValue(Duration, NumDigits - 1);
+			Str += TEXT("s");
 		}
 	}
 	else if (Duration < TimeUtils::Day)
diff --git a/Engine/Source/Developer/TraceInsights/Private/Insights/Common/TimeUtils.h b/Engine/Source/Developer/TraceInsights/Private/Insights/Common/TimeUtils.h
--- a/Engine/Source/Developer/TraceInsights/Private/Insights/Common/TimeUtils.h
+++ b/Engine/Source/Developer/TraceInsights/Private/Insights/Common/TimeUtils.h
@@ -54,7 +54,8 @@ namespace TimeUtils
 		bool bIsInfinite;
 	};
 
-	FString FormatTimeAuto(const double Duration);
+	FString FormatTimeValue(const double Duration, const int32 NumDigits = 1);
+	FString FormatTimeAuto(const double Duration, const int32 NumDigits = 1);
 	FString FormatTimeMs(const double Duration, const int32 NumDigits = 2, bool bAddTimeUnit = false);
 	FString FormatTime(const double Time, const double Precision = 0.0);
 	FString FormatTimeHMS(const double Time, const double Precision = 0.0);
diff --git a/Engine/Source/Developer/TraceInsights/Private/Insights/Widgets/SLogView.cpp b/Engine/Source/Developer/TraceInsights/Private/Insights/Widgets/SLogView.cpp
--- a/Engine/Source/Developer/TraceInsights/Private/Insights/Widgets/SLogView.cpp
+++ b/Engine/Source/Developer/TraceInsights/Private/Insights/Widgets/SLogView.cpp
@@ -633,6 +633,7 @@ void SLogView::Construct(const FArguments& InArgs)
 							.ExternalScrollbar(ExternalScrollbar)
 							.ItemHeight(20.0f)
 							.SelectionMode(ESelectionMode::Single)
+							.OnMouseButtonClick(this, &SLogView::OnMouseButtonClick)
 							.OnSelectionChanged(this, &SLogView::OnSelectionChanged)
 							.ListItemsSource(&Messages)
 							.OnGenerateRow(this, &SLogView::OnGenerateRow)
@@ -934,29 +935,25 @@ void SLogView::SelectedLogMessageByLogIndex(int32 LogIndex)
 
 ////////////////////////////////////////////////////////////////////////////////////////////////////
 
-void SLogView::OnSelectionChanged(TSharedPtr<FLogMessage> LogMessage, ESelectInfo::Type SelectInfo)
+void SLogView::SelectLogMessage(TSharedPtr<FLogMessage> LogMessage)
 {
-	if (SelectInfo != ESelectInfo::Direct)
+	if (LogMessage.IsValid())
 	{
-		// Single item selection.
-		if (LogMessage.IsValid())
+		TSharedPtr<STimingProfilerWindow> Window = FTimingProfilerManager::Get()->GetProfilerWindow();
+		if (Window)
 		{
-			TSharedPtr<STimingProfilerWindow> Window = FTimingProfilerManager::Get()->GetProfilerWindow();
-			if (Window)
+			TSharedPtr<STimingView> TimingView = Window->GetTimingView();
+			if (TimingView)
 			{
-				TSharedPtr<STimingView> TimingView = Window->GetTimingView();
-				if (TimingView)
-				{
-					const double Time = Cache.Get(LogMessage->GetIndex()).Time;
+				const double Time = Cache.Get(LogMessage->GetIndex()).Time;
 
-					if (FSlateApplication::Get().GetModifierKeys().IsShiftDown())
-					{
-						TimingView->SelectToTimeMarker(Time);
-					}
-					else
-					{
-						TimingView->SetAndCenterOnTimeMarker(Time);
-					}
+				if (FSlateApplication::Get().GetModifierKeys().IsShiftDown())
+				{
+					TimingView->SelectToTimeMarker(Time);
+				}
+				else
+				{
+					TimingView->SetAndCenterOnTimeMarker(Time);
 				}
 			}
 		}
@@ -965,6 +962,24 @@ void SLogView::OnSelectionChanged(TSharedPtr<FLogMessage> LogMessage, ESelectInf
 
 ////////////////////////////////////////////////////////////////////////////////////////////////////
 
+void SLogView::OnMouseButtonClick(TSharedPtr<FLogMessage> LogMessage)
+{
+	SelectLogMessage(LogMessage);
+}
+
+////////////////////////////////////////////////////////////////////////////////////////////////////
+
+void SLogView::OnSelectionChanged(TSharedPtr<FLogMessage> LogMessage, ESelectInfo::Type SelectInfo)
+{
+	if (SelectInfo != ESelectInfo::Direct &&
+		SelectInfo != ESelectInfo::OnMouseClick)
+	{
+		SelectLogMessage(LogMessage);
+	}
+}
+
+////////////////////////////////////////////////////////////////////////////////////////////////////
+
 void SLogView::FilterTextBox_OnTextChanged(const FText& InFilterText)
 {
 	if (Filter.GetFilterText().ToString().Equals(InFilterText.ToString(), ESearchCase::CaseSensitive))
diff --git a/Engine/Source/Developer/TraceInsights/Private/Insights/Widgets/SLogView.h b/Engine/Source/Developer/TraceInsights/Private/Insights/Widgets/SLogView.h
--- a/Engine/Source/Developer/TraceInsights/Private/Insights/Widgets/SLogView.h
+++ b/Engine/Source/Developer/TraceInsights/Private/Insights/Widgets/SLogView.h
@@ -79,6 +79,8 @@ protected:
 	/** Generate a new list view row. */
 	TSharedRef<ITableRow> OnGenerateRow(TSharedPtr<FLogMessage> InLogMessage, const TSharedRef<STableViewBase>& OwnerTable);
 
+	void SelectLogMessage(TSharedPtr<FLogMessage> LogMessage);
+	void OnMouseButtonClick(TSharedPtr<FLogMessage> LogMessage);
 	void OnSelectionChanged(TSharedPtr<FLogMessage> LogMessage, ESelectInfo::Type SelectInfo);
 
 	void FilterTextBox_OnTextChanged(const FText& InFilterText);
diff --git a/Engine/Source/Developer/TraceInsights/Private/Insights/Widgets/STimersViewTooltip.cpp b/Engine/Source/Developer/TraceInsights/Private/Insights/Widgets/STimersViewTooltip.cpp
--- a/Engine/Source/Developer/TraceInsights/Private/Insights/Widgets/STimersViewTooltip.cpp
+++ b/Engine/Source/Developer/TraceInsights/Private/Insights/Widgets/STimersViewTooltip.cpp
@@ -94,13 +94,16 @@ TSharedPtr<SToolTip> STimersViewTooltip::GetCellTooltip(const TSharedPtr<FTimerN
 
 	const int32 NumDigits = 5;
 
-	const FText TotalInclusiveTimeText = FText::FromString(TimeUtils::FormatTimeAuto(Stats.TotalInclusiveTime));
+	TCHAR FormatString[32];
+	FCString::Snprintf(FormatString, sizeof(FormatString), TEXT("%%.%dfs (%%s)"), NumDigits);
+
+	const FText TotalInclusiveTimeText = FText::FromString(FString::Printf(FormatString, Stats.TotalInclusiveTime, *TimeUtils::FormatTimeAuto(Stats.TotalInclusiveTime, 2)));
 	const FText MinInclusiveTimeText = FText::FromString(TimeUtils::FormatTimeMs(Stats.MinInclusiveTime, NumDigits, true));
 	const FText MaxInclusiveTimeText = FText::FromString(TimeUtils::FormatTimeMs(Stats.MaxInclusiveTime, NumDigits, true));
 	const FText AvgInclusiveTimeText = FText::FromString(TimeUtils::FormatTimeMs(Stats.AverageInclusiveTime, NumDigits, true));
 	const FText MedInclusiveTimeText = FText::FromString(TimeUtils::FormatTimeMs(Stats.MedianInclusiveTime, NumDigits, true));
 
-	const FText TotalExclusiveTimeText = FText::FromString(TimeUtils::FormatTimeAuto(Stats.TotalExclusiveTime));
+	const FText TotalExclusiveTimeText = FText::FromString(FString::Printf(FormatString, Stats.TotalExclusiveTime, *TimeUtils::FormatTimeAuto(Stats.TotalExclusiveTime, 2)));
 	const FText MinExclusiveTimeText = FText::FromString(TimeUtils::FormatTimeMs(Stats.MinExclusiveTime, NumDigits, true));
 	const FText MaxExclusiveTimeText = FText::FromString(TimeUtils::FormatTimeMs(Stats.MaxExclusiveTime, NumDigits, true));
 	const FText AvgExclusiveTimeText = FText::FromString(TimeUtils::FormatTimeMs(Stats.AverageExclusiveTime, NumDigits, true));
@@ -253,7 +256,7 @@ TSharedPtr<SToolTip> STimersViewTooltip::GetCellTooltip(const TSharedPtr<FTimerN
 						.TextStyle(FEditorStyle::Get(), TEXT("Profiler.TooltipBold"))
 					]
 					+ SGridPanel::Slot(2, 0)
-					.Padding(2.0f)
+					.Padding(FMargin(8.0f, 2.0f, 2.0f, 2.0f))
 					.HAlign(HAlign_Right)
 					[
 						SNew(STextBlock)
@@ -306,7 +309,7 @@ void STimersViewTooltip::AddStatsRow(TSharedPtr<SGridPanel> Grid, int32& Row, co
 		];
 
 	Grid->AddSlot(2, Row)
-		.Padding(2.0f)
+		.Padding(FMargin(8.0f, 2.0f, 2.0f, 2.0f))
 		.HAlign(HAlign_Right)
 		[
 			SNew(STextBlock)
