[Insights]
  - TraceServices/NetProfiler: Added iteration for NetEventTypes. Added aggregation stats for net events. (by [REDACTED_PERSON])
  - Networking: Added NetStats tree view
    * Columns: Name, Level and stats (Count, Total/Max/Average Inclusive Size, Total/Max Eclusive Size).
    * Grouping: Flat, By Name or By Level.
    * Sorting by any column.
    * Filtering by event name.
    * Aggregated stats are computed for selected packet.
  - Networking: Added mouse click on an event in Packedt Content view to select the event and also select its type in NetStats view.
  - Timing Insights: Renamed "Stats" to "Counters".
  - Timers view: Changed sorting to use the generic Sorter implementation.

diff --git a/Engine/Source/Developer/TraceInsights/Private/Insights/NetworkingProfiler/ViewModels/NetEventGroupingAndSorting.cpp b/Engine/Source/Developer/TraceInsights/Private/Insights/NetworkingProfiler/ViewModels/NetEventGroupingAndSorting.cpp
new file mode 100644
--- /dev/null
+++ b/Engine/Source/Developer/TraceInsights/Private/Insights/NetworkingProfiler/ViewModels/NetEventGroupingAndSorting.cpp
@@ -0,0 +1,333 @@
+// Copyright 1998-2019 Epic Games, Inc. All Rights Reserved.
+
+#include "NetEventGroupingAndSorting.h"
+
+#include "Insights/Table/ViewModels/TableColumn.h"
+#include "Insights/NetworkingProfiler/ViewModels/NetEventNodeHelper.h"
+
+#define LOCTEXT_NAMESPACE "NetEventNode"
+
+// Sort by name (ascending).
+#define INSIGHTS_DEFAULT_SORTING_NODES(A, B) return A->GetName().LexicalLess(B->GetName());
+
+////////////////////////////////////////////////////////////////////////////////////////////////////
+// Sorting by Event Type
+////////////////////////////////////////////////////////////////////////////////////////////////////
+
+FNetEventNodeSortingByEventType::FNetEventNodeSortingByEventType(TSharedRef<Insights::FTableColumn> InColumnRef)
+	: Insights::FTableCellValueSorter(
+		FName(TEXT("ByEventType")),
+		LOCTEXT("Sorting_ByEventType_Name", "By Type"),
+		LOCTEXT("Sorting_ByEventType_Title", "Sort By Type"),
+		LOCTEXT("Sorting_ByEventType_Desc", "Sort by event type (ascending or descending), then by name (ascending)."),
+		InColumnRef)
+{
+}
+
+////////////////////////////////////////////////////////////////////////////////////////////////////
+
+void FNetEventNodeSortingByEventType::Sort(TArray<Insights::FBaseTreeNodePtr>& NodesToSort, Insights::ESortMode SortMode) const
+{
+	if (SortMode == Insights::ESortMode::Ascending)
+	{
+		NodesToSort.Sort([](const Insights::FBaseTreeNodePtr& A, const Insights::FBaseTreeNodePtr& B) -> bool
+		{
+			ensure(A.IsValid() && A->GetTypeName() == FNetEventNode::TypeName);
+			const FNetEventNodePtr NetEventNodeA = StaticCastSharedPtr<FNetEventNode, Insights::FBaseTreeNode>(A);
+
+			ensure(B.IsValid() && B->GetTypeName() == FNetEventNode::TypeName);
+			const FNetEventNodePtr NetEventNodeB = StaticCastSharedPtr<FNetEventNode, Insights::FBaseTreeNode>(B);
+
+			if (NetEventNodeA->GetType() == NetEventNodeB->GetType())
+			{
+				INSIGHTS_DEFAULT_SORTING_NODES(A, B)
+			}
+			else
+			{
+				// Sort by net event type (ascending).
+				return NetEventNodeA->GetType() < NetEventNodeB->GetType();
+			}
+		});
+	}
+	else // if (SortMode == ESortMode::Descending)
+	{
+		NodesToSort.Sort([](const Insights::FBaseTreeNodePtr& A, const Insights::FBaseTreeNodePtr& B) -> bool
+		{
+			ensure(A.IsValid() && A->GetTypeName() == FNetEventNode::TypeName);
+			const FNetEventNodePtr NetEventNodeA = StaticCastSharedPtr<FNetEventNode, Insights::FBaseTreeNode>(A);
+
+			ensure(B.IsValid() && B->GetTypeName() == FNetEventNode::TypeName);
+			const FNetEventNodePtr NetEventNodeB = StaticCastSharedPtr<FNetEventNode, Insights::FBaseTreeNode>(B);
+
+			if (NetEventNodeA->GetType() == NetEventNodeB->GetType())
+			{
+				INSIGHTS_DEFAULT_SORTING_NODES(A, B)
+			}
+			else
+			{
+				// Sort by net event type (descending).
+				return NetEventNodeB->GetType() < NetEventNodeA->GetType();
+			}
+		});
+	}
+}
+
+////////////////////////////////////////////////////////////////////////////////////////////////////
+// Sort by Instance Count
+////////////////////////////////////////////////////////////////////////////////////////////////////
+
+FNetEventNodeSortingByInstanceCount::FNetEventNodeSortingByInstanceCount(TSharedRef<Insights::FTableColumn> InColumnRef)
+	: Insights::FTableCellValueSorter(
+		FName(TEXT("ByInstanceCount")),
+		LOCTEXT("Sorting_ByInstanceCount_Name", "By Instance Count"),
+		LOCTEXT("Sorting_ByInstanceCount_Title", "Sort By Instance Count"),
+		LOCTEXT("Sorting_ByInstanceCount_Desc", "Sort by aggregated instance count (ascending or descending), then by name (ascending)."),
+		InColumnRef)
+{
+}
+
+////////////////////////////////////////////////////////////////////////////////////////////////////
+
+void FNetEventNodeSortingByInstanceCount::Sort(TArray<Insights::FBaseTreeNodePtr>& NodesToSort, Insights::ESortMode SortMode) const
+{
+	if (SortMode == Insights::ESortMode::Ascending)
+	{
+		NodesToSort.Sort([](const Insights::FBaseTreeNodePtr& A, const Insights::FBaseTreeNodePtr& B) -> bool
+		{
+			ensure(A.IsValid() && A->GetTypeName() == FNetEventNode::TypeName);
+			const FNetEventNodePtr NetEventNodeA = StaticCastSharedPtr<FNetEventNode, Insights::FBaseTreeNode>(A);
+			const uint64 ValueA = NetEventNodeA->GetAggregatedStats().InstanceCount;
+
+			ensure(B.IsValid() && B->GetTypeName() == FNetEventNode::TypeName);
+			const FNetEventNodePtr NetEventNodeB = StaticCastSharedPtr<FNetEventNode, Insights::FBaseTreeNode>(B);
+			const uint64 ValueB = NetEventNodeB->GetAggregatedStats().InstanceCount;
+
+			if (ValueA == ValueB)
+			{
+				INSIGHTS_DEFAULT_SORTING_NODES(A, B)
+			}
+			else
+			{
+				// Sort by instance count (ascending).
+				return ValueA < ValueB;
+			}
+		});
+	}
+	else // if (SortMode == ESortMode::Descending)
+	{
+		NodesToSort.Sort([](const Insights::FBaseTreeNodePtr& A, const Insights::FBaseTreeNodePtr& B) -> bool
+		{
+			ensure(A.IsValid() && A->GetTypeName() == FNetEventNode::TypeName);
+			const FNetEventNodePtr NetEventNodeA = StaticCastSharedPtr<FNetEventNode, Insights::FBaseTreeNode>(A);
+			const uint64 ValueA = NetEventNodeA->GetAggregatedStats().InstanceCount;
+
+			ensure(B.IsValid() && B->GetTypeName() == FNetEventNode::TypeName);
+			const FNetEventNodePtr NetEventNodeB = StaticCastSharedPtr<FNetEventNode, Insights::FBaseTreeNode>(B);
+			const uint64 ValueB = NetEventNodeB->GetAggregatedStats().InstanceCount;
+
+			if (ValueA == ValueB)
+			{
+				INSIGHTS_DEFAULT_SORTING_NODES(A, B)
+			}
+			else
+			{
+				// Sort by instance count (descending).
+				return ValueB < ValueA;
+			}
+		});
+	}
+}
+
+////////////////////////////////////////////////////////////////////////////////////////////////////
+// Sort by Total Inclusive Size
+////////////////////////////////////////////////////////////////////////////////////////////////////
+
+FNetEventNodeSortingByTotalInclusiveSize::FNetEventNodeSortingByTotalInclusiveSize(TSharedRef<Insights::FTableColumn> InColumnRef)
+	: Insights::FTableCellValueSorter(
+		FName(TEXT("ByTotalInclusiveSize")),
+		LOCTEXT("Sorting_ByTotalInclusiveSize_Name", "By Total Inclusive Size"),
+		LOCTEXT("Sorting_ByTotalInclusiveSize_Title", "Sort By Total Inclusive Size"),
+		LOCTEXT("Sorting_ByTotalInclusiveSize_Desc", "Sort by aggregated total inclusive size (ascending or descending), then by name (ascending)."),
+		InColumnRef)
+{
+}
+
+////////////////////////////////////////////////////////////////////////////////////////////////////
+
+void FNetEventNodeSortingByTotalInclusiveSize::Sort(TArray<Insights::FBaseTreeNodePtr>& NodesToSort, Insights::ESortMode SortMode) const
+{
+	if (SortMode == Insights::ESortMode::Ascending)
+	{
+		NodesToSort.Sort([](const Insights::FBaseTreeNodePtr& A, const Insights::FBaseTreeNodePtr& B) -> bool
+		{
+			ensure(A.IsValid() && A->GetTypeName() == FNetEventNode::TypeName);
+			const FNetEventNodePtr NetEventNodeA = StaticCastSharedPtr<FNetEventNode, Insights::FBaseTreeNode>(A);
+			const uint32 ValueA = NetEventNodeA->GetAggregatedStats().TotalInclusive;
+
+			ensure(B.IsValid() && B->GetTypeName() == FNetEventNode::TypeName);
+			const FNetEventNodePtr NetEventNodeB = StaticCastSharedPtr<FNetEventNode, Insights::FBaseTreeNode>(B);
+			const uint32 ValueB = NetEventNodeB->GetAggregatedStats().TotalInclusive;
+
+			if (ValueA == ValueB)
+			{
+				INSIGHTS_DEFAULT_SORTING_NODES(A, B)
+			}
+			else
+			{
+				// Sort by total inclusive size (ascending).
+				return ValueA < ValueB;
+			}
+		});
+	}
+	else // if (SortMode == ESortMode::Descending)
+	{
+		NodesToSort.Sort([](const Insights::FBaseTreeNodePtr& A, const Insights::FBaseTreeNodePtr& B) -> bool
+		{
+			ensure(A.IsValid() && A->GetTypeName() == FNetEventNode::TypeName);
+			const FNetEventNodePtr NetEventNodeA = StaticCastSharedPtr<FNetEventNode, Insights::FBaseTreeNode>(A);
+			const uint32 ValueA = NetEventNodeA->GetAggregatedStats().TotalInclusive;
+
+			ensure(B.IsValid() && B->GetTypeName() == FNetEventNode::TypeName);
+			const FNetEventNodePtr NetEventNodeB = StaticCastSharedPtr<FNetEventNode, Insights::FBaseTreeNode>(B);
+			const uint32 ValueB = NetEventNodeB->GetAggregatedStats().TotalInclusive;
+
+			if (ValueA == ValueB)
+			{
+				INSIGHTS_DEFAULT_SORTING_NODES(A, B)
+			}
+			else
+			{
+				// Sort by total inclusive size (descending).
+				return ValueB < ValueA;
+			}
+		});
+	}
+}
+
+////////////////////////////////////////////////////////////////////////////////////////////////////
+// Sort by Total Exclusive Size
+////////////////////////////////////////////////////////////////////////////////////////////////////
+
+FNetEventNodeSortingByTotalExclusiveSize::FNetEventNodeSortingByTotalExclusiveSize(TSharedRef<Insights::FTableColumn> InColumnRef)
+	: Insights::FTableCellValueSorter(
+		FName(TEXT("ByTotalExclusiveSize")),
+		LOCTEXT("Sorting_ByTotalExclusiveSize_Name", "By Total Exclusive Size"),
+		LOCTEXT("Sorting_ByTotalExclusiveSize_Title", "Sort By Total Exclusive Size"),
+		LOCTEXT("Sorting_ByTotalExclusiveSize_Desc", "Sort by aggregated total exclusive size (ascending or descending), then by name (ascending)."),
+		InColumnRef)
+{
+}
+
+////////////////////////////////////////////////////////////////////////////////////////////////////
+
+void FNetEventNodeSortingByTotalExclusiveSize::Sort(TArray<Insights::FBaseTreeNodePtr>& NodesToSort, Insights::ESortMode SortMode) const
+{
+	if (SortMode == Insights::ESortMode::Ascending)
+	{
+		NodesToSort.Sort([](const Insights::FBaseTreeNodePtr& A, const Insights::FBaseTreeNodePtr& B) -> bool
+		{
+			ensure(A.IsValid() && A->GetTypeName() == FNetEventNode::TypeName);
+			const FNetEventNodePtr NetEventNodeA = StaticCastSharedPtr<FNetEventNode, Insights::FBaseTreeNode>(A);
+			const uint32 ValueA = NetEventNodeA->GetAggregatedStats().TotalExclusive;
+
+			ensure(B.IsValid() && B->GetTypeName() == FNetEventNode::TypeName);
+			const FNetEventNodePtr NetEventNodeB = StaticCastSharedPtr<FNetEventNode, Insights::FBaseTreeNode>(B);
+			const uint32 ValueB = NetEventNodeB->GetAggregatedStats().TotalExclusive;
+
+			if (ValueA == ValueB)
+			{
+				INSIGHTS_DEFAULT_SORTING_NODES(A, B)
+			}
+			else
+			{
+				// Sort by total exclusive size (ascending).
+				return ValueA < ValueB;
+			}
+		});
+	}
+	else // if (SortMode == ESortMode::Descending)
+	{
+		NodesToSort.Sort([](const Insights::FBaseTreeNodePtr& A, const Insights::FBaseTreeNodePtr& B) -> bool
+		{
+			ensure(A.IsValid() && A->GetTypeName() == FNetEventNode::TypeName);
+			const FNetEventNodePtr NetEventNodeA = StaticCastSharedPtr<FNetEventNode, Insights::FBaseTreeNode>(A);
+			const uint32 ValueA = NetEventNodeA->GetAggregatedStats().TotalExclusive;
+
+			ensure(B.IsValid() && B->GetTypeName() == FNetEventNode::TypeName);
+			const FNetEventNodePtr NetEventNodeB = StaticCastSharedPtr<FNetEventNode, Insights::FBaseTreeNode>(B);
+			const uint32 ValueB = NetEventNodeB->GetAggregatedStats().TotalExclusive;
+
+			if (ValueA == ValueB)
+			{
+				INSIGHTS_DEFAULT_SORTING_NODES(A, B)
+			}
+			else
+			{
+				// Sort by total exclusive size (descending).
+				return ValueB < ValueA;
+			}
+		});
+	}
+}
+
+////////////////////////////////////////////////////////////////////////////////////////////////////
+// Grouping by Event Type
+////////////////////////////////////////////////////////////////////////////////////////////////////
+
+FNetEventNodeGroupingByEventType::FNetEventNodeGroupingByEventType()
+	: FTreeNodeGrouping(
+		LOCTEXT("Grouping_ByEventType_ShortName", "Type"),
+		LOCTEXT("Grouping_ByEventType_TitleName", "Type"),
+		LOCTEXT("Grouping_ByEventType_Desc", "Creates a group for each net event type."),
+		TEXT("Profiler.FiltersAndPresets.StatTypeIcon"), //TODO
+		nullptr)
+{
+}
+
+////////////////////////////////////////////////////////////////////////////////////////////////////
+
+Insights::FTreeNodeGroupInfo FNetEventNodeGroupingByEventType::GetGroupForNode(const Insights::FBaseTreeNodePtr InNodePtr) const
+{
+	if (InNodePtr->GetTypeName() == FNetEventNode::TypeName)
+	{
+		const FNetEventNodePtr& NetEventNodePtr = StaticCastSharedPtr<FNetEventNode>(InNodePtr);
+		return { FName(*NetEventNodeTypeHelper::ToText(NetEventNodePtr->GetType()).ToString()), true };
+	}
+	else
+	{
+		return { FName("Unknown"), true };
+	}
+}
+
+////////////////////////////////////////////////////////////////////////////////////////////////////
+// Grouping by Level
+////////////////////////////////////////////////////////////////////////////////////////////////////
+
+FNetEventNodeGroupingByLevel::FNetEventNodeGroupingByLevel()
+	: FTreeNodeGrouping(
+		LOCTEXT("Grouping_ByLevel_ShortName", "Level"),
+		LOCTEXT("Grouping_ByLevel_TitleName", "Level"),
+		LOCTEXT("Grouping_ByLevel_Desc", "Creates a group for each level."),
+		TEXT("Profiler.FiltersAndPresets.StatTypeIcon"), //TODO
+		nullptr)
+{
+}
+
+////////////////////////////////////////////////////////////////////////////////////////////////////
+
+Insights::FTreeNodeGroupInfo FNetEventNodeGroupingByLevel::GetGroupForNode(const Insights::FBaseTreeNodePtr InNodePtr) const
+{
+	if (InNodePtr->GetTypeName() == FNetEventNode::TypeName)
+	{
+		const FNetEventNodePtr& NetEventNodePtr = StaticCastSharedPtr<FNetEventNode>(InNodePtr);
+		return { FName(*FString::Printf(TEXT("Level %d"), NetEventNodePtr->GetLevel())), true };
+	}
+	else
+	{
+		return { FName("Unknown"), true };
+	}
+}
+
+////////////////////////////////////////////////////////////////////////////////////////////////////
+
+#undef LOCTEXT_NAMESPACE
diff --git a/Engine/Source/Developer/TraceInsights/Private/Insights/NetworkingProfiler/ViewModels/NetEventGroupingAndSorting.h b/Engine/Source/Developer/TraceInsights/Private/Insights/NetworkingProfiler/ViewModels/NetEventGroupingAndSorting.h
new file mode 100644
--- /dev/null
+++ b/Engine/Source/Developer/TraceInsights/Private/Insights/NetworkingProfiler/ViewModels/NetEventGroupingAndSorting.h
@@ -0,0 +1,105 @@
+// Copyright 1998-2019 Epic Games, Inc. All Rights Reserved.
+
+#pragma once
+
+#include "CoreMinimal.h"
+
+// Insights
+#include "Insights/NetworkingProfiler/ViewModels/NetEventNode.h"
+#include "Insights/Table/ViewModels/TableColumn.h"
+#include "Insights/Table/ViewModels/TableCellValueSorter.h"
+#include "Insights/Table/ViewModels/TreeNodeGrouping.h"
+
+////////////////////////////////////////////////////////////////////////////////////////////////////
+// Sorters
+////////////////////////////////////////////////////////////////////////////////////////////////////
+
+class FNetEventNodeSortingByEventType: public Insights::FTableCellValueSorter
+{
+public:
+	FNetEventNodeSortingByEventType(TSharedRef<Insights::FTableColumn> InColumnRef);
+
+	virtual void Sort(TArray<Insights::FBaseTreeNodePtr>& NodesToSort, Insights::ESortMode SortMode) const override;
+};
+
+////////////////////////////////////////////////////////////////////////////////////////////////////
+
+class FNetEventNodeSortingByInstanceCount : public Insights::FTableCellValueSorter
+{
+public:
+	FNetEventNodeSortingByInstanceCount(TSharedRef<Insights::FTableColumn> InColumnRef);
+
+	virtual void Sort(TArray<Insights::FBaseTreeNodePtr>& NodesToSort, Insights::ESortMode SortMode) const override;
+};
+
+////////////////////////////////////////////////////////////////////////////////////////////////////
+
+class FNetEventNodeSortingByTotalInclusiveSize : public Insights::FTableCellValueSorter
+{
+public:
+	FNetEventNodeSortingByTotalInclusiveSize(TSharedRef<Insights::FTableColumn> InColumnRef);
+
+	virtual void Sort(TArray<Insights::FBaseTreeNodePtr>& NodesToSort, Insights::ESortMode SortMode) const override;
+};
+
+////////////////////////////////////////////////////////////////////////////////////////////////////
+
+class FNetEventNodeSortingByTotalExclusiveSize : public Insights::FTableCellValueSorter
+{
+public:
+	FNetEventNodeSortingByTotalExclusiveSize(TSharedRef<Insights::FTableColumn> InColumnRef);
+
+	virtual void Sort(TArray<Insights::FBaseTreeNodePtr>& NodesToSort, Insights::ESortMode SortMode) const override;
+};
+
+////////////////////////////////////////////////////////////////////////////////////////////////////
+// Organizers
+////////////////////////////////////////////////////////////////////////////////////////////////////
+
+/** Enumerates types of grouping or sorting for the net event nodes. */
+enum class ENetEventGroupingMode
+{
+	/** Creates a single group for all net events. */
+	Flat,
+
+	/** Creates one group for one letter. */
+	ByName,
+
+	/** Creates one group for each net event type. */
+	ByType,
+
+	/** Creates one group for each net event level. */
+	ByLevel,
+
+	/** Invalid enum type, may be used as a number of enumerations. */
+	InvalidOrMax,
+};
+
+/** Type definition for shared pointers to instances of ENetEventGroupingMode. */
+typedef TSharedPtr<ENetEventGroupingMode> ENetEventGroupingModePtr;
+
+////////////////////////////////////////////////////////////////////////////////////////////////////
+
+/** Creates a group for each event type. */
+class FNetEventNodeGroupingByEventType : public Insights::FTreeNodeGrouping
+{
+public:
+	FNetEventNodeGroupingByEventType();
+	virtual ~FNetEventNodeGroupingByEventType() {}
+
+	virtual Insights::FTreeNodeGroupInfo GetGroupForNode(const Insights::FBaseTreeNodePtr InNode) const override;
+};
+
+////////////////////////////////////////////////////////////////////////////////////////////////////
+
+/** Creates a group for each event level. */
+class FNetEventNodeGroupingByLevel : public Insights::FTreeNodeGrouping
+{
+public:
+	FNetEventNodeGroupingByLevel();
+	virtual ~FNetEventNodeGroupingByLevel() {}
+
+	virtual Insights::FTreeNodeGroupInfo GetGroupForNode(const Insights::FBaseTreeNodePtr InNode) const override;
+};
+
+////////////////////////////////////////////////////////////////////////////////////////////////////
diff --git a/Engine/Source/Developer/TraceInsights/Private/Insights/NetworkingProfiler/ViewModels/NetEventNode.cpp b/Engine/Source/Developer/TraceInsights/Private/Insights/NetworkingProfiler/ViewModels/NetEventNode.cpp
new file mode 100644
--- /dev/null
+++ b/Engine/Source/Developer/TraceInsights/Private/Insights/NetworkingProfiler/ViewModels/NetEventNode.cpp
@@ -0,0 +1,25 @@
+// Copyright 1998-2019 Epic Games, Inc. All Rights Reserved.
+
+#include "NetEventNode.h"
+
+#define LOCTEXT_NAMESPACE "NetEventNode"
+
+const FName FNetEventNode::TypeName(TEXT("FNetEventNode"));
+
+////////////////////////////////////////////////////////////////////////////////////////////////////
+
+void FNetEventNode::ResetAggregatedStats()
+{
+	AggregatedStats = Trace::FNetProfilerAggregatedStats();
+}
+
+////////////////////////////////////////////////////////////////////////////////////////////////////
+
+void FNetEventNode::SetAggregatedStats(const Trace::FNetProfilerAggregatedStats& InAggregatedStats)
+{
+	AggregatedStats = InAggregatedStats;
+}
+
+////////////////////////////////////////////////////////////////////////////////////////////////////
+
+#undef LOCTEXT_NAMESPACE
diff --git a/Engine/Source/Developer/TraceInsights/Private/Insights/NetworkingProfiler/ViewModels/NetEventNode.h b/Engine/Source/Developer/TraceInsights/Private/Insights/NetworkingProfiler/ViewModels/NetEventNode.h
new file mode 100644
--- /dev/null
+++ b/Engine/Source/Developer/TraceInsights/Private/Insights/NetworkingProfiler/ViewModels/NetEventNode.h
@@ -0,0 +1,100 @@
+// Copyright 1998-2019 Epic Games, Inc. All Rights Reserved.
+
+#pragma once
+
+#include "CoreMinimal.h"
+#include "TraceServices/Model/NetProfiler.h"
+
+// Insights
+#include "Insights/Table/ViewModels/BaseTreeNode.h"
+
+////////////////////////////////////////////////////////////////////////////////////////////////////
+
+enum class ENetEventNodeType
+{
+	/** The NetEventNode is a Net Event. */
+	NetEvent,
+
+	/** The NetEventNode is a Net Object. */
+	//NetObject,
+
+	/** The NetEventNode is a group node. */
+	Group,
+
+	/** Invalid enum type, may be used as a number of enumerations. */
+	InvalidOrMax,
+};
+
+////////////////////////////////////////////////////////////////////////////////////////////////////
+
+class FNetEventNode;
+
+/** Type definition for shared pointers to instances of FNetEventNode. */
+typedef TSharedPtr<class FNetEventNode> FNetEventNodePtr;
+
+/** Type definition for shared references to instances of FNetEventNode. */
+typedef TSharedRef<class FNetEventNode> FNetEventNodeRef;
+
+/** Type definition for shared references to const instances of FNetEventNode. */
+typedef TSharedRef<const class FNetEventNode> FNetEventNodeRefConst;
+
+/** Type definition for weak references to instances of FNetEventNode. */
+typedef TWeakPtr<class FNetEventNode> FNetEventNodeWeak;
+
+////////////////////////////////////////////////////////////////////////////////////////////////////
+/**
+ * Class used to store information about a timer node (used in the STimersView).
+ */
+class FNetEventNode : public Insights::FBaseTreeNode
+{
+public:
+	static const FName TypeName;
+
+public:
+	/** Initialization constructor for the NetEvent node. */
+	FNetEventNode(uint64 InId, const FName InName, ENetEventNodeType InType, uint32 InLevel)
+		: FBaseTreeNode(InId, InName, InType == ENetEventNodeType::Group)
+		, Type(InType)
+		, Level(InLevel)
+	{
+		ResetAggregatedStats();
+	}
+
+	/** Initialization constructor for the group node. */
+	FNetEventNode(const FName InGroupName)
+		: FBaseTreeNode(0, InGroupName, true)
+		, Type(ENetEventNodeType::Group)
+		, Level(0)
+	{
+		ResetAggregatedStats();
+	}
+
+	virtual const FName& GetTypeName() const override { return TypeName; }
+
+	/**
+	 * @return a type of this NetEvent node or ENetEventNodeType::Group for group nodes.
+	 */
+	const ENetEventNodeType& GetType() const { return Type; }
+
+	void SetLevel(uint32 InLevel) const { return ; }
+	uint32 GetLevel() const { return Level; }
+
+	/**
+	 * @return the aggregated stats for this NetEvent node.
+	 */
+	const Trace::FNetProfilerAggregatedStats& GetAggregatedStats() const { return AggregatedStats; }
+
+	void ResetAggregatedStats();
+	void SetAggregatedStats(const Trace::FNetProfilerAggregatedStats& AggregatedStats);
+
+private:
+	/** Holds the type of this NetEvent node. */
+	const ENetEventNodeType Type;
+
+	uint32 Level;
+
+	/** Aggregated stats. */
+	Trace::FNetProfilerAggregatedStats AggregatedStats;
+};
+
+////////////////////////////////////////////////////////////////////////////////////////////////////
diff --git a/Engine/Source/Developer/TraceInsights/Private/Insights/NetworkingProfiler/ViewModels/NetEventNodeHelper.cpp b/Engine/Source/Developer/TraceInsights/Private/Insights/NetworkingProfiler/ViewModels/NetEventNodeHelper.cpp
new file mode 100644
--- /dev/null
+++ b/Engine/Source/Developer/TraceInsights/Private/Insights/NetworkingProfiler/ViewModels/NetEventNodeHelper.cpp
@@ -0,0 +1,118 @@
+// Copyright 1998-2019 Epic Games, Inc. All Rights Reserved.
+
+#include "NetEventNodeHelper.h"
+#include "EditorStyleSet.h"
+
+#define LOCTEXT_NAMESPACE "NetEventNode"
+
+////////////////////////////////////////////////////////////////////////////////////////////////////
+// NetEventNode Type Helper
+////////////////////////////////////////////////////////////////////////////////////////////////////
+
+FText NetEventNodeTypeHelper::ToText(const ENetEventNodeType NodeType)
+{
+	static_assert(static_cast<int>(ENetEventNodeType::InvalidOrMax) == 2, "Not all cases are handled in switch below!?");
+	switch (NodeType)
+	{
+		case ENetEventNodeType::NetEvent:	return LOCTEXT("Type_Name_NetEvent", "Net Event");
+		case ENetEventNodeType::Group:		return LOCTEXT("Type_Name_Group", "Group");
+		default:							return LOCTEXT("InvalidOrMax", "InvalidOrMax");
+	}
+}
+
+////////////////////////////////////////////////////////////////////////////////////////////////////
+
+FText NetEventNodeTypeHelper::ToDescription(const ENetEventNodeType NodeType)
+{
+	static_assert(static_cast<int>(ENetEventNodeType::InvalidOrMax) == 2, "Not all cases are handled in switch below!?");
+	switch (NodeType)
+	{
+		case ENetEventNodeType::NetEvent:	return LOCTEXT("Type_Desc_NetEvent", "Net event");
+		case ENetEventNodeType::Group:		return LOCTEXT("Type_Desc_Group", "Group node");
+		default:							return LOCTEXT("InvalidOrMax", "InvalidOrMax");
+	}
+}
+
+////////////////////////////////////////////////////////////////////////////////////////////////////
+
+FName NetEventNodeTypeHelper::ToBrushName(const ENetEventNodeType NodeType)
+{
+	static_assert(static_cast<int>(ENetEventNodeType::InvalidOrMax) == 2, "Not all cases are handled in switch below!?");
+	switch (NodeType)
+	{
+		case ENetEventNodeType::NetEvent:	return TEXT("Profiler.FiltersAndPresets.StatTypeIcon"); //TODO: "Icons.NetEvent"
+		case ENetEventNodeType::Group:		return TEXT("Profiler.Misc.GenericGroup"); //TODO: "Icons.GenericGroup"
+		default:							return NAME_None;
+	}
+}
+
+////////////////////////////////////////////////////////////////////////////////////////////////////
+
+const FSlateBrush* NetEventNodeTypeHelper::GetIconForGroup()
+{
+	return FEditorStyle::GetBrush(TEXT("Profiler.Misc.GenericGroup")); //TODO: FInsightsStyle::GetBrush(TEXT("Icons.GenericGroup"));
+}
+
+////////////////////////////////////////////////////////////////////////////////////////////////////
+
+const FSlateBrush* NetEventNodeTypeHelper::GetIconForNetEventNodeType(const ENetEventNodeType NodeType)
+{
+	static_assert(static_cast<int>(ENetEventNodeType::InvalidOrMax) == 2, "Not all cases are handled in switch below!?");
+	switch (NodeType)
+	{
+		case ENetEventNodeType::NetEvent:	return FEditorStyle::GetBrush(TEXT("Profiler.Type.NumberFloat")); //TODO: FInsightsStyle::GetBrush(TEXT("Icons.NetEvent"));
+		case ENetEventNodeType::Group:		return FEditorStyle::GetBrush(TEXT("Profiler.Misc.GenericGroup")); //TODO: FInsightsStyle::GetBrush(TEXT("Icons.GenericGroup"));
+		default:							return nullptr;
+	}
+}
+
+////////////////////////////////////////////////////////////////////////////////////////////////////
+// NetEventNode Grouping Helper
+////////////////////////////////////////////////////////////////////////////////////////////////////
+
+FText NetEventNodeGroupingHelper::ToText(const ENetEventGroupingMode GroupingMode)
+{
+	static_assert(static_cast<int>(ENetEventGroupingMode::InvalidOrMax) == 4, "Not all cases are handled in switch below!?");
+	switch (GroupingMode)
+	{
+		case ENetEventGroupingMode::Flat:		return LOCTEXT("Grouping_Name_Flat",	"Flat");
+		case ENetEventGroupingMode::ByName:		return LOCTEXT("Grouping_Name_ByName",	"Name");
+		case ENetEventGroupingMode::ByType:		return LOCTEXT("Grouping_Name_ByType",	"Event Type");
+		case ENetEventGroupingMode::ByLevel:	return LOCTEXT("Grouping_Name_ByLevel", "Level");
+		default:								return LOCTEXT("InvalidOrMax", "InvalidOrMax");
+	}
+}
+
+////////////////////////////////////////////////////////////////////////////////////////////////////
+
+FText NetEventNodeGroupingHelper::ToDescription(const ENetEventGroupingMode GroupingMode)
+{
+	static_assert(static_cast<int>(ENetEventGroupingMode::InvalidOrMax) == 4, "Not all cases are handled in switch below!?");
+	switch (GroupingMode)
+	{
+		case ENetEventGroupingMode::Flat:		return LOCTEXT("Grouping_Desc_Flat",	"Creates a single group. Includes all net events.");
+		case ENetEventGroupingMode::ByName:		return LOCTEXT("Grouping_Desc_ByName",	"Creates one group for one letter.");
+		case ENetEventGroupingMode::ByType:		return LOCTEXT("Grouping_Desc_ByType",	"Creates one group for each net event type.");
+		case ENetEventGroupingMode::ByLevel:	return LOCTEXT("Grouping_Desc_ByLevel", "Creates one group for each level.");
+		default:								return LOCTEXT("InvalidOrMax", "InvalidOrMax");
+	}
+}
+
+////////////////////////////////////////////////////////////////////////////////////////////////////
+
+FName NetEventNodeGroupingHelper::ToBrushName(const ENetEventGroupingMode GroupingMode)
+{
+	static_assert(static_cast<int>(ENetEventGroupingMode::InvalidOrMax) == 4, "Not all cases are handled in switch below!?");
+	switch (GroupingMode)
+	{
+		case ENetEventGroupingMode::Flat:		return TEXT("Profiler.FiltersAndPresets.GroupNameIcon"); //TODO: "Icons.Grouping.Flat"
+		case ENetEventGroupingMode::ByName:		return TEXT("Profiler.FiltersAndPresets.GroupNameIcon"); //TODO: "Icons.Grouping.ByName"
+		case ENetEventGroupingMode::ByType:		return TEXT("Profiler.FiltersAndPresets.StatTypeIcon"); //TODO
+		case ENetEventGroupingMode::ByLevel:	return TEXT("Profiler.FiltersAndPresets.StatTypeIcon"); //TODO
+		default:								return NAME_None;
+	}
+}
+
+////////////////////////////////////////////////////////////////////////////////////////////////////
+
+#undef LOCTEXT_NAMESPACE
diff --git a/Engine/Source/Developer/TraceInsights/Private/Insights/NetworkingProfiler/ViewModels/NetEventNodeHelper.h b/Engine/Source/Developer/TraceInsights/Private/Insights/NetworkingProfiler/ViewModels/NetEventNodeHelper.h
new file mode 100644
--- /dev/null
+++ b/Engine/Source/Developer/TraceInsights/Private/Insights/NetworkingProfiler/ViewModels/NetEventNodeHelper.h
@@ -0,0 +1,69 @@
+// Copyright 1998-2019 Epic Games, Inc. All Rights Reserved.
+
+#pragma once
+
+#include "CoreMinimal.h"
+#include "Styling/SlateBrush.h"
+
+// Insights
+#include "Insights/NetworkingProfiler/ViewModels/NetEventNode.h"
+#include "Insights/NetworkingProfiler/ViewModels/NetEventGroupingAndSorting.h"
+
+////////////////////////////////////////////////////////////////////////////////////////////////////
+
+/** Helper struct that contains helper functions and classes for ENetEventNodeType enum. */
+struct NetEventNodeTypeHelper
+{
+	/**
+	 * @param Type - The value to get the text for.
+	 *
+	 * @return text representation of the specified ENetEventNodeType value.
+	 */
+	static FText ToText(const ENetEventNodeType Type);
+
+	/**
+	 * @param Type - The value to get the text for.
+	 *
+	 * @return text representation with more detailed explanation of the specified ENetEventNodeType value.
+	 */
+	static FText ToDescription(const ENetEventNodeType Type);
+
+	/**
+	 * @param Type - The value to get the brush name for.
+	 *
+	 * @return brush name of the specified ENetEventNodeType value.
+	 */
+	static FName ToBrushName(const ENetEventNodeType Type);
+
+	static const FSlateBrush* GetIconForGroup();
+	static const FSlateBrush* GetIconForNetEventNodeType(const ENetEventNodeType Type);
+};
+
+////////////////////////////////////////////////////////////////////////////////////////////////////
+
+/** Helper struct that contains grouping static functions and classes. */
+struct NetEventNodeGroupingHelper
+{
+	/**
+	 * @param NetEventGroupingMode - The value to get the text for.
+	 *
+	 * @return text representation of the specified ENetEventGroupingMode value.
+	 */
+	static FText ToText(const ENetEventGroupingMode NetEventGroupingMode);
+
+	/**
+	 * @param NetEventGroupingMode - The value to get the text for.
+	 *
+	 * @return text representation with more detailed explanation of the specified ENetEventGroupingMode value.
+	 */
+	static FText ToDescription(const ENetEventGroupingMode NetEventGroupingMode);
+
+	/**
+	 * @param NetEventGroupingMode - The value to get the brush name for.
+	 *
+	 * @return brush name of the specified ENetEventGroupingMode value.
+	 */
+	static FName ToBrushName(const ENetEventGroupingMode NetEventGroupingMode);
+};
+
+////////////////////////////////////////////////////////////////////////////////////////////////////
diff --git a/Engine/Source/Developer/TraceInsights/Private/Insights/NetworkingProfiler/ViewModels/NetStatsViewColumnFactory.cpp b/Engine/Source/Developer/TraceInsights/Private/Insights/NetworkingProfiler/ViewModels/NetStatsViewColumnFactory.cpp
new file mode 100644
--- /dev/null
+++ b/Engine/Source/Developer/TraceInsights/Private/Insights/NetworkingProfiler/ViewModels/NetStatsViewColumnFactory.cpp
@@ -0,0 +1,505 @@
+// Copyright 1998-2019 Epic Games, Inc. All Rights Reserved.
+
+#include "NetStatsViewColumnFactory.h"
+
+// Insights
+#include "Insights/Table/ViewModels/TableCellValueFormatter.h"
+#include "Insights/Table/ViewModels/TableCellValueGetter.h"
+#include "Insights/Table/ViewModels/TableCellValueSorter.h"
+#include "Insights/NetworkingProfiler/ViewModels/NetEventGroupingAndSorting.h"
+#include "Insights/NetworkingProfiler/ViewModels/NetEventNodeHelper.h"
+
+#define LOCTEXT_NAMESPACE "SNetStatsView"
+
+////////////////////////////////////////////////////////////////////////////////////////////////////
+// Column identifiers
+
+const FName FNetStatsViewColumns::NameColumnID(TEXT("Name")); // TEXT("_Hierarchy")
+const FName FNetStatsViewColumns::TypeColumnID(TEXT("Type"));
+const FName FNetStatsViewColumns::LevelColumnID(TEXT("Level"));
+const FName FNetStatsViewColumns::InstanceCountColumnID(TEXT("Count"));
+
+// Inclusive  columns
+const FName FNetStatsViewColumns::TotalInclusiveSizeColumnID(TEXT("TotalIncl"));
+const FName FNetStatsViewColumns::MaxInclusiveSizeColumnID(TEXT("MaxIncl"));
+const FName FNetStatsViewColumns::AverageInclusiveSizeColumnID(TEXT("AverageIncl"));
+
+// Exclusive  columns
+const FName FNetStatsViewColumns::TotalExclusiveSizeColumnID(TEXT("TotalExcl"));
+const FName FNetStatsViewColumns::MaxExclusiveSizeColumnID(TEXT("MaxExcl"));
+const FName FNetStatsViewColumns::AverageExclusiveSizeColumnID(TEXT("AverageExcl"));
+
+////////////////////////////////////////////////////////////////////////////////////////////////////
+
+void FNetStatsViewColumnFactory::CreateNetStatsViewColumns(TArray<TSharedPtr<Insights::FTableColumn>>& Columns)
+{
+	Columns.Reset();
+
+	Columns.Add(CreateNameColumn());
+	Columns.Add(CreateTypeColumn());
+	Columns.Add(CreateLevelColumn());
+	Columns.Add(CreateInstanceCountColumn());
+
+	Columns.Add(CreateTotalInclusiveSizeColumn());
+	Columns.Add(CreateMaxInclusiveSizeColumn());
+	Columns.Add(CreateAverageInclusiveSizeColumn());
+
+	Columns.Add(CreateTotalExclusiveSizeColumn());
+	Columns.Add(CreateMaxExclusiveSizeColumn());
+	//Columns.Add(CreateAverageExclusiveSizeColumn());
+}
+
+////////////////////////////////////////////////////////////////////////////////////////////////////
+
+TSharedRef<Insights::FTableColumn> FNetStatsViewColumnFactory::CreateNameColumn()
+{
+	using namespace Insights;
+
+	TSharedPtr<FTableColumn> ColumnPtr = MakeShareable(new FTableColumn(FNetStatsViewColumns::NameColumnID));
+	FTableColumn& Column = *ColumnPtr;
+
+	Column.SetShortName(LOCTEXT("Name_ColumnName", "Name"));
+	Column.SetTitleName(LOCTEXT("Name_ColumnTitle", "NetEvent or Group Name"));
+	Column.SetDescription(LOCTEXT("Name_ColumnDesc", "Name of the timer or group"));
+
+	Column.SetFlags(ETableColumnFlags::ShouldBeVisible |
+					ETableColumnFlags::CanBeFiltered |
+					ETableColumnFlags::IsHierarchy);
+
+	Column.SetHorizontalAlignment(HAlign_Left);
+	Column.SetInitialWidth(206.0f);
+	Column.SetMinWidth(42.0f);
+
+	Column.SetDataType(ETableCellDataType::Text);
+
+	TSharedPtr<ITableCellValueGetter> GetterPtr = MakeShareable(new FDisplayNameValueGetter());
+	Column.SetValueGetter(GetterPtr.ToSharedRef());
+
+	TSharedPtr<ITableCellValueFormatter> FormatterPtr = MakeShareable(new FTextValueFormatter());
+	Column.SetValueFormatter(FormatterPtr.ToSharedRef());
+
+	TSharedPtr<ITableCellValueSorter> SorterPtr = MakeShareable(new FSorterByName(ColumnPtr.ToSharedRef()));
+	Column.SetValueSorter(SorterPtr);
+
+	return ColumnPtr.ToSharedRef();
+}
+
+////////////////////////////////////////////////////////////////////////////////////////////////////
+
+TSharedRef<Insights::FTableColumn> FNetStatsViewColumnFactory::CreateTypeColumn()
+{
+	using namespace Insights;
+
+	TSharedPtr<FTableColumn> ColumnPtr = MakeShareable(new FTableColumn(FNetStatsViewColumns::TypeColumnID));
+	FTableColumn& Column = *ColumnPtr;
+
+	Column.SetShortName(LOCTEXT("Type_ColumnName", "Type"));
+	Column.SetTitleName(LOCTEXT("Type_ColumnTitle", "Type"));
+	Column.SetDescription(LOCTEXT("Type_ColumnDesc", "Type of timer or group"));
+
+	Column.SetFlags(ETableColumnFlags::CanBeHidden |
+					//ETableColumnFlags::ShouldBeVisible |
+					ETableColumnFlags::CanBeFiltered);
+
+	Column.SetHorizontalAlignment(HAlign_Left);
+	Column.SetInitialWidth(60.0f);
+
+	Column.SetDataType(ETableCellDataType::Text);
+
+	class FNetEventTypeValueGetter : public FTableCellValueGetter
+	{
+	public:
+		virtual const TOptional<FTableCellValue> GetValue(const FTableColumn& Column, const FBaseTreeNode& Node) const
+		{
+			ensure(Column.GetId() == FNetStatsViewColumns::TypeColumnID);
+			const FNetEventNode& NetEventNode = static_cast<const FNetEventNode&>(Node);
+			return TOptional<FTableCellValue>(FTableCellValue(NetEventNodeTypeHelper::ToText(NetEventNode.GetType())));
+		}
+	};
+
+	TSharedPtr<FTableCellValueGetter> GetterPtr = MakeShareable(new FNetEventTypeValueGetter());
+	Column.SetValueGetter(GetterPtr.ToSharedRef());
+
+	TSharedPtr<FTableCellValueFormatter> FormatterPtr = MakeShareable(new FTextValueFormatter());
+	Column.SetValueFormatter(FormatterPtr.ToSharedRef());
+
+	//TSharedPtr<ITableCellValueSorter> SorterPtr = MakeShareable(new FSorterByTextValue(ColumnPtr.ToSharedRef()));
+	TSharedPtr<ITableCellValueSorter> SorterPtr = MakeShareable(new FNetEventNodeSortingByEventType(ColumnPtr.ToSharedRef()));
+	Column.SetValueSorter(SorterPtr);
+
+	return ColumnPtr.ToSharedRef();
+}
+
+////////////////////////////////////////////////////////////////////////////////////////////////////
+
+TSharedRef<Insights::FTableColumn> FNetStatsViewColumnFactory::CreateLevelColumn()
+{
+	using namespace Insights;
+
+	TSharedPtr<FTableColumn> ColumnPtr = MakeShareable(new FTableColumn(FNetStatsViewColumns::LevelColumnID));
+	FTableColumn& Column = *ColumnPtr;
+
+	Column.SetShortName(LOCTEXT("Level_ColumnName", "Level"));
+	Column.SetTitleName(LOCTEXT("Level_ColumnTitle", "Level"));
+	Column.SetDescription(LOCTEXT("Level_ColumnDesc", "Level of net event instances"));
+
+	Column.SetFlags(ETableColumnFlags::CanBeHidden |
+					ETableColumnFlags::ShouldBeVisible |
+					ETableColumnFlags::CanBeFiltered);
+
+	Column.SetHorizontalAlignment(HAlign_Right);
+	Column.SetInitialWidth(40.0f);
+
+	Column.SetDataType(ETableCellDataType::Int64);
+
+	class FLevelValueGetter : public FTableCellValueGetter
+	{
+	public:
+		virtual const TOptional<FTableCellValue> GetValue(const FTableColumn& Column, const FBaseTreeNode& Node) const
+		{
+			ensure(Column.GetId() == FNetStatsViewColumns::LevelColumnID);
+			const FNetEventNode& NetEventNode = static_cast<const FNetEventNode&>(Node);
+			return TOptional<FTableCellValue>(FTableCellValue(static_cast<int64>(NetEventNode.GetLevel())));
+		}
+	};
+
+	TSharedPtr<FTableCellValueGetter> GetterPtr = MakeShareable(new FLevelValueGetter());
+	Column.SetValueGetter(GetterPtr.ToSharedRef());
+
+	TSharedPtr<FTableCellValueFormatter> FormatterPtr = MakeShareable(new FInt64ValueFormatterAsNumber());
+	Column.SetValueFormatter(FormatterPtr.ToSharedRef());
+
+	TSharedPtr<ITableCellValueSorter> SorterPtr = MakeShareable(new FSorterByInt64Value(ColumnPtr.ToSharedRef()));
+	Column.SetValueSorter(SorterPtr);
+
+	return ColumnPtr.ToSharedRef();
+}
+
+////////////////////////////////////////////////////////////////////////////////////////////////////
+
+TSharedRef<Insights::FTableColumn> FNetStatsViewColumnFactory::CreateInstanceCountColumn()
+{
+	using namespace Insights;
+
+	TSharedPtr<FTableColumn> ColumnPtr = MakeShareable(new FTableColumn(FNetStatsViewColumns::InstanceCountColumnID));
+	FTableColumn& Column = *ColumnPtr;
+
+	Column.SetShortName(LOCTEXT("InstanceCount_ColumnName", "Count"));
+	Column.SetTitleName(LOCTEXT("InstanceCount_ColumnTitle", "Instance Count"));
+	Column.SetDescription(LOCTEXT("InstanceCount_ColumnDesc", "Number of net event instances"));
+
+	Column.SetFlags(ETableColumnFlags::CanBeHidden |
+					ETableColumnFlags::ShouldBeVisible |
+					ETableColumnFlags::CanBeFiltered);
+
+	Column.SetHorizontalAlignment(HAlign_Right);
+	Column.SetInitialWidth(60.0f);
+
+	Column.SetDataType(ETableCellDataType::Int64);
+
+	class FInstanceCountValueGetter : public FTableCellValueGetter
+	{
+	public:
+		virtual const TOptional<FTableCellValue> GetValue(const FTableColumn& Column, const FBaseTreeNode& Node) const
+		{
+			ensure(Column.GetId() == FNetStatsViewColumns::InstanceCountColumnID);
+			const FNetEventNode& NetEventNode = static_cast<const FNetEventNode&>(Node);
+			return TOptional<FTableCellValue>(FTableCellValue(static_cast<int64>(NetEventNode.GetAggregatedStats().InstanceCount)));
+		}
+	};
+
+	TSharedPtr<FTableCellValueGetter> GetterPtr = MakeShareable(new FInstanceCountValueGetter());
+	Column.SetValueGetter(GetterPtr.ToSharedRef());
+
+	TSharedPtr<FTableCellValueFormatter> FormatterPtr = MakeShareable(new FInt64ValueFormatterAsNumber());
+	Column.SetValueFormatter(FormatterPtr.ToSharedRef());
+
+	//TSharedPtr<ITableCellValueSorter> SorterPtr = MakeShareable(new FSorterByInt64Value(ColumnPtr.ToSharedRef()));
+	TSharedPtr<ITableCellValueSorter> SorterPtr = MakeShareable(new FNetEventNodeSortingByInstanceCount(ColumnPtr.ToSharedRef()));
+	Column.SetValueSorter(SorterPtr);
+
+	return ColumnPtr.ToSharedRef();
+}
+
+////////////////////////////////////////////////////////////////////////////////////////////////////
+// Inclusive  Columns
+////////////////////////////////////////////////////////////////////////////////////////////////////
+
+TSharedRef<Insights::FTableColumn> FNetStatsViewColumnFactory::CreateTotalInclusiveSizeColumn()
+{
+	using namespace Insights;
+
+	TSharedPtr<FTableColumn> ColumnPtr = MakeShareable(new FTableColumn(FNetStatsViewColumns::TotalInclusiveSizeColumnID));
+	FTableColumn& Column = *ColumnPtr;
+
+	Column.SetShortName(LOCTEXT("TotalInclusive_ColumnName", "Incl"));
+	Column.SetTitleName(LOCTEXT("TotalInclusive_ColumnTitle", "Total Inclusive Size"));
+	Column.SetDescription(LOCTEXT("TotalInclusive_ColumnDesc", "Total inclusive size (bits) of selected net event instances"));
+
+	Column.SetFlags(ETableColumnFlags::CanBeHidden |
+					ETableColumnFlags::ShouldBeVisible |
+					ETableColumnFlags::CanBeFiltered);
+
+	Column.SetHorizontalAlignment(HAlign_Right);
+	Column.SetInitialWidth(TotalSizeColumnInitialWidth);
+
+	Column.SetDataType(ETableCellDataType::Int64);
+	Column.SetAggregation(ETableColumnAggregation::Sum);
+
+	class FTotalInclusiveValueGetter : public FTableCellValueGetter
+	{
+	public:
+		virtual const TOptional<FTableCellValue> GetValue(const FTableColumn& Column, const FBaseTreeNode& Node) const
+		{
+			ensure(Column.GetId() == FNetStatsViewColumns::TotalInclusiveSizeColumnID);
+			const FNetEventNode& NetEventNode = static_cast<const FNetEventNode&>(Node);
+			return TOptional<FTableCellValue>(FTableCellValue(static_cast<int64>(NetEventNode.GetAggregatedStats().TotalInclusive)));
+		}
+	};
+
+	TSharedPtr<FTableCellValueGetter> GetterPtr = MakeShareable(new FTotalInclusiveValueGetter());
+	Column.SetValueGetter(GetterPtr.ToSharedRef());
+
+	TSharedPtr<FTableCellValueFormatter> FormatterPtr = MakeShareable(new FInt64ValueFormatterAsNumber());
+	Column.SetValueFormatter(FormatterPtr.ToSharedRef());
+
+	//TSharedPtr<ITableCellValueSorter> SorterPtr = MakeShareable(new FSorterByInt64Value(ColumnPtr.ToSharedRef()));
+	TSharedPtr<ITableCellValueSorter> SorterPtr = MakeShareable(new FNetEventNodeSortingByTotalInclusiveSize(ColumnPtr.ToSharedRef()));
+	Column.SetValueSorter(SorterPtr);
+
+	return ColumnPtr.ToSharedRef();
+}
+
+////////////////////////////////////////////////////////////////////////////////////////////////////
+
+TSharedRef<Insights::FTableColumn> FNetStatsViewColumnFactory::CreateMaxInclusiveSizeColumn()
+{
+	using namespace Insights;
+
+	TSharedPtr<FTableColumn> ColumnPtr = MakeShareable(new FTableColumn(FNetStatsViewColumns::MaxInclusiveSizeColumnID));
+	FTableColumn& Column = *ColumnPtr;
+
+	Column.SetShortName(LOCTEXT("MaxInclusive_ColumnName", "I.Max"));
+	Column.SetTitleName(LOCTEXT("MaxInclusive_ColumnTitle", "Max Inclusive Size"));
+	Column.SetDescription(LOCTEXT("MaxInclusive_ColumnDesc", "Maximum inclusive size (bits) of selected net event instances"));
+
+	Column.SetFlags(ETableColumnFlags::CanBeHidden |
+					//ETableColumnFlags::ShouldBeVisible |
+					ETableColumnFlags::CanBeFiltered);
+
+	Column.SetHorizontalAlignment(HAlign_Right);
+	Column.SetInitialWidth(SizeColumnInitialWidth);
+
+	Column.SetDataType(ETableCellDataType::Int64);
+
+	class FMaxInclusiveValueGetter : public FTableCellValueGetter
+	{
+	public:
+		virtual const TOptional<FTableCellValue> GetValue(const FTableColumn& Column, const FBaseTreeNode& Node) const
+		{
+			ensure(Column.GetId() == FNetStatsViewColumns::MaxInclusiveSizeColumnID);
+			const FNetEventNode& NetEventNode = static_cast<const FNetEventNode&>(Node);
+			return TOptional<FTableCellValue>(FTableCellValue(static_cast<int64>(NetEventNode.GetAggregatedStats().MaxInclusive)));
+		}
+	};
+
+	TSharedPtr<FTableCellValueGetter> GetterPtr = MakeShareable(new FMaxInclusiveValueGetter());
+	Column.SetValueGetter(GetterPtr.ToSharedRef());
+
+	TSharedPtr<FTableCellValueFormatter> FormatterPtr = MakeShareable(new FInt64ValueFormatterAsNumber());
+	Column.SetValueFormatter(FormatterPtr.ToSharedRef());
+
+	TSharedPtr<ITableCellValueSorter> SorterPtr = MakeShareable(new FSorterByInt64Value(ColumnPtr.ToSharedRef()));
+	Column.SetValueSorter(SorterPtr);
+
+	return ColumnPtr.ToSharedRef();
+}
+
+////////////////////////////////////////////////////////////////////////////////////////////////////
+
+TSharedRef<Insights::FTableColumn> FNetStatsViewColumnFactory::CreateAverageInclusiveSizeColumn()
+{
+	using namespace Insights;
+
+	TSharedPtr<FTableColumn> ColumnPtr = MakeShareable(new FTableColumn(FNetStatsViewColumns::AverageInclusiveSizeColumnID));
+	FTableColumn& Column = *ColumnPtr;
+
+	Column.SetShortName(LOCTEXT("AvgInclusive_ColumnName", "I.Avg"));
+	Column.SetTitleName(LOCTEXT("AvgInclusive_ColumnTitle", "Average Inclusive Size"));
+	Column.SetDescription(LOCTEXT("AvgInclusive_ColumnDesc", "Average inclusive size (bits) of selected net event instances"));
+
+	Column.SetFlags(ETableColumnFlags::CanBeHidden |
+					//ETableColumnFlags::ShouldBeVisible |
+					ETableColumnFlags::CanBeFiltered);
+
+	Column.SetHorizontalAlignment(HAlign_Right);
+	Column.SetInitialWidth(SizeColumnInitialWidth);
+
+	Column.SetDataType(ETableCellDataType::Int64);
+
+	class FAverageInclusiveValueGetter : public FTableCellValueGetter
+	{
+	public:
+		virtual const TOptional<FTableCellValue> GetValue(const FTableColumn& Column, const FBaseTreeNode& Node) const
+		{
+			ensure(Column.GetId() == FNetStatsViewColumns::AverageInclusiveSizeColumnID);
+			const FNetEventNode& NetEventNode = static_cast<const FNetEventNode&>(Node);
+			return TOptional<FTableCellValue>(FTableCellValue(static_cast<int64>(NetEventNode.GetAggregatedStats().AverageInclusive)));
+		}
+	};
+
+	TSharedPtr<FTableCellValueGetter> GetterPtr = MakeShareable(new FAverageInclusiveValueGetter());
+	Column.SetValueGetter(GetterPtr.ToSharedRef());
+
+	TSharedPtr<FTableCellValueFormatter> FormatterPtr = MakeShareable(new FInt64ValueFormatterAsNumber());
+	Column.SetValueFormatter(FormatterPtr.ToSharedRef());
+
+	TSharedPtr<ITableCellValueSorter> SorterPtr = MakeShareable(new FSorterByInt64Value(ColumnPtr.ToSharedRef()));
+	Column.SetValueSorter(SorterPtr);
+
+	return ColumnPtr.ToSharedRef();
+}
+
+////////////////////////////////////////////////////////////////////////////////////////////////////
+// Exclusive  Columns
+////////////////////////////////////////////////////////////////////////////////////////////////////
+
+TSharedRef<Insights::FTableColumn> FNetStatsViewColumnFactory::CreateTotalExclusiveSizeColumn()
+{
+	using namespace Insights;
+
+	TSharedPtr<FTableColumn> ColumnPtr = MakeShareable(new FTableColumn(FNetStatsViewColumns::TotalExclusiveSizeColumnID));
+	FTableColumn& Column = *ColumnPtr;
+
+	Column.SetShortName(LOCTEXT("TotalExclusive_ColumnName", "Excl"));
+	Column.SetTitleName(LOCTEXT("TotalExclusive_ColumnTitle", "Total Exclusive Size"));
+	Column.SetDescription(LOCTEXT("TotalExclusive_ColumnDesc", "Total exclusive size (bits) of selected net event instances"));
+
+	Column.SetFlags(ETableColumnFlags::CanBeHidden |
+					ETableColumnFlags::ShouldBeVisible |
+					ETableColumnFlags::CanBeFiltered);
+
+	Column.SetHorizontalAlignment(HAlign_Right);
+	Column.SetInitialWidth(TotalSizeColumnInitialWidth);
+
+	Column.SetDataType(ETableCellDataType::Int64);
+	Column.SetAggregation(ETableColumnAggregation::Sum);
+
+	class FTotalExclusiveValueGetter : public FTableCellValueGetter
+	{
+	public:
+		virtual const TOptional<FTableCellValue> GetValue(const FTableColumn& Column, const FBaseTreeNode& Node) const
+		{
+			ensure(Column.GetId() == FNetStatsViewColumns::TotalExclusiveSizeColumnID);
+			const FNetEventNode& NetEventNode = static_cast<const FNetEventNode&>(Node);
+			return TOptional<FTableCellValue>(FTableCellValue(static_cast<int64>(NetEventNode.GetAggregatedStats().TotalExclusive)));
+		}
+	};
+
+	TSharedPtr<FTableCellValueGetter> GetterPtr = MakeShareable(new FTotalExclusiveValueGetter());
+	Column.SetValueGetter(GetterPtr.ToSharedRef());
+
+	TSharedPtr<FTableCellValueFormatter> FormatterPtr = MakeShareable(new FInt64ValueFormatterAsNumber());
+	Column.SetValueFormatter(FormatterPtr.ToSharedRef());
+
+	//TSharedPtr<ITableCellValueSorter> SorterPtr = MakeShareable(new FSorterByInt64Value(ColumnPtr.ToSharedRef()));
+	TSharedPtr<ITableCellValueSorter> SorterPtr = MakeShareable(new FNetEventNodeSortingByTotalExclusiveSize(ColumnPtr.ToSharedRef()));
+	Column.SetValueSorter(SorterPtr);
+
+	return ColumnPtr.ToSharedRef();
+}
+
+////////////////////////////////////////////////////////////////////////////////////////////////////
+
+TSharedRef<Insights::FTableColumn> FNetStatsViewColumnFactory::CreateMaxExclusiveSizeColumn()
+{
+	using namespace Insights;
+
+	TSharedPtr<FTableColumn> ColumnPtr = MakeShareable(new FTableColumn(FNetStatsViewColumns::MaxExclusiveSizeColumnID));
+	FTableColumn& Column = *ColumnPtr;
+
+	Column.SetShortName(LOCTEXT("MaxExclusive_ColumnName", "E.Max"));
+	Column.SetTitleName(LOCTEXT("MaxExclusive_ColumnTitle", "Max Exclusive Size"));
+	Column.SetDescription(LOCTEXT("MaxExclusive_ColumnDesc", "Maximum exclusive size of selected net event instances"));
+
+	Column.SetFlags(ETableColumnFlags::CanBeHidden |
+					//ETableColumnFlags::ShouldBeVisible |
+					ETableColumnFlags::CanBeFiltered);
+
+	Column.SetHorizontalAlignment(HAlign_Right);
+	Column.SetInitialWidth(SizeColumnInitialWidth);
+
+	Column.SetDataType(ETableCellDataType::Int64);
+
+	class FMaxExclusiveValueGetter : public FTableCellValueGetter
+	{
+	public:
+		virtual const TOptional<FTableCellValue> GetValue(const FTableColumn& Column, const FBaseTreeNode& Node) const
+		{
+			ensure(Column.GetId() == FNetStatsViewColumns::MaxExclusiveSizeColumnID);
+			const FNetEventNode& NetEventNode = static_cast<const FNetEventNode&>(Node);
+			return TOptional<FTableCellValue>(FTableCellValue(static_cast<int64>(NetEventNode.GetAggregatedStats().MaxExclusive)));
+		}
+	};
+
+	TSharedPtr<FTableCellValueGetter> GetterPtr = MakeShareable(new FMaxExclusiveValueGetter());
+	Column.SetValueGetter(GetterPtr.ToSharedRef());
+
+	TSharedPtr<FTableCellValueFormatter> FormatterPtr = MakeShareable(new FInt64ValueFormatterAsNumber());
+	Column.SetValueFormatter(FormatterPtr.ToSharedRef());
+
+	TSharedPtr<ITableCellValueSorter> SorterPtr = MakeShareable(new FSorterByInt64Value(ColumnPtr.ToSharedRef()));
+	Column.SetValueSorter(SorterPtr);
+
+	return ColumnPtr.ToSharedRef();
+}
+
+////////////////////////////////////////////////////////////////////////////////////////////////////
+
+TSharedRef<Insights::FTableColumn> FNetStatsViewColumnFactory::CreateAverageExclusiveSizeColumn()
+{
+	using namespace Insights;
+
+	TSharedPtr<FTableColumn> ColumnPtr = MakeShareable(new FTableColumn(FNetStatsViewColumns::AverageExclusiveSizeColumnID));
+	FTableColumn& Column = *ColumnPtr;
+
+	Column.SetShortName(LOCTEXT("AvgExclusiveName", "E.Avg"));
+	Column.SetTitleName(LOCTEXT("AvgExclusiveTitle", "Average Exclusive  (ms)"));
+	Column.SetDescription(LOCTEXT("AvgExclusiveDesc", "Average exclusive size of selected net event instances"));
+
+	Column.SetFlags(ETableColumnFlags::CanBeHidden |
+					//ETableColumnFlags::ShouldBeVisible |
+					ETableColumnFlags::CanBeFiltered);
+
+	Column.SetHorizontalAlignment(HAlign_Right);
+	Column.SetInitialWidth(SizeColumnInitialWidth);
+
+	Column.SetDataType(ETableCellDataType::Int64);
+
+	class FAverageExclusiveValueGetter : public FTableCellValueGetter
+	{
+	public:
+		virtual const TOptional<FTableCellValue> GetValue(const FTableColumn& Column, const FBaseTreeNode& Node) const
+		{
+			ensure(Column.GetId() == FNetStatsViewColumns::AverageExclusiveSizeColumnID);
+			const FNetEventNode& NetEventNode = static_cast<const FNetEventNode&>(Node);
+			//return TOptional<FTableCellValue>(FTableCellValue(static_cast<int64>(NetEventNode.GetAggregatedStats().AverageExclusive)));
+			return TOptional<FTableCellValue>(FTableCellValue(int64(0)));
+		}
+	};
+
+	TSharedPtr<FTableCellValueGetter> GetterPtr = MakeShareable(new FAverageExclusiveValueGetter());
+	Column.SetValueGetter(GetterPtr.ToSharedRef());
+
+	TSharedPtr<FTableCellValueFormatter> FormatterPtr = MakeShareable(new FInt64ValueFormatterAsNumber());
+	Column.SetValueFormatter(FormatterPtr.ToSharedRef());
+
+	TSharedPtr<ITableCellValueSorter> SorterPtr = MakeShareable(new FSorterByInt64Value(ColumnPtr.ToSharedRef()));
+	Column.SetValueSorter(SorterPtr);
+
+	return ColumnPtr.ToSharedRef();
+}
+
+////////////////////////////////////////////////////////////////////////////////////////////////////
+
+#undef LOCTEXT_NAMESPACE
diff --git a/Engine/Source/Developer/TraceInsights/Private/Insights/NetworkingProfiler/ViewModels/NetStatsViewColumnFactory.h b/Engine/Source/Developer/TraceInsights/Private/Insights/NetworkingProfiler/ViewModels/NetStatsViewColumnFactory.h
new file mode 100644
--- /dev/null
+++ b/Engine/Source/Developer/TraceInsights/Private/Insights/NetworkingProfiler/ViewModels/NetStatsViewColumnFactory.h
@@ -0,0 +1,62 @@
+// Copyright 1998-2019 Epic Games, Inc. All Rights Reserved.
+
+#pragma once
+
+#include "CoreMinimal.h"
+
+////////////////////////////////////////////////////////////////////////////////////////////////////
+
+struct FNetStatsViewColumns
+{
+	//////////////////////////////////////////////////
+	// Column identifiers
+
+	static const FName NameColumnID;
+	static const FName TypeColumnID;
+	static const FName LevelColumnID;
+	static const FName InstanceCountColumnID;
+
+	// Inclusive  columns
+	static const FName TotalInclusiveSizeColumnID;
+	static const FName MaxInclusiveSizeColumnID;
+	static const FName AverageInclusiveSizeColumnID;
+
+	// Exclusive  columns
+	static const FName TotalExclusiveSizeColumnID;
+	static const FName MaxExclusiveSizeColumnID;
+	static const FName AverageExclusiveSizeColumnID;
+};
+
+////////////////////////////////////////////////////////////////////////////////////////////////////
+
+namespace Insights
+{
+	class FTableColumn;
+}
+
+////////////////////////////////////////////////////////////////////////////////////////////////////
+
+struct FNetStatsViewColumnFactory
+{
+public:
+	static void CreateNetStatsViewColumns(TArray<TSharedPtr<Insights::FTableColumn>>& Columns);
+
+	static TSharedRef<Insights::FTableColumn> CreateNameColumn();
+	static TSharedRef<Insights::FTableColumn> CreateTypeColumn();
+	static TSharedRef<Insights::FTableColumn> CreateLevelColumn();
+	static TSharedRef<Insights::FTableColumn> CreateInstanceCountColumn();
+
+	static TSharedRef<Insights::FTableColumn> CreateTotalInclusiveSizeColumn();
+	static TSharedRef<Insights::FTableColumn> CreateMaxInclusiveSizeColumn();
+	static TSharedRef<Insights::FTableColumn> CreateAverageInclusiveSizeColumn();
+
+	static TSharedRef<Insights::FTableColumn> CreateTotalExclusiveSizeColumn();
+	static TSharedRef<Insights::FTableColumn> CreateMaxExclusiveSizeColumn();
+	static TSharedRef<Insights::FTableColumn> CreateAverageExclusiveSizeColumn();
+
+private:
+	static constexpr float TotalSizeColumnInitialWidth = 60.0f;
+	static constexpr float SizeColumnInitialWidth = 50.0f;
+};
+
+////////////////////////////////////////////////////////////////////////////////////////////////////
diff --git a/Engine/Source/Developer/TraceInsights/Private/Insights/NetworkingProfiler/ViewModels/PacketBreakdownViewHelper.cpp b/Engine/Source/Developer/TraceInsights/Private/Insights/NetworkingProfiler/ViewModels/PacketBreakdownViewHelper.cpp
--- a/Engine/Source/Developer/TraceInsights/Private/Insights/NetworkingProfiler/ViewModels/PacketBreakdownViewHelper.cpp
+++ b/Engine/Source/Developer/TraceInsights/Private/Insights/NetworkingProfiler/ViewModels/PacketBreakdownViewHelper.cpp
@@ -35,16 +35,16 @@ FPacketContentViewDrawStateBuilder::FPacketContentViewDrawStateBuilder(FPacketCo
 
 ////////////////////////////////////////////////////////////////////////////////////////////////////
 
-void FPacketContentViewDrawStateBuilder::AddEvent(const Trace::FNetProfilerContentEvent& Event, const TCHAR* EventName)
+void FPacketContentViewDrawStateBuilder::AddEvent(const Trace::FNetProfilerContentEvent& Event, const TCHAR* EventName, uint32 NetId)
 {
-	const int32 Depth = Event.Level;
-
 	DrawState.Events.AddUninitialized();
-	FPacketContentViewDrawState::FPacketEvent& PacketEvent = DrawState.Events.Last();
-	PacketEvent.Offset = Event.StartPos;
-	PacketEvent.Size = Event.EndPos - Event.StartPos;
-	PacketEvent.Type = Event.NameIndex;
-	PacketEvent.Depth = Depth;
+	FNetworkPacketEvent& PacketEvent = DrawState.Events.Last();
+	PacketEvent.EventTypeIndex = Event.EventTypeIndex;
+	PacketEvent.ObjectInstanceIndex = Event.ObjectInstanceIndex;
+	PacketEvent.NetId = NetId;
+	PacketEvent.BitOffset = Event.StartPos;
+	PacketEvent.BitSize = Event.EndPos - Event.StartPos;
+	PacketEvent.Level = Event.Level;
 
 	const FAxisViewportDouble& ViewportX = Viewport.GetHorizontalAxisViewport();
 
@@ -71,6 +71,7 @@ void FPacketContentViewDrawStateBuilder::AddEvent(const Trace::FNetProfilerConte
 		EventX2 = EventX1 + 1.0f;
 	}
 
+	const int32 Depth = static_cast<int32>(Event.Level);
 	if (Depth > MaxDepth)
 	{
 		MaxDepth = Depth;
@@ -205,7 +206,7 @@ void FPacketContentViewDrawStateBuilder::AddEvent(const Trace::FNetProfilerConte
 			if (Event.ObjectInstanceIndex != 0)
 			{
 				Name += TEXT(" (NetId:");
-				Name += FText::AsNumber(Event.ObjectInstanceIndex).ToString();
+				Name += FText::AsNumber(NetId).ToString();
 				Name += TEXT(", ");
 			}
 			else
@@ -373,17 +374,17 @@ void FPacketContentViewDrawHelper::Draw(const FPacketContentViewDrawState& DrawS
 
 ////////////////////////////////////////////////////////////////////////////////////////////////////
 
-void FPacketContentViewDrawHelper::DrawEventHighlight(const FNetworkPacketEvent& Event) const
+void FPacketContentViewDrawHelper::DrawEventHighlight(const FNetworkPacketEvent& Event, EHighlightMode Mode) const
 {
 	const FAxisViewportDouble& ViewportX = Viewport.GetHorizontalAxisViewport();
 
-	float EventX1 = ViewportX.GetRoundedOffsetForValue(static_cast<double>(Event.Offset));
+	float EventX1 = ViewportX.GetRoundedOffsetForValue(static_cast<double>(Event.BitOffset));
 	if (EventX1 > Viewport.GetWidth())
 	{
 		return;
 	}
 
-	float EventX2 = ViewportX.GetRoundedOffsetForValue(static_cast<double>(Event.Offset + Event.Size));
+	float EventX2 = ViewportX.GetRoundedOffsetForValue(static_cast<double>(Event.BitOffset + Event.BitSize));
 	if (EventX2 < 0)
 	{
 		return;
@@ -415,11 +416,29 @@ void FPacketContentViewDrawHelper::DrawEventHighlight(const FNetworkPacketEvent&
 	constexpr float EventH = 14.0f;
 	constexpr float EventDY = 2.0f;
 
-	const float EventY = Y0 + (EventH + EventDY) * Event.Depth;
+	const float EventY = Y0 + (EventH + EventDY) * Event.Level;
+
+	if (Mode == EHighlightMode::Hovered)
+	{
+		const FLinearColor Color(1.0f, 1.0f, 0.0f, 1.0f); // yellow
 
-	const FLinearColor Color(1.0f, 1.0f, 0.0f, 1.0f);
-	DrawContext.DrawBox(EventX1 - 2.0f, EventY - 2.0f, EventW + 4.0f, EventH + 4.0f, HoveredEventBorderBrush, Color);
+		// Draw border around the timing event box.
+		DrawContext.DrawBox(EventX1 - 2.0f, EventY - 2.0f, EventW + 4.0f, EventH + 4.0f, HoveredEventBorderBrush, Color);
+	}
+	else // EHighlightMode::Selected or EHighlightMode::SelectedAndHovered
+	{
+		// Animate color from white (if selected and hovered) or yellow (if only selected) to black, using a squared sine function.
+		const double Time = static_cast<double>(FPlatformTime::Cycles64()) * FPlatformTime::GetSecondsPerCycle64();
+		float S = FMath::Sin(2.0 * Time);
+		S = S * S; // squared, to ensure only positive [0 - 1] values
+		const float Blue = (Mode == EHighlightMode::SelectedAndHovered) ? 0.0f : S;
+		const FLinearColor Color(S, S, Blue, 1.0f);
+
+		// Draw border around the timing event box.
+		DrawContext.DrawBox(EventX1 - 2.0f, EventY - 2.0f, EventW + 4.0f, EventH + 4.0f, SelectedEventBorderBrush, Color);
+	}
 	DrawContext.LayerId++;
+
 }
 
 ////////////////////////////////////////////////////////////////////////////////////////////////////
diff --git a/Engine/Source/Developer/TraceInsights/Private/Insights/NetworkingProfiler/ViewModels/PacketBreakdownViewHelper.h b/Engine/Source/Developer/TraceInsights/Private/Insights/NetworkingProfiler/ViewModels/PacketBreakdownViewHelper.h
--- a/Engine/Source/Developer/TraceInsights/Private/Insights/NetworkingProfiler/ViewModels/PacketBreakdownViewHelper.h
+++ b/Engine/Source/Developer/TraceInsights/Private/Insights/NetworkingProfiler/ViewModels/PacketBreakdownViewHelper.h
@@ -20,32 +20,60 @@ class FSlateWindowElementList;
 
 struct FNetworkPacketEvent
 {
-	int64 Offset;
-	int64 Size;
-	int32 Type;
-	int32 Depth;
-
-	FNetworkPacketEvent(int64 InOffset, int64 InSize, int32 InType, int32 InDepth)
-		: Offset(InOffset), Size(InSize), Type(InType), Depth(InDepth)
+	uint32 EventTypeIndex;
+	uint32 ObjectInstanceIndex;
+	uint32 NetId;
+	uint32 BitOffset;
+	uint32 BitSize;
+	uint32 Level;
+
+	FNetworkPacketEvent()
+		: EventTypeIndex(0)
+		, ObjectInstanceIndex(0)
+		, NetId(0)
+		, BitOffset(0)
+		, BitSize(0)
+		, Level(0)
 	{}
 
+	FNetworkPacketEvent(uint32 InEventTypeIndex, uint32 InObjectInstanceIndex, uint32 InNetId, uint32 InBitOffset, uint32 InBitSize, uint32 InLevel)
+		: EventTypeIndex(InEventTypeIndex)
+		, ObjectInstanceIndex(InObjectInstanceIndex)
+		, NetId(InNetId)
+		, BitOffset(InBitOffset)
+		, BitSize(InBitSize)
+		, Level(InLevel)
+	{}
+	
 	FNetworkPacketEvent(const FNetworkPacketEvent& Other)
-		: Offset(Other.Offset), Size(Other.Size), Type(Other.Type), Depth(Other.Depth)
+		: EventTypeIndex(Other.EventTypeIndex)
+		, ObjectInstanceIndex(Other.ObjectInstanceIndex)
+		, NetId(Other.NetId)
+		, BitOffset(Other.BitOffset)
+		, BitSize(Other.BitSize)
+		, Level(Other.Level)
 	{
 	}
 
 	FNetworkPacketEvent& operator=(const FNetworkPacketEvent& Other)
 	{
-		Offset = Other.Offset;
-		Size = Other.Size;
-		Type = Other.Type;
-		Depth = Other.Depth;
+		EventTypeIndex = Other.EventTypeIndex;
+		ObjectInstanceIndex = Other.ObjectInstanceIndex;
+		NetId = Other.NetId;
+		BitOffset = Other.BitOffset;
+		BitSize = Other.BitSize;
+		Level = Other.Level;
 		return *this;
 	}
 
 	bool Equals(const FNetworkPacketEvent& Other) const
 	{
-		return Offset == Other.Offset && Size == Other.Size && Type == Other.Type && Depth == Other.Depth;
+		return EventTypeIndex == Other.EventTypeIndex
+			&& ObjectInstanceIndex == Other.ObjectInstanceIndex
+			&& NetId == Other.NetId
+			&& BitOffset == Other.BitOffset
+			&& BitSize == Other.BitSize
+			&& Level == Other.Level;
 	}
 
 	static bool AreEquals(const FNetworkPacketEvent& A, const FNetworkPacketEvent& B)
@@ -58,14 +86,6 @@ struct FNetworkPacketEvent
 
 struct FPacketContentViewDrawState
 {
-	struct FPacketEvent
-	{
-		int64 Offset; // [bit]
-		int64 Size; // [bit]
-		int32 Type;
-		int32 Depth;
-	};
-
 	struct FBox
 	{
 		float X;
@@ -98,7 +118,7 @@ struct FPacketContentViewDrawState
 
 	int32 GetNumMergedBoxes() const { return NumMergedBoxes; }
 
-	TArray<FPacketEvent> Events;
+	TArray<FNetworkPacketEvent> Events;
 	TArray<FBox> Boxes;
 	TArray<FBox> Borders;
 	TArray<FText> Texts;
@@ -132,7 +152,7 @@ public:
 	FPacketContentViewDrawStateBuilder(const FPacketContentViewDrawStateBuilder&) = delete;
 	FPacketContentViewDrawStateBuilder& operator=(const FPacketContentViewDrawStateBuilder&) = delete;
 
-	void AddEvent(const Trace::FNetProfilerContentEvent& Event, const TCHAR* Name = nullptr);
+	void AddEvent(const Trace::FNetProfilerContentEvent& Event, const TCHAR* Name, uint32 NetId);
 	void Flush();
 
 private:
@@ -176,7 +196,7 @@ public:
 
 	void DrawBackground() const;
 	void Draw(const FPacketContentViewDrawState& DrawState) const;
-	void DrawEventHighlight(const FNetworkPacketEvent& Event) const;
+	void DrawEventHighlight(const FNetworkPacketEvent& Event, EHighlightMode Mode) const;
 
 	static FLinearColor GetColorByType(int32 Type);
 
diff --git a/Engine/Source/Developer/TraceInsights/Private/Insights/NetworkingProfiler/Widgets/SDataStreamBreakdownView.cpp b/Engine/Source/Developer/TraceInsights/Private/Insights/NetworkingProfiler/Widgets/SDataStreamBreakdownView.cpp
deleted file mode 100644
diff --git a/Engine/Source/Developer/TraceInsights/Private/Insights/NetworkingProfiler/Widgets/SDataStreamBreakdownView.h b/Engine/Source/Developer/TraceInsights/Private/Insights/NetworkingProfiler/Widgets/SDataStreamBreakdownView.h
deleted file mode 100644
diff --git a/Engine/Source/Developer/TraceInsights/Private/Insights/NetworkingProfiler/Widgets/SNetStatsTableCell.cpp b/Engine/Source/Developer/TraceInsights/Private/Insights/NetworkingProfiler/Widgets/SNetStatsTableCell.cpp
new file mode 100644
--- /dev/null
+++ b/Engine/Source/Developer/TraceInsights/Private/Insights/NetworkingProfiler/Widgets/SNetStatsTableCell.cpp
@@ -0,0 +1,131 @@
+// Copyright 1998-2019 Epic Games, Inc. All Rights Reserved.
+
+#include "SNetStatsTableCell.h"
+
+#include "EditorStyleSet.h"
+#include "SlateOptMacros.h"
+#include "Widgets/Images/SImage.h"
+#include "Widgets/Input/SButton.h"
+#include "Widgets/SBoxPanel.h"
+#include "Widgets/SOverlay.h"
+#include "Widgets/SToolTip.h"
+#include "Widgets/Text/STextBlock.h"
+#include "Widgets/Views/SExpanderArrow.h"
+
+// Insights
+#include "Insights/Table/ViewModels/Table.h"
+#include "Insights/Table/ViewModels/TableColumn.h"
+#include "Insights/NetworkingProfiler/Widgets/SNetStatsViewTooltip.h"
+
+#define LOCTEXT_NAMESPACE "SNetStatsView"
+
+////////////////////////////////////////////////////////////////////////////////////////////////////
+
+BEGIN_SLATE_FUNCTION_BUILD_OPTIMIZATION
+
+void SNetStatsTableCell::Construct(const FArguments& InArgs, const TSharedRef<class ITableRow>& TableRow)
+{
+	TablePtr = InArgs._TablePtr;
+	ColumnPtr = InArgs._ColumnPtr;
+	NetEventNodePtr = InArgs._NetEventNodePtr;
+
+	ensure(TablePtr.IsValid());
+	ensure(ColumnPtr.IsValid());
+	ensure(NetEventNodePtr.IsValid());
+
+	SetHoveredCellDelegate = InArgs._OnSetHoveredCell;
+
+	ChildSlot
+		[
+			GenerateWidgetForColumn(InArgs, TableRow)
+		];
+}
+
+////////////////////////////////////////////////////////////////////////////////////////////////////
+
+TSharedRef<SWidget> SNetStatsTableCell::GenerateWidgetForColumn(const FArguments& InArgs, const TSharedRef<class ITableRow>& TableRow)
+{
+	if (InArgs._IsNameColumn)
+	{
+		return GenerateWidgetForNameColumn(InArgs, TableRow);
+	}
+	else
+	{
+		return GenerateWidgetForStatsColumn(InArgs, TableRow);
+	}
+}
+
+////////////////////////////////////////////////////////////////////////////////////////////////////
+
+TSharedRef<SWidget> SNetStatsTableCell::GenerateWidgetForNameColumn(const FArguments& InArgs, const TSharedRef<class ITableRow>& TableRow)
+{
+	return
+		SNew(SHorizontalBox)
+
+		+ SHorizontalBox::Slot()
+		.AutoWidth()
+		.HAlign(HAlign_Right)
+		.VAlign(VAlign_Center)
+		[
+			SNew(SExpanderArrow, TableRow)
+		]
+
+		// Info icon + tooltip
+		+ SHorizontalBox::Slot()
+		.AutoWidth()
+		.HAlign(HAlign_Center)
+		.VAlign(VAlign_Center)
+		[
+			SNew(SImage)
+			.Visibility(this, &SNetStatsTableCell::GetHintIconVisibility)
+			.Image(FEditorStyle::GetBrush("Profiler.Tooltip.HintIcon10"))
+			.ToolTip(SNetStatsViewTooltip::GetCellTooltip(NetEventNodePtr, ColumnPtr))
+		]
+
+		// Name
+		+ SHorizontalBox::Slot()
+		.AutoWidth()
+		.VAlign(VAlign_Center)
+		.HAlign(ColumnPtr->GetHorizontalAlignment())
+		.Padding(FMargin(2.0f, 0.0f))
+		[
+			SNew(STextBlock)
+			.Text(this, &SNetStatsTableCell::GetDisplayName)
+			.HighlightText(InArgs._HighlightText)
+			.TextStyle(FEditorStyle::Get(), TEXT("Profiler.Tooltip"))
+			.ColorAndOpacity(this, &SNetStatsTableCell::GetColorAndOpacity)
+			.ShadowColorAndOpacity(this, &SNetStatsTableCell::GetShadowColorAndOpacity)
+		]
+	;
+}
+
+////////////////////////////////////////////////////////////////////////////////////////////////////
+
+TSharedRef<SWidget> SNetStatsTableCell::GenerateWidgetForStatsColumn(const FArguments& InArgs, const TSharedRef<class ITableRow>& TableRow)
+{
+	const FText CellText = NetEventNodePtr->IsGroup() ? FText::GetEmpty() : ColumnPtr->GetValueAsText(*NetEventNodePtr);
+
+	return
+		SNew(SHorizontalBox)
+
+		// Value
+		+ SHorizontalBox::Slot()
+		.FillWidth(1.0f)
+		.VAlign(VAlign_Center)
+		.HAlign(ColumnPtr->GetHorizontalAlignment())
+		.Padding(FMargin(2.0f, 0.0f))
+		[
+			SNew(STextBlock)
+			.Text(CellText)
+			.TextStyle(FEditorStyle::Get(), TEXT("Profiler.Tooltip"))
+			.ColorAndOpacity(this, &SNetStatsTableCell::GetStatsColorAndOpacity)
+			.ShadowColorAndOpacity(this, &SNetStatsTableCell::GetShadowColorAndOpacity)
+		]
+	;
+}
+
+END_SLATE_FUNCTION_BUILD_OPTIMIZATION
+
+////////////////////////////////////////////////////////////////////////////////////////////////////
+
+#undef LOCTEXT_NAMESPACE
diff --git a/Engine/Source/Developer/TraceInsights/Private/Insights/NetworkingProfiler/Widgets/SNetStatsTableCell.h b/Engine/Source/Developer/TraceInsights/Private/Insights/NetworkingProfiler/Widgets/SNetStatsTableCell.h
new file mode 100644
--- /dev/null
+++ b/Engine/Source/Developer/TraceInsights/Private/Insights/NetworkingProfiler/Widgets/SNetStatsTableCell.h
@@ -0,0 +1,146 @@
+// Copyright 1998-2019 Epic Games, Inc. All Rights Reserved.
+
+#pragma once
+
+#include "CoreMinimal.h"
+#include "Widgets/DeclarativeSyntaxSupport.h"
+#include "Widgets/SCompoundWidget.h"
+
+// Insights
+#include "Insights/NetworkingProfiler/ViewModels/NetEventNode.h"
+
+namespace Insights
+{
+	class FTable;
+	class FTableColumn;
+}
+
+DECLARE_DELEGATE_ThreeParams(FSetHoveredNetStatsTableCell, TSharedPtr<Insights::FTable> /*TablePtr*/, TSharedPtr<Insights::FTableColumn> /*ColumnPtr*/, const FNetEventNodePtr /*NetEventNodePtr*/);
+
+class SNetStatsTableCell : public SCompoundWidget
+{
+public:
+	SLATE_BEGIN_ARGS(SNetStatsTableCell) {}
+		SLATE_EVENT(FSetHoveredNetStatsTableCell, OnSetHoveredCell)
+		SLATE_ATTRIBUTE(FText, HighlightText)
+		SLATE_ARGUMENT(TSharedPtr<Insights::FTable>, TablePtr)
+		SLATE_ARGUMENT(TSharedPtr<Insights::FTableColumn>, ColumnPtr)
+		SLATE_ARGUMENT(FNetEventNodePtr, NetEventNodePtr)
+		SLATE_ARGUMENT(bool, IsNameColumn)
+	SLATE_END_ARGS()
+
+	void Construct(const FArguments& InArgs, const TSharedRef<class ITableRow>& TableRow);
+
+protected:
+	TSharedRef<SWidget> GenerateWidgetForColumn(const FArguments& InArgs, const TSharedRef<class ITableRow>& TableRow);
+	TSharedRef<SWidget> GenerateWidgetForNameColumn(const FArguments& InArgs, const TSharedRef<class ITableRow>& TableRow);
+	TSharedRef<SWidget> GenerateWidgetForStatsColumn(const FArguments& InArgs, const TSharedRef<class ITableRow>& TableRow);
+
+	/**
+	 * The system will use this event to notify a widget that the cursor has entered it. This event is NOT bubbled.
+	 *
+	 * @param MyGeometry The Geometry of the widget receiving the event
+	 * @param MouseEvent Information about the input event
+	 */
+	virtual void OnMouseEnter(const FGeometry& MyGeometry, const FPointerEvent& MouseEvent) override
+	{
+		SCompoundWidget::OnMouseEnter(MyGeometry, MouseEvent);
+		SetHoveredCellDelegate.ExecuteIfBound(TablePtr, ColumnPtr, NetEventNodePtr);
+	}
+
+	/**
+	 * The system will use this event to notify a widget that the cursor has left it. This event is NOT bubbled.
+	 *
+	 * @param MouseEvent Information about the input event
+	 */
+	virtual void OnMouseLeave(const FPointerEvent& MouseEvent) override
+	{
+		SCompoundWidget::OnMouseLeave(MouseEvent);
+		SetHoveredCellDelegate.ExecuteIfBound(nullptr, nullptr, nullptr);
+	}
+
+	/**
+	 * Called during drag and drop when the drag enters a widget.
+	 *
+	 * Enter/Leave events in slate are meant as lightweight notifications.
+	 * So we do not want to capture mouse or set focus in response to these.
+	 * However, OnDragEnter must also support external APIs (e.g. OLE Drag/Drop)
+	 * Those require that we let them know whether we can handle the content
+	 * being dragged OnDragEnter.
+	 *
+	 * The concession is to return a can_handled/cannot_handle
+	 * boolean rather than a full FReply.
+	 *
+	 * @param MyGeometry      The geometry of the widget receiving the event.
+	 * @param DragDropEvent   The drag and drop event.
+	 *
+	 * @return A reply that indicated whether the contents of the DragDropEvent can potentially be processed by this widget.
+	 */
+	virtual void OnDragEnter(const FGeometry& MyGeometry, const FDragDropEvent& DragDropEvent) override
+	{
+		SCompoundWidget::OnDragEnter(MyGeometry, DragDropEvent);
+		SetHoveredCellDelegate.ExecuteIfBound(TablePtr, ColumnPtr, NetEventNodePtr);
+	}
+
+	/**
+	 * Called during drag and drop when the drag leaves a widget.
+	 *
+	 * @param DragDropEvent   The drag and drop event.
+	 */
+	virtual void OnDragLeave(const FDragDropEvent& DragDropEvent)  override
+	{
+		SCompoundWidget::OnDragLeave(DragDropEvent);
+		SetHoveredCellDelegate.ExecuteIfBound(nullptr, nullptr, nullptr);
+	}
+
+	EVisibility GetHintIconVisibility() const
+	{
+		return IsHovered() ? EVisibility::Visible : EVisibility::Hidden;
+	}
+
+	FText GetDisplayName() const
+	{
+		return NetEventNodePtr->GetDisplayName();
+	}
+
+	FSlateColor GetColorAndOpacity() const
+	{
+		const FLinearColor TextColor =
+			NetEventNodePtr->IsFiltered() ?
+				FLinearColor(1.0f, 1.0f, 1.0f, 0.5f) :
+				FLinearColor(1.0f, 1.0f, 1.0f, 1.0f);
+		return TextColor;
+	}
+
+	FSlateColor GetStatsColorAndOpacity() const
+	{
+		const FLinearColor TextColor =
+			NetEventNodePtr->IsFiltered() ?
+				FLinearColor(1.0f, 1.0f, 1.0f, 0.5f) :
+			NetEventNodePtr->GetAggregatedStats().InstanceCount == 0 ?
+				FLinearColor(1.0f, 1.0f, 1.0f, 0.6f) :
+				FLinearColor(1.0f, 1.0f, 1.0f, 1.0f);
+		return TextColor;
+	}
+
+	FLinearColor GetShadowColorAndOpacity() const
+	{
+		const FLinearColor ShadowColor =
+			NetEventNodePtr->IsFiltered() ?
+				FLinearColor(0.f, 0.f, 0.f, 0.25f) :
+				FLinearColor(0.0f, 0.0f, 0.0f, 0.5f);
+		return ShadowColor;
+	}
+
+protected:
+	/** A shared pointer to the table view model. */
+	TSharedPtr<Insights::FTable> TablePtr; // TODO: TSharedRef
+
+	/** A shared pointer to the table column view model. */
+	TSharedPtr<Insights::FTableColumn> ColumnPtr; // TODO: TSharedRef
+
+	/** A shared pointer to the timer node. */
+	FNetEventNodePtr NetEventNodePtr; // TODO: TSharedRef
+
+	FSetHoveredNetStatsTableCell SetHoveredCellDelegate;
+};
diff --git a/Engine/Source/Developer/TraceInsights/Private/Insights/NetworkingProfiler/Widgets/SNetStatsTableRow.cpp b/Engine/Source/Developer/TraceInsights/Private/Insights/NetworkingProfiler/Widgets/SNetStatsTableRow.cpp
new file mode 100644
--- /dev/null
+++ b/Engine/Source/Developer/TraceInsights/Private/Insights/NetworkingProfiler/Widgets/SNetStatsTableRow.cpp
@@ -0,0 +1,211 @@
+// Copyright 1998-2019 Epic Games, Inc. All Rights Reserved.
+
+#include "SNetStatsTableRow.h"
+
+#include "EditorStyleSet.h"
+#include "SlateOptMacros.h"
+#include "Widgets/SOverlay.h"
+#include "Widgets/Images/SImage.h"
+
+// Insights
+#include "Insights/Common/TimeUtils.h"
+#include "Insights/Table/ViewModels/Table.h"
+#include "Insights/Table/ViewModels/TableColumn.h"
+#include "Insights/NetworkingProfiler/Widgets/SNetStatsTableCell.h"
+
+#define LOCTEXT_NAMESPACE "SNetStatsView"
+
+////////////////////////////////////////////////////////////////////////////////////////////////////
+
+BEGIN_SLATE_FUNCTION_BUILD_OPTIMIZATION
+
+void SNetStatsTableRow::Construct(const FArguments& InArgs, const TSharedRef<STableViewBase>& InOwnerTableView)
+{
+	OnShouldBeEnabled = InArgs._OnShouldBeEnabled;
+	IsColumnVisibleDelegate = InArgs._OnIsColumnVisible;
+	SetHoveredCellDelegate = InArgs._OnSetHoveredCell;
+	GetColumnOutlineHAlignmentDelegate = InArgs._OnGetColumnOutlineHAlignmentDelegate;
+
+	HighlightText = InArgs._HighlightText;
+	HighlightedNodeName = InArgs._HighlightedNodeName;
+
+	TablePtr = InArgs._TablePtr;
+	NetEventNodePtr = InArgs._NetEventNodePtr;
+
+	SetEnabled(TAttribute<bool>(this, &SNetStatsTableRow::HandleShouldBeEnabled));
+
+	SMultiColumnTableRow<FNetEventNodePtr>::Construct(SMultiColumnTableRow<FNetEventNodePtr>::FArguments(), InOwnerTableView);
+}
+
+////////////////////////////////////////////////////////////////////////////////////////////////////
+
+TSharedRef<SWidget> SNetStatsTableRow::GenerateWidgetForColumn(const FName& ColumnId)
+{
+	TSharedPtr<Insights::FTableColumn> ColumnPtr = TablePtr->FindColumnChecked(ColumnId);
+
+	return
+		SNew(SOverlay)
+		.Visibility(EVisibility::SelfHitTestInvisible)
+
+		+SOverlay::Slot()
+		.Padding(0.0f)
+		[
+			SNew(SImage)
+			.Image(FEditorStyle::GetBrush("Profiler.LineGraphArea"))
+			.ColorAndOpacity(this, &SNetStatsTableRow::GetBackgroundColorAndOpacity)
+		]
+
+		+SOverlay::Slot()
+		.Padding(0.0f)
+		[
+			SNew(SImage)
+			.Image(this, &SNetStatsTableRow::GetOutlineBrush, ColumnId)
+			.ColorAndOpacity(this, &SNetStatsTableRow::GetOutlineColorAndOpacity)
+		]
+
+		+SOverlay::Slot()
+		[
+			SNew(SNetStatsTableCell, SharedThis(this))
+			.Visibility(this, &SNetStatsTableRow::IsColumnVisible, ColumnId)
+			.TablePtr(TablePtr)
+			.ColumnPtr(ColumnPtr)
+			.NetEventNodePtr(NetEventNodePtr)
+			.HighlightText(HighlightText)
+			.IsNameColumn(ColumnPtr->IsHierarchy())
+			.OnSetHoveredCell(this, &SNetStatsTableRow::OnSetHoveredCell)
+		];
+}
+
+END_SLATE_FUNCTION_BUILD_OPTIMIZATION
+
+////////////////////////////////////////////////////////////////////////////////////////////////////
+
+FReply SNetStatsTableRow::OnDragDetected(const FGeometry& MyGeometry, const FPointerEvent& MouseEvent)
+{
+	//im:TODO
+	//if (MouseEvent.IsMouseButtonDown(EKeys::LeftMouseButton))
+	//{
+	//	if (NetEventNode->IsGroup())
+	//	{
+	//		// Add all timer Ids for the group.
+	//		TArray<int32> NetEventIds;
+	//		const TArray<FNetEventNodePtr>& FilteredChildren = NetEventNode->GetFilteredChildren();
+	//		const int32 NumFilteredChildren = FilteredChildren.Num();
+	//
+	//		NetEventIds.Reserve(NumFilteredChildren);
+	//		for (int32 Nx = 0; Nx < NumFilteredChildren; ++Nx)
+	//		{
+	//			NetEventIds.Add(FilteredChildren[Nx]->GetId());
+	//		}
+	//
+	//		return FReply::Handled().BeginDragDrop(FStatIDDragDropOp::NewGroup(NetEventIds, NetEventNode->GetName().GetPlainNameString()));
+	//	}
+	//	else
+	//	{
+	//		return FReply::Handled().BeginDragDrop(FStatIDDragDropOp::NewSingle(NetEventNode->GetId(), NetEventNode->GetName().GetPlainNameString()));
+	//	}
+	//}
+
+	return SMultiColumnTableRow<FNetEventNodePtr>::OnDragDetected(MyGeometry, MouseEvent);
+}
+
+////////////////////////////////////////////////////////////////////////////////////////////////////
+
+FSlateColor SNetStatsTableRow::GetBackgroundColorAndOpacity() const
+{
+	return GetBackgroundColorAndOpacity(NetEventNodePtr->GetAggregatedStats().TotalInclusive);
+	//return FLinearColor(0.0f, 0.0f, 0.0f, 1.0f)
+}
+
+////////////////////////////////////////////////////////////////////////////////////////////////////
+
+FSlateColor SNetStatsTableRow::GetBackgroundColorAndOpacity(uint32 Size) const
+{
+	const FLinearColor Color =	Size > 1000 ? FLinearColor(0.3f, 0.0f, 0.0f, 1.0f) :
+								Size > 100  ? FLinearColor(0.3f, 0.1f, 0.0f, 1.0f) :
+								Size > 10   ? FLinearColor(0.0f, 0.1f, 0.0f, 1.0f) :
+								              FLinearColor(0.0f, 0.0f, 0.0f, 1.0f);
+	return Color;
+}
+
+////////////////////////////////////////////////////////////////////////////////////////////////////
+
+FSlateColor SNetStatsTableRow::GetOutlineColorAndOpacity() const
+{
+	const FLinearColor NoColor(0.0f, 0.0f, 0.0f, 0.0f);
+	const bool bShouldBeHighlighted = NetEventNodePtr->GetName() == HighlightedNodeName.Get();
+	const FLinearColor OutlineColorAndOpacity = bShouldBeHighlighted ? FLinearColor(FColorList::SlateBlue) : NoColor;
+	return OutlineColorAndOpacity;
+}
+
+////////////////////////////////////////////////////////////////////////////////////////////////////
+
+const FSlateBrush* SNetStatsTableRow::GetOutlineBrush(const FName ColumnId) const
+{
+	EHorizontalAlignment Result = HAlign_Center;
+	if (IsColumnVisibleDelegate.IsBound())
+	{
+		Result = GetColumnOutlineHAlignmentDelegate.Execute(ColumnId);
+	}
+
+	const FSlateBrush* Brush = nullptr;
+	if (Result == HAlign_Left)
+	{
+		Brush = FEditorStyle::GetBrush("Profiler.EventGraph.Border.L");
+	}
+	else if(Result == HAlign_Right)
+	{
+		Brush = FEditorStyle::GetBrush("Profiler.EventGraph.Border.R");
+	}
+	else
+	{
+		Brush = FEditorStyle::GetBrush("Profiler.EventGraph.Border.TB");
+	}
+	return Brush;
+}
+
+////////////////////////////////////////////////////////////////////////////////////////////////////
+
+bool SNetStatsTableRow::HandleShouldBeEnabled() const
+{
+	bool bResult = false;
+
+	if (NetEventNodePtr->IsGroup())
+	{
+		bResult = true;
+	}
+	else
+	{
+		if (OnShouldBeEnabled.IsBound())
+		{
+			bResult = OnShouldBeEnabled.Execute(NetEventNodePtr->GetId());
+		}
+	}
+
+	return bResult;
+}
+
+////////////////////////////////////////////////////////////////////////////////////////////////////
+
+EVisibility SNetStatsTableRow::IsColumnVisible(const FName ColumnId) const
+{
+	EVisibility Result = EVisibility::Collapsed;
+
+	if (IsColumnVisibleDelegate.IsBound())
+	{
+		Result = IsColumnVisibleDelegate.Execute(ColumnId) ? EVisibility::Visible : EVisibility::Collapsed;
+	}
+
+	return Result;
+}
+
+////////////////////////////////////////////////////////////////////////////////////////////////////
+
+void SNetStatsTableRow::OnSetHoveredCell(TSharedPtr<Insights::FTable> InTablePtr, TSharedPtr<Insights::FTableColumn> InColumnPtr, const FNetEventNodePtr InNetEventNodePtr)
+{
+	SetHoveredCellDelegate.ExecuteIfBound(InTablePtr, InColumnPtr, InNetEventNodePtr);
+}
+
+////////////////////////////////////////////////////////////////////////////////////////////////////
+
+#undef LOCTEXT_NAMESPACE
diff --git a/Engine/Source/Developer/TraceInsights/Private/Insights/NetworkingProfiler/Widgets/SNetStatsTableRow.h b/Engine/Source/Developer/TraceInsights/Private/Insights/NetworkingProfiler/Widgets/SNetStatsTableRow.h
new file mode 100644
--- /dev/null
+++ b/Engine/Source/Developer/TraceInsights/Private/Insights/NetworkingProfiler/Widgets/SNetStatsTableRow.h
@@ -0,0 +1,88 @@
+// Copyright 1998-2019 Epic Games, Inc. All Rights Reserved.
+
+#pragma once
+
+#include "CoreMinimal.h"
+#include "Widgets/DeclarativeSyntaxSupport.h"
+#include "Widgets/Views/SHeaderRow.h"
+#include "Widgets/Views/STableViewBase.h"
+#include "Widgets/Views/STableRow.h"
+#include "Widgets/Views/STreeView.h"
+
+// Insights
+#include "Insights/NetworkingProfiler/ViewModels/NetEventNodeHelper.h"
+
+namespace Insights
+{
+	class FTable;
+	class FTableColumn;
+}
+
+DECLARE_DELEGATE_RetVal_OneParam(bool, FShouldBeEnabledDelegate, const uint32 /*NetEventId*/);
+DECLARE_DELEGATE_RetVal_OneParam(bool, FIsColumnVisibleDelegate, const FName /*ColumnId*/);
+DECLARE_DELEGATE_RetVal_OneParam(EHorizontalAlignment, FGetColumnOutlineHAlignmentDelegate, const FName /*ColumnId*/);
+DECLARE_DELEGATE_ThreeParams(FSetHoveredNetEventTableCell, TSharedPtr<Insights::FTable> /*TablePtr*/, TSharedPtr<Insights::FTableColumn> /*ColumnPtr*/, const FNetEventNodePtr /*NetEventNodePtr*/);
+
+/** Widget that represents a table row in a timer tree control. Generates widgets for each column on demand. */
+class SNetStatsTableRow : public SMultiColumnTableRow<FNetEventNodePtr>
+{
+public:
+	SLATE_BEGIN_ARGS(SNetStatsTableRow) {}
+		SLATE_EVENT(FShouldBeEnabledDelegate, OnShouldBeEnabled)
+		SLATE_EVENT(FIsColumnVisibleDelegate, OnIsColumnVisible)
+		SLATE_EVENT(FGetColumnOutlineHAlignmentDelegate, OnGetColumnOutlineHAlignmentDelegate)
+		SLATE_EVENT(FSetHoveredNetEventTableCell, OnSetHoveredCell)
+		SLATE_ATTRIBUTE(FText, HighlightText)
+		SLATE_ATTRIBUTE(FName, HighlightedNodeName)
+		SLATE_ARGUMENT(TSharedPtr<Insights::FTable>, TablePtr)
+		SLATE_ARGUMENT(FNetEventNodePtr, NetEventNodePtr)
+	SLATE_END_ARGS()
+
+public:
+	void Construct(const FArguments& InArgs, const TSharedRef<STableViewBase>& InOwnerTableView);
+
+	virtual TSharedRef<SWidget> GenerateWidgetForColumn(const FName& ColumnId) override;
+
+	/**
+	 * Called when Slate detects that a widget started to be dragged.
+	 * Usage:
+	 * A widget can ask Slate to detect a drag.
+	 * OnMouseDown() reply with FReply::Handled().DetectDrag(SharedThis(this)).
+	 * Slate will either send an OnDragDetected() event or do nothing.
+	 * If the user releases a mouse button or leaves the widget before
+	 * a drag is triggered (maybe user started at the very edge) then no event will be
+	 * sent.
+	 *
+	 * @param  InMyGeometry  Widget geometry
+	 * @param  InMouseEvent  MouseMove that triggered the drag
+	 *
+	 */
+	virtual FReply OnDragDetected(const FGeometry& MyGeometry, const FPointerEvent& MouseEvent) override;
+
+protected:
+	FSlateColor GetBackgroundColorAndOpacity() const;
+	FSlateColor GetBackgroundColorAndOpacity(uint32 Size) const;
+	FSlateColor GetOutlineColorAndOpacity() const;
+	const FSlateBrush* GetOutlineBrush(const FName ColumnId) const;
+	bool HandleShouldBeEnabled() const;
+	EVisibility IsColumnVisible(const FName ColumnId) const;
+	void OnSetHoveredCell(TSharedPtr<Insights::FTable> InTablePtr, TSharedPtr<Insights::FTableColumn> InColumnPtr, const FNetEventNodePtr InNetEventNodePtr);
+
+protected:
+	/** A shared pointer to the table view model. */
+	TSharedPtr<Insights::FTable> TablePtr;
+
+	/** Data context for this table row. */
+	FNetEventNodePtr NetEventNodePtr;
+
+	FShouldBeEnabledDelegate OnShouldBeEnabled;
+	FIsColumnVisibleDelegate IsColumnVisibleDelegate;
+	FSetHoveredNetEventTableCell SetHoveredCellDelegate;
+	FGetColumnOutlineHAlignmentDelegate GetColumnOutlineHAlignmentDelegate;
+
+	/** Text to be highlighted on timer name. */
+	TAttribute<FText> HighlightText;
+
+	/** Name of the timer node that should be drawn as highlighted. */
+	TAttribute<FName> HighlightedNodeName;
+};
diff --git a/Engine/Source/Developer/TraceInsights/Private/Insights/NetworkingProfiler/Widgets/SNetStatsView.cpp b/Engine/Source/Developer/TraceInsights/Private/Insights/NetworkingProfiler/Widgets/SNetStatsView.cpp
new file mode 100644
--- /dev/null
+++ b/Engine/Source/Developer/TraceInsights/Private/Insights/NetworkingProfiler/Widgets/SNetStatsView.cpp
@@ -0,0 +1,1630 @@
+// Copyright 1998-2019 Epic Games, Inc. All Rights Reserved.
+
+#include "SNetStatsView.h"
+
+#include "EditorStyleSet.h"
+#include "Framework/MultiBox/MultiBoxBuilder.h"
+#include "SlateOptMacros.h"
+#include "Templates/UniquePtr.h"
+#include "TraceServices/AnalysisService.h"
+#include "TraceServices/Model/NetProfiler.h"
+#include "Widgets/Input/SCheckBox.h"
+#include "Widgets/Layout/SScrollBox.h"
+#include "Widgets/Input/SSearchBox.h"
+#include "Widgets/Images/SImage.h"
+#include "Widgets/Layout/SGridPanel.h"
+#include "Widgets/Layout/SSeparator.h"
+#include "Widgets/SToolTip.h"
+#include "Widgets/Views/STableViewBase.h"
+
+// Insights
+#include "Insights/Table/ViewModels/Table.h"
+#include "Insights/Table/ViewModels/TableColumn.h"
+#include "Insights/NetworkingProfiler/NetworkingProfilerManager.h"
+#include "Insights/NetworkingProfiler/ViewModels/NetStatsViewColumnFactory.h"
+#include "Insights/NetworkingProfiler/Widgets/SNetStatsViewTooltip.h"
+#include "Insights/NetworkingProfiler/Widgets/SNetStatsTableRow.h"
+
+#define LOCTEXT_NAMESPACE "SNetStatsView"
+
+////////////////////////////////////////////////////////////////////////////////////////////////////
+
+SNetStatsView::SNetStatsView()
+	: Table(MakeShareable(new Insights::FTable()))
+	, bExpansionSaved(false)
+	, GroupingMode(ENetEventGroupingMode::Flat)
+	, AvailableSorters()
+	, CurrentSorter(nullptr)
+	, ColumnBeingSorted(GetDefaultColumnBeingSorted())
+	, ColumnSortMode(GetDefaultColumnSortMode())
+	, NextTimestamp(0)
+	, ObjectsChangeCount(0)
+	, GameInstanceIndex(0)
+	, ConnectionIndex(0)
+	, ConnectionMode(Trace::ENetProfilerConnectionMode::Outgoing)
+	, StatsPacketStartIndex(0)
+	, StatsPacketEndIndex(0)
+	, StatsStartPosition(0)
+	, StatsEndPosition(0)
+{
+	FMemory::Memset(bNetEventTypeIsVisible, 1);
+}
+
+////////////////////////////////////////////////////////////////////////////////////////////////////
+
+SNetStatsView::~SNetStatsView()
+{
+	// Remove ourselves from the Insights manager.
+	if (FInsightsManager::Get().IsValid())
+	{
+		FInsightsManager::Get()->GetSessionChangedEvent().RemoveAll(this);
+	}
+}
+
+////////////////////////////////////////////////////////////////////////////////////////////////////
+
+BEGIN_SLATE_FUNCTION_BUILD_OPTIMIZATION
+void SNetStatsView::Construct(const FArguments& InArgs)
+{
+	SAssignNew(ExternalScrollbar, SScrollBar)
+	.AlwaysShowScrollbar(true);
+
+	ChildSlot
+	[
+		SNew(SVerticalBox)
+
+		+SVerticalBox::Slot()
+		.VAlign(VAlign_Center)
+		.AutoHeight()
+		[
+			SNew(SBorder)
+			.BorderImage(FEditorStyle::GetBrush("ToolPanel.GroupBorder"))
+			.Padding(2.0f)
+			[
+				SNew(SVerticalBox)
+
+				// Search box
+				+SVerticalBox::Slot()
+				.VAlign(VAlign_Center)
+				.Padding(2.0f)
+				.AutoHeight()
+				[
+					SAssignNew(SearchBox, SSearchBox)
+					.HintText(LOCTEXT("SearchBoxHint", "Search net events or groups"))
+					.OnTextChanged(this, &SNetStatsView::SearchBox_OnTextChanged)
+					.IsEnabled(this, &SNetStatsView::SearchBox_IsEnabled)
+					.ToolTipText(LOCTEXT("FilterSearchHint", "Type here to search net events or groups"))
+				]
+
+				// Group by
+				+SVerticalBox::Slot()
+				.VAlign(VAlign_Center)
+				.Padding(2.0f)
+				.AutoHeight()
+				[
+					SNew(SHorizontalBox)
+
+					+SHorizontalBox::Slot()
+					.FillWidth(1.0f)
+					.VAlign(VAlign_Center)
+					[
+						SNew(STextBlock)
+						.Text(LOCTEXT("GroupByText", "Group by"))
+					]
+
+					+SHorizontalBox::Slot()
+					.FillWidth(2.0f)
+					.VAlign(VAlign_Center)
+					[
+						SAssignNew(GroupByComboBox, SComboBox<TSharedPtr<ENetEventGroupingMode>>)
+						.ToolTipText(this, &SNetStatsView::GroupBy_GetSelectedTooltipText)
+						.OptionsSource(&GroupByOptionsSource)
+						.OnSelectionChanged(this, &SNetStatsView::GroupBy_OnSelectionChanged)
+						.OnGenerateWidget(this, &SNetStatsView::GroupBy_OnGenerateWidget)
+						[
+							SNew(STextBlock)
+							.Text(this, &SNetStatsView::GroupBy_GetSelectedText)
+						]
+					]
+				]
+
+				// TODO: Check boxes for: NetEvent, ...
+				/*
+				+ SVerticalBox::Slot()
+				.VAlign(VAlign_Center)
+				.Padding(2.0f)
+				.AutoHeight()
+				[
+					SNew(SHorizontalBox)
+
+					+SHorizontalBox::Slot()
+					.Padding(FMargin(0.0f,0.0f,1.0f,0.0f))
+					.FillWidth(1.0f)
+					[
+						GetToggleButtonForNetEventType(ENetEventNodeType::NetEvent)
+					]
+
+					//+SHorizontalBox::Slot()
+					//.Padding(FMargin(1.0f,0.0f,1.0f,0.0f))
+					//.FillWidth(1.0f)
+					//[
+					//	GetToggleButtonForNetEventType(ENetEventNodeType::NetEvent1)
+					//]
+
+					//+SHorizontalBox::Slot()
+					.//Padding(FMargin(1.0f,0.0f,1.0f,0.0f))
+					//.FillWidth(1.0f)
+					//[
+					//	GetToggleButtonForNetEventType(ENetEventNodeType::NetEvent2)
+					//]
+				]
+				*/
+			]
+		]
+
+		// Tree view
+		+SVerticalBox::Slot()
+		.FillHeight(1.0f)
+		.Padding(0.0f, 6.0f, 0.0f, 0.0f)
+		[
+			SNew(SHorizontalBox)
+
+			+ SHorizontalBox::Slot()
+			.FillWidth(1.0f)
+			.Padding(0.0f)
+			[
+				SNew(SScrollBox)
+				.Orientation(Orient_Horizontal)
+
+				+ SScrollBox::Slot()
+				[
+					SNew(SBorder)
+					.BorderImage(FEditorStyle::GetBrush("ToolPanel.GroupBorder"))
+					.Padding(0.0f)
+					[
+						SAssignNew(TreeView, STreeView<FNetEventNodePtr>)
+						.ExternalScrollbar(ExternalScrollbar)
+						.SelectionMode(ESelectionMode::Multi)
+						.TreeItemsSource(&FilteredGroupNodes)
+						.OnGetChildren(this, &SNetStatsView::TreeView_OnGetChildren)
+						.OnGenerateRow(this, &SNetStatsView::TreeView_OnGenerateRow)
+						.OnSelectionChanged(this, &SNetStatsView::TreeView_OnSelectionChanged)
+						.OnMouseButtonDoubleClick(this, &SNetStatsView::TreeView_OnMouseButtonDoubleClick)
+						.OnContextMenuOpening(FOnContextMenuOpening::CreateSP(this, &SNetStatsView::TreeView_GetMenuContent))
+						.ItemHeight(12.0f)
+						.HeaderRow
+						(
+							SAssignNew(TreeViewHeaderRow, SHeaderRow)
+							.Visibility(EVisibility::Visible)
+						)
+					]
+				]
+			]
+
+			+ SHorizontalBox::Slot()
+			.AutoWidth()
+			.Padding(0.0f)
+			[
+				SNew(SBox)
+				.WidthOverride(FOptionalSize(13.0f))
+				[
+					ExternalScrollbar.ToSharedRef()
+				]
+			]
+		]
+	];
+
+	InitializeAndShowHeaderColumns();
+	//BindCommands();
+
+	// Create the search filters: text based, type based etc.
+	TextFilter = MakeShareable(new FNetEventNodeTextFilter(FNetEventNodeTextFilter::FItemToStringArray::CreateSP(this, &SNetStatsView::HandleItemToStringArray)));
+	Filters = MakeShareable(new FNetEventNodeFilterCollection());
+	Filters->Add(TextFilter);
+
+	CreateGroupByOptionsSources();
+	CreateSortings();
+
+	// Register ourselves with the Insights manager.
+	FInsightsManager::Get()->GetSessionChangedEvent().AddSP(this, &SNetStatsView::InsightsManager_OnSessionChanged);
+
+	// Update the Session (i.e. when analysis session was already started).
+	InsightsManager_OnSessionChanged();
+}
+END_SLATE_FUNCTION_BUILD_OPTIMIZATION
+
+////////////////////////////////////////////////////////////////////////////////////////////////////
+
+void SNetStatsView::Tick(const FGeometry& AllottedGeometry, const double InCurrentTime, const float InDeltaTime)
+{
+	//SCompoundWidget::Tick(AllottedGeometry, InCurrentTime, InDeltaTime);
+
+	// Check if we need to update the lists of net events, but not too often.
+	const uint64 Time = FPlatformTime::Cycles64();
+	if (Time > NextTimestamp)
+	{
+		const uint64 WaitTime = static_cast<uint64>(0.5 / FPlatformTime::GetSecondsPerCycle64()); // 500ms
+		NextTimestamp = Time + WaitTime;
+
+		RebuildTree(false);
+	}
+}
+
+////////////////////////////////////////////////////////////////////////////////////////////////////
+
+TSharedPtr<SWidget> SNetStatsView::TreeView_GetMenuContent()
+{
+	const TArray<FNetEventNodePtr> SelectedNodes = TreeView->GetSelectedItems();
+	const int32 NumSelectedNodes = SelectedNodes.Num();
+	FNetEventNodePtr SelectedNode = NumSelectedNodes ? SelectedNodes[0] : nullptr;
+
+	const TSharedPtr<Insights::FTableColumn> HoveredColumnPtr = Table->FindColumn(HoveredColumnId);
+
+	FText SelectionStr;
+	FText PropertyName;
+	FText PropertyValue;
+
+	if (NumSelectedNodes == 0)
+	{
+		SelectionStr = LOCTEXT("NothingSelected", "Nothing selected");
+	}
+	else if (NumSelectedNodes == 1)
+	{
+		if (HoveredColumnPtr != nullptr)
+		{
+			PropertyName = HoveredColumnPtr->GetShortName();
+			PropertyValue = HoveredColumnPtr->GetValueAsTooltipText(*SelectedNode);
+		}
+		SelectionStr = FText::FromName(SelectedNode->GetName());
+	}
+	else
+	{
+		SelectionStr = LOCTEXT("MultipleSelection", "Multiple selection");
+	}
+
+	const bool bShouldCloseWindowAfterMenuSelection = true;
+	FMenuBuilder MenuBuilder(bShouldCloseWindowAfterMenuSelection, NULL);
+
+	// Selection menu
+	MenuBuilder.BeginSection("Selection", LOCTEXT("ContextMenu_Header_Selection", "Selection"));
+	{
+		struct FLocal
+		{
+			static bool ReturnFalse()
+			{
+				return false;
+			}
+		};
+
+		FUIAction DummyUIAction;
+		DummyUIAction.CanExecuteAction = FCanExecuteAction::CreateStatic(&FLocal::ReturnFalse);
+		MenuBuilder.AddMenuEntry
+		(
+			SelectionStr,
+			LOCTEXT("ContextMenu_Selection", "Currently selected items"),
+			FSlateIcon(FEditorStyle::GetStyleSetName(), "@missing.icon"), DummyUIAction, NAME_None, EUserInterfaceActionType::Button
+		);
+	}
+	MenuBuilder.EndSection();
+
+	MenuBuilder.BeginSection("Misc", LOCTEXT("ContextMenu_Header_Misc", "Miscellaneous"));
+	{
+		/*TODO
+		FUIAction Action_CopySelectedToClipboard
+		(
+			FExecuteAction::CreateSP(this, &SNetStatsView::ContextMenu_CopySelectedToClipboard_Execute),
+			FCanExecuteAction::CreateSP(this, &SNetStatsView::ContextMenu_CopySelectedToClipboard_CanExecute)
+		);
+		MenuBuilder.AddMenuEntry
+		(
+			LOCTEXT("ContextMenu_Header_Misc_CopySelectedToClipboard", "Copy To Clipboard"),
+			LOCTEXT("ContextMenu_Header_Misc_CopySelectedToClipboard_Desc", "Copies selection to clipboard"),
+			FSlateIcon(FEditorStyle::GetStyleSetName(), "Profiler.Misc.CopyToClipboard"), Action_CopySelectedToClipboard, NAME_None, EUserInterfaceActionType::Button
+		);
+		*/
+
+		MenuBuilder.AddSubMenu
+		(
+			LOCTEXT("ContextMenu_Header_Misc_Sort", "Sort By"),
+			LOCTEXT("ContextMenu_Header_Misc_Sort_Desc", "Sort by column"),
+			FNewMenuDelegate::CreateSP(this, &SNetStatsView::TreeView_BuildSortByMenu),
+			false,
+			FSlateIcon(FEditorStyle::GetStyleSetName(), "Profiler.Misc.SortBy")
+		);
+	}
+	MenuBuilder.EndSection();
+
+	MenuBuilder.BeginSection("Columns", LOCTEXT("ContextMenu_Header_Columns", "Columns"));
+	{
+		MenuBuilder.AddSubMenu
+		(
+			LOCTEXT("ContextMenu_Header_Columns_View", "View Column"),
+			LOCTEXT("ContextMenu_Header_Columns_View_Desc", "Hides or shows columns"),
+			FNewMenuDelegate::CreateSP(this, &SNetStatsView::TreeView_BuildViewColumnMenu),
+			false,
+			FSlateIcon(FEditorStyle::GetStyleSetName(), "Profiler.EventGraph.ViewColumn")
+		);
+
+		FUIAction Action_ShowAllColumns
+		(
+			FExecuteAction::CreateSP(this, &SNetStatsView::ContextMenu_ShowAllColumns_Execute),
+			FCanExecuteAction::CreateSP(this, &SNetStatsView::ContextMenu_ShowAllColumns_CanExecute)
+		);
+		MenuBuilder.AddMenuEntry
+		(
+			LOCTEXT("ContextMenu_Header_Columns_ShowAllColumns", "Show All Columns"),
+			LOCTEXT("ContextMenu_Header_Columns_ShowAllColumns_Desc", "Resets tree view to show all columns"),
+			FSlateIcon(FEditorStyle::GetStyleSetName(), "Profiler.EventGraph.ResetColumn"), Action_ShowAllColumns, NAME_None, EUserInterfaceActionType::Button
+		);
+
+		FUIAction Action_ShowMinMaxMedColumns
+		(
+			FExecuteAction::CreateSP(this, &SNetStatsView::ContextMenu_ShowMinMaxMedColumns_Execute),
+			FCanExecuteAction::CreateSP(this, &SNetStatsView::ContextMenu_ShowMinMaxMedColumns_CanExecute)
+		);
+		MenuBuilder.AddMenuEntry
+		(
+			LOCTEXT("ContextMenu_Header_Columns_ShowMinMaxMedColumns", "Reset Columns to Min/Max/Median Preset"),
+			LOCTEXT("ContextMenu_Header_Columns_ShowMinMaxMedColumns_Desc", "Resets columns to Min/Max/Median preset"),
+			FSlateIcon(FEditorStyle::GetStyleSetName(), "Profiler.EventGraph.ResetColumn"), Action_ShowMinMaxMedColumns, NAME_None, EUserInterfaceActionType::Button
+		);
+
+		FUIAction Action_ResetColumns
+		(
+			FExecuteAction::CreateSP(this, &SNetStatsView::ContextMenu_ResetColumns_Execute),
+			FCanExecuteAction::CreateSP(this, &SNetStatsView::ContextMenu_ResetColumns_CanExecute)
+		);
+		MenuBuilder.AddMenuEntry
+		(
+			LOCTEXT("ContextMenu_Header_Columns_ResetColumns", "Reset Columns to Default"),
+			LOCTEXT("ContextMenu_Header_Columns_ResetColumns_Desc", "Resets columns to default"),
+			FSlateIcon(FEditorStyle::GetStyleSetName(), "Profiler.EventGraph.ResetColumn"), Action_ResetColumns, NAME_None, EUserInterfaceActionType::Button
+		);
+	}
+	MenuBuilder.EndSection();
+
+	return MenuBuilder.MakeWidget();
+}
+
+////////////////////////////////////////////////////////////////////////////////////////////////////
+
+void SNetStatsView::TreeView_BuildSortByMenu(FMenuBuilder& MenuBuilder)
+{
+	// TODO: Refactor later @see TSharedPtr<SWidget> SCascadePreviewViewportToolBar::GenerateViewMenu() const
+
+	MenuBuilder.BeginSection("ColumnName", LOCTEXT("ContextMenu_Header_Misc_ColumnName", "Column Name"));
+
+	for (const TSharedPtr<Insights::FTableColumn>& ColumnPtr : Table->GetColumns())
+	{
+		const Insights::FTableColumn& Column = *ColumnPtr;
+
+		if (Column.IsVisible() && Column.CanBeSorted())
+		{
+			FUIAction Action_SortByColumn
+			(
+				FExecuteAction::CreateSP(this, &SNetStatsView::ContextMenu_SortByColumn_Execute, Column.GetId()),
+				FCanExecuteAction::CreateSP(this, &SNetStatsView::ContextMenu_SortByColumn_CanExecute, Column.GetId()),
+				FIsActionChecked::CreateSP(this, &SNetStatsView::ContextMenu_SortByColumn_IsChecked, Column.GetId())
+			);
+			MenuBuilder.AddMenuEntry
+			(
+				Column.GetTitleName(),
+				Column.GetDescription(),
+				FSlateIcon(), Action_SortByColumn, NAME_None, EUserInterfaceActionType::RadioButton
+			);
+		}
+	}
+
+	MenuBuilder.EndSection();
+
+	//-----------------------------------------------------------------------------
+
+	MenuBuilder.BeginSection("SortMode", LOCTEXT("ContextMenu_Header_Misc_Sort_SortMode", "Sort Mode"));
+	{
+		FUIAction Action_SortAscending
+		(
+			FExecuteAction::CreateSP(this, &SNetStatsView::ContextMenu_SortMode_Execute, EColumnSortMode::Ascending),
+			FCanExecuteAction::CreateSP(this, &SNetStatsView::ContextMenu_SortMode_CanExecute, EColumnSortMode::Ascending),
+			FIsActionChecked::CreateSP(this, &SNetStatsView::ContextMenu_SortMode_IsChecked, EColumnSortMode::Ascending)
+		);
+		MenuBuilder.AddMenuEntry
+		(
+			LOCTEXT("ContextMenu_Header_Misc_Sort_SortAscending", "Sort Ascending"),
+			LOCTEXT("ContextMenu_Header_Misc_Sort_SortAscending_Desc", "Sorts ascending"),
+			FSlateIcon(FEditorStyle::GetStyleSetName(), "Profiler.Misc.SortAscending"), Action_SortAscending, NAME_None, EUserInterfaceActionType::RadioButton
+		);
+
+		FUIAction Action_SortDescending
+		(
+			FExecuteAction::CreateSP(this, &SNetStatsView::ContextMenu_SortMode_Execute, EColumnSortMode::Descending),
+			FCanExecuteAction::CreateSP(this, &SNetStatsView::ContextMenu_SortMode_CanExecute, EColumnSortMode::Descending),
+			FIsActionChecked::CreateSP(this, &SNetStatsView::ContextMenu_SortMode_IsChecked, EColumnSortMode::Descending)
+		);
+		MenuBuilder.AddMenuEntry
+		(
+			LOCTEXT("ContextMenu_Header_Misc_Sort_SortDescending", "Sort Descending"),
+			LOCTEXT("ContextMenu_Header_Misc_Sort_SortDescending_Desc", "Sorts descending"),
+			FSlateIcon(FEditorStyle::GetStyleSetName(), "Profiler.Misc.SortDescending"), Action_SortDescending, NAME_None, EUserInterfaceActionType::RadioButton
+		);
+	}
+	MenuBuilder.EndSection();
+}
+
+////////////////////////////////////////////////////////////////////////////////////////////////////
+
+void SNetStatsView::TreeView_BuildViewColumnMenu(FMenuBuilder& MenuBuilder)
+{
+	MenuBuilder.BeginSection("ViewColumn", LOCTEXT("ContextMenu_Header_Columns_View", "View Column"));
+
+	for (const TSharedPtr<Insights::FTableColumn>& ColumnPtr : Table->GetColumns())
+	{
+		const Insights::FTableColumn& Column = *ColumnPtr;
+
+		FUIAction Action_ToggleColumn
+		(
+			FExecuteAction::CreateSP(this, &SNetStatsView::ToggleColumnVisibility, Column.GetId()),
+			FCanExecuteAction::CreateSP(this, &SNetStatsView::CanToggleColumnVisibility, Column.GetId()),
+			FIsActionChecked::CreateSP(this, &SNetStatsView::IsColumnVisible, Column.GetId())
+		);
+		MenuBuilder.AddMenuEntry
+		(
+			Column.GetTitleName(),
+			Column.GetDescription(),
+			FSlateIcon(), Action_ToggleColumn, NAME_None, EUserInterfaceActionType::ToggleButton
+		);
+	}
+
+	MenuBuilder.EndSection();
+}
+
+////////////////////////////////////////////////////////////////////////////////////////////////////
+
+void SNetStatsView::InitializeAndShowHeaderColumns()
+{
+	// Create columns.
+	TArray<TSharedPtr<Insights::FTableColumn>> Columns;
+	FNetStatsViewColumnFactory::CreateNetStatsViewColumns(Columns);
+	Table->SetColumns(Columns);
+
+	// Show columns.
+	for (const TSharedPtr<Insights::FTableColumn>& ColumnPtr : Table->GetColumns())
+	{
+		if (ColumnPtr->ShouldBeVisible())
+		{
+			ShowColumn(ColumnPtr->GetId());
+		}
+	}
+}
+
+////////////////////////////////////////////////////////////////////////////////////////////////////
+
+FText SNetStatsView::GetColumnHeaderText(const FName ColumnId) const
+{
+	const Insights::FTableColumn& Column = *Table->FindColumnChecked(ColumnId);
+	return Column.GetShortName();
+}
+
+////////////////////////////////////////////////////////////////////////////////////////////////////
+
+TSharedRef<SWidget> SNetStatsView::TreeViewHeaderRow_GenerateColumnMenu(const Insights::FTableColumn& Column)
+{
+	bool bIsMenuVisible = false;
+
+	const bool bShouldCloseWindowAfterMenuSelection = true;
+	FMenuBuilder MenuBuilder(bShouldCloseWindowAfterMenuSelection, NULL);
+	{
+		if (Column.CanBeHidden())
+		{
+			MenuBuilder.BeginSection("Column", LOCTEXT("TreeViewHeaderRow_Header_Column", "Column"));
+
+			FUIAction Action_HideColumn
+			(
+				FExecuteAction::CreateSP(this, &SNetStatsView::HideColumn, Column.GetId()),
+				FCanExecuteAction::CreateSP(this, &SNetStatsView::CanHideColumn, Column.GetId())
+			);
+
+			MenuBuilder.AddMenuEntry
+			(
+				LOCTEXT("TreeViewHeaderRow_HideColumn", "Hide"),
+				LOCTEXT("TreeViewHeaderRow_HideColumn_Desc", "Hides the selected column"),
+				FSlateIcon(), Action_HideColumn, NAME_None, EUserInterfaceActionType::Button
+			);
+			bIsMenuVisible = true;
+
+			MenuBuilder.EndSection();
+		}
+
+		if (Column.CanBeSorted())
+		{
+			MenuBuilder.BeginSection("SortMode", LOCTEXT("ContextMenu_Header_Misc_Sort_SortMode", "Sort Mode"));
+
+			FUIAction Action_SortAscending
+			(
+				FExecuteAction::CreateSP(this, &SNetStatsView::HeaderMenu_SortMode_Execute, Column.GetId(), EColumnSortMode::Ascending),
+				FCanExecuteAction::CreateSP(this, &SNetStatsView::HeaderMenu_SortMode_CanExecute, Column.GetId(), EColumnSortMode::Ascending),
+				FIsActionChecked::CreateSP(this, &SNetStatsView::HeaderMenu_SortMode_IsChecked, Column.GetId(), EColumnSortMode::Ascending)
+			);
+			MenuBuilder.AddMenuEntry
+			(
+				LOCTEXT("ContextMenu_Header_Misc_Sort_SortAscending", "Sort Ascending"),
+				LOCTEXT("ContextMenu_Header_Misc_Sort_SortAscending_Desc", "Sorts ascending"),
+				FSlateIcon(FEditorStyle::GetStyleSetName(), "Profiler.Misc.SortAscending"), Action_SortAscending, NAME_None, EUserInterfaceActionType::RadioButton
+			);
+
+			FUIAction Action_SortDescending
+			(
+				FExecuteAction::CreateSP(this, &SNetStatsView::HeaderMenu_SortMode_Execute, Column.GetId(), EColumnSortMode::Descending),
+				FCanExecuteAction::CreateSP(this, &SNetStatsView::HeaderMenu_SortMode_CanExecute, Column.GetId(), EColumnSortMode::Descending),
+				FIsActionChecked::CreateSP(this, &SNetStatsView::HeaderMenu_SortMode_IsChecked, Column.GetId(), EColumnSortMode::Descending)
+			);
+			MenuBuilder.AddMenuEntry
+			(
+				LOCTEXT("ContextMenu_Header_Misc_Sort_SortDescending", "Sort Descending"),
+				LOCTEXT("ContextMenu_Header_Misc_Sort_SortDescending_Desc", "Sorts descending"),
+				FSlateIcon(FEditorStyle::GetStyleSetName(), "Profiler.Misc.SortDescending"), Action_SortDescending, NAME_None, EUserInterfaceActionType::RadioButton
+			);
+			bIsMenuVisible = true;
+
+			MenuBuilder.EndSection();
+		}
+
+		//if (Column.CanBeFiltered())
+		//{
+		//	MenuBuilder.BeginSection("FilterMode", LOCTEXT("ContextMenu_Header_Misc_Filter_FilterMode", "Filter Mode"));
+		//	bIsMenuVisible = true;
+		//	MenuBuilder.EndSection();
+		//}
+	}
+
+	/*
+	TODO:
+	- Show top ten
+	- Show top bottom
+	- Filter by list (avg, median, 10%, 90%, etc.)
+	- Text box for filtering for each column instead of one text box used for filtering
+	- Grouping button for flat view modes (show at most X groups, show all groups for names)
+	*/
+
+	return bIsMenuVisible ? MenuBuilder.MakeWidget() : (TSharedRef<SWidget>)SNullWidget::NullWidget;
+}
+
+////////////////////////////////////////////////////////////////////////////////////////////////////
+
+void SNetStatsView::InsightsManager_OnSessionChanged()
+{
+	TSharedPtr<const Trace::IAnalysisSession> NewSession = FInsightsManager::Get()->GetSession();
+
+	if (NewSession != Session)
+	{
+		Session = NewSession;
+
+		GameInstanceIndex = 0;
+		ConnectionIndex = 0;
+		ConnectionMode = Trace::ENetProfilerConnectionMode::Outgoing;
+		StatsPacketStartIndex = 0;
+		StatsPacketEndIndex = 0;
+		StatsStartPosition = 0;
+		StatsEndPosition = 0;
+
+		RebuildTree();
+	}
+	else
+	{
+		UpdateTree();
+	}
+}
+
+////////////////////////////////////////////////////////////////////////////////////////////////////
+
+void SNetStatsView::UpdateTree()
+{
+	CreateGroups();
+	SortTreeNodes();
+	ApplyFiltering();
+}
+
+////////////////////////////////////////////////////////////////////////////////////////////////////
+
+void SNetStatsView::ApplyFiltering()
+{
+	FilteredGroupNodes.Reset();
+
+	// Apply filter to all groups and its children.
+	const int32 NumGroups = GroupNodes.Num();
+	for (int32 GroupIndex = 0; GroupIndex < NumGroups; ++GroupIndex)
+	{
+		FNetEventNodePtr& GroupPtr = GroupNodes[GroupIndex];
+		GroupPtr->ClearFilteredChildren();
+		const bool bIsGroupVisible = Filters->PassesAllFilters(GroupPtr);
+
+		const TArray<Insights::FBaseTreeNodePtr>& GroupChildren = GroupPtr->GetChildren();
+		const int32 NumChildren = GroupChildren.Num();
+		int32 NumVisibleChildren = 0;
+		for (int32 Cx = 0; Cx < NumChildren; ++Cx)
+		{
+			// Add a child.
+			const FNetEventNodePtr& NodePtr = StaticCastSharedPtr<FNetEventNode, Insights::FBaseTreeNode>(GroupChildren[Cx]);
+			const bool bIsChildVisible = Filters->PassesAllFilters(NodePtr) && bNetEventTypeIsVisible[static_cast<int>(NodePtr->GetType())];
+			if (bIsChildVisible)
+			{
+				GroupPtr->AddFilteredChild(NodePtr);
+				NumVisibleChildren++;
+			}
+		}
+
+		if (bIsGroupVisible || NumVisibleChildren > 0)
+		{
+			// Add a group.
+			FilteredGroupNodes.Add(GroupPtr);
+			GroupPtr->SetExpansion(true);
+		}
+		else
+		{
+			GroupPtr->SetExpansion(false);
+		}
+	}
+
+	// Only expand net event nodes if we have a text filter.
+	const bool bNonEmptyTextFilter = !TextFilter->GetRawFilterText().IsEmpty();
+	if (bNonEmptyTextFilter)
+	{
+		if (!bExpansionSaved)
+		{
+			ExpandedNodes.Empty();
+			TreeView->GetExpandedItems(ExpandedNodes);
+			bExpansionSaved = true;
+		}
+
+		for (int32 Fx = 0; Fx < FilteredGroupNodes.Num(); Fx++)
+		{
+			const FNetEventNodePtr& GroupPtr = FilteredGroupNodes[Fx];
+			TreeView->SetItemExpansion(GroupPtr, GroupPtr->IsExpanded());
+		}
+	}
+	else
+	{
+		if (bExpansionSaved)
+		{
+			// Restore previously expanded nodes when the text filter is disabled.
+			TreeView->ClearExpandedItems();
+			for (auto It = ExpandedNodes.CreateConstIterator(); It; ++It)
+			{
+				TreeView->SetItemExpansion(*It, true);
+			}
+			bExpansionSaved = false;
+		}
+	}
+
+	// Request tree refresh
+	TreeView->RequestTreeRefresh();
+}
+
+////////////////////////////////////////////////////////////////////////////////////////////////////
+
+void SNetStatsView::HandleItemToStringArray(const FNetEventNodePtr& FNetEventNodePtr, TArray<FString>& OutSearchStrings) const
+{
+	OutSearchStrings.Add(FNetEventNodePtr->GetName().GetPlainNameString());
+}
+
+////////////////////////////////////////////////////////////////////////////////////////////////////
+
+TSharedRef<SWidget> SNetStatsView::GetToggleButtonForNetEventType(const ENetEventNodeType NodeType)
+{
+	return SNew(SCheckBox)
+		.Style(FEditorStyle::Get(), "ToggleButtonCheckbox")
+		.HAlign(HAlign_Center)
+		.Padding(2.0f)
+		.OnCheckStateChanged(this, &SNetStatsView::FilterByNetEventType_OnCheckStateChanged, NodeType)
+		.IsChecked(this, &SNetStatsView::FilterByNetEventType_IsChecked, NodeType)
+		.ToolTipText(NetEventNodeTypeHelper::ToDescription(NodeType))
+		[
+			SNew(SHorizontalBox)
+
+			+ SHorizontalBox::Slot()
+				.AutoWidth()
+				.VAlign(VAlign_Center)
+				[
+					SNew(SImage)
+						.Image(NetEventNodeTypeHelper::GetIconForNetEventNodeType(NodeType))
+				]
+
+			+SHorizontalBox::Slot()
+				.Padding(2.0f, 0.0f, 0.0f, 0.0f)
+				.VAlign(VAlign_Center)
+				[
+					SNew(STextBlock)
+						.Text(NetEventNodeTypeHelper::ToText(NodeType))
+						.TextStyle(FEditorStyle::Get(), TEXT("Profiler.Caption"))
+				]
+		];
+}
+
+////////////////////////////////////////////////////////////////////////////////////////////////////
+
+void SNetStatsView::FilterByNetEventType_OnCheckStateChanged(ECheckBoxState NewRadioState, const ENetEventNodeType InStatType)
+{
+	bNetEventTypeIsVisible[static_cast<int>(InStatType)] = (NewRadioState == ECheckBoxState::Checked);
+	ApplyFiltering();
+}
+
+////////////////////////////////////////////////////////////////////////////////////////////////////
+
+ECheckBoxState SNetStatsView::FilterByNetEventType_IsChecked(const ENetEventNodeType InStatType) const
+{
+	return bNetEventTypeIsVisible[static_cast<int>(InStatType)] ? ECheckBoxState::Checked : ECheckBoxState::Unchecked;
+}
+
+////////////////////////////////////////////////////////////////////////////////////////////////////
+// TreeView
+////////////////////////////////////////////////////////////////////////////////////////////////////
+
+void SNetStatsView::TreeView_Refresh()
+{
+	if (TreeView.IsValid())
+	{
+		TreeView->RequestTreeRefresh();
+	}
+}
+
+////////////////////////////////////////////////////////////////////////////////////////////////////
+
+void SNetStatsView::TreeView_OnSelectionChanged(FNetEventNodePtr SelectedItem, ESelectInfo::Type SelectInfo)
+{
+	if (SelectInfo != ESelectInfo::Direct)
+	{
+		TArray<FNetEventNodePtr> SelectedItems = TreeView->GetSelectedItems();
+		if (SelectedItems.Num() == 1 && !SelectedItems[0]->IsGroup())
+		{
+			//TODO: FNetworkingProfilerManager::Get()->SetSelectedNetEvent(SelectedItems[0]->GetId());
+		}
+	}
+}
+
+////////////////////////////////////////////////////////////////////////////////////////////////////
+
+void SNetStatsView::TreeView_OnGetChildren(FNetEventNodePtr InParent, TArray<FNetEventNodePtr>& OutChildren)
+{
+	const TArray<Insights::FBaseTreeNodePtr>& Children = InParent->GetFilteredChildren();
+	OutChildren.Reset(Children.Num());
+	for (const Insights::FBaseTreeNodePtr& Child : Children)
+	{
+		OutChildren.Add(StaticCastSharedPtr<FNetEventNode, Insights::FBaseTreeNode>(Child));
+	}
+}
+
+////////////////////////////////////////////////////////////////////////////////////////////////////
+
+void SNetStatsView::TreeView_OnMouseButtonDoubleClick(FNetEventNodePtr NetEventNodePtr)
+{
+	if (NetEventNodePtr->IsGroup())
+	{
+		const bool bIsGroupExpanded = TreeView->IsItemExpanded(NetEventNodePtr);
+		TreeView->SetItemExpansion(NetEventNodePtr, !bIsGroupExpanded);
+	}
+}
+
+////////////////////////////////////////////////////////////////////////////////////////////////////
+// Tree View's Table Row
+////////////////////////////////////////////////////////////////////////////////////////////////////
+
+TSharedRef<ITableRow> SNetStatsView::TreeView_OnGenerateRow(FNetEventNodePtr NetEventNodePtr, const TSharedRef<STableViewBase>& OwnerTable)
+{
+	TSharedRef<ITableRow> TableRow =
+		SNew(SNetStatsTableRow, OwnerTable)
+		.OnShouldBeEnabled(this, &SNetStatsView::TableRow_ShouldBeEnabled)
+		.OnIsColumnVisible(this, &SNetStatsView::IsColumnVisible)
+		.OnSetHoveredCell(this, &SNetStatsView::TableRow_SetHoveredCell)
+		.OnGetColumnOutlineHAlignmentDelegate(this, &SNetStatsView::TableRow_GetColumnOutlineHAlignment)
+		.HighlightText(this, &SNetStatsView::TableRow_GetHighlightText)
+		.HighlightedNodeName(this, &SNetStatsView::TableRow_GetHighlightedNodeName)
+		.TablePtr(Table)
+		.NetEventNodePtr(NetEventNodePtr);
+
+	return TableRow;
+}
+
+////////////////////////////////////////////////////////////////////////////////////////////////////
+
+bool SNetStatsView::TableRow_ShouldBeEnabled(const uint32 NodeId) const
+{
+	return true;//im:TODO: Session->GetAggregatedStat(NodeId) != nullptr;
+}
+
+////////////////////////////////////////////////////////////////////////////////////////////////////
+
+void SNetStatsView::TableRow_SetHoveredCell(TSharedPtr<Insights::FTable> InTablePtr, TSharedPtr<Insights::FTableColumn> InColumnPtr, const FNetEventNodePtr InNodePtr)
+{
+	HoveredColumnId = InColumnPtr ? InColumnPtr->GetId() : FName();
+
+	const bool bIsAnyMenusVisible = FSlateApplication::Get().AnyMenusVisible();
+	if (!HasMouseCapture() && !bIsAnyMenusVisible)
+	{
+		HoveredNodePtr = InNodePtr;
+	}
+}
+
+////////////////////////////////////////////////////////////////////////////////////////////////////
+
+EHorizontalAlignment SNetStatsView::TableRow_GetColumnOutlineHAlignment(const FName ColumnId) const
+{
+	const TIndirectArray<SHeaderRow::FColumn>& Columns = TreeViewHeaderRow->GetColumns();
+	const int32 LastColumnIdx = Columns.Num() - 1;
+
+	// First column
+	if (Columns[0].ColumnId == ColumnId)
+	{
+		return HAlign_Left;
+	}
+	// Last column
+	else if (Columns[LastColumnIdx].ColumnId == ColumnId)
+	{
+		return HAlign_Right;
+	}
+	// Middle columns
+	{
+		return HAlign_Center;
+	}
+}
+
+////////////////////////////////////////////////////////////////////////////////////////////////////
+
+FText SNetStatsView::TableRow_GetHighlightText() const
+{
+	return SearchBox->GetText();
+}
+
+////////////////////////////////////////////////////////////////////////////////////////////////////
+
+FName SNetStatsView::TableRow_GetHighlightedNodeName() const
+{
+	return HighlightedNodeName;
+}
+
+////////////////////////////////////////////////////////////////////////////////////////////////////
+// SearchBox
+////////////////////////////////////////////////////////////////////////////////////////////////////
+
+void SNetStatsView::SearchBox_OnTextChanged(const FText& InFilterText)
+{
+	TextFilter->SetRawFilterText(InFilterText);
+	SearchBox->SetError(TextFilter->GetFilterErrorText());
+	ApplyFiltering();
+}
+
+////////////////////////////////////////////////////////////////////////////////////////////////////
+
+bool SNetStatsView::SearchBox_IsEnabled() const
+{
+	return NetEventNodes.Num() > 0;
+}
+
+////////////////////////////////////////////////////////////////////////////////////////////////////
+// Grouping
+////////////////////////////////////////////////////////////////////////////////////////////////////
+
+void SNetStatsView::CreateGroups()
+{
+	TMap<FName, FNetEventNodePtr> GroupNodeSet;
+
+	if (GroupingMode == ENetEventGroupingMode::Flat)
+	{
+		const FName GroupName(TEXT("All"));
+		FNetEventNodePtr* GroupPtr = GroupNodeSet.Find(GroupName);
+		if (!GroupPtr)
+		{
+			GroupPtr = &GroupNodeSet.Add(GroupName, MakeShareable(new FNetEventNode(GroupName)));
+		}
+
+		for (const FNetEventNodePtr& NetEventNodePtr : NetEventNodes)
+		{
+			(*GroupPtr)->AddChildAndSetGroupPtr(NetEventNodePtr);
+		}
+
+		TreeView->SetItemExpansion(*GroupPtr, true);
+	}
+	// Creates one group for one letter.
+	else if (GroupingMode == ENetEventGroupingMode::ByName)
+	{
+		for (const FNetEventNodePtr& NetEventNodePtr : NetEventNodes)
+		{
+			const FName GroupName = *NetEventNodePtr->GetName().GetPlainNameString().Left(1).ToUpper();
+
+			FNetEventNodePtr* GroupPtr = GroupNodeSet.Find(GroupName);
+			if (!GroupPtr)
+			{
+				GroupPtr = &GroupNodeSet.Add(GroupName, MakeShareable(new FNetEventNode(GroupName)));
+			}
+
+			(*GroupPtr)->AddChildAndSetGroupPtr(NetEventNodePtr);
+		}
+	}
+	// Creates one group for each stat type.
+	else if (GroupingMode == ENetEventGroupingMode::ByType)
+	{
+		for (const FNetEventNodePtr& NetEventNodePtr : NetEventNodes)
+		{
+			const FName GroupName = *NetEventNodeTypeHelper::ToText(NetEventNodePtr->GetType()).ToString();
+
+			FNetEventNodePtr* GroupPtr = GroupNodeSet.Find(GroupName);
+			if (!GroupPtr)
+			{
+				GroupPtr = &GroupNodeSet.Add(GroupName, MakeShareable(new FNetEventNode(GroupName)));
+			}
+
+			(*GroupPtr)->AddChildAndSetGroupPtr(NetEventNodePtr);
+			TreeView->SetItemExpansion(*GroupPtr, true);
+		}
+	}
+	// Creates one group for each stat type.
+	else if (GroupingMode == ENetEventGroupingMode::ByLevel)
+	{
+		for (const FNetEventNodePtr& NetEventNodePtr : NetEventNodes)
+		{
+			const FName GroupName(*FString::Printf(TEXT("Level %d"), NetEventNodePtr->GetLevel()));
+
+			FNetEventNodePtr* GroupPtr = GroupNodeSet.Find(GroupName);
+			if (!GroupPtr)
+			{
+				GroupPtr = &GroupNodeSet.Add(GroupName, MakeShareable(new FNetEventNode(GroupName)));
+			}
+
+			(*GroupPtr)->AddChildAndSetGroupPtr(NetEventNodePtr);
+			TreeView->SetItemExpansion(*GroupPtr, true);
+		}
+	}
+
+	GroupNodeSet.GenerateValueArray(GroupNodes);
+}
+
+////////////////////////////////////////////////////////////////////////////////////////////////////
+
+void SNetStatsView::CreateGroupByOptionsSources()
+{
+	GroupByOptionsSource.Reset(3);
+
+	// Must be added in order of elements in the ENetEventGroupingMode.
+	GroupByOptionsSource.Add(MakeShareable(new ENetEventGroupingMode(ENetEventGroupingMode::Flat)));
+	GroupByOptionsSource.Add(MakeShareable(new ENetEventGroupingMode(ENetEventGroupingMode::ByName)));
+	//GroupByOptionsSource.Add(MakeShareable(new ENetEventGroupingMode(ENetEventGroupingMode::ByType)));
+	GroupByOptionsSource.Add(MakeShareable(new ENetEventGroupingMode(ENetEventGroupingMode::ByLevel)));
+
+	ENetEventGroupingModePtr* GroupingModePtrPtr = GroupByOptionsSource.FindByPredicate([&](const ENetEventGroupingModePtr InGroupingModePtr) { return *InGroupingModePtr == GroupingMode; });
+	if (GroupingModePtrPtr != nullptr)
+	{
+		GroupByComboBox->SetSelectedItem(*GroupingModePtrPtr);
+	}
+
+	GroupByComboBox->RefreshOptions();
+}
+
+////////////////////////////////////////////////////////////////////////////////////////////////////
+
+void SNetStatsView::GroupBy_OnSelectionChanged(TSharedPtr<ENetEventGroupingMode> NewGroupingMode, ESelectInfo::Type SelectInfo)
+{
+	if (SelectInfo != ESelectInfo::Direct)
+	{
+		GroupingMode = *NewGroupingMode;
+
+		CreateGroups();
+		SortTreeNodes();
+		ApplyFiltering();
+	}
+}
+
+////////////////////////////////////////////////////////////////////////////////////////////////////
+
+TSharedRef<SWidget> SNetStatsView::GroupBy_OnGenerateWidget(TSharedPtr<ENetEventGroupingMode> InGroupingMode) const
+{
+	return SNew(STextBlock)
+		.Text(NetEventNodeGroupingHelper::ToText(*InGroupingMode))
+		.ToolTipText(NetEventNodeGroupingHelper::ToDescription(*InGroupingMode));
+}
+
+////////////////////////////////////////////////////////////////////////////////////////////////////
+
+FText SNetStatsView::GroupBy_GetSelectedText() const
+{
+	return NetEventNodeGroupingHelper::ToText(GroupingMode);
+}
+
+////////////////////////////////////////////////////////////////////////////////////////////////////
+
+FText SNetStatsView::GroupBy_GetSelectedTooltipText() const
+{
+	return NetEventNodeGroupingHelper::ToDescription(GroupingMode);
+}
+
+////////////////////////////////////////////////////////////////////////////////////////////////////
+// Sorting
+////////////////////////////////////////////////////////////////////////////////////////////////////
+
+const FName SNetStatsView::GetDefaultColumnBeingSorted()
+{
+	return FNetStatsViewColumns::TotalInclusiveSizeColumnID;
+}
+
+////////////////////////////////////////////////////////////////////////////////////////////////////
+
+const EColumnSortMode::Type SNetStatsView::GetDefaultColumnSortMode()
+{
+	return EColumnSortMode::Descending;
+}
+
+////////////////////////////////////////////////////////////////////////////////////////////////////
+
+void SNetStatsView::CreateSortings()
+{
+	AvailableSorters.Reset();
+	CurrentSorter = nullptr;
+
+	for (const TSharedPtr<Insights::FTableColumn> ColumnPtr : Table->GetColumns())
+	{
+		if (ColumnPtr->CanBeSorted())
+		{
+			TSharedPtr<Insights::ITableCellValueSorter> SorterPtr = ColumnPtr->GetValueSorter();
+			if (ensure(SorterPtr.IsValid()))
+			{
+				AvailableSorters.Add(SorterPtr);
+			}
+		}
+	}
+
+	UpdateCurrentSortingByColumn();
+}
+
+////////////////////////////////////////////////////////////////////////////////////////////////////
+
+void SNetStatsView::UpdateCurrentSortingByColumn()
+{
+	TSharedPtr<Insights::FTableColumn> ColumnPtr = Table->FindColumn(ColumnBeingSorted);
+	CurrentSorter = ColumnPtr.IsValid() ? ColumnPtr->GetValueSorter() : nullptr;
+}
+
+////////////////////////////////////////////////////////////////////////////////////////////////////
+
+void SNetStatsView::SortTreeNodes()
+{
+	if (CurrentSorter.IsValid())
+	{
+		for (FNetEventNodePtr& Root : GroupNodes)
+		{
+			SortTreeNodesRec(*Root, *CurrentSorter);
+		}
+	}
+}
+
+////////////////////////////////////////////////////////////////////////////////////////////////////
+
+void SNetStatsView::SortTreeNodesRec(FNetEventNode& Node, const Insights::ITableCellValueSorter& Sorter)
+{
+	if (ColumnSortMode == EColumnSortMode::Type::Descending)
+	{
+		Node.SortChildrenDescending(Sorter);
+	}
+	else // if (ColumnSortMode == EColumnSortMode::Type::Ascending)
+	{
+		Node.SortChildrenAscending(Sorter);
+	}
+
+	for (Insights::FBaseTreeNodePtr ChildPtr : Node.GetChildren())
+	{
+		//if (ChildPtr->IsGroup())
+		if (ChildPtr->GetChildren().Num() > 0)
+		{
+			SortTreeNodesRec(*StaticCastSharedPtr<FNetEventNode>(ChildPtr), Sorter);
+		}
+	}
+}
+
+////////////////////////////////////////////////////////////////////////////////////////////////////
+
+EColumnSortMode::Type SNetStatsView::GetSortModeForColumn(const FName ColumnId) const
+{
+	if (ColumnBeingSorted != ColumnId)
+	{
+		return EColumnSortMode::None;
+	}
+
+	return ColumnSortMode;
+}
+
+////////////////////////////////////////////////////////////////////////////////////////////////////
+
+void SNetStatsView::SetSortModeForColumn(const FName& ColumnId, const EColumnSortMode::Type SortMode)
+{
+	ColumnBeingSorted = ColumnId;
+	ColumnSortMode = SortMode;
+	UpdateCurrentSortingByColumn();
+
+	SortTreeNodes();
+	ApplyFiltering();
+}
+
+////////////////////////////////////////////////////////////////////////////////////////////////////
+
+void SNetStatsView::OnSortModeChanged(const EColumnSortPriority::Type SortPriority, const FName& ColumnId, const EColumnSortMode::Type SortMode)
+{
+	SetSortModeForColumn(ColumnId, SortMode);
+	TreeView_Refresh();
+}
+
+////////////////////////////////////////////////////////////////////////////////////////////////////
+// SortMode action (HeaderMenu)
+////////////////////////////////////////////////////////////////////////////////////////////////////
+
+bool SNetStatsView::HeaderMenu_SortMode_IsChecked(const FName ColumnId, const EColumnSortMode::Type InSortMode)
+{
+	return ColumnBeingSorted == ColumnId && ColumnSortMode == InSortMode;
+}
+
+////////////////////////////////////////////////////////////////////////////////////////////////////
+
+bool SNetStatsView::HeaderMenu_SortMode_CanExecute(const FName ColumnId, const EColumnSortMode::Type InSortMode) const
+{
+	const Insights::FTableColumn& Column = *Table->FindColumnChecked(ColumnId);
+	return Column.CanBeSorted();
+}
+
+////////////////////////////////////////////////////////////////////////////////////////////////////
+
+void SNetStatsView::HeaderMenu_SortMode_Execute(const FName ColumnId, const EColumnSortMode::Type InSortMode)
+{
+	SetSortModeForColumn(ColumnId, InSortMode);
+	TreeView_Refresh();
+}
+
+////////////////////////////////////////////////////////////////////////////////////////////////////
+// SortMode action (ContextMenu)
+////////////////////////////////////////////////////////////////////////////////////////////////////
+
+bool SNetStatsView::ContextMenu_SortMode_IsChecked(const EColumnSortMode::Type InSortMode)
+{
+	return ColumnSortMode == InSortMode;
+}
+
+////////////////////////////////////////////////////////////////////////////////////////////////////
+
+bool SNetStatsView::ContextMenu_SortMode_CanExecute(const EColumnSortMode::Type InSortMode) const
+{
+	return true; //ColumnSortMode != InSortMode;
+}
+
+////////////////////////////////////////////////////////////////////////////////////////////////////
+
+void SNetStatsView::ContextMenu_SortMode_Execute(const EColumnSortMode::Type InSortMode)
+{
+	SetSortModeForColumn(ColumnBeingSorted, InSortMode);
+	TreeView_Refresh();
+}
+
+////////////////////////////////////////////////////////////////////////////////////////////////////
+// SortByColumn action (ContextMenu)
+////////////////////////////////////////////////////////////////////////////////////////////////////
+
+bool SNetStatsView::ContextMenu_SortByColumn_IsChecked(const FName ColumnId)
+{
+	return ColumnId == ColumnBeingSorted;
+}
+
+////////////////////////////////////////////////////////////////////////////////////////////////////
+
+bool SNetStatsView::ContextMenu_SortByColumn_CanExecute(const FName ColumnId) const
+{
+	return true; //ColumnId != ColumnBeingSorted;
+}
+
+////////////////////////////////////////////////////////////////////////////////////////////////////
+
+void SNetStatsView::ContextMenu_SortByColumn_Execute(const FName ColumnId)
+{
+	SetSortModeForColumn(ColumnId, EColumnSortMode::Descending);
+	TreeView_Refresh();
+}
+
+////////////////////////////////////////////////////////////////////////////////////////////////////
+// ShowColumn action
+////////////////////////////////////////////////////////////////////////////////////////////////////
+
+bool SNetStatsView::CanShowColumn(const FName ColumnId) const
+{
+	return true;
+}
+
+////////////////////////////////////////////////////////////////////////////////////////////////////
+
+void SNetStatsView::ShowColumn(const FName ColumnId)
+{
+	Insights::FTableColumn& Column = *Table->FindColumnChecked(ColumnId);
+	Column.Show();
+
+	SHeaderRow::FColumn::FArguments ColumnArgs;
+	ColumnArgs
+		.ColumnId(Column.GetId())
+		.DefaultLabel(Column.GetShortName())
+		.HAlignHeader(HAlign_Fill)
+		.VAlignHeader(VAlign_Fill)
+		.HeaderContentPadding(FMargin(2.0f))
+		.HAlignCell(HAlign_Fill)
+		.VAlignCell(VAlign_Fill)
+		.SortMode(this, &SNetStatsView::GetSortModeForColumn, Column.GetId())
+		.OnSort(this, &SNetStatsView::OnSortModeChanged)
+		.ManualWidth(Column.GetInitialWidth())
+		.FixedWidth(Column.IsFixedWidth() ? Column.GetInitialWidth() : TOptional<float>())
+		.HeaderContent()
+		[
+			SNew(SBox)
+			.ToolTip(SNetStatsViewTooltip::GetColumnTooltip(Column))
+			.HAlign(Column.GetHorizontalAlignment())
+			.VAlign(VAlign_Center)
+			[
+				SNew(STextBlock)
+				.Text(this, &SNetStatsView::GetColumnHeaderText, Column.GetId())
+			]
+		]
+		.MenuContent()
+		[
+			TreeViewHeaderRow_GenerateColumnMenu(Column)
+		];
+
+	int32 ColumnIndex = 0;
+	const int32 NewColumnPosition = Table->GetColumnPositionIndex(ColumnId);
+	const int32 NumColumns = TreeViewHeaderRow->GetColumns().Num();
+	for (; ColumnIndex < NumColumns; ColumnIndex++)
+	{
+		const SHeaderRow::FColumn& CurrentColumn = TreeViewHeaderRow->GetColumns()[ColumnIndex];
+		const int32 CurrentColumnPosition = Table->GetColumnPositionIndex(CurrentColumn.ColumnId);
+		if (NewColumnPosition < CurrentColumnPosition)
+		{
+			break;
+		}
+	}
+
+	TreeViewHeaderRow->InsertColumn(ColumnArgs, ColumnIndex);
+}
+
+////////////////////////////////////////////////////////////////////////////////////////////////////
+// HideColumn action
+////////////////////////////////////////////////////////////////////////////////////////////////////
+
+bool SNetStatsView::CanHideColumn(const FName ColumnId) const
+{
+	const Insights::FTableColumn& Column = *Table->FindColumnChecked(ColumnId);
+	return Column.CanBeHidden();
+}
+
+////////////////////////////////////////////////////////////////////////////////////////////////////
+
+void SNetStatsView::HideColumn(const FName ColumnId)
+{
+	Insights::FTableColumn& Column = *Table->FindColumnChecked(ColumnId);
+	Column.Hide();
+
+	TreeViewHeaderRow->RemoveColumn(ColumnId);
+}
+
+////////////////////////////////////////////////////////////////////////////////////////////////////
+// ToggleColumn action
+////////////////////////////////////////////////////////////////////////////////////////////////////
+
+bool SNetStatsView::IsColumnVisible(const FName ColumnId)
+{
+	const Insights::FTableColumn& Column = *Table->FindColumnChecked(ColumnId);
+	return Column.IsVisible();
+}
+
+////////////////////////////////////////////////////////////////////////////////////////////////////
+
+bool SNetStatsView::CanToggleColumnVisibility(const FName ColumnId) const
+{
+	const Insights::FTableColumn& Column = *Table->FindColumnChecked(ColumnId);
+	return !Column.IsVisible() || Column.CanBeHidden();
+}
+
+////////////////////////////////////////////////////////////////////////////////////////////////////
+
+void SNetStatsView::ToggleColumnVisibility(const FName ColumnId)
+{
+	const Insights::FTableColumn& Column = *Table->FindColumnChecked(ColumnId);
+	if (Column.IsVisible())
+	{
+		HideColumn(ColumnId);
+	}
+	else
+	{
+		ShowColumn(ColumnId);
+	}
+}
+
+////////////////////////////////////////////////////////////////////////////////////////////////////
+// "Show All Columns" action (ContextMenu)
+////////////////////////////////////////////////////////////////////////////////////////////////////
+
+bool SNetStatsView::ContextMenu_ShowAllColumns_CanExecute() const
+{
+	return true;
+}
+
+////////////////////////////////////////////////////////////////////////////////////////////////////
+
+void SNetStatsView::ContextMenu_ShowAllColumns_Execute()
+{
+	ColumnSortMode = GetDefaultColumnSortMode();
+	ColumnBeingSorted = GetDefaultColumnBeingSorted();
+
+	for (const TSharedPtr<Insights::FTableColumn>& ColumnPtr : Table->GetColumns())
+	{
+		const Insights::FTableColumn& Column = *ColumnPtr;
+
+		if (!Column.IsVisible())
+		{
+			ShowColumn(Column.GetId());
+		}
+	}
+}
+
+////////////////////////////////////////////////////////////////////////////////////////////////////
+// "Show Min/Max/Median Columns" action (ContextMenu)
+////////////////////////////////////////////////////////////////////////////////////////////////////
+
+bool SNetStatsView::ContextMenu_ShowMinMaxMedColumns_CanExecute() const
+{
+	return true;
+}
+
+////////////////////////////////////////////////////////////////////////////////////////////////////
+
+void SNetStatsView::ContextMenu_ShowMinMaxMedColumns_Execute()
+{
+	TSet<FName> Preset =
+	{
+		FNetStatsViewColumns::NameColumnID,
+		//FNetStatsViewColumns::MetaGroupNameColumnID,
+		//FNetStatsViewColumns::TypeColumnID,
+		FNetStatsViewColumns::InstanceCountColumnID,
+
+		FNetStatsViewColumns::TotalInclusiveSizeColumnID,
+		FNetStatsViewColumns::MaxInclusiveSizeColumnID,
+		FNetStatsViewColumns::AverageInclusiveSizeColumnID,
+
+		FNetStatsViewColumns::TotalExclusiveSizeColumnID,
+		FNetStatsViewColumns::MaxExclusiveSizeColumnID,
+		//FNetStatsViewColumns::AverageExclusiveSizeColumnID,
+	};
+
+	ColumnBeingSorted = FNetStatsViewColumns::TotalInclusiveSizeColumnID;
+	ColumnSortMode = EColumnSortMode::Descending;
+	UpdateCurrentSortingByColumn();
+
+	for (const TSharedPtr<Insights::FTableColumn>& ColumnPtr : Table->GetColumns())
+	{
+		const Insights::FTableColumn& Column = *ColumnPtr;
+
+		const bool bShouldBeVisible = Preset.Contains(Column.GetId());
+
+		if (bShouldBeVisible && !Column.IsVisible())
+		{
+			ShowColumn(Column.GetId());
+		}
+		else if (!bShouldBeVisible && Column.IsVisible())
+		{
+			HideColumn(Column.GetId());
+		}
+	}
+}
+
+////////////////////////////////////////////////////////////////////////////////////////////////////
+// ResetColumns action (ContextMenu)
+////////////////////////////////////////////////////////////////////////////////////////////////////
+
+bool SNetStatsView::ContextMenu_ResetColumns_CanExecute() const
+{
+	return true;
+}
+
+////////////////////////////////////////////////////////////////////////////////////////////////////
+
+void SNetStatsView::ContextMenu_ResetColumns_Execute()
+{
+	ColumnBeingSorted = GetDefaultColumnBeingSorted();
+	ColumnSortMode = GetDefaultColumnSortMode();
+	UpdateCurrentSortingByColumn();
+
+	for (const TSharedPtr<Insights::FTableColumn>& ColumnPtr : Table->GetColumns())
+	{
+		const Insights::FTableColumn& Column = *ColumnPtr;
+
+		if (Column.ShouldBeVisible() && !Column.IsVisible())
+		{
+			ShowColumn(Column.GetId());
+		}
+		else if (!Column.ShouldBeVisible() && Column.IsVisible())
+		{
+			HideColumn(Column.GetId());
+		}
+	}
+}
+
+////////////////////////////////////////////////////////////////////////////////////////////////////
+
+void SNetStatsView::RebuildTree(bool bResync)
+{
+	TArray<FNetEventNodePtr> SelectedItems;
+	bool bListHasChanged = false;
+
+	if (bResync)
+	{
+		const int32 PreviousNodeCount = NetEventNodes.Num();
+		NetEventNodes.Empty(PreviousNodeCount);
+		NetEventNodesIdMap.Empty(PreviousNodeCount);
+		bListHasChanged = true;
+	}
+
+	if (Session.IsValid())
+	{
+		Trace::FAnalysisSessionReadScope SessionReadScope(*Session.Get());
+
+		const Trace::INetProfilerProvider& NetProfilerProvider = Trace::ReadNetProfilerProvider(*Session.Get());
+
+		NetProfilerProvider.ReadEventTypes([this, &bResync, &SelectedItems, &bListHasChanged, &NetProfilerProvider](const Trace::FNetProfilerEventType* NetEvents, uint64 NetEventCount)
+		{
+			if (NetEventCount != NetEventNodes.Num())
+			{
+				bResync = true;
+			}
+
+			if (bResync)
+			{
+				// Save selection.
+				TreeView->GetSelectedItems(SelectedItems);
+
+				const int32 PreviousNodeCount = NetEventNodes.Num();
+				NetEventNodes.Empty(PreviousNodeCount);
+				NetEventNodesIdMap.Empty(PreviousNodeCount);
+				bListHasChanged = true;
+
+				for (uint64 Index = 0; Index < NetEventCount; ++Index)
+				{
+					const Trace::FNetProfilerEventType& NetEvent = NetEvents[Index];
+
+					const uint64 Id = NetEvent.EventTypeIndex;
+
+					const TCHAR* NamePtr;
+					NetProfilerProvider.ReadName(NetEvent.NameIndex, [&NamePtr](const Trace::FNetProfilerName& InName)
+					{
+						NamePtr = InName.Name;
+					});
+					const FName Name(NamePtr);
+					const ENetEventNodeType Type = ENetEventNodeType::NetEvent;
+					FNetEventNodePtr NetEventNodePtr = MakeShareable(new FNetEventNode(Id, Name, Type, NetEvent.Level));
+					NetEventNodes.Add(NetEventNodePtr);
+					NetEventNodesIdMap.Add(Id, NetEventNodePtr);
+				}
+			}
+		});
+	}
+
+	if (bListHasChanged)
+	{
+		//UpdateTree();
+		UpdateStatsInternal();
+
+		// Restore selection.
+		if (SelectedItems.Num() > 0)
+		{
+			TreeView->ClearSelection();
+			TArray<FNetEventNodePtr> NewSelectedItems;
+			for (const FNetEventNodePtr& NetEventNode : SelectedItems)
+			{
+				FNetEventNodePtr* NetEventNodePtrPtr = NetEventNodesIdMap.Find(NetEventNode->GetId());
+				if (NetEventNodePtrPtr != nullptr)
+				{
+					NewSelectedItems.Add(*NetEventNodePtrPtr);
+				}
+			}
+			if (NewSelectedItems.Num() > 0)
+			{
+				TreeView->SetItemSelection(NewSelectedItems, true);
+				TreeView->RequestScrollIntoView(NewSelectedItems[0]);
+			}
+		}
+	}
+}
+
+////////////////////////////////////////////////////////////////////////////////////////////////////
+
+void SNetStatsView::ResetStats()
+{
+	GameInstanceIndex = 0;
+	ConnectionIndex = 0;
+	ConnectionMode = Trace::ENetProfilerConnectionMode::Outgoing;
+	StatsPacketStartIndex = 0;
+	StatsPacketEndIndex = 0;
+	StatsStartPosition = 0;
+	StatsEndPosition = 0;
+
+	UpdateStatsInternal();
+}
+
+////////////////////////////////////////////////////////////////////////////////////////////////////
+
+void SNetStatsView::UpdateStats(uint32 InGameInstanceIndex, uint32 InConnectionIndex, Trace::ENetProfilerConnectionMode InConnectionMode, uint32 InStatsPacketStartIndex, uint32 InStatsPacketEndIndex, uint32 InStatsStartPosition, uint32 InStatsEndPosition)
+{
+	GameInstanceIndex = InGameInstanceIndex;
+	ConnectionIndex = InConnectionIndex;
+	ConnectionMode = InConnectionMode;
+	StatsPacketStartIndex = InStatsPacketStartIndex;
+	StatsPacketEndIndex = InStatsPacketEndIndex;
+	StatsStartPosition = InStatsStartPosition;
+	StatsEndPosition = InStatsEndPosition;
+
+	UpdateStatsInternal();
+}
+
+////////////////////////////////////////////////////////////////////////////////////////////////////
+
+void SNetStatsView::UpdateStatsInternal()
+{
+	for (const FNetEventNodePtr& NetEventNodePtr : NetEventNodes)
+	{
+		NetEventNodePtr->ResetAggregatedStats();
+	}
+
+	if (Session.IsValid() && StatsStartPosition < StatsEndPosition)
+	{
+		TUniquePtr<Trace::ITable<Trace::FNetProfilerAggregatedStats>> AggregatedStatsTable;
+
+		{
+			Trace::FAnalysisSessionReadScope SessionReadScope(*Session.Get());
+			const Trace::INetProfilerProvider& NetProfilerProvider = Trace::ReadNetProfilerProvider(*Session.Get());
+			AggregatedStatsTable.Reset(NetProfilerProvider.CreateAggregation(ConnectionIndex, ConnectionMode, StatsPacketStartIndex, StatsPacketEndIndex - 1, StatsStartPosition, StatsEndPosition));
+		}
+
+		if (AggregatedStatsTable.IsValid())
+		{
+			TUniquePtr<Trace::ITableReader<Trace::FNetProfilerAggregatedStats>> TableReader(AggregatedStatsTable->CreateReader());
+			while (TableReader->IsValid())
+			{
+				const Trace::FNetProfilerAggregatedStats* Row = TableReader->GetCurrentRow();
+				FNetEventNodePtr* NetEventNodePtrPtr = NetEventNodesIdMap.Find(static_cast<uint64>(Row->EventTypeIndex));
+				if (NetEventNodePtrPtr != nullptr)
+				{
+					(*NetEventNodePtrPtr)->SetAggregatedStats(*Row);
+				}
+				TableReader->NextRow();
+			}
+		}
+	}
+
+	UpdateTree();
+
+	// Save selection.
+	TArray<FNetEventNodePtr> SelectedItems = TreeView->GetSelectedItems();
+
+	TreeView->RebuildList();
+
+	// Restore selection.
+	if (SelectedItems.Num() > 0)
+	{
+		TreeView->ClearSelection();
+		TreeView->SetItemSelection(SelectedItems, true);
+		TreeView->RequestScrollIntoView(SelectedItems[0]);
+	}
+}
+
+////////////////////////////////////////////////////////////////////////////////////////////////////
+
+void SNetStatsView::SelectNetEventNode(uint64 Id)
+{
+	FNetEventNodePtr* NodePtrPtr = NetEventNodesIdMap.Find(Id);
+	if (NodePtrPtr != nullptr)
+	{
+		FNetEventNodePtr NodePtr = *NodePtrPtr;
+
+		TreeView->SetSelection(NodePtr);
+		TreeView->RequestScrollIntoView(NodePtr);
+	}
+}
+
+////////////////////////////////////////////////////////////////////////////////////////////////////
+
+#undef LOCTEXT_NAMESPACE
diff --git a/Engine/Source/Developer/TraceInsights/Private/Insights/NetworkingProfiler/Widgets/SNetStatsView.h b/Engine/Source/Developer/TraceInsights/Private/Insights/NetworkingProfiler/Widgets/SNetStatsView.h
new file mode 100644
--- /dev/null
+++ b/Engine/Source/Developer/TraceInsights/Private/Insights/NetworkingProfiler/Widgets/SNetStatsView.h
@@ -0,0 +1,362 @@
+// Copyright 1998-2019 Epic Games, Inc. All Rights Reserved.
+
+#pragma once
+
+#include "CoreMinimal.h"
+#include "Misc/FilterCollection.h"
+#include "Misc/TextFilter.h"
+#include "SlateFwd.h"
+#include "Widgets/DeclarativeSyntaxSupport.h"
+#include "Widgets/Input/SComboBox.h"
+#include "Widgets/SCompoundWidget.h"
+#include "Widgets/SWidget.h"
+#include "Widgets/Views/SHeaderRow.h"
+#include "Widgets/Views/STableRow.h"
+#include "Widgets/Views/STableViewBase.h"
+#include "Widgets/Views/STreeView.h"
+
+// Insights
+#include "Insights/NetworkingProfiler/ViewModels/NetEventGroupingAndSorting.h"
+#include "Insights/NetworkingProfiler/ViewModels/NetEventNode.h"
+
+////////////////////////////////////////////////////////////////////////////////////////////////////
+
+class FMenuBuilder;
+
+namespace Trace
+{
+	class IAnalysisSession;
+}
+
+namespace Insights
+{
+	class FTable;
+	class FTableColumn;
+	class ITableCellValueSorter;
+}
+
+////////////////////////////////////////////////////////////////////////////////////////////////////
+
+/** The filter collection - used for updating the list of tree nodes. */
+typedef TFilterCollection<const FNetEventNodePtr&> FNetEventNodeFilterCollection;
+
+/** The text based filter - used for updating the list of tree nodes. */
+typedef TTextFilter<const FNetEventNodePtr&> FNetEventNodeTextFilter;
+
+////////////////////////////////////////////////////////////////////////////////////////////////////
+/**
+ * A custom widget used to display the list of net events and thier aggregated stats.
+ */
+class SNetStatsView : public SCompoundWidget
+{
+public:
+	/** Default constructor. */
+	SNetStatsView();
+
+	/** Virtual destructor. */
+	virtual ~SNetStatsView();
+
+	SLATE_BEGIN_ARGS(SNetStatsView){}
+	SLATE_END_ARGS()
+
+	/**
+	 * Construct this widget
+	 * @param InArgs - The declaration data for this widget
+	 */
+	void Construct(const FArguments& InArgs);
+
+	TSharedPtr<Insights::FTable> GetTable() const { return Table; }
+
+	/**
+	 * Ticks this widget. Override in derived classes, but always call the parent implementation.
+	 *
+	 * @param  AllottedGeometry - the space allotted for this widget
+	 * @param  InCurrentTime    - current absolute real time
+	 * @param  InDeltaTime      - real time passed since last tick
+	 */
+	virtual void Tick(const FGeometry& AllottedGeometry, const double InCurrentTime, const float InDeltaTime) override;
+
+	/**
+	 * Rebuilds the tree (if necessary).
+	 * @param bResync - If true, it forces a resync with list of net events from Analysis, even if the list did not changed since last sync.
+	 */
+	void RebuildTree(bool bResync = true);
+
+	void ResetStats();
+	void UpdateStats(uint32 InGameInstanceIndex, uint32 InConnectionIndex, Trace::ENetProfilerConnectionMode InConnectionMode, uint32 InStatsPacketStartIndex, uint32 InStatsPacketEndIndex, uint32 InStatsStartPosition, uint32 InStatsEndPosition);
+
+	void SelectNetEventNode(uint64 Id);
+
+	//const TSet<FNetEventNodePtr>& GetNetEventNodes() const { return NetEventNodes; }
+	//const TMap<uint64, FNetEventNodePtr> GetNetEventNodesIdMap() const { return NetEventNodesIdMap; }
+	const FNetEventNodePtr* GetNetEventNode(uint64 Id) const { return NetEventNodesIdMap.Find(Id); }
+
+protected:
+	void UpdateTree();
+
+	void UpdateStatsInternal();
+
+	/** Called when the analysis session has changed. */
+	void InsightsManager_OnSessionChanged();
+
+	/**
+	 * Populates OutSearchStrings with the strings that should be used in searching.
+	 *
+	 * @param GroupOrStatNodePtr - the group and stat node to get a text description from.
+	 * @param OutSearchStrings   - an array of strings to use in searching.
+	 *
+	 */
+	void HandleItemToStringArray(const FNetEventNodePtr& GroupOrStatNodePtr, TArray<FString>& OutSearchStrings) const;
+
+	////////////////////////////////////////////////////////////////////////////////////////////////////
+	// Tree View - Context Menu
+
+	TSharedPtr<SWidget> TreeView_GetMenuContent();
+	void TreeView_BuildSortByMenu(FMenuBuilder& MenuBuilder);
+	void TreeView_BuildViewColumnMenu(FMenuBuilder& MenuBuilder);
+
+	////////////////////////////////////////////////////////////////////////////////////////////////////
+	// Tree View - Columns' Header
+
+	void InitializeAndShowHeaderColumns();
+
+	FText GetColumnHeaderText(const FName ColumnId) const;
+
+	TSharedRef<SWidget> TreeViewHeaderRow_GenerateColumnMenu(const Insights::FTableColumn& Column);
+
+	////////////////////////////////////////////////////////////////////////////////////////////////////
+	// Tree View - Misc
+
+	void TreeView_Refresh();
+
+	/**
+	 * Called by STreeView to retrieves the children for the specified parent item.
+	 * @param InParent    - The parent node to retrieve the children from.
+	 * @param OutChildren - List of children for the parent node.
+	 */
+	void TreeView_OnGetChildren(FNetEventNodePtr InParent, TArray<FNetEventNodePtr>& OutChildren);
+
+	/** Called by STreeView when selection has changed. */
+	void TreeView_OnSelectionChanged(FNetEventNodePtr SelectedItem, ESelectInfo::Type SelectInfo);
+
+	/** Called by STreeView when a tree item is double clicked. */
+	void TreeView_OnMouseButtonDoubleClick(FNetEventNodePtr TreeNode);
+
+	////////////////////////////////////////////////////////////////////////////////////////////////////
+	// Tree View - Table Row
+
+	/** Called by STreeView to generate a table row for the specified item. */
+	TSharedRef<ITableRow> TreeView_OnGenerateRow(FNetEventNodePtr TreeNode, const TSharedRef<STableViewBase>& OwnerTable);
+
+	void TableRow_SetHoveredCell(TSharedPtr<Insights::FTable> TablePtr, TSharedPtr<Insights::FTableColumn> ColumnPtr, const FNetEventNodePtr NodePtr);
+	EHorizontalAlignment TableRow_GetColumnOutlineHAlignment(const FName ColumnId) const;
+
+	FText TableRow_GetHighlightText() const;
+	FName TableRow_GetHighlightedNodeName() const;
+
+	bool TableRow_ShouldBeEnabled(const uint32 NodeId) const;
+
+	////////////////////////////////////////////////////////////////////////////////////////////////////
+	// Filtering
+
+	/** Populates the group and stat tree with items based on the current data. */
+	void ApplyFiltering();
+
+	TSharedRef<SWidget> GetToggleButtonForNetEventType(const ENetEventNodeType InNetEventType);
+	void FilterByNetEventType_OnCheckStateChanged(ECheckBoxState NewRadioState, const ENetEventNodeType InNetEventType);
+	ECheckBoxState FilterByNetEventType_IsChecked(const ENetEventNodeType InNetEventType) const;
+
+	bool SearchBox_IsEnabled() const;
+	void SearchBox_OnTextChanged(const FText& InFilterText);
+
+	////////////////////////////////////////////////////////////////////////////////////////////////////
+	// Grouping
+
+	void CreateGroups();
+	void CreateGroupByOptionsSources();
+
+	void GroupBy_OnSelectionChanged(TSharedPtr<ENetEventGroupingMode> NewGroupingMode, ESelectInfo::Type SelectInfo);
+
+	TSharedRef<SWidget> GroupBy_OnGenerateWidget(TSharedPtr<ENetEventGroupingMode> InGroupingMode) const;
+
+	FText GroupBy_GetSelectedText() const;
+
+	FText GroupBy_GetSelectedTooltipText() const;
+
+	////////////////////////////////////////////////////////////////////////////////////////////////////
+	// Sorting
+
+	static const FName GetDefaultColumnBeingSorted();
+	static const EColumnSortMode::Type GetDefaultColumnSortMode();
+
+	void CreateSortings();
+
+	void UpdateCurrentSortingByColumn();
+	void SortTreeNodes();
+	void SortTreeNodesRec(FNetEventNode& Node, const Insights::ITableCellValueSorter& Sorter);
+
+	EColumnSortMode::Type GetSortModeForColumn(const FName ColumnId) const;
+	void SetSortModeForColumn(const FName& ColumnId, EColumnSortMode::Type SortMode);
+	void OnSortModeChanged(const EColumnSortPriority::Type SortPriority, const FName& ColumnId, const EColumnSortMode::Type SortMode);
+
+	////////////////////////////////////////////////////////////////////////////////////////////////////
+	// Sorting actions
+
+	// SortMode (HeaderMenu)
+	bool HeaderMenu_SortMode_IsChecked(const FName ColumnId, const EColumnSortMode::Type InSortMode);
+	bool HeaderMenu_SortMode_CanExecute(const FName ColumnId, const EColumnSortMode::Type InSortMode) const;
+	void HeaderMenu_SortMode_Execute(const FName ColumnId, const EColumnSortMode::Type InSortMode);
+
+	// SortMode (ContextMenu)
+	bool ContextMenu_SortMode_IsChecked(const EColumnSortMode::Type InSortMode);
+	bool ContextMenu_SortMode_CanExecute(const EColumnSortMode::Type InSortMode) const;
+	void ContextMenu_SortMode_Execute(const EColumnSortMode::Type InSortMode);
+
+	// SortByColumn (ContextMenu)
+	bool ContextMenu_SortByColumn_IsChecked(const FName ColumnId);
+	bool ContextMenu_SortByColumn_CanExecute(const FName ColumnId) const;
+	void ContextMenu_SortByColumn_Execute(const FName ColumnId);
+
+	////////////////////////////////////////////////////////////////////////////////////////////////////
+	// Column visibility actions
+
+	// ShowColumn
+	bool CanShowColumn(const FName ColumnId) const;
+	void ShowColumn(const FName ColumnId);
+
+	// HideColumn
+	bool CanHideColumn(const FName ColumnId) const;
+	void HideColumn(const FName ColumnId);
+
+	// ToggleColumnVisibility
+	bool IsColumnVisible(const FName ColumnId);
+	bool CanToggleColumnVisibility(const FName ColumnId) const;
+	void ToggleColumnVisibility(const FName ColumnId);
+
+	// ShowAllColumns (ContextMenu)
+	bool ContextMenu_ShowAllColumns_CanExecute() const;
+	void ContextMenu_ShowAllColumns_Execute();
+
+	// MinMaxMedColumns (ContextMenu)
+	bool ContextMenu_ShowMinMaxMedColumns_CanExecute() const;
+	void ContextMenu_ShowMinMaxMedColumns_Execute();
+
+	// ResetColumns (ContextMenu)
+	bool ContextMenu_ResetColumns_CanExecute() const;
+	void ContextMenu_ResetColumns_Execute();
+
+	////////////////////////////////////////////////////////////////////////////////////////////////////
+
+protected:
+	/** Table view model. */
+	TSharedPtr<Insights::FTable> Table;
+
+	/** A weak pointer to the profiler session used to populate this widget. */
+	TSharedPtr<const Trace::IAnalysisSession>/*Weak*/ Session;
+
+	//////////////////////////////////////////////////
+	// Tree View, Columns
+
+	/** The tree widget which holds the net event tree nodes. */
+	TSharedPtr<STreeView<FNetEventNodePtr>> TreeView;
+
+	/** Holds the tree view header row widget which display all columns in the tree view. */
+	TSharedPtr<SHeaderRow> TreeViewHeaderRow;
+
+	/** External scrollbar used to synchronize tree view position. */
+	TSharedPtr<SScrollBar> ExternalScrollbar;
+
+	//////////////////////////////////////////////////
+	// Hovered Column, Hovered Node
+
+	/** Name of the column currently being hovered by the mouse. */
+	FName HoveredColumnId;
+
+	/** A shared pointer to the node currently being hovered by the mouse. */
+	FNetEventNodePtr HoveredNodePtr;
+
+	/** Name of the node that should be drawn as highlighted. */
+	FName HighlightedNodeName;
+
+	//////////////////////////////////////////////////
+	// Net Event Nodes
+
+	/** An array of group nodes. */
+	TArray<FNetEventNodePtr> GroupNodes;
+
+	/** A filtered array of group nodes to be displayed in the tree widget. */
+	TArray<FNetEventNodePtr> FilteredGroupNodes;
+
+	/** All net event nodes. */
+	TSet<FNetEventNodePtr> NetEventNodes;
+
+	/** All net event nodes, stored as NodeId -> FNetEventNodePtr. */
+	TMap<uint64, FNetEventNodePtr> NetEventNodesIdMap;
+
+	/** Currently expanded group nodes. */
+	TSet<FNetEventNodePtr> ExpandedNodes;
+
+	/** If true, the expanded nodes have been saved before applying a text filter. */
+	bool bExpansionSaved;
+
+	//////////////////////////////////////////////////
+	// Search box and filters
+
+	//bool bUseFiltering;
+
+	/** The search box widget used to filter items displayed in the stats and groups tree. */
+	TSharedPtr<SSearchBox> SearchBox;
+
+	/** The text based filter. */
+	TSharedPtr<FNetEventNodeTextFilter> TextFilter;
+
+	/** The filter collection. */
+	TSharedPtr<FNetEventNodeFilterCollection> Filters;
+
+	/** Holds the visibility of each net event type. */
+	bool bNetEventTypeIsVisible[static_cast<int>(ENetEventNodeType::InvalidOrMax)];
+
+	//////////////////////////////////////////////////
+	// Grouping
+
+	//bool bUseGrouping;
+
+	TArray<TSharedPtr<ENetEventGroupingMode>> GroupByOptionsSource;
+
+	TSharedPtr<SComboBox<TSharedPtr<ENetEventGroupingMode>>> GroupByComboBox;
+
+	/** How we group the net event nodes? */
+	ENetEventGroupingMode GroupingMode;
+
+	//////////////////////////////////////////////////
+	// Sorting
+
+	//bool bUseSorting;
+
+	/** All available sorters. */
+	TArray<TSharedPtr<Insights::ITableCellValueSorter>> AvailableSorters;
+
+	/** Current sorter. It is nullptr if sorting is disabled. */
+	TSharedPtr<Insights::ITableCellValueSorter> CurrentSorter;
+
+	/** Name of the column currently being sorted. Can be NAME_None if sorting is disabled (CurrentSorting == nullptr) or if a complex sorting is used (CurrentSorting != nullptr). */
+	FName ColumnBeingSorted;
+
+	/** How we sort the nodes? Ascending or Descending. */
+	EColumnSortMode::Type ColumnSortMode;
+
+	//////////////////////////////////////////////////
+
+	uint64 NextTimestamp;
+	uint32 ObjectsChangeCount;
+
+	uint32 GameInstanceIndex;
+	uint32 ConnectionIndex;
+	Trace::ENetProfilerConnectionMode ConnectionMode;
+	uint32 StatsPacketStartIndex;
+	uint32 StatsPacketEndIndex;
+	uint32 StatsStartPosition;
+	uint32 StatsEndPosition;
+};
+
+////////////////////////////////////////////////////////////////////////////////////////////////////
diff --git a/Engine/Source/Developer/TraceInsights/Private/Insights/NetworkingProfiler/Widgets/SNetStatsViewTooltip.cpp b/Engine/Source/Developer/TraceInsights/Private/Insights/NetworkingProfiler/Widgets/SNetStatsViewTooltip.cpp
new file mode 100644
--- /dev/null
+++ b/Engine/Source/Developer/TraceInsights/Private/Insights/NetworkingProfiler/Widgets/SNetStatsViewTooltip.cpp
@@ -0,0 +1,316 @@
+// Copyright 1998-2019 Epic Games, Inc. All Rights Reserved.
+
+#include "SNetStatsViewTooltip.h"
+
+#include "EditorStyleSet.h"
+#include "SlateOptMacros.h"
+#include "TraceServices/Model/NetProfiler.h"
+#include "Widgets/Layout/SGridPanel.h"
+#include "Widgets/Layout/SSeparator.h"
+#include "Widgets/SBoxPanel.h"
+#include "Widgets/SToolTip.h"
+#include "Widgets/Text/STextBlock.h"
+
+// Insights
+#include "Insights/Table/ViewModels/Table.h"
+#include "Insights/Table/ViewModels/TableColumn.h"
+#include "Insights/NetworkingProfiler/ViewModels/NetEventNode.h"
+#include "Insights/NetworkingProfiler/ViewModels/NetEventNodeHelper.h"
+
+////////////////////////////////////////////////////////////////////////////////////////////////////
+
+#define LOCTEXT_NAMESPACE "SNetStatsView"
+
+////////////////////////////////////////////////////////////////////////////////////////////////////
+
+BEGIN_SLATE_FUNCTION_BUILD_OPTIMIZATION
+
+TSharedPtr<SToolTip> SNetStatsViewTooltip::GetTableTooltip(const Insights::FTable& Table)
+{
+	TSharedPtr<SToolTip> ColumnTooltip =
+		SNew(SToolTip)
+		[
+			SNew(SVerticalBox)
+
+			+ SVerticalBox::Slot()
+			.AutoHeight()
+			.Padding(2.0f)
+			[
+				SNew(STextBlock)
+				.Text(Table.GetDisplayName())
+				.TextStyle(FEditorStyle::Get(), TEXT("Profiler.TooltipBold"))
+			]
+
+			+ SVerticalBox::Slot()
+			.AutoHeight()
+			.Padding(2.0f)
+			[
+				SNew(STextBlock)
+				.Text(Table.GetDescription())
+				.TextStyle(FEditorStyle::Get(), TEXT("Profiler.Tooltip"))
+			]
+		];
+
+	return ColumnTooltip;
+}
+
+////////////////////////////////////////////////////////////////////////////////////////////////////
+
+TSharedPtr<SToolTip> SNetStatsViewTooltip::GetColumnTooltip(const Insights::FTableColumn& Column)
+{
+	TSharedPtr<SToolTip> ColumnTooltip =
+		SNew(SToolTip)
+		[
+			SNew(SVerticalBox)
+
+			+ SVerticalBox::Slot()
+			.AutoHeight()
+			.Padding(2.0f)
+			[
+				SNew(STextBlock)
+				.Text(Column.GetTitleName())
+				.TextStyle(FEditorStyle::Get(), TEXT("Profiler.TooltipBold"))
+			]
+
+			+ SVerticalBox::Slot()
+			.AutoHeight()
+			.Padding(2.0f)
+			[
+				SNew(STextBlock)
+				.Text(Column.GetDescription())
+				.TextStyle(FEditorStyle::Get(), TEXT("Profiler.Tooltip"))
+			]
+		];
+
+	return ColumnTooltip;
+}
+
+////////////////////////////////////////////////////////////////////////////////////////////////////
+
+TSharedPtr<SToolTip> SNetStatsViewTooltip::GetCellTooltip(const TSharedPtr<FNetEventNode> NetEventNodePtr, const TSharedPtr<Insights::FTableColumn> ColumnPtr)
+{
+	const Trace::FNetProfilerAggregatedStats& Stats = NetEventNodePtr->GetAggregatedStats();
+
+	const FText TotalInclusiveSizeText = FText::AsNumber(Stats.TotalInclusive);
+	const FText MaxInclusiveSizeText = FText::AsNumber(Stats.MaxInclusive);
+	const FText AvgInclusiveSizeText = FText::AsNumber(Stats.AverageInclusive);
+
+	const FText TotalExclusiveSizeText = FText::AsNumber(Stats.TotalExclusive);
+	const FText MaxExclusiveSizeText = FText::AsNumber(Stats.MaxExclusive);
+	//const FText AvgExclusiveSizeText = FText::AsNumber(Stats.AverageExclusive);
+	const FText AvgExclusiveSizeText;
+
+	const FText InstanceCountText = FText::AsNumber(Stats.InstanceCount);
+
+	TSharedPtr<SGridPanel> GridPanel;
+	TSharedPtr<SHorizontalBox> HBox;
+
+	TSharedPtr<SToolTip> TableCellTooltip =
+		SNew(SToolTip)
+		[
+			SAssignNew(HBox, SHorizontalBox)
+
+			+ SHorizontalBox::Slot()
+			.AutoWidth()
+			[
+				SNew(SVerticalBox)
+
+				+ SVerticalBox::Slot()
+				.AutoHeight()
+				.Padding(2.0f)
+				[
+					SNew(SSeparator)
+					.Orientation(Orient_Horizontal)
+				]
+
+				+ SVerticalBox::Slot()
+				.AutoHeight()
+				.Padding(2.0f)
+				[
+					SNew(SGridPanel)
+
+					// Id: [Id]
+					+ SGridPanel::Slot(0, 0)
+					.Padding(2.0f)
+					[
+						SNew(STextBlock)
+						.Text(LOCTEXT("TT_Id", "Id:"))
+						.TextStyle(FEditorStyle::Get(), TEXT("Profiler.TooltipBold"))
+					]
+					+ SGridPanel::Slot(1, 0)
+					.Padding(2.0f)
+					[
+						SNew(STextBlock)
+						.Text(FText::AsNumber(NetEventNodePtr->GetId()))
+						.TextStyle(FEditorStyle::Get(), TEXT("Profiler.Tooltip"))
+					]
+
+					// Name: [Name]
+					+ SGridPanel::Slot(0, 1)
+					.Padding(2.0f)
+					[
+						SNew(STextBlock)
+						.Text(LOCTEXT("TT_Name", "Name:"))
+						.TextStyle(FEditorStyle::Get(), TEXT("Profiler.TooltipBold"))
+					]
+					+ SGridPanel::Slot(1, 1)
+					.Padding(2.0f)
+					[
+						SNew(STextBlock)
+						.Text(FText::FromName(NetEventNodePtr->GetName()))
+						.TextStyle(FEditorStyle::Get(), TEXT("Profiler.Tooltip"))
+					]
+
+					//// Group: [MetaGroupName]
+					//+ SGridPanel::Slot(0, 2)
+					//.Padding(2.0f)
+					//[
+					//	SNew(STextBlock)
+					//	.Text(LOCTEXT("TT_Group", "Group:"))
+					//	.TextStyle(FEditorStyle::Get(), TEXT("Profiler.TooltipBold"))
+					//]
+					//+ SGridPanel::Slot(1, 2)
+					//.Padding(2.0f)
+					//[
+					//	SNew(STextBlock)
+					//	.Text(FText::FromName(NetEventNodePtr->GetMetaGroupName()))
+					//	.TextStyle(FEditorStyle::Get(), TEXT("Profiler.Tooltip"))
+					//]
+
+					// Sizer Type: [Type]
+					+ SGridPanel::Slot(0, 3)
+					.Padding(2.0f)
+					[
+						SNew(STextBlock)
+						.Text(LOCTEXT("TT_Type", "Type:"))
+						.TextStyle(FEditorStyle::Get(), TEXT("Profiler.TooltipBold"))
+					]
+					+ SGridPanel::Slot(1, 3)
+					.Padding(2.0f)
+					[
+						SNew(STextBlock)
+						.Text(NetEventNodeTypeHelper::ToText(NetEventNodePtr->GetType()))
+						.TextStyle(FEditorStyle::Get(), TEXT("Profiler.Tooltip"))
+					]
+				]
+
+				+ SVerticalBox::Slot()
+				.AutoHeight()
+				.Padding(2.0f)
+				[
+					SNew(SSeparator)
+					.Orientation(Orient_Horizontal)
+				]
+
+				+ SVerticalBox::Slot()
+				.AutoHeight()
+				.Padding(2.0f)
+				[
+					SNew(SGridPanel)
+
+					+ SGridPanel::Slot(0, 0)
+					.Padding(2.0f)
+					[
+						SNew(STextBlock)
+						.Text(LOCTEXT("TT_NumInstances", "Num Instances:"))
+						.TextStyle(FEditorStyle::Get(), TEXT("Profiler.TooltipBold"))
+					]
+					+ SGridPanel::Slot(1, 0)
+					.Padding(2.0f)
+					[
+						SNew(STextBlock)
+						.Text(InstanceCountText)
+						.TextStyle(FEditorStyle::Get(), TEXT("Profiler.Tooltip"))
+					]
+				]
+
+				+ SVerticalBox::Slot()
+				.AutoHeight()
+				.Padding(2.0f)
+				[
+					SNew(SSeparator)
+					.Orientation(Orient_Horizontal)
+				]
+
+				+ SVerticalBox::Slot()
+				.AutoHeight()
+				.Padding(2.0f)
+				[
+					SAssignNew(GridPanel, SGridPanel)
+
+					+ SGridPanel::Slot(1, 0)
+					.Padding(2.0f)
+					.HAlign(HAlign_Right)
+					[
+						SNew(STextBlock)
+						.Text(LOCTEXT("TT_InclusiveSize", "Inclusive"))
+						.TextStyle(FEditorStyle::Get(), TEXT("Profiler.TooltipBold"))
+					]
+					+ SGridPanel::Slot(2, 0)
+					.Padding(FMargin(8.0f, 2.0f, 2.0f, 2.0f))
+					.HAlign(HAlign_Right)
+					[
+						SNew(STextBlock)
+						.Text(LOCTEXT("TT_ExclusiveSize", "Exclusive"))
+						.TextStyle(FEditorStyle::Get(), TEXT("Profiler.TooltipBold"))
+					]
+
+					// Stats are added here.
+				]
+
+				+ SVerticalBox::Slot()
+				.AutoHeight()
+				.Padding(2.0f)
+				[
+					SNew(SSeparator)
+					.Orientation(Orient_Horizontal)
+				]
+			]
+		];
+
+	int32 Row = 1;
+	AddStatsRow(GridPanel, Row, LOCTEXT("TT_TotalSize",   "Total Size:"),   TotalInclusiveSizeText, TotalExclusiveSizeText);
+	AddStatsRow(GridPanel, Row, LOCTEXT("TT_MaxSize",     "Max Size:"),     MaxInclusiveSizeText,   MaxExclusiveSizeText);
+	AddStatsRow(GridPanel, Row, LOCTEXT("TT_AverageSize", "Average Size:"), AvgInclusiveSizeText,   AvgExclusiveSizeText);
+
+	return TableCellTooltip;
+}
+
+////////////////////////////////////////////////////////////////////////////////////////////////////
+
+void SNetStatsViewTooltip::AddStatsRow(TSharedPtr<SGridPanel> Grid, int32& Row, const FText& Name, const FText& Value1, const FText& Value2)
+{
+	Grid->AddSlot(0, Row)
+		.Padding(2.0f)
+		[
+			SNew(STextBlock)
+			.Text(Name)
+			.TextStyle(FEditorStyle::Get(), TEXT("Profiler.TooltipBold"))
+		];
+
+	Grid->AddSlot(1, Row)
+		.Padding(2.0f)
+		.HAlign(HAlign_Right)
+		[
+			SNew(STextBlock)
+			.Text(Value1)
+			.TextStyle(FEditorStyle::Get(), TEXT("Profiler.Tooltip"))
+		];
+
+	Grid->AddSlot(2, Row)
+		.Padding(FMargin(8.0f, 2.0f, 2.0f, 2.0f))
+		.HAlign(HAlign_Right)
+		[
+			SNew(STextBlock)
+			.Text(Value2)
+			.TextStyle(FEditorStyle::Get(), TEXT("Profiler.Tooltip"))
+		];
+
+	Row++;
+}
+
+END_SLATE_FUNCTION_BUILD_OPTIMIZATION
+
+////////////////////////////////////////////////////////////////////////////////////////////////////
+
+#undef LOCTEXT_NAMESPACE
diff --git a/Engine/Source/Developer/TraceInsights/Private/Insights/NetworkingProfiler/Widgets/SNetStatsViewTooltip.h b/Engine/Source/Developer/TraceInsights/Private/Insights/NetworkingProfiler/Widgets/SNetStatsViewTooltip.h
new file mode 100644
--- /dev/null
+++ b/Engine/Source/Developer/TraceInsights/Private/Insights/NetworkingProfiler/Widgets/SNetStatsViewTooltip.h
@@ -0,0 +1,37 @@
+// Copyright 1998-2019 Epic Games, Inc. All Rights Reserved.
+
+#pragma once
+
+#include "CoreMinimal.h"
+
+// Insights
+#include "Insights/InsightsManager.h"
+
+class SToolTip;
+class SGridPanel;
+
+namespace Insights
+{
+	class FTable;
+	class FTableColumn;
+}
+
+class FNetEventNode;
+
+#define LOCTEXT_NAMESPACE "SNetStatsView"
+
+/** NetStats View Tooltip */
+class SNetStatsViewTooltip
+{
+public:
+	SNetStatsViewTooltip() = delete;
+
+	static TSharedPtr<SToolTip> GetTableTooltip(const Insights::FTable& Table);
+	static TSharedPtr<SToolTip> GetColumnTooltip(const Insights::FTableColumn& Column);
+	static TSharedPtr<SToolTip> GetCellTooltip(const TSharedPtr<FNetEventNode> TreeNodePtr, const TSharedPtr<Insights::FTableColumn> ColumnPtr);
+
+private:
+	static void AddStatsRow(TSharedPtr<SGridPanel> Grid, int32& Row, const FText& Name, const FText& Value1, const FText& Value2);
+};
+
+#undef LOCTEXT_NAMESPACE
diff --git a/Engine/Source/Developer/TraceInsights/Private/Insights/NetworkingProfiler/Widgets/SNetworkingProfilerWindow.cpp b/Engine/Source/Developer/TraceInsights/Private/Insights/NetworkingProfiler/Widgets/SNetworkingProfilerWindow.cpp
--- a/Engine/Source/Developer/TraceInsights/Private/Insights/NetworkingProfiler/Widgets/SNetworkingProfilerWindow.cpp
+++ b/Engine/Source/Developer/TraceInsights/Private/Insights/NetworkingProfiler/Widgets/SNetworkingProfilerWindow.cpp
@@ -26,8 +26,7 @@
 #include "Insights/InsightsManager.h"
 #include "Insights/InsightsStyle.h"
 #include "Insights/NetworkingProfiler/NetworkingProfilerManager.h"
-#include "Insights/NetworkingProfiler/Widgets/SDataStreamBreakdownView.h"
-//#include "Insights/NetworkingProfiler/Widgets/SNetStatsView.h"
+#include "Insights/NetworkingProfiler/Widgets/SNetStatsView.h"
 #include "Insights/NetworkingProfiler/Widgets/SNetworkingProfilerToolbar.h"
 #include "Insights/NetworkingProfiler/Widgets/SPacketBreakdownView.h"
 //#include "Insights/NetworkingProfiler/Widgets/SPacketContentView.h"
@@ -89,6 +88,7 @@ void SNetworkingProfilerWindow::Reset()
 	if (NetStatsView)
 	{
 		//TODO: NetStatsView->Reset();
+		NetStatsView->RebuildTree(true);
 	}
 
 	AvailableGameInstances.Reset();
@@ -97,6 +97,11 @@ void SNetworkingProfilerWindow::Reset()
 	SelectedConnection.Reset();
 	AvailableConnectionModes.Reset();
 	SelectedConnectionMode.Reset();
+
+	SelectedPacketStartIndex = 0;
+	SelectedPacketEndIndex = 0;
+	SelectionStartPosition = 0;
+	SelectionEndPosition = 0;
 }
 
 ////////////////////////////////////////////////////////////////////////////////////////////////////
@@ -154,7 +159,7 @@ TSharedRef<SDockTab> SNetworkingProfilerWindow::SpawnTab_PacketContentView(const
 		.ShouldAutosize(false)
 		.TabRole(ETabRole::PanelTab)
 		[
-			SAssignNew(PacketContentView, SPacketContentView)
+			SAssignNew(PacketContentView, SPacketContentView, SharedThis(this))
 		];
 
 	DockTab->SetOnTabClosed(SDockTab::FOnTabClosedCallback::CreateRaw(this, &SNetworkingProfilerWindow::OnPacketContentViewTabClosed));
@@ -266,7 +271,7 @@ void SNetworkingProfilerWindow::Construct(const FArguments& InArgs, const TShare
 				(
 					FTabManager::NewStack()
 					->SetSizeCoefficient(0.35f)
-					->AddTab(FNetworkingProfilerTabs::NetStatsViewID, ETabState::ClosedTab)
+					->AddTab(FNetworkingProfilerTabs::NetStatsViewID, ETabState::OpenedTab)
 					//->SetForegroundTab(FNetworkingProfilerTabs::NetStatsViewID)
 				)
 			)
@@ -751,6 +756,7 @@ void SNetworkingProfilerWindow::GameInstance_OnSelectionChanged(TSharedPtr<FGame
 	{
 		PacketView->SetConnection(GetSelectedGameInstanceIndex(), GetSelectedConnectionIndex(), GetSelectedConnectionMode());
 	}
+	UpdateAggregatedNetStats();
 }
 
 ////////////////////////////////////////////////////////////////////////////////////////////////////
@@ -807,6 +813,7 @@ void SNetworkingProfilerWindow::Connection_OnSelectionChanged(TSharedPtr<FConnec
 	{
 		PacketView->SetConnection(GetSelectedGameInstanceIndex(), GetSelectedConnectionIndex(), GetSelectedConnectionMode());
 	}
+	UpdateAggregatedNetStats();
 }
 
 ////////////////////////////////////////////////////////////////////////////////////////////////////
@@ -862,6 +869,7 @@ void SNetworkingProfilerWindow::ConnectionMode_OnSelectionChanged(TSharedPtr<FCo
 	{
 		PacketView->SetConnection(GetSelectedGameInstanceIndex(), GetSelectedConnectionIndex(), GetSelectedConnectionMode());
 	}
+	UpdateAggregatedNetStats();
 }
 
 ////////////////////////////////////////////////////////////////////////////////////////////////////
@@ -880,4 +888,93 @@ FText SNetworkingProfilerWindow::ConnectionMode_GetSelectedTooltipText() const
 
 ////////////////////////////////////////////////////////////////////////////////////////////////////
 
+void SNetworkingProfilerWindow::SetSelectedPacket(uint32 StartIndex, uint32 EndIndex, uint32 SinglePacketBitSize)
+{
+	if (StartIndex != SelectedPacketStartIndex || EndIndex != SelectedPacketEndIndex)
+	{
+		SelectedPacketStartIndex = StartIndex;
+		SelectedPacketEndIndex = EndIndex;
+
+		UpdateAggregatedNetStats();
+
+		if (PacketContentView.IsValid())
+		{
+			if (SelectedPacketEndIndex == SelectedPacketStartIndex + 1 && // only one packet selected
+				SelectedGameInstance.IsValid() &&
+				SelectedConnection.IsValid() &&
+				SelectedConnectionMode.IsValid())
+			{
+				PacketContentView->SetPacket(SelectedGameInstance->GetIndex(),
+											 SelectedConnection->GetIndex(),
+											 SelectedConnectionMode->Mode,
+											 SelectedPacketStartIndex,
+											 SinglePacketBitSize);
+			}
+			else
+			{
+				PacketContentView->ResetPacket();
+			}
+		}
+	}
+}
+
+////////////////////////////////////////////////////////////////////////////////////////////////////
+
+void SNetworkingProfilerWindow::SetSelectedBitRange(uint32 StartPos, uint32 EndPos)
+{
+	if (StartPos != SelectionStartPosition || EndPos != SelectionEndPosition)
+	{
+		SelectionStartPosition = StartPos;
+		SelectionEndPosition = EndPos;
+		UpdateAggregatedNetStats();
+	}
+}
+
+////////////////////////////////////////////////////////////////////////////////////////////////////
+
+void SNetworkingProfilerWindow::SetSelectedEventType(const uint64 InEventTypeId)
+{
+	if (InEventTypeId != SelectedEventTypeId)
+	{
+		SelectedEventTypeId = InEventTypeId;
+
+		if (SelectedEventTypeId != InvalidEventTypeId)
+		{
+			if (NetStatsView)
+			{
+				NetStatsView->SelectNetEventNode(SelectedEventTypeId);
+			}
+		}
+	}
+}
+
+////////////////////////////////////////////////////////////////////////////////////////////////////
+
+void SNetworkingProfilerWindow::UpdateAggregatedNetStats()
+{
+	if (NetStatsView)
+	{
+		if (SelectedGameInstance.IsValid() &&
+			SelectedConnection.IsValid() &&
+			SelectedConnectionMode.IsValid() &&
+			SelectedPacketStartIndex < SelectedPacketEndIndex &&
+			SelectionStartPosition < SelectionEndPosition)
+		{
+			NetStatsView->UpdateStats(SelectedGameInstance->GetIndex(),
+									  SelectedConnection->GetIndex(),
+									  SelectedConnectionMode->Mode,
+									  SelectedPacketStartIndex,
+									  SelectedPacketEndIndex,
+									  SelectionStartPosition,
+									  SelectionEndPosition);
+		}
+		else
+		{
+			//NetStatsView->ResetStats();
+		}
+	}
+}
+
+////////////////////////////////////////////////////////////////////////////////////////////////////
+
 #undef LOCTEXT_NAMESPACE
diff --git a/Engine/Source/Developer/TraceInsights/Private/Insights/NetworkingProfiler/Widgets/SNetworkingProfilerWindow.h b/Engine/Source/Developer/TraceInsights/Private/Insights/NetworkingProfiler/Widgets/SNetworkingProfilerWindow.h
--- a/Engine/Source/Developer/TraceInsights/Private/Insights/NetworkingProfiler/Widgets/SNetworkingProfilerWindow.h
+++ b/Engine/Source/Developer/TraceInsights/Private/Insights/NetworkingProfiler/Widgets/SNetworkingProfilerWindow.h
@@ -132,6 +132,10 @@ public:
 	TSharedRef<SWidget> CreateConnectionComboBox();
 	TSharedRef<SWidget> CreateConnectionModeComboBox();
 
+	void SetSelectedPacket(uint32 StartIndex, uint32 EndIndex, uint32 SinglePacketBitSize = 0);
+	void SetSelectedBitRange(uint32 StartPos, uint32 EndPos);
+	void SetSelectedEventType(uint64 EventTypeId);
+
 private:
 	TSharedRef<SDockTab> SpawnTab_Toolbar(const FSpawnTabArgs& Args);
 	void OnToolbarTabClosed(TSharedRef<SDockTab> TabBeingClosed);
@@ -156,7 +160,7 @@ private:
 public:\
 	void Map_##CmdName(); /**< Maps UI command info CmdName with the specified UI command list. */\
 	const FUIAction CmdName##_Custom(); /**< UI action for CmdName command. */\
-protected:\
+private:\
 	void CmdName##_Execute(); /**< Handles FExecuteAction for CmdName. */\
 	bool CmdName##_CanExecute() const; /**< Handles FCanExecuteAction for CmdName. */\
 	ECheckBoxState CmdName##_GetCheckState() const; /**< Handles FGetActionCheckState for CmdName. */
@@ -167,7 +171,8 @@ protected:\
 #undef DECLARE_TOGGLE_COMMAND
 
 	//////////////////////////////////////////////////
-private:
+
+	void UpdateAggregatedNetStats();
 
 	/**
 	 * Fill the main menu with menu items.
@@ -293,6 +298,14 @@ private:
 	TArray<TSharedPtr<FConnectionModeItem>> AvailableConnectionModes;
 	TSharedPtr<FConnectionModeItem> SelectedConnectionMode;
 
+	uint32 SelectedPacketStartIndex;
+	uint32 SelectedPacketEndIndex;
+	uint32 SelectionStartPosition;
+	uint32 SelectionEndPosition;
+
+	static const uint64 InvalidEventTypeId = uint64(-1);
+	uint64 SelectedEventTypeId;
+
 	/** The handle to the active update duration tick */
 	TWeakPtr<FActiveTimerHandle> ActiveTimerHandle;
 
diff --git a/Engine/Source/Developer/TraceInsights/Private/Insights/NetworkingProfiler/Widgets/SPacketBreakdownView.cpp b/Engine/Source/Developer/TraceInsights/Private/Insights/NetworkingProfiler/Widgets/SPacketBreakdownView.cpp
--- a/Engine/Source/Developer/TraceInsights/Private/Insights/NetworkingProfiler/Widgets/SPacketBreakdownView.cpp
+++ b/Engine/Source/Developer/TraceInsights/Private/Insights/NetworkingProfiler/Widgets/SPacketBreakdownView.cpp
@@ -13,9 +13,11 @@
 
 // Insights
 #include "Insights/Common/PaintUtils.h"
+#include "Insights/Common/TimeUtils.h"
 #include "Insights/Common/Stopwatch.h"
 #include "Insights/InsightsManager.h"
 #include "Insights/InsightsStyle.h"
+#include "Insights/NetworkingProfiler/Widgets/SNetworkingProfilerWindow.h"
 
 ////////////////////////////////////////////////////////////////////////////////////////////////////
 
@@ -24,7 +26,8 @@
 ////////////////////////////////////////////////////////////////////////////////////////////////////
 
 SPacketContentView::SPacketContentView()
-	: DrawState(MakeShareable(new FPacketContentViewDrawState()))
+	: ProfilerWindow()
+	, DrawState(MakeShareable(new FPacketContentViewDrawState()))
 {
 	Reset();
 }
@@ -39,6 +42,8 @@ SPacketContentView::~SPacketContentView()
 
 void SPacketContentView::Reset()
 {
+	//ProfilerWindow
+
 	Viewport.Reset();
 	//FAxisViewportDouble& ViewportX = Viewport.GetHorizontalAxisViewport();
 	//ViewportX.SetScaleLimits(0.000001, 100000.0);
@@ -67,6 +72,7 @@ void SPacketContentView::Reset()
 	bIsScrolling = false;
 
 	HoveredEvent.Reset();
+	SelectedEvent.Reset();
 	Tooltip.Reset();
 
 	//ThisGeometry
@@ -81,8 +87,10 @@ void SPacketContentView::Reset()
 
 ////////////////////////////////////////////////////////////////////////////////////////////////////
 
-void SPacketContentView::Construct(const FArguments& InArgs)
+void SPacketContentView::Construct(const FArguments& InArgs, TSharedPtr<SNetworkingProfilerWindow> InProfilerWindow)
 {
+	ProfilerWindow = InProfilerWindow;
+
 	ChildSlot
 	[
 		SNew(SOverlay)
@@ -163,6 +171,9 @@ void SPacketContentView::ResetPacket()
 
 	DrawState->Reset();
 	bIsStateDirty = true;
+
+	HoveredEvent.Reset();
+	SelectedEvent.Reset();
 }
 
 ////////////////////////////////////////////////////////////////////////////////////////////////////
@@ -182,36 +193,9 @@ void SPacketContentView::SetPacket(uint32 InGameInstanceIndex, uint32 InConnecti
 
 	DrawState->Reset();
 	bIsStateDirty = true;
-}
-
-////////////////////////////////////////////////////////////////////////////////////////////////////
-
-void AddMockEventsRec(FPacketContentViewDrawStateBuilder& Builder, FRandomStream& RandomStream , int32 Depth, int64 Offset, int64 Size, int32 Type)
-{
-	Trace::FNetProfilerContentEvent Event;
-	Event.StartPos = Offset;
-	Event.EndPos = Offset + Size;
-	Event.Level = Depth;
-	Event.NameIndex = Type;
-	Builder.AddEvent(Event);
-
-	int64 ChildEventOffset = Offset;
-	while (ChildEventOffset < Offset + Size)
-	{
-		int64 ChildEventSize = RandomStream.RandRange(1, Size);
-		if (ChildEventOffset + ChildEventSize > Offset + Size)
-		{
-			ChildEventSize = Offset + Size - ChildEventOffset;
-		}
-
-		if (ChildEventOffset != Offset || ChildEventSize != Size)
-		{
-			int32 ChildEventType = RandomStream.RandRange(1, 1000);
-			AddMockEventsRec(Builder, RandomStream, Depth + 1, ChildEventOffset, ChildEventSize, ChildEventType);
-		}
 
-		ChildEventOffset += ChildEventSize;
-	}
+	HoveredEvent.Reset();
+	SelectedEvent.Reset();
 }
 
 ////////////////////////////////////////////////////////////////////////////////////////////////////
@@ -249,22 +233,25 @@ void SPacketContentView::UpdateState()
 
 			const uint32 StartPos = 0;
 			const uint32 EndPos = PacketBitSize;
-			NetProfilerProvider.EnumeratePacketContentEventsByPosition(ConnectionIndex, ConnectionMode, PacketIndex, StartPos, EndPos, [&Builder, &NetProfilerProvider](const Trace::FNetProfilerContentEvent& Event)
+			NetProfilerProvider.EnumeratePacketContentEventsByPosition(ConnectionIndex, ConnectionMode, PacketIndex, StartPos, EndPos, [this, &Builder, &NetProfilerProvider](const Trace::FNetProfilerContentEvent& Event)
 			{
 				const TCHAR* Name = nullptr;
 				NetProfilerProvider.ReadName(Event.NameIndex, [&Name](const Trace::FNetProfilerName& NetProfilerName)
 				{
 					Name = NetProfilerName.Name;
 				});
-				Builder.AddEvent(Event, Name);
-			});
-		}
 
-		// Init packet content with mock data.
-		if (false)
-		{
-			FRandomStream RandomStream(PacketIndex);
-			AddMockEventsRec(Builder, RandomStream, 0, 0, PacketBitSize, 0);
+				uint32 NetId = 0;
+				if (Event.ObjectInstanceIndex != 0)
+				{
+					NetProfilerProvider.ReadObject(GameInstanceIndex, Event.ObjectInstanceIndex, [&NetId](const Trace::FNetProfilerObjectInstance& ObjectInstance)
+					{
+						NetId = ObjectInstance.NetId;
+					});
+				}
+
+				Builder.AddEvent(Event, Name, NetId);
+			});
 		}
 
 		Builder.Flush();
@@ -292,18 +279,59 @@ void SPacketContentView::UpdateHoveredEvent()
 	{
 		// Init the tooltip's content.
 		Tooltip.ResetContent();
-		Tooltip.AddTitle(FString::Printf(TEXT("NetObject%d"), HoveredEvent.Event.Type));
-		Tooltip.AddNameValueTextLine(TEXT("Offset:"), FString::Format(TEXT("bit {0}"), { FText::AsNumber(HoveredEvent.Event.Offset).ToString() }));
-		if (HoveredEvent.Event.Size == 1)
+
+		const FNetworkPacketEvent& Event = HoveredEvent.Event;
+		FString Name(TEXT("?"));
+		Trace::FNetProfilerEventType EventType;
+		Trace::FNetProfilerObjectInstance ObjectInstance;
+
+		TSharedPtr<const Trace::IAnalysisSession> Session = FInsightsManager::Get()->GetSession();
+		if (Session.IsValid())
+		{
+			Trace::FAnalysisSessionReadScope SessionReadScope(*Session.Get());
+			const Trace::INetProfilerProvider& NetProfilerProvider = Trace::ReadNetProfilerProvider(*Session.Get());
+
+			NetProfilerProvider.ReadEventType(Event.EventTypeIndex, [&EventType](const Trace::FNetProfilerEventType& InEventType)
+			{
+				EventType = InEventType;
+			});
+
+			NetProfilerProvider.ReadName(EventType.NameIndex, [&Name](const Trace::FNetProfilerName& NetProfilerName)
+			{
+				Name = NetProfilerName.Name;
+			});
+
+			if (Event.ObjectInstanceIndex != 0)
+			{
+				NetProfilerProvider.ReadObject(GameInstanceIndex, Event.ObjectInstanceIndex, [&ObjectInstance](const Trace::FNetProfilerObjectInstance& InObjectInstance)
+				{
+					ObjectInstance = InObjectInstance;
+				});
+			}
+		}
+
+		Tooltip.AddTitle(Name);
+
+		if (Event.ObjectInstanceIndex != 0)
+		{
+			Tooltip.AddNameValueTextLine(TEXT("Net Id:"), FText::AsNumber(ObjectInstance.NetId).ToString());
+			Tooltip.AddNameValueTextLine(TEXT("Type Id:"), FString::Printf(TEXT("0x%016X"), ObjectInstance.TypeId));
+			Tooltip.AddNameValueTextLine(TEXT("Obj. LifeTime:"), FString::Format(TEXT("from {0} to {1}"),
+				{ TimeUtils::FormatTimeAuto(ObjectInstance.LifeTime.Begin), TimeUtils::FormatTimeAuto(ObjectInstance.LifeTime.End) }));
+		}
+
+		Tooltip.AddNameValueTextLine(TEXT("Offset:"), FString::Format(TEXT("bit {0}"), { FText::AsNumber(Event.BitOffset).ToString() }));
+		if (Event.BitSize == 1)
 		{
 			Tooltip.AddNameValueTextLine(TEXT("Size:"), TEXT("1 bit"));
 		}
 		else
 		{
-			Tooltip.AddNameValueTextLine(TEXT("Size:"), FString::Format(TEXT("{0} bits"), { FText::AsNumber(HoveredEvent.Event.Size).ToString() }));
+			Tooltip.AddNameValueTextLine(TEXT("Size:"), FString::Format(TEXT("{0} bits"), { FText::AsNumber(Event.BitSize).ToString() }));
 		}
-		Tooltip.AddNameValueTextLine(TEXT("Type:"), FText::AsNumber(HoveredEvent.Event.Type).ToString());
-		Tooltip.AddNameValueTextLine(TEXT("Depth:"), FText::AsNumber(HoveredEvent.Event.Depth).ToString());
+
+		Tooltip.AddNameValueTextLine(TEXT("Level:"), FText::AsNumber(Event.Level).ToString());
+
 		Tooltip.UpdateLayout();
 
 		Tooltip.SetDesiredOpacity(1.0f);
@@ -316,28 +344,53 @@ void SPacketContentView::UpdateHoveredEvent()
 
 ////////////////////////////////////////////////////////////////////////////////////////////////////
 
+void SPacketContentView::OnSelectedEventChanged()
+{
+	if (SelectedEvent.IsValid() && ProfilerWindow.IsValid())
+	{
+		// Select the node coresponding to net event type of selected net event instance.
+		const uint64 EventTypeId = static_cast<uint64>(SelectedEvent.Event.EventTypeIndex);
+		ProfilerWindow->SetSelectedEventType(EventTypeId);
+	}
+}
+
+////////////////////////////////////////////////////////////////////////////////////////////////////
+
+void SPacketContentView::SelectHoveredEvent()
+{
+	SelectedEvent = HoveredEvent;
+	//if (SelectedEvent.IsValid())
+	//{
+	//	LastSelectionType = ESelectionType::TimingEvent;
+	//	BringIntoView(SelectedEvent.StartPos, SelectedEvent.EndPos);
+	//}
+	OnSelectedEventChanged();
+}
+
+////////////////////////////////////////////////////////////////////////////////////////////////////
+
 FNetworkPacketEventRef SPacketContentView::GetEventAtMousePosition(float X, float Y)
 {
 	if (!bIsStateDirty)
 	{
-		for (const auto& Event : DrawState->Events)
+		for (const FNetworkPacketEvent& Event : DrawState->Events)
 		{
 			const FAxisViewportDouble& ViewportX = Viewport.GetHorizontalAxisViewport();
 
-			const float EventX1 = ViewportX.GetRoundedOffsetForValue(static_cast<double>(Event.Offset));
-			const float EventX2 = ViewportX.GetRoundedOffsetForValue(static_cast<double>(Event.Offset + Event.Size));
+			const float EventX1 = ViewportX.GetRoundedOffsetForValue(static_cast<double>(Event.BitOffset));
+			const float EventX2 = ViewportX.GetRoundedOffsetForValue(static_cast<double>(Event.BitOffset + Event.BitSize));
 
 			constexpr float Y0 = 0.0f;
 			constexpr float EventH = 14.0f;
 			constexpr float EventDY = 2.0f;
-			const float EventY = Y0 + (EventH + EventDY) * Event.Depth;
+			const float EventY = Y0 + (EventH + EventDY) * Event.Level;
 
 			constexpr float ToleranceX = 1.0f;
 
 			if (X >= EventX1 - ToleranceX && X <= EventX2 &&
 				Y >= EventY - EventDY / 2 && Y < EventY + EventH + EventDY / 2)
 			{
-				return FNetworkPacketEventRef(Event.Offset, Event.Size, Event.Type, Event.Depth);
+				return FNetworkPacketEventRef(Event);
 			}
 		}
 	}
@@ -367,10 +420,27 @@ int32 SPacketContentView::OnPaint(const FPaintArgs& Args, const FGeometry& Allot
 		// Draw the cached state (the network events of the package breakdown).
 		Helper.Draw(*DrawState);
 
-		// Highlight the hovered event.
-		if (HoveredEvent.IsValid())
+		if (!FNetworkPacketEventRef::AreEquals(SelectedEvent, HoveredEvent))
 		{
-			Helper.DrawEventHighlight(HoveredEvent.Event);
+			// Highlight the selected event (if any).
+			if (SelectedEvent.IsValid())
+			{
+				Helper.DrawEventHighlight(SelectedEvent.Event, FPacketContentViewDrawHelper::EHighlightMode::Selected);
+			}
+
+			// Highlight the hovered event (if any).
+			if (HoveredEvent.IsValid())
+			{
+				Helper.DrawEventHighlight(HoveredEvent.Event, FPacketContentViewDrawHelper::EHighlightMode::Hovered);
+			}
+		}
+		else
+		{
+			// Highlight the selected and hovered event (if any).
+			if (SelectedEvent.IsValid())
+			{
+				Helper.DrawEventHighlight(SelectedEvent.Event, FPacketContentViewDrawHelper::EHighlightMode::SelectedAndHovered);
+			}
 		}
 
 		// Draw tooltip for hovered Event.
@@ -484,7 +554,9 @@ FReply SPacketContentView::OnMouseButtonDown(const FGeometry& MyGeometry, const
 {
 	FReply Reply = FReply::Unhandled();
 
-	MousePositionOnButtonDown = MyGeometry.AbsoluteToLocal(MouseEvent.GetScreenSpacePosition());
+	MousePosition = MyGeometry.AbsoluteToLocal(MouseEvent.GetScreenSpacePosition());
+	MousePositionOnButtonDown = MousePosition;
+
 	ViewportPosXOnButtonDown = Viewport.GetHorizontalAxisViewport().GetPos();
 
 	if (MouseEvent.GetEffectingButton() == EKeys::LeftMouseButton)
@@ -498,7 +570,7 @@ FReply SPacketContentView::OnMouseButtonDown(const FGeometry& MyGeometry, const
 	{
 		bIsRMB_Pressed = true;
 
-		// Capture mouse, so we can scroll outside this widget.
+		// Capture mouse, so we can drag outside this widget.
 		Reply = FReply::Handled().CaptureMouse(SharedThis(this));
 	}
 
@@ -511,7 +583,8 @@ FReply SPacketContentView::OnMouseButtonUp(const FGeometry& MyGeometry, const FP
 {
 	FReply Reply = FReply::Unhandled();
 
-	MousePositionOnButtonUp = MyGeometry.AbsoluteToLocal(MouseEvent.GetScreenSpacePosition());
+	MousePosition = MyGeometry.AbsoluteToLocal(MouseEvent.GetScreenSpacePosition());
+	MousePositionOnButtonUp = MousePosition;
 
 	const bool bIsValidForMouseClick = MousePositionOnButtonUp.Equals(MousePositionOnButtonDown, MOUSE_SNAP_DISTANCE);
 
@@ -526,13 +599,15 @@ FReply SPacketContentView::OnMouseButtonUp(const FGeometry& MyGeometry, const FP
 			}
 			else if (bIsValidForMouseClick)
 			{
-				//SelectEventAtMousePosition(MousePositionOnButtonUp.X, MousePositionOnButtonUp.Y);
+				// Select the hovered timing event (if any).
+				UpdateHoveredEvent();
+				SelectHoveredEvent();
 			}
 
-			bIsLMB_Pressed = false;
-
-			// Release the mouse.
+			// Release mouse as we no longer drag.
 			Reply = FReply::Handled().ReleaseMouseCapture();
+
+			bIsLMB_Pressed = false;
 		}
 	}
 	else if (MouseEvent.GetEffectingButton() == EKeys::RightMouseButton)
@@ -549,10 +624,10 @@ FReply SPacketContentView::OnMouseButtonUp(const FGeometry& MyGeometry, const FP
 				//ShowContextMenu(MouseEvent);
 			}
 
-			bIsRMB_Pressed = false;
-
-			// Release mouse as we no longer scroll.
+			// Release mouse as we no longer drag.
 			Reply = FReply::Handled().ReleaseMouseCapture();
+
+			bIsRMB_Pressed = false;
 		}
 	}
 
@@ -613,7 +688,7 @@ void SPacketContentView::OnMouseLeave(const FPointerEvent& MouseEvent)
 {
 	if (!HasMouseCapture())
 	{
-		// No longer scrolling (unless we have mouse capture).
+		// No longer dragging (unless we have mouse capture).
 		bIsScrolling = false;
 
 		bIsLMB_Pressed = false;
diff --git a/Engine/Source/Developer/TraceInsights/Private/Insights/NetworkingProfiler/Widgets/SPacketBreakdownView.h b/Engine/Source/Developer/TraceInsights/Private/Insights/NetworkingProfiler/Widgets/SPacketBreakdownView.h
--- a/Engine/Source/Developer/TraceInsights/Private/Insights/NetworkingProfiler/Widgets/SPacketBreakdownView.h
+++ b/Engine/Source/Developer/TraceInsights/Private/Insights/NetworkingProfiler/Widgets/SPacketBreakdownView.h
@@ -20,47 +20,53 @@
 #include "Insights/ViewModels/TooltipDrawState.h"
 
 class SScrollBar;
+class SNetworkingProfilerWindow;
 
 ////////////////////////////////////////////////////////////////////////////////////////////////////
 
 struct FNetworkPacketEventRef
 {
 	FNetworkPacketEvent Event;
+	bool bIsValid;
 
 	FNetworkPacketEventRef()
-	: Event(0, 0, -1, -1)
+		: Event()
+		, bIsValid(false)
 	{
 	}
 
-	FNetworkPacketEventRef(int64 InOffset, int64 InSize, int32 InType, int32 InDepth)
-		: Event(InOffset, InSize, InType, InDepth)
+	FNetworkPacketEventRef(const FNetworkPacketEvent& InEvent)
+		: Event(InEvent)
+		, bIsValid(true)
 	{
 	}
 
 	FNetworkPacketEventRef(const FNetworkPacketEventRef& Other)
 		: Event(Other.Event)
+		, bIsValid(Other.bIsValid)
 	{
 	}
 
 	FNetworkPacketEventRef& operator=(const FNetworkPacketEventRef& Other)
 	{
 		Event = Other.Event;
+		bIsValid = Other.bIsValid;
 		return *this;
 	}
 
 	void Reset()
 	{
-		Event.Type = -1;
+		bIsValid = false;
 	}
 
 	bool IsValid() const
 	{
-		return Event.Type >= 0;
+		return bIsValid;
 	}
 
 	bool Equals(const FNetworkPacketEventRef& Other) const
 	{
-		return Event.Equals(Other.Event);
+		return bIsValid == Other.bIsValid && Event.Equals(Other.Event);
 	}
 
 	static bool AreEquals(const FNetworkPacketEventRef& A, const FNetworkPacketEventRef& B)
@@ -99,11 +105,7 @@ public:
 	SLATE_BEGIN_ARGS(SPacketContentView) {}
 	SLATE_END_ARGS()
 
-	/**
-	 * Construct this widget
-	 * @param InArgs The declaration data for this widget
-	 */
-	void Construct(const FArguments& InArgs);
+	void Construct(const FArguments& InArgs, TSharedPtr<SNetworkingProfilerWindow> InProfilerWindow);
 
 	/**
 	 * Ticks this widget. Override in derived classes, but always call the parent implementation.
@@ -129,15 +131,13 @@ public:
 	void SetPacket(uint32 InGameInstanceIndex, uint32 InConnectionIndex, Trace::ENetProfilerConnectionMode InConnectionMode, uint32 InPacketIndex, int64 InPacketBitSize);
 
 private:
-	void UpdateState();
-	void UpdateHoveredEvent();
-	FNetworkPacketEventRef GetEventAtMousePosition(float X, float Y);
-
 	//void ShowContextMenu(const FPointerEvent& MouseEvent);
 
 	/** Binds our UI commands to delegates. */
 	void BindCommands();
 
+	////////////////////////////////////////////////////////////////////////////////////////////////////
+
 	/**
 	 * Called when the user scrolls the horizontal scrollbar.
 	 * @param ScrollOffset Scroll offset as a fraction between 0 and 1.
@@ -147,7 +147,19 @@ private:
 
 	void ZoomHorizontally(const float Delta, const float X);
 
+	////////////////////////////////////////////////////////////////////////////////////////////////////
+
+	void UpdateState();
+
+	void UpdateHoveredEvent();
+	FNetworkPacketEventRef GetEventAtMousePosition(float X, float Y);
+
+	void OnSelectedEventChanged();
+	void SelectHoveredEvent();
+
 private:
+	TSharedPtr<SNetworkingProfilerWindow> ProfilerWindow;
+
 	/** The track's viewport. Encapsulates info about position and scale. */
 	FPacketContentViewport Viewport;
 	bool bIsViewportDirty;
@@ -188,6 +200,8 @@ private:
 	// Selection
 
 	FNetworkPacketEventRef HoveredEvent;
+	FNetworkPacketEventRef SelectedEvent;
+
 	FTooltipDrawState Tooltip;
 
 	//////////////////////////////////////////////////
diff --git a/Engine/Source/Developer/TraceInsights/Private/Insights/NetworkingProfiler/Widgets/SPacketSizesView.cpp b/Engine/Source/Developer/TraceInsights/Private/Insights/NetworkingProfiler/Widgets/SPacketSizesView.cpp
--- a/Engine/Source/Developer/TraceInsights/Private/Insights/NetworkingProfiler/Widgets/SPacketSizesView.cpp
+++ b/Engine/Source/Developer/TraceInsights/Private/Insights/NetworkingProfiler/Widgets/SPacketSizesView.cpp
@@ -375,19 +375,16 @@ void SPacketView::OnSelectedSampleChanged()
 {
 	if (ProfilerWindow.IsValid())
 	{
-		TSharedPtr<SPacketContentView> PacketContentView = ProfilerWindow->GetPacketContentView();
-		if (PacketContentView.IsValid())
+		if (SelectedSample.IsValid())
 		{
-			if (SelectedSample.IsValid())
-			{
-				PacketContentView->SetPacket(GameInstanceIndex, ConnectionIndex, ConnectionMode,
-											 SelectedSample.Sample->LargestPacket.Index,
-											 SelectedSample.Sample->LargestPacket.TotalSizeInBytes * 8);
-			}
-			else
-			{
-				PacketContentView->ResetPacket();
-			}
+			const uint32 BitSize = SelectedSample.Sample->LargestPacket.TotalSizeInBytes * 8;
+			ProfilerWindow->SetSelectedPacket(SelectedSample.Sample->LargestPacket.Index, SelectedSample.Sample->LargestPacket.Index + 1, BitSize);
+			ProfilerWindow->SetSelectedBitRange(0, BitSize);
+		}
+		else
+		{
+			ProfilerWindow->SetSelectedPacket(0, 0);
+			ProfilerWindow->SetSelectedBitRange(0, 0);
 		}
 	}
 }
diff --git a/Engine/Source/Developer/TraceInsights/Private/Insights/NetworkingProfiler/Widgets/SPacketSizesView.h b/Engine/Source/Developer/TraceInsights/Private/Insights/NetworkingProfiler/Widgets/SPacketSizesView.h
--- a/Engine/Source/Developer/TraceInsights/Private/Insights/NetworkingProfiler/Widgets/SPacketSizesView.h
+++ b/Engine/Source/Developer/TraceInsights/Private/Insights/NetworkingProfiler/Widgets/SPacketSizesView.h
@@ -108,10 +108,6 @@ public:
 	}
 	SLATE_END_ARGS()
 
-	/**
-	 * Construct this widget
-	 * @param InArgs The declaration data for this widget
-	 */
 	void Construct(const FArguments& InArgs, TSharedPtr<SNetworkingProfilerWindow> InProfilerWindow);
 
 	/**
diff --git a/Engine/Source/Developer/TraceInsights/Private/Insights/Table/Widgets/STableTreeView.cpp b/Engine/Source/Developer/TraceInsights/Private/Insights/Table/Widgets/STableTreeView.cpp
--- a/Engine/Source/Developer/TraceInsights/Private/Insights/Table/Widgets/STableTreeView.cpp
+++ b/Engine/Source/Developer/TraceInsights/Private/Insights/Table/Widgets/STableTreeView.cpp
@@ -1216,7 +1216,7 @@ TSharedRef<SWidget> STableTreeView::GetGroupingCrumbMenuContent(const TSharedPtr
 		{
 			MenuBuilder.AddSubMenu
 			(
-				LOCTEXT("ChangeGrouping_Section", "Change With..."),
+				LOCTEXT("ChangeGrouping_Section", "Change To..."),
 				LOCTEXT("GroupingMenu_Add_Desc", "Change selected grouping."),
 				FNewMenuDelegate::CreateSP(this, &STableTreeView::BuildGroupingSubMenu_Change, CrumbGrouping),
 				false,
@@ -1487,6 +1487,8 @@ void STableTreeView::CreateSortings()
 			}
 		}
 	}
+
+	UpdateCurrentSortingByColumn();
 }
 
 ////////////////////////////////////////////////////////////////////////////////////////////////////
diff --git a/Engine/Source/Developer/TraceInsights/Private/Insights/Table/Widgets/STableTreeViewTooltip.cpp b/Engine/Source/Developer/TraceInsights/Private/Insights/Table/Widgets/STableTreeViewTooltip.cpp
--- a/Engine/Source/Developer/TraceInsights/Private/Insights/Table/Widgets/STableTreeViewTooltip.cpp
+++ b/Engine/Source/Developer/TraceInsights/Private/Insights/Table/Widgets/STableTreeViewTooltip.cpp
@@ -133,16 +133,32 @@ TSharedPtr<SToolTip> STableTreeViewTooltip::GetCellTooltip(const TSharedPtr<FTab
 						.TextStyle(FEditorStyle::Get(), TEXT("Profiler.Tooltip"))
 					]
 
-					// Name: [Name]
+					// Item Type: [Type]
 					+ SGridPanel::Slot(0, 1)
 					.Padding(2.0f)
 					[
 						SNew(STextBlock)
-						.Text(LOCTEXT("TT_Name", "Name:"))
+						.Text(LOCTEXT("TT_Type", "Item Type:"))
 						.TextStyle(FEditorStyle::Get(), TEXT("Profiler.TooltipBold"))
 					]
 					+ SGridPanel::Slot(1, 1)
 					.Padding(2.0f)
+					[
+						SNew(STextBlock)
+						.Text(TreeNodePtr->IsGroup() ? LOCTEXT("TT_Type_Group", "Group Node") : LOCTEXT("TT_Type_TableRow", "Table Row"))
+						.TextStyle(FEditorStyle::Get(), TEXT("Profiler.Tooltip"))
+					]
+
+					// Item Name: [Name]
+					+ SGridPanel::Slot(0, 2)
+					.Padding(2.0f)
+					[
+						SNew(STextBlock)
+						.Text(LOCTEXT("TT_Name", "Item Name:"))
+						.TextStyle(FEditorStyle::Get(), TEXT("Profiler.TooltipBold"))
+					]
+					+ SGridPanel::Slot(1, 2)
+					.Padding(2.0f)
 					[
 						SNew(STextBlock)
 						.Text(FText::FromName(TreeNodePtr->GetName()))
diff --git a/Engine/Source/Developer/TraceInsights/Private/Insights/TimingProfilerCommands.cpp b/Engine/Source/Developer/TraceInsights/Private/Insights/TimingProfilerCommands.cpp
--- a/Engine/Source/Developer/TraceInsights/Private/Insights/TimingProfilerCommands.cpp
+++ b/Engine/Source/Developer/TraceInsights/Private/Insights/TimingProfilerCommands.cpp
@@ -59,7 +59,7 @@ void FTimingProfilerCommands::RegisterCommands()
 	UI_COMMAND(ToggleTimersViewVisibility, "Timers", "Toggles the visibility of the Timers view", EUserInterfaceActionType::ToggleButton, FInputChord());
 	UI_COMMAND(ToggleCallersTreeViewVisibility, "Callers", "Toggles the visibility of the Callers tree view", EUserInterfaceActionType::ToggleButton, FInputChord());
 	UI_COMMAND(ToggleCalleesTreeViewVisibility, "Callees", "Toggles the visibility of the Callees tree view", EUserInterfaceActionType::ToggleButton, FInputChord());
-	UI_COMMAND(ToggleStatsCountersViewVisibility, "Stats", "Toggles the visibility of the Stats Counters view", EUserInterfaceActionType::ToggleButton, FInputChord());
+	UI_COMMAND(ToggleStatsCountersViewVisibility, "Counters", "Toggles the visibility of the Counters view", EUserInterfaceActionType::ToggleButton, FInputChord());
 	UI_COMMAND(ToggleLogViewVisibility, "Log", "Toggles the visibility of the Log view", EUserInterfaceActionType::ToggleButton, FInputChord(EKeys::L));
 }
 PRAGMA_ENABLE_OPTIMIZATION
diff --git a/Engine/Source/Developer/TraceInsights/Private/Insights/ViewModels/TimerGroupingAndSorting.cpp b/Engine/Source/Developer/TraceInsights/Private/Insights/ViewModels/TimerGroupingAndSorting.cpp
--- a/Engine/Source/Developer/TraceInsights/Private/Insights/ViewModels/TimerGroupingAndSorting.cpp
+++ b/Engine/Source/Developer/TraceInsights/Private/Insights/ViewModels/TimerGroupingAndSorting.cpp
@@ -146,7 +146,7 @@ FTimerNodeSortingByTotalInclusiveTime::FTimerNodeSortingByTotalInclusiveTime(TSh
 	: Insights::FTableCellValueSorter(
 		FName(TEXT("ByTotalInclusiveTime")),
 		LOCTEXT("Sorting_ByTotalInclusiveTime_Name", "By Total Inclusive Time"),
-		LOCTEXT("Sorting_ByTotalInclusiveTime_Title", "Sort By Total Inxclusive Time"),
+		LOCTEXT("Sorting_ByTotalInclusiveTime_Title", "Sort By Total Inclusive Time"),
 		LOCTEXT("Sorting_ByTotalInclusiveTime_Desc", "Sort by aggregated total inclusive time (ascending or descending), then by name (ascending)."),
 		InColumnRef)
 {
diff --git a/Engine/Source/Developer/TraceInsights/Private/Insights/ViewModels/TimerNode.cpp b/Engine/Source/Developer/TraceInsights/Private/Insights/ViewModels/TimerNode.cpp
--- a/Engine/Source/Developer/TraceInsights/Private/Insights/ViewModels/TimerNode.cpp
+++ b/Engine/Source/Developer/TraceInsights/Private/Insights/ViewModels/TimerNode.cpp
@@ -2,10 +2,6 @@
 
 #include "TimerNode.h"
 
-// Insights
-#include "Insights/Table/ViewModels/TableColumn.h"
-#include "Insights/ViewModels/TimersViewColumnFactory.h"
-
 #define LOCTEXT_NAMESPACE "TimerNode"
 
 const FName FTimerNode::TypeName(TEXT("FTimerNode"));
diff --git a/Engine/Source/Developer/TraceInsights/Private/Insights/ViewModels/TimerNode.h b/Engine/Source/Developer/TraceInsights/Private/Insights/ViewModels/TimerNode.h
--- a/Engine/Source/Developer/TraceInsights/Private/Insights/ViewModels/TimerNode.h
+++ b/Engine/Source/Developer/TraceInsights/Private/Insights/ViewModels/TimerNode.h
@@ -79,42 +79,21 @@ public:
 	/**
 	 * @return a name of the meta group that this timer node belongs to, taken from the metadata.
 	 */
-	const FName& GetMetaGroupName() const
-	{
-		return MetaGroupName;
-	}
+	const FName& GetMetaGroupName() const { return MetaGroupName; }
 
 	/**
 	 * @return a type of this timer node or ETimerNodeType::Group for group nodes.
 	 */
-	const ETimerNodeType& GetType() const
-	{
-		return Type;
-	}
+	const ETimerNodeType& GetType() const { return Type; }
 
 	/**
 	 * @return the aggregated stats for this timer.
 	 */
-	const Trace::FAggregatedTimingStats& GetAggregatedStats() const
-	{
-		return AggregatedStats;
-	}
+	const Trace::FAggregatedTimingStats& GetAggregatedStats() const { return AggregatedStats; }
 
 	void ResetAggregatedStats();
-
 	void SetAggregatedStats(const Trace::FAggregatedTimingStats& AggregatedStats);
 
-	/** Sorts children using the specified class instance. */
-	template<typename TSortingClass>
-	void SortChildren(const TSortingClass& Instance)
-	{
-		auto Projection = [](Insights::FBaseTreeNodePtr Node) -> FTimerNodePtr
-		{
-			return StaticCastSharedPtr<FTimerNode, Insights::FBaseTreeNode>(Node);
-		};
-		Algo::SortBy(GetChildrenMutable(), Projection, Instance);
-	}
-
 	bool IsHotPath() const { return bIsHotPath; }
 	void SetIsHotPath(bool bOnOff) { bIsHotPath = bOnOff; }
 
diff --git a/Engine/Source/Developer/TraceInsights/Private/Insights/ViewModels/TimerNodeHelper.h b/Engine/Source/Developer/TraceInsights/Private/Insights/ViewModels/TimerNodeHelper.h
--- a/Engine/Source/Developer/TraceInsights/Private/Insights/ViewModels/TimerNodeHelper.h
+++ b/Engine/Source/Developer/TraceInsights/Private/Insights/ViewModels/TimerNodeHelper.h
@@ -67,174 +67,3 @@ struct TimerNodeGroupingHelper
 };
 
 ////////////////////////////////////////////////////////////////////////////////////////////////////
-
-/** Helper struct that contains sorting static functions and classes. */
-struct TimerNodeSortingHelper
-{
-	//////////////////////////////////////////////////
-	// Sorting by timer's name.
-
-	struct ByNameAscending
-	{
-		FORCEINLINE_DEBUGGABLE bool operator()(const FTimerNodePtr& A, const FTimerNodePtr& B) const
-		{
-			return A->GetName().LexicalLess(B->GetName());
-		}
-	};
-
-	struct ByNameDescending
-	{
-		FORCEINLINE_DEBUGGABLE bool operator()(const FTimerNodePtr& A, const FTimerNodePtr& B) const
-		{
-			return B->GetName().LexicalLess(A->GetName());
-		}
-	};
-
-	//////////////////////////////////////////////////
-	// Sorting by timer's meta group name.
-	// If meta group names are the same then sort by name.
-
-	struct ByMetaGroupNameAscending
-	{
-		FORCEINLINE_DEBUGGABLE bool operator()(const FTimerNodePtr& A, const FTimerNodePtr& B) const
-		{
-			if (A->GetMetaGroupName() == B->GetMetaGroupName())
-			{
-				// Sort by timer name (ascending).
-				return A->GetName().LexicalLess(B->GetName());
-			}
-			else
-			{
-				return A->GetMetaGroupName().LexicalLess(B->GetMetaGroupName());
-			}
-		}
-	};
-
-	struct ByMetaGroupNameDescending
-	{
-		FORCEINLINE_DEBUGGABLE bool operator()(const FTimerNodePtr& A, const FTimerNodePtr& B) const
-		{
-			if (A->GetMetaGroupName() == B->GetMetaGroupName())
-			{
-				// Sort by timer name (ascending).
-				return A->GetName().LexicalLess(B->GetName());
-			}
-			else
-			{
-				return B->GetMetaGroupName().LexicalLess(A->GetMetaGroupName());
-			}
-		}
-	};
-
-	//////////////////////////////////////////////////
-	// Sorting by timer's type.
-	// If types are the same then sort by name.
-
-	struct ByTypeAscending
-	{
-		FORCEINLINE_DEBUGGABLE bool operator()(const FTimerNodePtr& A, const FTimerNodePtr& B) const
-		{
-			const ETimerNodeType& TypeA = A->GetType();
-			const ETimerNodeType& TypeB = B->GetType();
-
-			if (TypeA == TypeB)
-			{
-				// Sort by timer name (ascending).
-				return A->GetName().LexicalLess(B->GetName());
-			}
-			else
-			{
-				return TypeA < TypeB;
-			}
-		}
-	};
-
-	struct ByTypeDescending
-	{
-		FORCEINLINE_DEBUGGABLE bool operator()(const FTimerNodePtr& A, const FTimerNodePtr& B) const
-		{
-			const ETimerNodeType& TypeA = A->GetType();
-			const ETimerNodeType& TypeB = B->GetType();
-
-			if (TypeA == TypeB)
-			{
-				// Sort by timer name (ascending).
-				return A->GetName().LexicalLess(B->GetName());
-			}
-			else
-			{
-				return TypeA > TypeB;
-			}
-		}
-	};
-
-	//////////////////////////////////////////////////
-	// Sorting by an aggregated stats value.
-	// If aggregated stats values are the same then sort by name.
-
-	#define SORT_BY_STATS_ASCENDING(Type, SortName, AggregatedStatsMember) \
-		struct SortName \
-		{ \
-			FORCEINLINE_DEBUGGABLE bool operator()(const FTimerNodePtr& A, const FTimerNodePtr& B) const \
-			{ \
-				const Type ValueA = A->GetAggregatedStats().AggregatedStatsMember; \
-				const Type ValueB = B->GetAggregatedStats().AggregatedStatsMember; \
-				\
-				if (ValueA == ValueB) \
-				{ \
-					/* Sort by name (ascending). */ \
-					return A->GetName().LexicalLess(B->GetName()); \
-				} \
-				else \
-				{ \
-					return ValueA < ValueB; \
-				} \
-			} \
-		};
-
-	#define SORT_BY_STATS_DESCENDING(Type, SortName, AggregatedStatsMember) \
-		struct SortName \
-		{ \
-			FORCEINLINE_DEBUGGABLE bool operator()(const FTimerNodePtr& A, const FTimerNodePtr& B) const \
-			{ \
-				const Type ValueA = A->GetAggregatedStats().AggregatedStatsMember; \
-				const Type ValueB = B->GetAggregatedStats().AggregatedStatsMember; \
-				\
-				if (ValueA == ValueB) \
-				{ \
-					/* Sort by name (ascending). */ \
-					return A->GetName().LexicalLess(B->GetName()); \
-				} \
-				else \
-				{ \
-					return ValueA > ValueB; \
-				} \
-			} \
-		};
-
-	#define SORT_BY_STATS(Type, SortName, AggregatedStatsMember) \
-		SORT_BY_STATS_ASCENDING(Type, SortName##Ascending, AggregatedStatsMember) \
-		SORT_BY_STATS_DESCENDING(Type, SortName##Descending, AggregatedStatsMember)
-
-	SORT_BY_STATS(uint64, ByInstanceCount, InstanceCount);
-
-	SORT_BY_STATS(double, ByTotalInclusiveTime, TotalInclusiveTime);
-	SORT_BY_STATS(double, ByMinInclusiveTime, MinInclusiveTime);
-	SORT_BY_STATS(double, ByMaxInclusiveTime, MaxInclusiveTime);
-	SORT_BY_STATS(double, ByAverageInclusiveTime, AverageInclusiveTime);
-	SORT_BY_STATS(double, ByMedianInclusiveTime, MedianInclusiveTime);
-
-	SORT_BY_STATS(double, ByTotalExclusiveTime, TotalExclusiveTime);
-	SORT_BY_STATS(double, ByMinExclusiveTime, MinExclusiveTime);
-	SORT_BY_STATS(double, ByMaxExclusiveTime, MaxExclusiveTime);
-	SORT_BY_STATS(double, ByAverageExclusiveTime, AverageExclusiveTime);
-	SORT_BY_STATS(double, ByMedianExclusiveTime, MedianExclusiveTime);
-
-	#undef SORT_BY_STATS_ASCENDING
-	#undef SORT_BY_STATS_DESCENDING
-	#undef SORT_BY_STATS
-
-	//////////////////////////////////////////////////
-};
-
-////////////////////////////////////////////////////////////////////////////////////////////////////
diff --git a/Engine/Source/Developer/TraceInsights/Private/Insights/Widgets/STimerTreeView.cpp b/Engine/Source/Developer/TraceInsights/Private/Insights/Widgets/STimerTreeView.cpp
--- a/Engine/Source/Developer/TraceInsights/Private/Insights/Widgets/STimerTreeView.cpp
+++ b/Engine/Source/Developer/TraceInsights/Private/Insights/Widgets/STimerTreeView.cpp
@@ -544,16 +544,16 @@ FName STimerTreeView::TableRow_GetHighlightedNodeName() const
 // Sorting
 ////////////////////////////////////////////////////////////////////////////////////////////////////
 
-const EColumnSortMode::Type STimerTreeView::GetDefaultColumnSortMode()
+const FName STimerTreeView::GetDefaultColumnBeingSorted()
 {
-	return EColumnSortMode::Descending;
+	return FTimersViewColumns::TotalInclusiveTimeColumnID;
 }
 
 ////////////////////////////////////////////////////////////////////////////////////////////////////
 
-const FName STimerTreeView::GetDefaultColumnBeingSorted()
+const EColumnSortMode::Type STimerTreeView::GetDefaultColumnSortMode()
 {
-	return FTimersViewColumns::TotalInclusiveTimeColumnID;
+	return EColumnSortMode::Descending;
 }
 
 ////////////////////////////////////////////////////////////////////////////////////////////////////
@@ -574,6 +574,8 @@ void STimerTreeView::CreateSortings()
 			}
 		}
 	}
+
+	UpdateCurrentSortingByColumn();
 }
 
 ////////////////////////////////////////////////////////////////////////////////////////////////////
diff --git a/Engine/Source/Developer/TraceInsights/Private/Insights/Widgets/STimerTreeView.h b/Engine/Source/Developer/TraceInsights/Private/Insights/Widgets/STimerTreeView.h
--- a/Engine/Source/Developer/TraceInsights/Private/Insights/Widgets/STimerTreeView.h
+++ b/Engine/Source/Developer/TraceInsights/Private/Insights/Widgets/STimerTreeView.h
@@ -16,7 +16,6 @@
 // Insights
 #include "Insights/ViewModels/TimerGroupingAndSorting.h"
 #include "Insights/ViewModels/TimerNode.h"
-#include "Insights/ViewModels/TimersViewColumn.h"
 
 ////////////////////////////////////////////////////////////////////////////////////////////////////
 
@@ -110,8 +109,8 @@ private:
 	////////////////////////////////////////////////////////////////////////////////////////////////////
 	// Sorting
 
-	static const EColumnSortMode::Type GetDefaultColumnSortMode();
 	static const FName GetDefaultColumnBeingSorted();
+	static const EColumnSortMode::Type GetDefaultColumnSortMode();
 
 	void CreateSortings();
 
diff --git a/Engine/Source/Developer/TraceInsights/Private/Insights/Widgets/STimersView.cpp b/Engine/Source/Developer/TraceInsights/Private/Insights/Widgets/STimersView.cpp
--- a/Engine/Source/Developer/TraceInsights/Private/Insights/Widgets/STimersView.cpp
+++ b/Engine/Source/Developer/TraceInsights/Private/Insights/Widgets/STimersView.cpp
@@ -32,26 +32,14 @@
 
 ////////////////////////////////////////////////////////////////////////////////////////////////////
 
-const EColumnSortMode::Type STimersView::GetDefaultColumnSortMode()
-{
-	return EColumnSortMode::Descending;
-}
-
-////////////////////////////////////////////////////////////////////////////////////////////////////
-
-const FName STimersView::GetDefaultColumnBeingSorted()
-{
-	return FTimersViewColumns::TotalInclusiveTimeColumnID;
-}
-
-////////////////////////////////////////////////////////////////////////////////////////////////////
-
 STimersView::STimersView()
 	: Table(MakeShareable(new Insights::FTable()))
 	, bExpansionSaved(false)
 	, GroupingMode(ETimerGroupingMode::ByType)
-	, ColumnSortMode(GetDefaultColumnSortMode())
+	, AvailableSorters()
+	, CurrentSorter(nullptr)
 	, ColumnBeingSorted(GetDefaultColumnBeingSorted())
+	, ColumnSortMode(GetDefaultColumnSortMode())
 	, StatsStartTime(0.0)
 	, StatsEndTime(0.0)
 {
@@ -229,6 +217,7 @@ void STimersView::Construct(const FArguments& InArgs)
 	Filters->Add(TextFilter);
 
 	CreateGroupByOptionsSources();
+	CreateSortings();
 
 	// Register ourselves with the Insights manager.
 	FInsightsManager::Get()->GetSessionChangedEvent().AddSP(this, &STimersView::InsightsManager_OnSessionChanged);
@@ -598,7 +587,7 @@ void STimersView::InsightsManager_OnSessionChanged()
 void STimersView::UpdateTree()
 {
 	CreateGroups();
-	SortTimers();
+	SortTreeNodes();
 	ApplyFiltering();
 }
 
@@ -832,7 +821,7 @@ void STimersView::TableRow_SetHoveredCell(TSharedPtr<Insights::FTable> InTablePt
 	const bool bIsAnyMenusVisible = FSlateApplication::Get().AnyMenusVisible();
 	if (!HasMouseCapture() && !bIsAnyMenusVisible)
 	{
-		HoveredTimerNodePtr = InNodePtr;
+		HoveredNodePtr = InNodePtr;
 	}
 }
 
@@ -982,9 +971,6 @@ void STimersView::CreateGroups()
 	}
 
 	GroupNodeSet.GenerateValueArray(GroupNodes);
-
-	// Sort by a fake group name.
-	GroupNodes.Sort(TimerNodeSortingHelper::ByNameAscending());
 }
 
 ////////////////////////////////////////////////////////////////////////////////////////////////////
@@ -1017,7 +1003,7 @@ void STimersView::GroupBy_OnSelectionChanged(TSharedPtr<ETimerGroupingMode> NewG
 		GroupingMode = *NewGroupingMode;
 
 		CreateGroups();
-		SortTimers();
+		SortTreeNodes();
 		ApplyFiltering();
 	}
 }
@@ -1049,47 +1035,82 @@ FText STimersView::GroupBy_GetSelectedTooltipText() const
 // Sorting
 ////////////////////////////////////////////////////////////////////////////////////////////////////
 
-void STimersView::SortTimers()
+const FName STimersView::GetDefaultColumnBeingSorted()
 {
-	const int32 NumGroups = GroupNodes.Num();
+	return FTimersViewColumns::TotalInclusiveTimeColumnID;
+}
+
+////////////////////////////////////////////////////////////////////////////////////////////////////
+
+const EColumnSortMode::Type STimersView::GetDefaultColumnSortMode()
+{
+	return EColumnSortMode::Descending;
+}
+
+////////////////////////////////////////////////////////////////////////////////////////////////////
+
+void STimersView::CreateSortings()
+{
+	AvailableSorters.Reset();
+	CurrentSorter = nullptr;
+
+	for (const TSharedPtr<Insights::FTableColumn> ColumnPtr : Table->GetColumns())
+	{
+		if (ColumnPtr->CanBeSorted())
+		{
+			TSharedPtr<Insights::ITableCellValueSorter> SorterPtr = ColumnPtr->GetValueSorter();
+			if (ensure(SorterPtr.IsValid()))
+			{
+				AvailableSorters.Add(SorterPtr);
+			}
+		}
+	}
+
+	UpdateCurrentSortingByColumn();
+}
+
+////////////////////////////////////////////////////////////////////////////////////////////////////
+
+void STimersView::UpdateCurrentSortingByColumn()
+{
+	TSharedPtr<Insights::FTableColumn> ColumnPtr = Table->FindColumn(ColumnBeingSorted);
+	CurrentSorter = ColumnPtr.IsValid() ? ColumnPtr->GetValueSorter() : nullptr;
+}
 
-	#define CHECK_AND_SORT_COLUMN(ColumnId, SortTypeName) \
-		if (ColumnBeingSorted == ColumnId) \
-		{ \
-			if (ColumnSortMode == EColumnSortMode::Type::Descending) \
-			{ \
-				for (int32 ID = 0; ID < NumGroups; ++ID) \
-				{ \
-					GroupNodes[ID]->SortChildren(TimerNodeSortingHelper::SortTypeName##Descending()); \
-				} \
-			} \
-			else /*if (ColumnSortMode == EColumnSortMode::Type::Ascending)*/ \
-			{ \
-				for (int32 ID = 0; ID < NumGroups; ++ID) \
-				{ \
-					GroupNodes[ID]->SortChildren(TimerNodeSortingHelper::SortTypeName##Ascending()); \
-				} \
-			} \
+////////////////////////////////////////////////////////////////////////////////////////////////////
+
+void STimersView::SortTreeNodes()
+{
+	if (CurrentSorter.IsValid())
+	{
+		for (FTimerNodePtr& Root : GroupNodes)
+		{
+			SortTreeNodesRec(*Root, *CurrentSorter);
 		}
+	}
+}
 
-		 CHECK_AND_SORT_COLUMN(FTimersViewColumns::NameColumnID,          ByName)
-	else CHECK_AND_SORT_COLUMN(FTimersViewColumns::MetaGroupNameColumnID, ByMetaGroupName)
-	else CHECK_AND_SORT_COLUMN(FTimersViewColumns::TypeColumnID,          ByType)
-	else CHECK_AND_SORT_COLUMN(FTimersViewColumns::InstanceCountColumnID, ByInstanceCount)
-	// Inclusive Time columns
-	else CHECK_AND_SORT_COLUMN(FTimersViewColumns::TotalInclusiveTimeColumnID,   ByTotalInclusiveTime)
-	else CHECK_AND_SORT_COLUMN(FTimersViewColumns::MaxInclusiveTimeColumnID,     ByMaxInclusiveTime)
-	else CHECK_AND_SORT_COLUMN(FTimersViewColumns::AverageInclusiveTimeColumnID, ByAverageInclusiveTime)
-	else CHECK_AND_SORT_COLUMN(FTimersViewColumns::MedianInclusiveTimeColumnID,  ByMedianInclusiveTime)
-	else CHECK_AND_SORT_COLUMN(FTimersViewColumns::MinInclusiveTimeColumnID,     ByMinInclusiveTime)
-	// Exclusive Time columns
-	else CHECK_AND_SORT_COLUMN(FTimersViewColumns::TotalExclusiveTimeColumnID,   ByTotalExclusiveTime)
-	else CHECK_AND_SORT_COLUMN(FTimersViewColumns::MaxExclusiveTimeColumnID,     ByMaxExclusiveTime)
-	else CHECK_AND_SORT_COLUMN(FTimersViewColumns::AverageExclusiveTimeColumnID, ByAverageExclusiveTime)
-	else CHECK_AND_SORT_COLUMN(FTimersViewColumns::MedianExclusiveTimeColumnID,  ByMedianExclusiveTime)
-	else CHECK_AND_SORT_COLUMN(FTimersViewColumns::MinExclusiveTimeColumnID,     ByMinExclusiveTime)
+////////////////////////////////////////////////////////////////////////////////////////////////////
 
-	#undef CHECK_AND_SORT_COLUMN
+void STimersView::SortTreeNodesRec(FTimerNode& Node, const Insights::ITableCellValueSorter& Sorter)
+{
+	if (ColumnSortMode == EColumnSortMode::Type::Descending)
+	{
+		Node.SortChildrenDescending(Sorter);
+	}
+	else // if (ColumnSortMode == EColumnSortMode::Type::Ascending)
+	{
+		Node.SortChildrenAscending(Sorter);
+	}
+
+	for (Insights::FBaseTreeNodePtr ChildPtr : Node.GetChildren())
+	{
+		//if (ChildPtr->IsGroup())
+		if (ChildPtr->GetChildren().Num() > 0)
+		{
+			SortTreeNodesRec(*StaticCastSharedPtr<FTimerNode>(ChildPtr), Sorter);
+		}
+	}
 }
 
 ////////////////////////////////////////////////////////////////////////////////////////////////////
@@ -1110,8 +1131,9 @@ void STimersView::SetSortModeForColumn(const FName& ColumnId, const EColumnSortM
 {
 	ColumnBeingSorted = ColumnId;
 	ColumnSortMode = SortMode;
+	UpdateCurrentSortingByColumn();
 
-	SortTimers();
+	SortTreeNodes();
 	ApplyFiltering();
 }
 
@@ -1323,8 +1345,9 @@ bool STimersView::ContextMenu_ShowAllColumns_CanExecute() const
 
 void STimersView::ContextMenu_ShowAllColumns_Execute()
 {
-	ColumnSortMode = GetDefaultColumnSortMode();
 	ColumnBeingSorted = GetDefaultColumnBeingSorted();
+	ColumnSortMode = GetDefaultColumnSortMode();
+	UpdateCurrentSortingByColumn();
 
 	for (const TSharedPtr<Insights::FTableColumn>& ColumnPtr : Table->GetColumns())
 	{
@@ -1368,8 +1391,9 @@ void STimersView::ContextMenu_ShowMinMaxMedColumns_Execute()
 		FTimersViewColumns::MinExclusiveTimeColumnID,
 	};
 
-	ColumnSortMode = EColumnSortMode::Descending;
 	ColumnBeingSorted = FTimersViewColumns::TotalInclusiveTimeColumnID;
+	ColumnSortMode = EColumnSortMode::Descending;
+	UpdateCurrentSortingByColumn();
 
 	for (const TSharedPtr<Insights::FTableColumn>& ColumnPtr : Table->GetColumns())
 	{
@@ -1401,8 +1425,9 @@ bool STimersView::ContextMenu_ResetColumns_CanExecute() const
 
 void STimersView::ContextMenu_ResetColumns_Execute()
 {
-	ColumnSortMode = GetDefaultColumnSortMode();
 	ColumnBeingSorted = GetDefaultColumnBeingSorted();
+	ColumnSortMode = GetDefaultColumnSortMode();
+	UpdateCurrentSortingByColumn();
 
 	for (const TSharedPtr<Insights::FTableColumn>& ColumnPtr : Table->GetColumns())
 	{
diff --git a/Engine/Source/Developer/TraceInsights/Private/Insights/Widgets/STimersView.h b/Engine/Source/Developer/TraceInsights/Private/Insights/Widgets/STimersView.h
--- a/Engine/Source/Developer/TraceInsights/Private/Insights/Widgets/STimersView.h
+++ b/Engine/Source/Developer/TraceInsights/Private/Insights/Widgets/STimersView.h
@@ -18,7 +18,6 @@
 // Insights
 #include "Insights/ViewModels/TimerGroupingAndSorting.h"
 #include "Insights/ViewModels/TimerNode.h"
-#include "Insights/ViewModels/TimersViewColumn.h"
 
 ////////////////////////////////////////////////////////////////////////////////////////////////////
 
@@ -32,6 +31,7 @@ namespace Trace
 namespace Insights
 {
 	class FTable;
+	class FTableColumn;
 	class ITableCellValueSorter;
 }
 
@@ -149,9 +149,9 @@ protected:
 	/** Populates the group and stat tree with items based on the current data. */
 	void ApplyFiltering();
 
-	TSharedRef<SWidget> GetToggleButtonForTimerType(const ETimerNodeType StatType);
-	void FilterByTimerType_OnCheckStateChanged(ECheckBoxState NewRadioState, const ETimerNodeType InStatType);
-	ECheckBoxState FilterByTimerType_IsChecked(const ETimerNodeType InStatType) const;
+	TSharedRef<SWidget> GetToggleButtonForTimerType(const ETimerNodeType InTimerType);
+	void FilterByTimerType_OnCheckStateChanged(ECheckBoxState NewRadioState, const ETimerNodeType InTimerType);
+	ECheckBoxState FilterByTimerType_IsChecked(const ETimerNodeType InTimerType) const;
 
 	bool SearchBox_IsEnabled() const;
 	void SearchBox_OnTextChanged(const FText& InFilterText);
@@ -173,10 +173,14 @@ protected:
 	////////////////////////////////////////////////////////////////////////////////////////////////////
 	// Sorting
 
-	static const EColumnSortMode::Type GetDefaultColumnSortMode();
 	static const FName GetDefaultColumnBeingSorted();
+	static const EColumnSortMode::Type GetDefaultColumnSortMode();
 
-	void SortTimers();
+	void CreateSortings();
+
+	void UpdateCurrentSortingByColumn();
+	void SortTreeNodes();
+	void SortTreeNodesRec(FTimerNode& Node, const Insights::ITableCellValueSorter& Sorter);
 
 	EColumnSortMode::Type GetSortModeForColumn(const FName ColumnId) const;
 	void SetSortModeForColumn(const FName& ColumnId, EColumnSortMode::Type SortMode);
@@ -256,7 +260,7 @@ protected:
 	FName HoveredColumnId;
 
 	/** A shared pointer to the timer node currently being hovered by the mouse. */
-	FTimerNodePtr HoveredTimerNodePtr;
+	FTimerNodePtr HoveredNodePtr;
 
 	/** Name of the timer that should be drawn as highlighted. */
 	FName HighlightedNodeName;
@@ -273,9 +277,6 @@ protected:
 	/** All timer nodes. */
 	TSet<FTimerNodePtr> TimerNodes;
 
-	/** All timer nodes, stored as TimerName -> FTimerNodePtr. */
-	//TMap<FName, FTimerNodePtr> TimerNodesMap;
-
 	/** All timer nodes, stored as TimerId -> FTimerNodePtr. */
 	TMap<uint64, FTimerNodePtr> TimerNodesIdMap;
 
@@ -319,12 +320,18 @@ protected:
 
 	//bool bUseSorting;
 
-	/** How we sort the timers? */
-	EColumnSortMode::Type ColumnSortMode;
+	/** All available sorters. */
+	TArray<TSharedPtr<Insights::ITableCellValueSorter>> AvailableSorters;
+
+	/** Current sorter. It is nullptr if sorting is disabled. */
+	TSharedPtr<Insights::ITableCellValueSorter> CurrentSorter;
 
-	/** Name of the column currently being sorted, NAME_None if sorting is disabled. */
+	/** Name of the column currently being sorted. Can be NAME_None if sorting is disabled (CurrentSorting == nullptr) or if a complex sorting is used (CurrentSorting != nullptr). */
 	FName ColumnBeingSorted;
 
+	/** How we sort the nodes? Ascending or Descending. */
+	EColumnSortMode::Type ColumnSortMode;
+
 	//////////////////////////////////////////////////
 
 	double StatsStartTime;
diff --git a/Engine/Source/Developer/TraceInsights/Private/Insights/Widgets/STimingProfilerWindow.cpp b/Engine/Source/Developer/TraceInsights/Private/Insights/Widgets/STimingProfilerWindow.cpp
--- a/Engine/Source/Developer/TraceInsights/Private/Insights/Widgets/STimingProfilerWindow.cpp
+++ b/Engine/Source/Developer/TraceInsights/Private/Insights/Widgets/STimingProfilerWindow.cpp
@@ -385,7 +385,7 @@ void STimingProfilerWindow::Construct(const FArguments& InArgs, const TSharedRef
 		.SetGroup(AppMenuGroup);
 
 	TabManager->RegisterTabSpawner(FTimingProfilerTabs::StatsCountersID, FOnSpawnTab::CreateRaw(this, &STimingProfilerWindow::SpawnTab_StatsCounters))
-		.SetDisplayName(LOCTEXT("StatsCountersTabTitle", "Stats Counters"))
+		.SetDisplayName(LOCTEXT("StatsCountersTabTitle", "Counters"))
 		.SetIcon(FSlateIcon(FInsightsStyle::GetStyleSetName(), "StatsCountersView.Icon.Small"))
 		.SetGroup(AppMenuGroup);
 
diff --git a/Engine/Source/Developer/TraceInsights/Private/Insights/Widgets/STimingView.cpp b/Engine/Source/Developer/TraceInsights/Private/Insights/Widgets/STimingView.cpp
--- a/Engine/Source/Developer/TraceInsights/Private/Insights/Widgets/STimingView.cpp
+++ b/Engine/Source/Developer/TraceInsights/Private/Insights/Widgets/STimingView.cpp
@@ -312,6 +312,8 @@ void STimingView::EnableAssetLoadingMode()
 
 void STimingView::Tick(const FGeometry& AllottedGeometry, const double InCurrentTime, const float InDeltaTime)
 {
+	//SCompoundWidget::Tick(AllottedGeometry, InCurrentTime, InDeltaTime);
+
 	//FStopwatch TickStopwatch;
 	//TickStopwatch.Start();
 
diff --git a/Engine/Source/Developer/TraceServices/Private/Analyzers/NetTraceAnalyzer.cpp b/Engine/Source/Developer/TraceServices/Private/Analyzers/NetTraceAnalyzer.cpp
--- a/Engine/Source/Developer/TraceServices/Private/Analyzers/NetTraceAnalyzer.cpp
+++ b/Engine/Source/Developer/TraceServices/Private/Analyzers/NetTraceAnalyzer.cpp
@@ -35,6 +35,24 @@ void FNetTraceAnalyzer::OnAnalysisEnd()
 {
 }
 
+uint32 FNetTraceAnalyzer::GetTracedEventTypeIndex(uint16 NameIndex, uint8 Level)
+{
+	const uint32 TracedEventTypeKey = uint32(NameIndex << 8U | Level);
+
+	if (const uint32* NetProfilerEventTypeIndex = TraceEventTypeToNetProfilerEventTypeIndexMap.Find(TracedEventTypeKey))
+	{
+		return *NetProfilerEventTypeIndex;
+	}
+	else
+	{
+		// Add new EventType
+		uint32 NewEventTypeIndex = NetProfilerProvider.AddNetProfilerEventType(NameIndex, Level);
+		TraceEventTypeToNetProfilerEventTypeIndexMap.Add(TracedEventTypeKey, NewEventTypeIndex);
+
+		return NewEventTypeIndex;
+	}
+}
+
 bool FNetTraceAnalyzer::OnEvent(uint16 RouteId, const FOnEventContext& Context)
 {
 	Trace::FAnalysisSessionEditScope _(Session);
@@ -86,7 +104,7 @@ bool FNetTraceAnalyzer::OnEvent(uint16 RouteId, const FOnEventContext& Context)
 			const uint16 ConnectionId = EventData.GetValue<uint16>("ConnectionId");
 			const uint8 GameInstanceId = EventData.GetValue<uint8>("ReplicationSystemId");
 			const uint8 PacketType =  EventData.GetValue<uint8>("PacketType");
-			
+
 			TSharedRef<FNetTraceGameInstanceState> GameInstanceState = GetOrCreateActiveGameInstanceState(GameInstanceId);
 			TSharedRef<FNetTraceConnectionState> ConnectionState = GetActiveConnectionState(GameInstanceId, ConnectionId);
 
@@ -96,7 +114,7 @@ bool FNetTraceAnalyzer::OnEvent(uint16 RouteId, const FOnEventContext& Context)
 			++ConnectionData.ContentEventChangeCount;
 
 			TPagedArray<Trace::FNetProfilerContentEvent>& Events = ConnectionData.ContentEvents;
-		
+
 			// New Packet, update state
 			if (ConnectionState->CurrentPacketStartIndex == Events.Num())
 			{
@@ -113,7 +131,7 @@ bool FNetTraceAnalyzer::OnEvent(uint16 RouteId, const FOnEventContext& Context)
 			const uint8* BufferPtr = EventData.GetAttachment();
 			const uint8* BufferEnd = BufferPtr + BufferSize;
 			uint64 LastOffset = 0;
-			
+
 			while (BufferPtr < BufferEnd)
 			{
 				// Decode data
@@ -150,9 +168,11 @@ bool FNetTraceAnalyzer::OnEvent(uint16 RouteId, const FOnEventContext& Context)
 						Event.NameIndex = *NetProfilerNameIndex;
 					}
 				}
+
+				// EventTypeIndex does not match NameIndex as we might see the same name on different levels
+				Event.EventTypeIndex = GetTracedEventTypeIndex(Event.NameIndex, Event.Level);
 			}
 			check(BufferPtr == BufferEnd);
-			
 		}
 		break;
 
@@ -162,7 +182,7 @@ bool FNetTraceAnalyzer::OnEvent(uint16 RouteId, const FOnEventContext& Context)
 			const uint32 SequenceNumber = EventData.GetValue<uint32>("SequenceNumber");
 			const uint8 GameInstanceId = EventData.GetValue<uint8>("ReplicationSystemId");
 			const uint16 ConnectionId = EventData.GetValue<uint16>("ConnectionId");
-			
+
 			// Update LastTimestamp, latar on we will be able to get timestamps piggybacked from other analyzers
 			LastTimeStamp = Context.SessionContext.TimestampFromCycle(TimestampCycles);
 
@@ -266,10 +286,10 @@ bool FNetTraceAnalyzer::OnEvent(uint16 RouteId, const FOnEventContext& Context)
 			const uint16 NameId = EventData.GetValue<uint16>("NameId");
 			const uint8 GameInstanceId = EventData.GetValue<uint8>("ReplicationSystemId");
 			const uint16 ConnectionId = EventData.GetValue<uint16>("ConnectionId");
-			
+
 			TSharedRef<FNetTraceGameInstanceState> GameInstanceState = GetOrCreateActiveGameInstanceState(GameInstanceId);
 			check(!GameInstanceState->ActiveObjects.Contains(ObjectId));
-	
+
 			// Add persistent object representation
 			Trace::FNetProfilerObjectInstance& ObjectInstance = NetProfilerProvider.CreateObject(GameInstanceState->GameInstanceIndex);
 
diff --git a/Engine/Source/Developer/TraceServices/Private/Analyzers/NetTraceAnalyzer.h b/Engine/Source/Developer/TraceServices/Private/Analyzers/NetTraceAnalyzer.h
--- a/Engine/Source/Developer/TraceServices/Private/Analyzers/NetTraceAnalyzer.h
+++ b/Engine/Source/Developer/TraceServices/Private/Analyzers/NetTraceAnalyzer.h
@@ -45,6 +45,8 @@ private:
 	// Shared for trace
 	TMap<uint16, uint32> TracedNameIdToNetProfilerNameIdMap;
 
+	TMap<uint32, uint32> TraceEventTypeToNetProfilerEventTypeIndexMap;
+
 	struct FNetTraceConnectionState
 	{
 		// Index into persistent connection array
@@ -70,6 +72,8 @@ private:
 		uint32 GameInstanceIndex;
 	};
 
+	uint32 GetTracedEventTypeIndex(uint16 NameIndex, uint8 Level);
+
 	TSharedRef<FNetTraceGameInstanceState> GetOrCreateActiveGameInstanceState(uint32 GameInstanceId);
 	void DestroyActiveGameInstanceState(uint32 GameInstanceId);
 
diff --git a/Engine/Source/Developer/TraceServices/Private/Model/NetProfilerProvider.cpp b/Engine/Source/Developer/TraceServices/Private/Model/NetProfilerProvider.cpp
--- a/Engine/Source/Developer/TraceServices/Private/Model/NetProfilerProvider.cpp
+++ b/Engine/Source/Developer/TraceServices/Private/Model/NetProfilerProvider.cpp
@@ -15,16 +15,18 @@ FNetProfilerProvider::FNetProfilerProvider(IAnalysisSession& InSession)
 	, Connections(InSession.GetLinearAllocator(), 4096)
 	, ConnectionChangeCount(0u)
 {
+	// Use name index 0 to indicate that we do not know the name
+	AddNetProfilerName(TEXT("N/A"));
 }
 
 FNetProfilerProvider::~FNetProfilerProvider()
 {
 }
 
-uint32 FNetProfilerProvider::AddNetProfilerName(TCHAR* Name)
+uint32 FNetProfilerProvider::AddNetProfilerName(const TCHAR* Name)
 {
 	Session.WriteAccessCheck();
-	
+
 	FNetProfilerName& NewName = Names.AddDefaulted_GetRef();
 	NewName.NameIndex = Names.Num() - 1;
 	NewName.Name = Session.StoreString(Name);
@@ -37,6 +39,24 @@ const FNetProfilerName* FNetProfilerProvider::GetNetProfilerName(uint32 NameInde
 	return NameIndex < (uint32)Names.Num() ? &Names[NameIndex] : nullptr;
 }
 
+uint32 FNetProfilerProvider::AddNetProfilerEventType(uint32 NameIndex, uint32 Level)
+{
+	Session.WriteAccessCheck();
+
+	FNetProfilerEventType& NewEventType = EventTypes.AddDefaulted_GetRef();
+	NewEventType.EventTypeIndex = EventTypes.Num() - 1;
+	NewEventType.NameIndex = NameIndex;
+	NewEventType.Name = GetNetProfilerName(NameIndex)->Name;
+	NewEventType.Level = Level;
+
+	return NewEventType.EventTypeIndex;
+}
+
+const FNetProfilerEventType* FNetProfilerProvider::GetNetProfilerEventType(uint32 EventTypeIndex) const
+{
+	return EventTypeIndex < (uint32)EventTypes.Num() ? &EventTypes[EventTypeIndex] : nullptr;
+}
+
 Trace::FNetProfilerGameInstanceInternal& FNetProfilerProvider::CreateGameInstance()
 {
 	Session.WriteAccessCheck();
@@ -48,6 +68,9 @@ Trace::FNetProfilerGameInstanceInternal& FNetProfilerProvider::CreateGameInstanc
 	new (GameInstance.Objects) TPagedArray<FNetProfilerObjectInstance>(Session.GetLinearAllocator(), 4096);
 	GameInstance.ObjectsChangeCount = 0u;
 
+	// We reserve object index 0 as an invalid object
+	CreateObject(GameInstance.Instance.GameInstanceIndex);
+
 	return GameInstance;
 }
 
@@ -77,7 +100,7 @@ Trace::FNetProfilerConnectionInternal& FNetProfilerProvider::CreateConnection(ui
 	FNetProfilerConnectionInternal& Connection = Connections.PushBack();
 
 	Connection.Connection.ConnectionIndex = ConnectionIndex;
-	Connection.Connection.GameIntanceIndex = GameInstanceIndex;
+	Connection.Connection.GameInstanceIndex = GameInstanceIndex;
 	Connection.Connection.bHasIncomingData = false;
 	Connection.Connection.bHasOutgoingData = false;
 
@@ -115,7 +138,7 @@ Trace::FNetProfilerObjectInstance* FNetProfilerProvider::EditObject(uint32 GameI
 
 	FNetProfilerGameInstanceInternal* GameInstance = EditGameInstance(GameInstanceIndex);
 	check(GameInstance);
-	
+
 	if (ensure(ObjectIndex < (uint32)GameInstance->Objects->Num()))
 	{
 		++GameInstance->ObjectsChangeCount;
@@ -130,7 +153,7 @@ Trace::FNetProfilerObjectInstance* FNetProfilerProvider::EditObject(uint32 GameI
 Trace::FNetProfilerConnectionInternal* FNetProfilerProvider::EditConnection(uint32 ConnectionIndex)
 {
 	Session.WriteAccessCheck();
-	
+
 	if (ensure(ConnectionIndex < (uint32)Connections.Num()))
 	{
 		++ConnectionChangeCount;
@@ -198,6 +221,21 @@ void FNetProfilerProvider::ReadName(uint32 NameIndex, TFunctionRef<void(const FN
 	Callback(*GetNetProfilerName(NameIndex));
 }
 
+void FNetProfilerProvider::ReadEventTypes(TFunctionRef<void(const FNetProfilerEventType*, uint64)> Callback) const
+{
+	Session.ReadAccessCheck();
+
+	Callback(EventTypes.GetData(), EventTypes.Num());
+}
+
+void FNetProfilerProvider::ReadEventType(uint32 EventTypeIndex, TFunctionRef<void(const FNetProfilerEventType&)> Callback) const
+{
+	Session.ReadAccessCheck();
+	check(EventTypeIndex < (uint32)EventTypes.Num());
+
+	Callback(*GetNetProfilerEventType(EventTypeIndex));
+}
+
 void FNetProfilerProvider::ReadGameInstances(TFunctionRef<void(const FNetProfilerGameInstance&)> Callback) const
 {
 	Session.ReadAccessCheck();
@@ -281,7 +319,7 @@ void FNetProfilerProvider::ReadObject(uint32 GameInstanceIndex, uint32 ObjectInd
 	Callback((*GameInstance.Objects)[ObjectIndex]);
 }
 
-uint32 FNetProfilerProvider::GeObjectsChangeCount(uint32 GameInstanceIndex) const
+uint32 FNetProfilerProvider::GetObjectsChangeCount(uint32 GameInstanceIndex) const
 {
 	Session.ReadAccessCheck();
 
@@ -306,7 +344,7 @@ uint32 FNetProfilerProvider::GetPacketCount(uint32 ConnectionIndex, ENetProfiler
 	check(ConnectionIndex < Connections.Num());
 
 	const auto& Packets = Connections[ConnectionIndex].Data[Mode]->Packets;
-	
+
 	return Packets.Num();
 }
 
@@ -317,7 +355,7 @@ void FNetProfilerProvider::EnumeratePackets(uint32 ConnectionIndex, ENetProfiler
 	check(ConnectionIndex < Connections.Num());
 
 	const auto& Packets = Connections[ConnectionIndex].Data[Mode]->Packets;
-	
+
 	const uint32 PacketCount = Packets.Num();
 
 	if (PacketCount == 0 || PacketIndexIntervalStart > PacketIndexIntervalEnd)
@@ -339,7 +377,7 @@ void FNetProfilerProvider::EnumeratePacketContentEventsByIndex(uint32 Connection
 	check(ConnectionIndex < Connections.Num());
 
 	const auto& ContentEvents = Connections[ConnectionIndex].Data[Mode]->ContentEvents;
-	
+
 	const uint32 EventCount = ContentEvents.Num();
 
 	if (EventCount == 0 || StartEventIndex > EndEventIndex)
@@ -371,8 +409,8 @@ void FNetProfilerProvider::EnumeratePacketContentEventsByPosition(uint32 Connect
 	const uint32 EndEventIndex = StartEventIndex + Packet.EventCount - 1u;
 
 	const auto& ContentEvents = ConnectionData.ContentEvents;
-	
-	uint32 EventIt = StartEventIndex; 
+
+	uint32 EventIt = StartEventIndex;
 	// Skip all Events outside of the scope
 	while (EventIt <= EndEventIndex &&	ContentEvents[EventIt].EndPos < StartPos)
 	{
@@ -403,4 +441,119 @@ uint32 FNetProfilerProvider::GetPacketContentEventChangeCount(uint32 ConnectionI
 	return Connections[ConnectionIndex].Data[Mode]->ContentEventChangeCount;
 }
 
+ITable<FNetProfilerAggregatedStats>* FNetProfilerProvider::CreateAggregation(uint32 ConnectionIndex, ENetProfilerConnectionMode Mode, uint32 PacketIndexIntervalStart, uint32 PacketIndexIntervalEnd, uint32 StartPosition, uint32 EndPosition) const
+{
+	Session.ReadAccessCheck();
+
+	if (!ensure(ConnectionIndex < Connections.Num()))
+	{
+		return nullptr;
+	}
+
+	if (!ensure(PacketIndexIntervalStart <= PacketIndexIntervalEnd))
+	{
+		return nullptr;
+	}
+
+	const auto& Packets = Connections[ConnectionIndex].Data[Mode]->Packets;
+	const uint32 PacketCount = Packets.Num();
+
+	if (!ensure(PacketCount > 0))
+	{
+		return nullptr;
+	}
+
+	// Aggregate stats for events in selection
+	TMap<uint32, FNetProfilerAggregatedStats> AggregatedStats;
+	AggregatedStats.Reserve(EventTypes.Num());
+
+	struct FStackEntry
+	{
+		uint32 EventTypeIndex = 0U;
+		uint32 StartPos = 0U;
+		uint32 EndPos = 0U;
+		uint32 ExclusiveAccumulator = 0U;
+	};
+
+	TArray<FStackEntry> Stack;
+	Stack.SetNum(256);
+
+	auto GetStatsFunction = [&AggregatedStats, &Stack, this](const FNetProfilerContentEvent& ContentEvent)
+	{
+		FNetProfilerAggregatedStats* StatsEntry = AggregatedStats.Find(ContentEvent.EventTypeIndex);
+		if (!StatsEntry)
+		{
+			StatsEntry = &AggregatedStats.Add(ContentEvent.EventTypeIndex);
+			StatsEntry->EventTypeIndex = ContentEvent.EventTypeIndex;
+		}
+
+		// Fill in basics
+		uint32 InclusiveSize = ContentEvent.EndPos - ContentEvent.StartPos;
+
+		++StatsEntry->InstanceCount;
+		StatsEntry->TotalInclusive += InclusiveSize;
+		StatsEntry->MaxInclusive = FMath::Max(InclusiveSize, StatsEntry->MaxInclusive);
+
+		// We track what we have visited to be able to update exclusive bits for our parent
+		// We accumulate during the first pass and finalize the values during the second
+		Stack[ContentEvent.Level].EventTypeIndex = ContentEvent.EventTypeIndex;
+		Stack[ContentEvent.Level].StartPos = ContentEvent.StartPos;
+		Stack[ContentEvent.Level].EndPos = ContentEvent.EndPos;
+		Stack[ContentEvent.Level].ExclusiveAccumulator = 0U;
+
+		if (ContentEvent.Level > 0U)
+		{
+			FStackEntry& ParentStackEntry = Stack[ContentEvent.Level - 1U];
+
+			ParentStackEntry.ExclusiveAccumulator += InclusiveSize;
+
+			if (ContentEvent.EndPos == ParentStackEntry.EndPos)
+			{
+				// Finalize exclusive
+				FNetProfilerAggregatedStats& ParentStateEntry = AggregatedStats.FindChecked(ContentEvent.EventTypeIndex);
+				uint32 ExclusiveTime = (ParentStackEntry.EndPos - ParentStackEntry.StartPos) - ParentStackEntry.ExclusiveAccumulator;
+				ParentStateEntry.TotalExclusive += ExclusiveTime;
+				ParentStateEntry.MaxExclusive = FMath::Max(ParentStateEntry.MaxExclusive, ExclusiveTime);
+			}
+		}
+	};
+
+	// Iterate over content events
+	if (PacketIndexIntervalStart == PacketIndexIntervalEnd)
+	{
+		EnumeratePacketContentEventsByPosition(ConnectionIndex, Mode, PacketIndexIntervalStart, StartPosition, EndPosition, GetStatsFunction);
+	}
+	else
+	{
+		// Iterate over packets
+		for (uint32 PacketIt = PacketIndexIntervalStart, PacketEndIt = FMath::Min(PacketIndexIntervalEnd, PacketCount - 1u); PacketIt <= PacketEndIt; ++PacketIt)
+		{
+			const auto& ContentEvents = Connections[ConnectionIndex].Data[Mode]->ContentEvents;
+			const uint32 EventCount = ContentEvents.Num();
+
+			const FNetProfilerPacket& Packet = Packets[PacketIt];
+			if (Packet.EventCount > 0U)
+			{
+				const uint32 StartEventIndex = Packet.StartEventIndex;
+				for (uint32 It = 0U; It < Packet.EventCount; ++It)
+				{
+					GetStatsFunction(ContentEvents[StartEventIndex + It]);
+				}
+			}
+		}
+	}
+
+	// Calculate averages and populate table
+	TTable<FAggregatedStatsTableLayout>* Table = new TTable<FAggregatedStatsTableLayout>();
+	for (const auto& KV : AggregatedStats)
+	{
+		FNetProfilerAggregatedStats& Row = Table->AddRow();
+		Row = KV.Value;
+
+		// Finalize StatsEntry
+		Row.AverageInclusive = (uint64)((double)KV.Value.TotalInclusive / KV.Value.InstanceCount);
+	}
+	return Table;
+}
+
 }
diff --git a/Engine/Source/Developer/TraceServices/Private/Model/NetProfilerProvider.h b/Engine/Source/Developer/TraceServices/Private/Model/NetProfilerProvider.h
--- a/Engine/Source/Developer/TraceServices/Private/Model/NetProfilerProvider.h
+++ b/Engine/Source/Developer/TraceServices/Private/Model/NetProfilerProvider.h
@@ -53,8 +53,9 @@ public:
 	explicit FNetProfilerProvider(IAnalysisSession& InSession);
 	virtual ~FNetProfilerProvider();
 
-	uint32 AddNetProfilerName(TCHAR* Name) ;
-	const FNetProfilerName* GetNetProfilerName(uint32 ProfilerNameId) const;
+	uint32 AddNetProfilerName(const TCHAR* Name);
+
+	uint32 AddNetProfilerEventType(uint32 NameIndex, uint32 Level);
 
 	Trace::FNetProfilerGameInstanceInternal& CreateGameInstance();
 	Trace::FNetProfilerGameInstanceInternal* EditGameInstance(uint32 GameInstanceIndex);
@@ -75,6 +76,11 @@ public:
 	virtual void ReadNames(TFunctionRef<void(const FNetProfilerName*, uint64)> Callback) const override;
 	virtual void ReadName(uint32 NameIndex, TFunctionRef<void(const FNetProfilerName&)> Callback) const override;
 
+	// Access EventTypers
+	virtual uint32 GetEventTypesCount() const override { return EventTypes.Num(); }
+	virtual void ReadEventTypes(TFunctionRef<void(const FNetProfilerEventType*, uint64)> Callback) const override;
+	virtual void ReadEventType(uint32 EventTypeIndex, TFunctionRef<void(const FNetProfilerEventType&)> Callback) const override;
+
 	// Access GameInstances
 	virtual uint32 GetGameInstanceCount() const override { return GameInstances.Num(); }
 	virtual void ReadGameInstances(TFunctionRef<void(const FNetProfilerGameInstance&)> Callback) const override;
@@ -89,7 +95,7 @@ public:
 	virtual uint32 GetObjectCount(uint32 GameInstanceIndex) const override;
 	virtual void ReadObjects(uint32 GameInstanceIndex, TFunctionRef<void(const FNetProfilerObjectInstance&)> Callback) const override;
 	virtual void ReadObject(uint32 GameInstanceIndex, uint32 ObjectIndex, TFunctionRef<void(const FNetProfilerObjectInstance&)> Callback) const override;
-	virtual uint32 GeObjectsChangeCount(uint32 GameInstanceIndex) const override;
+	virtual uint32 GetObjectsChangeCount(uint32 GameInstanceIndex) const override;
 
 	// Enumerate packets in the provided packet interval
 	virtual uint32 GetPacketCount(uint32 ConnectionIndex, ENetProfilerConnectionMode Mode) const override;
@@ -100,15 +106,34 @@ public:
 	virtual void EnumeratePacketContentEventsByIndex(uint32 ConnectionIndex, ENetProfilerConnectionMode Mode, uint32 StartEventIndex, uint32 EndEventIndex, TFunctionRef<void(const FNetProfilerContentEvent&)> Callback) const override;
 	virtual void EnumeratePacketContentEventsByPosition(uint32 ConnectionIndex, ENetProfilerConnectionMode Mode, uint32 PacketIndex, uint32 StartPos, uint32 EndPos, TFunctionRef<void(const FNetProfilerContentEvent&)> Callback) const override;
 	virtual uint32 GetPacketContentEventChangeCount(uint32 ConnectionIndex, ENetProfilerConnectionMode Mode) const override;
-	
+
+	// Stats queries
+	virtual ITable<FNetProfilerAggregatedStats>* CreateAggregation(uint32 ConnectionIndex, ENetProfilerConnectionMode Mode, uint32 PacketIndexIntervalStart, uint32 PacketIndexIntervalEnd, uint32 StartPosition, uint32 EndPosition) const override;
+
 private:
+
+	UE_TRACE_TABLE_LAYOUT_BEGIN(FAggregatedStatsTableLayout, FNetProfilerAggregatedStats)
+		UE_TRACE_TABLE_COLUMN(EventTypeIndex, TEXT("EventTypeIndex"))
+		UE_TRACE_TABLE_COLUMN(InstanceCount, TEXT("Count"))
+		UE_TRACE_TABLE_COLUMN(TotalInclusive, TEXT("Incl"))
+		UE_TRACE_TABLE_COLUMN(MaxInclusive, TEXT("I.Max"))
+		UE_TRACE_TABLE_COLUMN(AverageInclusive, TEXT("I.Avg"))
+		UE_TRACE_TABLE_COLUMN(TotalExclusive, TEXT("Excl"))
+		UE_TRACE_TABLE_COLUMN(MaxExclusive, TEXT("E.Max"))
+	UE_TRACE_TABLE_LAYOUT_END()
+
+	const FNetProfilerName* GetNetProfilerName(uint32 ProfilerNameId) const;
+	const FNetProfilerEventType* GetNetProfilerEventType(uint32 ProfilerEventTypeId) const;
+
 	IAnalysisSession& Session;
 
 	TArray<FNetProfilerName> Names;
 
+	TArray<FNetProfilerEventType> EventTypes;
+
 	// All GameInstances seen throughout the session
 	TArray<FNetProfilerGameInstanceInternal, TInlineAllocator<4>> GameInstances;
-	
+
 	// All connections we have seen throughout the session
 	TPagedArray<FNetProfilerConnectionInternal> Connections;
 	uint32 ConnectionChangeCount;
diff --git a/Engine/Source/Developer/TraceServices/Public/TraceServices/Containers/Tables.h b/Engine/Source/Developer/TraceServices/Public/TraceServices/Containers/Tables.h
--- a/Engine/Source/Developer/TraceServices/Public/TraceServices/Containers/Tables.h
+++ b/Engine/Source/Developer/TraceServices/Public/TraceServices/Containers/Tables.h
@@ -2,6 +2,8 @@
 
 #pragma once
 
+#include "CoreTypes.h"
+
 namespace Trace
 {
 
diff --git a/Engine/Source/Developer/TraceServices/Public/TraceServices/Model/NetProfiler.h b/Engine/Source/Developer/TraceServices/Public/TraceServices/Model/NetProfiler.h
--- a/Engine/Source/Developer/TraceServices/Public/TraceServices/Model/NetProfiler.h
+++ b/Engine/Source/Developer/TraceServices/Public/TraceServices/Model/NetProfiler.h
@@ -2,8 +2,11 @@
 
 #pragma once
 
-#include "TraceServices/Model/AnalysisSession.h"
+#include "CoreTypes.h"
 #include "Templates/Function.h"
+#include "TraceServices/Containers/Tables.h"
+#include "TraceServices/Model/AnalysisSession.h"
+#include <limits>
 
 namespace Trace
 {
@@ -28,12 +31,21 @@ struct FNetProfilerName
 	uint32 NameIndex;	// Index in the type array?
 };
 
+// This is our event type separate per level as the same name might be used on different levels
+struct FNetProfilerEventType
+{
+	uint32 EventTypeIndex;
+	const TCHAR* Name;
+	uint32 NameIndex : 16;
+	uint32 Level : 16;
+};
+
 typedef double FNetProfilerTimeStamp;
 
 struct FNetProfilerLifeTime
 {
 	FNetProfilerTimeStamp Begin = 0;
-	FNetProfilerTimeStamp End = 0;
+	FNetProfilerTimeStamp End = std::numeric_limits<double>::infinity();
 };
 
 //struct FNetProfilerProtocol
@@ -45,17 +57,18 @@ struct FNetProfilerLifeTime
 
 struct FNetProfilerObjectInstance
 {
-	uint32 ObjectIndex;				// Index in the object array
-	uint16 NameIndex;				// Index in the Name array
-	uint32 TypeId;					// ProtocolIdentifier
-	uint32 NetId;					// NetHandleIndex or NetGUID
+	uint32 ObjectIndex = 0U;		// Index in the object array
+	uint16 NameIndex = 0U;			// Index in the Name array
+	uint64 TypeId = uint64(0);		// ProtocolIdentifier
+	uint32 NetId = 0U;				// NetHandleIndex or NetGUID
 	FNetProfilerLifeTime LifeTime;	// Lifetime of this instance
 };
 
 struct FNetProfilerContentEvent
 {
+	uint32 EventTypeIndex;		// Will replace name index
 	uint32 NameIndex ;			// identify the name / type, should we store the actual Name as well
-	uint32 ObjectInstanceIndex;	// object instance, non zero if this is a NetObject
+	uint32 ObjectInstanceIndex;	// object instance, Non zero if this is a NetObject, we can then look up data by indexing into ObjectInstances
 
 	uint64 StartPos : 16;		// Start position in the packet
 	uint64 EndPos : 16;			// End position in the packet
@@ -80,7 +93,7 @@ struct FNetProfilerConnection
 {
 	const TCHAR* Name;
 	FNetProfilerLifeTime LifeTime;
-	uint32 GameIntanceIndex;
+	uint32 GameInstanceIndex;
 	uint32 ConnectionIndex : 16;
 	uint32 ConnectionId : 14;
 	uint32 bHasIncomingData : 1;
@@ -94,6 +107,18 @@ struct FNetProfilerGameInstance
 	uint32 GameInstanceId;
 };
 
+// What do we need?
+struct FNetProfilerAggregatedStats
+{
+	uint32 EventTypeIndex;
+	uint32 InstanceCount = 0U;
+	uint32 TotalInclusive = 0U;
+	uint32 MaxInclusive = 0U;
+	uint32 AverageInclusive = 0U;
+	uint32 TotalExclusive = 0U;
+	uint32 MaxExclusive = 0U;
+};
+
 // What queries do we need?
 class INetProfilerProvider : public IProvider
 {
@@ -103,6 +128,11 @@ public:
 	virtual void ReadNames(TFunctionRef<void(const FNetProfilerName*, uint64)> Callback) const = 0;
 	virtual void ReadName(uint32 NameIndex, TFunctionRef<void(const FNetProfilerName&)> Callback) const = 0;
 
+	// Access EventTypes
+	virtual uint32 GetEventTypesCount() const = 0;
+	virtual void ReadEventTypes(TFunctionRef<void(const FNetProfilerEventType*, uint64)> Callback) const = 0;
+	virtual void ReadEventType(uint32 EventTypeIndex, TFunctionRef<void(const FNetProfilerEventType&)> Callback) const = 0;
+
 	// Access GameInstances
 	virtual uint32 GetGameInstanceCount() const = 0;
 	virtual void ReadGameInstances(TFunctionRef<void(const FNetProfilerGameInstance&)> Callback) const = 0;
@@ -114,10 +144,10 @@ public:
 	virtual uint32 GetConnectionChangeCount() const = 0;
 
 	// Access Object Instances
-	virtual uint32 GetObjectCount(uint32 GameInstanceIndex) const = 0;	
+	virtual uint32 GetObjectCount(uint32 GameInstanceIndex) const = 0;
 	virtual void ReadObjects(uint32 GameInstanceIndex, TFunctionRef<void(const FNetProfilerObjectInstance&)> Callback) const = 0;
 	virtual void ReadObject(uint32 GameInstanceIndex, uint32 ObjectIndex, TFunctionRef<void(const FNetProfilerObjectInstance&)> Callback) const = 0;
-	virtual uint32 GeObjectsChangeCount(uint32 GameInstanceIndex) const = 0;
+	virtual uint32 GetObjectsChangeCount(uint32 GameInstanceIndex) const = 0;
 
 	// Enumerate packets in the provided packet interval
 	virtual uint32 GetPacketCount(uint32 ConnectionIndex, ENetProfilerConnectionMode Mode) const = 0;
@@ -128,6 +158,9 @@ public:
 	virtual void EnumeratePacketContentEventsByIndex(uint32 ConnectionIndex, ENetProfilerConnectionMode Mode, uint32 StartEventIndex, uint32 EndEventIndex, TFunctionRef<void(const FNetProfilerContentEvent&)> Callback) const = 0;
 	virtual void EnumeratePacketContentEventsByPosition(uint32 ConnectionIndex, ENetProfilerConnectionMode Mode, uint32 PacketIndex, uint32 StartPosition, uint32 EndPosition, TFunctionRef<void(const FNetProfilerContentEvent&)> Callback) const = 0;
 	virtual uint32 GetPacketContentEventChangeCount(uint32 ConnectionIndex, ENetProfilerConnectionMode Mode) const = 0;
+
+	// Stats queries
+	virtual ITable<FNetProfilerAggregatedStats>* CreateAggregation(uint32 ConnectionIndex, ENetProfilerConnectionMode Mode, uint32 PacketIndexIntervalStart, uint32 PacketIndexIntervalEnd, uint32 StartPosition, uint32 EndPosition) const = 0;
 };
 
 TRACESERVICES_API const INetProfilerProvider& ReadNetProfilerProvider(const IAnalysisSession& Session);
