Fixed more compile errors.

#rnx

diff --git a/Engine/Source/Developer/TraceInsights/Private/Insights/NetworkingProfiler/ViewModels/NetEventGroupingAndSorting.cpp b/Engine/Source/Developer/TraceInsights/Private/Insights/NetworkingProfiler/ViewModels/NetEventGroupingAndSorting.cpp
--- a/Engine/Source/Developer/TraceInsights/Private/Insights/NetworkingProfiler/ViewModels/NetEventGroupingAndSorting.cpp
+++ b/Engine/Source/Developer/TraceInsights/Private/Insights/NetworkingProfiler/ViewModels/NetEventGroupingAndSorting.cpp
@@ -330,4 +330,5 @@ Insights::FTreeNodeGroupInfo FNetEventNodeGroupingByLevel::GetGroupForNode(const
 
 ////////////////////////////////////////////////////////////////////////////////////////////////////
 
+#undef INSIGHTS_DEFAULT_SORTING_NODES
 #undef LOCTEXT_NAMESPACE
diff --git a/Engine/Source/Developer/TraceInsights/Private/Insights/ViewModels/TimerGroupingAndSorting.cpp b/Engine/Source/Developer/TraceInsights/Private/Insights/ViewModels/TimerGroupingAndSorting.cpp
--- a/Engine/Source/Developer/TraceInsights/Private/Insights/ViewModels/TimerGroupingAndSorting.cpp
+++ b/Engine/Source/Developer/TraceInsights/Private/Insights/ViewModels/TimerGroupingAndSorting.cpp
@@ -301,4 +301,5 @@ Insights::FTreeNodeGroupInfo FTimerNodeGroupingByTimerType::GetGroupForNode(cons
 
 ////////////////////////////////////////////////////////////////////////////////////////////////////
 
+#undef INSIGHTS_DEFAULT_SORTING_NODES
 #undef LOCTEXT_NAMESPACE
