[Insights]
  - Networking: Renamed PacketBreakdown* files to PacketContent* (classes were already renamed).
  - Networking: Renamed PacketSizesView* files to PacketView* (classes were already renamed).

diff --git a/Engine/Source/Developer/TraceInsights/Private/Insights/NetworkingProfiler/ViewModels/PacketBreakdownViewHelper.cpp b/Engine/Source/Developer/TraceInsights/Private/Insights/NetworkingProfiler/ViewModels/PacketBreakdownViewHelper.cpp
deleted file mode 100644
diff --git a/Engine/Source/Developer/TraceInsights/Private/Insights/NetworkingProfiler/ViewModels/PacketBreakdownViewHelper.h b/Engine/Source/Developer/TraceInsights/Private/Insights/NetworkingProfiler/ViewModels/PacketBreakdownViewHelper.h
deleted file mode 100644
diff --git a/Engine/Source/Developer/TraceInsights/Private/Insights/NetworkingProfiler/ViewModels/PacketBreakdownViewport.cpp b/Engine/Source/Developer/TraceInsights/Private/Insights/NetworkingProfiler/ViewModels/PacketBreakdownViewport.cpp
deleted file mode 100644
diff --git a/Engine/Source/Developer/TraceInsights/Private/Insights/NetworkingProfiler/ViewModels/PacketBreakdownViewport.h b/Engine/Source/Developer/TraceInsights/Private/Insights/NetworkingProfiler/ViewModels/PacketBreakdownViewport.h
deleted file mode 100644
diff --git a/Engine/Source/Developer/TraceInsights/Private/Insights/NetworkingProfiler/ViewModels/PacketContentViewDrawHelper.cpp b/Engine/Source/Developer/TraceInsights/Private/Insights/NetworkingProfiler/ViewModels/PacketContentViewDrawHelper.cpp
new file mode 100644
--- /dev/null
+++ b/Engine/Source/Developer/TraceInsights/Private/Insights/NetworkingProfiler/ViewModels/PacketContentViewDrawHelper.cpp
@@ -0,0 +1,441 @@
+// Copyright 1998-2019 Epic Games, Inc. All Rights Reserved.
+
+#include "PacketContentViewDrawHelper.h"
+
+#include "Fonts/FontMeasure.h"
+#include "Framework/Application/SlateApplication.h"
+#include "Rendering/DrawElements.h"
+#include "Styling/CoreStyle.h"
+#include "TraceServices/AnalysisService.h"
+
+// Insights
+#include "Insights/Common/PaintUtils.h"
+#include "Insights/InsightsStyle.h"
+#include "Insights/NetworkingProfiler/ViewModels/PacketContentViewport.h"
+#include "Insights/ViewModels/DrawHelpers.h"
+
+#include <limits>
+
+////////////////////////////////////////////////////////////////////////////////////////////////////
+// FPacketContentViewDrawStateBuilder
+////////////////////////////////////////////////////////////////////////////////////////////////////
+
+FPacketContentViewDrawStateBuilder::FPacketContentViewDrawStateBuilder(FPacketContentViewDrawState& InDrawState, const FPacketContentViewport& InViewport)
+	: DrawState(InDrawState)
+	, Viewport(InViewport)
+	, MaxDepth(-1)
+	, LastEventX2()
+	, LastBox()
+	, EventFont(FCoreStyle::GetDefaultFontStyle("Regular", 8))
+{
+	DrawState.Reset();
+}
+
+////////////////////////////////////////////////////////////////////////////////////////////////////
+
+void FPacketContentViewDrawStateBuilder::AddEvent(const Trace::FNetProfilerContentEvent& Event, const TCHAR* EventName, uint32 NetId)
+{
+	DrawState.Events.AddUninitialized();
+	FNetworkPacketEvent& PacketEvent = DrawState.Events.Last();
+	PacketEvent.EventTypeIndex = Event.EventTypeIndex;
+	PacketEvent.ObjectInstanceIndex = Event.ObjectInstanceIndex;
+	PacketEvent.NetId = NetId;
+	PacketEvent.BitOffset = Event.StartPos;
+	PacketEvent.BitSize = Event.EndPos - Event.StartPos;
+	PacketEvent.Level = Event.Level;
+
+	const FAxisViewportDouble& ViewportX = Viewport.GetHorizontalAxisViewport();
+
+	float EventX1 = ViewportX.GetRoundedOffsetForValue(static_cast<double>(Event.StartPos));
+	if (EventX1 > Viewport.GetWidth())
+	{
+		return;
+	}
+
+	//float EventX2 = EventSize == 0 ?
+	//	EventX1 + 1.0f :
+	//	ViewportX.GetScale() < 1.0f ?
+	//	ViewportX.GetRoundedOffsetForValue(static_cast<double>(Event.EndPos - 1)) + 1.0f : // x2 = pos of last inclusive bit + 1px
+	//	ViewportX.GetRoundedOffsetForValue(static_cast<double>(Event.EndPos));             // x2 = pos of first exclusive bit
+	float EventX2 = ViewportX.GetRoundedOffsetForValue(static_cast<double>(Event.EndPos));
+	if (EventX2 < 0)
+	{
+		return;
+	}
+
+	// Events are displayed with minimum 1px (including empty ones).
+	if (EventX1 == EventX2)
+	{
+		EventX2 = EventX1 + 1.0f;
+	}
+
+	const int32 Depth = static_cast<int32>(Event.Level);
+	if (Depth > MaxDepth)
+	{
+		MaxDepth = Depth;
+	}
+
+	// Ensure we have enough slots in array. See LastBox[Depth] usage.
+	if (LastBox.Num() <= Depth)
+	{
+		LastBox.AddDefaulted(Depth + 1 - LastBox.Num());
+	}
+
+	constexpr float Y0 = 0.0f;
+	constexpr float EventH = 14.0f;
+	constexpr float EventDY = 2.0f;
+
+	const float EventY = Y0 + (EventH + EventDY) * Depth;
+	if (EventY < -EventH || EventY > Viewport.GetHeight())
+	{
+		return;
+	}
+
+	// Ensure we have enough slots in array. See LastEventX2[Depth] usage.
+	while (LastEventX2.Num() <= Depth)
+	{
+		LastEventX2.Add(-1.0f);
+	}
+
+	// Limit event width on the viewport's left side.
+	// This also makes the text to be displayed in viewport, for very long events.
+	const float MinX = -2.0f; // -2 allows event border to remain outside screen
+	if (EventX1 < MinX)
+	{
+		EventX1 = MinX;
+	}
+
+	// Limit event width on the viewport's right side.
+	const float MaxX = Viewport.GetWidth() + 2.0f; // +2 allows event border to remain outside screen
+	if (EventX2 > MaxX)
+	{
+		EventX2 = MaxX;
+	}
+
+	const float EventW = EventX2 - EventX1;
+
+	// Optimization...
+	if (EventW == 1.0f && EventX1 == LastEventX2[Depth] - 1.0f)
+	{
+		// Do no draw 1 pixel event if the last event was ended on that pixel.
+		return;
+	}
+
+	//////////////////////////////////////////////////
+	// Coloring
+
+	const uint32 Color = (Event.NameIndex + 0x2c2c57ed) * 0x2c2c57ed;
+	const FLinearColor EventColorFill(((Color >> 16) & 0xFF) / 255.0f, ((Color >> 8) & 0xFF) / 255.0f, (Color & 0xFF) / 255.0f, 1.0f);
+	constexpr float BorderColorFactor = 1.25f; // brighter border
+
+	//////////////////////////////////////////////////
+
+	if (EventW > 2.0f && EventH > 2.0f)
+	{
+		// Fill inside of the event box.
+		DrawState.Boxes.AddUninitialized();
+		FPacketContentViewDrawState::FBox& DrawBox = DrawState.Boxes.Last();
+		DrawBox.X = EventX1 + 1.0f;
+		DrawBox.Y = EventY + 1.0f;
+		DrawBox.W = EventW - 2.0f;
+		DrawBox.H = EventH - 2.0f;
+		DrawBox.Color = EventColorFill;
+	}
+
+	// Save X2, for current depth.
+	LastEventX2[Depth] = EventX2;
+
+	// Add border around the event box.
+	if (EventW > 2.0f)
+	{
+		FBoxData& Box = LastBox[Depth];
+		if (Box.X1 < Box.X2)
+		{
+			FlushBox(Box, EventY, EventH);
+			Box.Reset();
+		}
+
+		DrawState.Borders.AddUninitialized();
+		FPacketContentViewDrawState::FBox& DrawBox = DrawState.Borders.Last();
+		DrawBox.X = EventX1;
+		DrawBox.Y = EventY;
+		DrawBox.W = EventW;
+		DrawBox.H = EventH;
+		DrawBox.Color.R = EventColorFill.R * BorderColorFactor;
+		DrawBox.Color.G = EventColorFill.G * BorderColorFactor;
+		DrawBox.Color.B = EventColorFill.B * BorderColorFactor;
+		DrawBox.Color.A = 1.0f;
+	}
+	else // 1px or 2px boxes
+	{
+		FBoxData& Box = LastBox[Depth];
+
+		// Check if we can merge this box with previous one, if any.
+		// Note: We are assuming events are processed in sorted order by X1.
+		if (Color == Box.Color && // same color
+			EventX1 <= Box.X2) // overlapping or adjacent
+		{
+			// Merge it with previous box.
+			Box.X2 = EventX2;
+			DrawState.NumMergedBoxes++;
+		}
+		else
+		{
+			// Flush previous box, if any.
+			if (Box.X1 < Box.X2)
+			{
+				FlushBox(Box, EventY, EventH);
+			}
+
+			// Start new "merge box".
+			Box.X1 = EventX1;
+			Box.X2 = EventX2;
+			Box.Color = Color;
+			Box.LinearColor = FLinearColor(EventColorFill.R * BorderColorFactor, EventColorFill.G * BorderColorFactor, EventColorFill.B * BorderColorFactor, 1.0f);
+		}
+	}
+
+	// Draw the name of the event.
+	if (EventW > 8.0f && EventH > 10.0f)
+	{
+		FString Name = EventName ? EventName : TEXT("?");
+		if (EventW > Name.Len() * 2.0f + 48.0f)
+		{
+			if (Event.ObjectInstanceIndex != 0)
+			{
+				Name += TEXT(" (NetId:");
+				Name += FText::AsNumber(NetId).ToString();
+				Name += TEXT(", ");
+			}
+			else
+			{
+				Name += TEXT(" (");
+			}
+			const uint32 EventSize = Event.EndPos - Event.StartPos;
+			if (EventSize == 1)
+			{
+				Name += TEXT("1 bit)");
+			}
+			else
+			{
+				Name += FText::AsNumber(EventSize).ToString();
+				Name += TEXT(" bits)");
+			}
+		}
+
+		const TSharedRef<FSlateFontMeasure> FontMeasureService = FSlateApplication::Get().GetRenderer()->GetFontMeasureService();
+		const int32 LastWholeCharacterIndex = FontMeasureService->FindLastWholeCharacterIndexBeforeOffset(Name, EventFont, FMath::RoundToInt(EventW - 2.0f));
+
+		if (LastWholeCharacterIndex >= 0)
+		{
+			// Grey threshold is shifted toward black (0.4 instead of 0.5 in test below) due to "area rule":
+			// a large gray surface (background of an event in this case) is perceived lighter than a smaller area (text pixels).
+			// Ref: https://books.google.ro/books?id=0pVr7dhmdWYC
+			const bool bIsDarkColor = (EventColorFill.ComputeLuminance() < 0.4f);
+
+			DrawState.Texts.AddDefaulted();
+			FPacketContentViewDrawState::FText& DrawText = DrawState.Texts.Last();
+			DrawText.X = EventX1 + 2.0f;
+			DrawText.Y = EventY + 1.0f;
+			DrawText.Text = Name.Left(LastWholeCharacterIndex + 1);
+			DrawText.bWhite = bIsDarkColor;
+		}
+	}
+}
+
+////////////////////////////////////////////////////////////////////////////////////////////////////
+
+void FPacketContentViewDrawStateBuilder::FlushBox(const FBoxData& Box, const float EventY, const float EventH)
+{
+	DrawState.Boxes.AddUninitialized();
+	FPacketContentViewDrawState::FBox& DrawBox = DrawState.Boxes.Last();
+	DrawBox.X = Box.X1;
+	DrawBox.W = Box.X2 - Box.X1;
+	DrawBox.Y = EventY;
+	DrawBox.H = EventH;
+	DrawBox.Color = Box.LinearColor;
+}
+
+////////////////////////////////////////////////////////////////////////////////////////////////////
+
+void FPacketContentViewDrawStateBuilder::Flush()
+{
+	constexpr float Y0 = 0.0f;
+	constexpr float EventH = 14.0f;
+	constexpr float EventDY = 2.0f;
+
+	// Flush merged boxes.
+	for (int32 Depth = 0; Depth <= MaxDepth; ++Depth)
+	{
+		const FBoxData& Box = LastBox[Depth];
+		if (Box.X1 < Box.X2)
+		{
+			const float EventY = Y0 + (EventH + EventDY) * Depth;
+			FlushBox(Box, EventY, EventH);
+		}
+	}
+}
+
+////////////////////////////////////////////////////////////////////////////////////////////////////
+// FPacketContentViewDrawHelper
+////////////////////////////////////////////////////////////////////////////////////////////////////
+
+FPacketContentViewDrawHelper::FPacketContentViewDrawHelper(const FDrawContext& InDrawContext, const FPacketContentViewport& InViewport)
+	: DrawContext(InDrawContext)
+	, Viewport(InViewport)
+	, WhiteBrush(FInsightsStyle::Get().GetBrush("WhiteBrush"))
+	, EventBorderBrush(FInsightsStyle::Get().GetBrush("EventBorder"))
+	, HoveredEventBorderBrush(FInsightsStyle::Get().GetBrush("HoveredEventBorder"))
+	, SelectedEventBorderBrush(FInsightsStyle::Get().GetBrush("SelectedEventBorder"))
+	, EventFont(FCoreStyle::GetDefaultFontStyle("Regular", 8))
+{
+}
+
+////////////////////////////////////////////////////////////////////////////////////////////////////
+
+void FPacketContentViewDrawHelper::DrawBackground() const
+{
+	const FAxisViewportDouble& ViewportX = Viewport.GetHorizontalAxisViewport();
+
+	const float X0 = 0.0f;
+	const float X1 = ViewportX.GetMinPos() - ViewportX.GetPos();
+	const float X2 = ViewportX.GetMaxPos() - ViewportX.GetPos();
+	const float X3 = FMath::CeilToFloat(Viewport.GetWidth());
+
+	const float Y = 0.0f;
+	const float H = FMath::CeilToFloat(Viewport.GetHeight());
+
+	FDrawHelpers::DrawBackground(DrawContext, WhiteBrush, X0, X1, X2, X3, Y, H);
+}
+
+////////////////////////////////////////////////////////////////////////////////////////////////////
+
+FLinearColor FPacketContentViewDrawHelper::GetColorByType(int32 Type)
+{
+	constexpr float Alpha = 1.0f;
+	switch (Type)
+	{
+	case 0:
+		return FLinearColor(0.75f, 0.25f, 0.25f, Alpha);
+
+	case 1:
+		return FLinearColor(0.25f, 0.75f, 0.25f, Alpha);
+
+	case 2:
+		return FLinearColor(0.25f, 0.25f, 0.75f, Alpha);
+
+	case 3:
+		return FLinearColor(0.75f, 0.75f, 0.25f, Alpha);
+
+	case 4:
+		return FLinearColor(0.25f, 0.75f, 0.75f, Alpha);
+
+	case 5:
+		return FLinearColor(0.75f, 0.25f, 0.75f, Alpha);
+
+	//default:
+	//	return FLinearColor(0.75f, 0.75f, 0.75f, Alpha);
+	}
+
+	uint32 Hash = (static_cast<uint32>(Type) + 0x2c2c57ed) * 0x2c2c57ed;
+	return FLinearColor(((Hash >> 16) & 0xFF) / 255.0f,
+						((Hash >>  8) & 0xFF) / 255.0f,
+						((Hash      ) & 0xFF) / 255.0f,
+						1.0f);
+}
+
+////////////////////////////////////////////////////////////////////////////////////////////////////
+
+void FPacketContentViewDrawHelper::Draw(const FPacketContentViewDrawState& DrawState) const
+{
+	// Draw filled boxes.
+	for (const FPacketContentViewDrawState::FBox& Box : DrawState.Boxes)
+	{
+		DrawContext.DrawBox(Box.X, Box.Y, Box.W, Box.H, WhiteBrush, Box.Color);
+	}
+	DrawContext.LayerId++;
+
+	// Draw borders.
+	for (const FPacketContentViewDrawState::FBox& Box : DrawState.Borders)
+	{
+		DrawContext.DrawBox(Box.X, Box.Y, Box.W, Box.H, EventBorderBrush, Box.Color);
+	}
+	DrawContext.LayerId++;
+
+	// Draw texts.
+	for (const FPacketContentViewDrawState::FText& Text : DrawState.Texts)
+	{
+		DrawContext.DrawText(Text.X, Text.Y, Text.Text, EventFont, Text.bWhite ? FLinearColor::White : FLinearColor::Black);
+	}
+	DrawContext.LayerId++;
+}
+
+////////////////////////////////////////////////////////////////////////////////////////////////////
+
+void FPacketContentViewDrawHelper::DrawEventHighlight(const FNetworkPacketEvent& Event, EHighlightMode Mode) const
+{
+	const FAxisViewportDouble& ViewportX = Viewport.GetHorizontalAxisViewport();
+
+	float EventX1 = ViewportX.GetRoundedOffsetForValue(static_cast<double>(Event.BitOffset));
+	if (EventX1 > Viewport.GetWidth())
+	{
+		return;
+	}
+
+	float EventX2 = ViewportX.GetRoundedOffsetForValue(static_cast<double>(Event.BitOffset + Event.BitSize));
+	if (EventX2 < 0)
+	{
+		return;
+	}
+
+	// Events are displayed with minimum 1px (including empty ones).
+	if (EventX1 == EventX2)
+	{
+		EventX2 = EventX1 + 1.0f;
+	}
+
+	// Limit event width on the viewport's left side.
+	const float MinX = -2.0f; // -2 allows event border to remain outside screen
+	if (EventX1 < MinX)
+	{
+		EventX1 = MinX;
+	}
+
+	// Limit event width on the viewport's right side.
+	const float MaxX = Viewport.GetWidth() + 2.0f; // +2 allows event border to remain outside screen
+	if (EventX2 > MaxX)
+	{
+		EventX2 = MaxX;
+	}
+
+	const float EventW = EventX2 - EventX1;
+
+	constexpr float Y0 = 0.0f;
+	constexpr float EventH = 14.0f;
+	constexpr float EventDY = 2.0f;
+
+	const float EventY = Y0 + (EventH + EventDY) * Event.Level;
+
+	if (Mode == EHighlightMode::Hovered)
+	{
+		const FLinearColor Color(1.0f, 1.0f, 0.0f, 1.0f); // yellow
+
+		// Draw border around the timing event box.
+		DrawContext.DrawBox(EventX1 - 2.0f, EventY - 2.0f, EventW + 4.0f, EventH + 4.0f, HoveredEventBorderBrush, Color);
+	}
+	else // EHighlightMode::Selected or EHighlightMode::SelectedAndHovered
+	{
+		// Animate color from white (if selected and hovered) or yellow (if only selected) to black, using a squared sine function.
+		const double Time = static_cast<double>(FPlatformTime::Cycles64()) * FPlatformTime::GetSecondsPerCycle64();
+		float S = FMath::Sin(2.0 * Time);
+		S = S * S; // squared, to ensure only positive [0 - 1] values
+		const float Blue = (Mode == EHighlightMode::SelectedAndHovered) ? 0.0f : S;
+		const FLinearColor Color(S, S, Blue, 1.0f);
+
+		// Draw border around the timing event box.
+		DrawContext.DrawBox(EventX1 - 2.0f, EventY - 2.0f, EventW + 4.0f, EventH + 4.0f, SelectedEventBorderBrush, Color);
+	}
+	DrawContext.LayerId++;
+}
+
+////////////////////////////////////////////////////////////////////////////////////////////////////
diff --git a/Engine/Source/Developer/TraceInsights/Private/Insights/NetworkingProfiler/ViewModels/PacketContentViewDrawHelper.h b/Engine/Source/Developer/TraceInsights/Private/Insights/NetworkingProfiler/ViewModels/PacketContentViewDrawHelper.h
new file mode 100644
--- /dev/null
+++ b/Engine/Source/Developer/TraceInsights/Private/Insights/NetworkingProfiler/ViewModels/PacketContentViewDrawHelper.h
@@ -0,0 +1,214 @@
+// Copyright 1998-2019 Epic Games, Inc. All Rights Reserved.
+
+#pragma once
+
+#include "CoreMinimal.h"
+#include "Fonts/SlateFontInfo.h"
+#include "TraceServices/Model/NetProfiler.h"
+#include "Styling/WidgetStyle.h"
+
+enum class ESlateDrawEffect : uint8;
+
+struct FDrawContext;
+struct FGeometry;
+struct FSlateBrush;
+
+class FPacketContentViewport;
+class FSlateWindowElementList;
+
+////////////////////////////////////////////////////////////////////////////////////////////////////
+
+struct FNetworkPacketEvent
+{
+	uint32 EventTypeIndex;
+	uint32 ObjectInstanceIndex;
+	uint32 NetId;
+	uint32 BitOffset;
+	uint32 BitSize;
+	uint32 Level;
+
+	FNetworkPacketEvent()
+		: EventTypeIndex(0)
+		, ObjectInstanceIndex(0)
+		, NetId(0)
+		, BitOffset(0)
+		, BitSize(0)
+		, Level(0)
+	{}
+
+	FNetworkPacketEvent(uint32 InEventTypeIndex, uint32 InObjectInstanceIndex, uint32 InNetId, uint32 InBitOffset, uint32 InBitSize, uint32 InLevel)
+		: EventTypeIndex(InEventTypeIndex)
+		, ObjectInstanceIndex(InObjectInstanceIndex)
+		, NetId(InNetId)
+		, BitOffset(InBitOffset)
+		, BitSize(InBitSize)
+		, Level(InLevel)
+	{}
+
+	FNetworkPacketEvent(const FNetworkPacketEvent& Other)
+		: EventTypeIndex(Other.EventTypeIndex)
+		, ObjectInstanceIndex(Other.ObjectInstanceIndex)
+		, NetId(Other.NetId)
+		, BitOffset(Other.BitOffset)
+		, BitSize(Other.BitSize)
+		, Level(Other.Level)
+	{
+	}
+
+	FNetworkPacketEvent& operator=(const FNetworkPacketEvent& Other)
+	{
+		EventTypeIndex = Other.EventTypeIndex;
+		ObjectInstanceIndex = Other.ObjectInstanceIndex;
+		NetId = Other.NetId;
+		BitOffset = Other.BitOffset;
+		BitSize = Other.BitSize;
+		Level = Other.Level;
+		return *this;
+	}
+
+	bool Equals(const FNetworkPacketEvent& Other) const
+	{
+		return EventTypeIndex == Other.EventTypeIndex
+			&& ObjectInstanceIndex == Other.ObjectInstanceIndex
+			&& NetId == Other.NetId
+			&& BitOffset == Other.BitOffset
+			&& BitSize == Other.BitSize
+			&& Level == Other.Level;
+	}
+
+	static bool AreEquals(const FNetworkPacketEvent& A, const FNetworkPacketEvent& B)
+	{
+		return A.Equals(B);
+	}
+};
+
+////////////////////////////////////////////////////////////////////////////////////////////////////
+
+struct FPacketContentViewDrawState
+{
+	struct FBox
+	{
+		float X;
+		float Y;
+		float W;
+		float H;
+		FLinearColor Color;
+	};
+
+	struct FText
+	{
+		float X;
+		float Y;
+		FString Text;
+		bool bWhite;
+	};
+
+	FPacketContentViewDrawState()
+	{
+	}
+
+	void Reset()
+	{
+		Events.Reset();
+		Boxes.Reset();
+		Borders.Reset();
+		Texts.Reset();
+		NumMergedBoxes = 0;
+	}
+
+	int32 GetNumMergedBoxes() const { return NumMergedBoxes; }
+
+	TArray<FNetworkPacketEvent> Events;
+	TArray<FBox> Boxes;
+	TArray<FBox> Borders;
+	TArray<FText> Texts;
+
+	// Debug stats.
+	int32 NumMergedBoxes;
+};
+
+////////////////////////////////////////////////////////////////////////////////////////////////////
+
+class FPacketContentViewDrawStateBuilder
+{
+private:
+	struct FBoxData
+	{
+		float X1;
+		float X2;
+		uint32 Color;
+		FLinearColor LinearColor;
+
+		FBoxData() : X1(0.0f), X2(0.0f), Color(0) {}
+		void Reset() { X1 = 0.0f; X2 = 0.0f; Color = 0; }
+	};
+
+public:
+	explicit FPacketContentViewDrawStateBuilder(FPacketContentViewDrawState& InState, const FPacketContentViewport& InViewport);
+
+	/**
+	 * Non-copyable
+	 */
+	FPacketContentViewDrawStateBuilder(const FPacketContentViewDrawStateBuilder&) = delete;
+	FPacketContentViewDrawStateBuilder& operator=(const FPacketContentViewDrawStateBuilder&) = delete;
+
+	void AddEvent(const Trace::FNetProfilerContentEvent& Event, const TCHAR* Name, uint32 NetId);
+	void Flush();
+
+private:
+	void FlushBox(const FBoxData& Box, const float EventY, const float EventH);
+
+private:
+	FPacketContentViewDrawState& DrawState; // cached draw state to build
+	const FPacketContentViewport& Viewport;
+
+	int32 MaxDepth;
+
+	TArray<float> LastEventX2; // X2 value for last event on each depth
+	TArray<FBoxData> LastBox;
+
+	const FSlateFontInfo EventFont;
+};
+
+////////////////////////////////////////////////////////////////////////////////////////////////////
+
+class FPacketContentViewDrawHelper
+{
+public:
+	enum class EHighlightMode : uint32
+	{
+		Hovered = 1,
+		Selected = 2,
+		SelectedAndHovered = 3
+	};
+
+public:
+	explicit FPacketContentViewDrawHelper(const FDrawContext& InDrawContext, const FPacketContentViewport& InViewport);
+
+	/**
+	 * Non-copyable
+	 */
+	FPacketContentViewDrawHelper(const FPacketContentViewDrawHelper&) = delete;
+	FPacketContentViewDrawHelper& operator=(const FPacketContentViewDrawHelper&) = delete;
+
+	const FSlateBrush* GetWhiteBrush() const { return WhiteBrush; }
+	const FSlateFontInfo& GetEventFont() const { return EventFont; }
+
+	void DrawBackground() const;
+	void Draw(const FPacketContentViewDrawState& DrawState) const;
+	void DrawEventHighlight(const FNetworkPacketEvent& Event, EHighlightMode Mode) const;
+
+	static FLinearColor GetColorByType(int32 Type);
+
+private:
+	const FDrawContext& DrawContext;
+	const FPacketContentViewport& Viewport;
+
+	const FSlateBrush* WhiteBrush;
+	const FSlateBrush* EventBorderBrush;
+	const FSlateBrush* HoveredEventBorderBrush;
+	const FSlateBrush* SelectedEventBorderBrush;
+	const FSlateFontInfo EventFont;
+};
+
+////////////////////////////////////////////////////////////////////////////////////////////////////
diff --git a/Engine/Source/Developer/TraceInsights/Private/Insights/NetworkingProfiler/ViewModels/PacketContentViewport.cpp b/Engine/Source/Developer/TraceInsights/Private/Insights/NetworkingProfiler/ViewModels/PacketContentViewport.cpp
new file mode 100644
--- /dev/null
+++ b/Engine/Source/Developer/TraceInsights/Private/Insights/NetworkingProfiler/ViewModels/PacketContentViewport.cpp
@@ -0,0 +1,5 @@
+// Copyright 1998-2019 Epic Games, Inc. All Rights Reserved.
+
+#include "PacketContentViewport.h"
+
+////////////////////////////////////////////////////////////////////////////////////////////////////
diff --git a/Engine/Source/Developer/TraceInsights/Private/Insights/NetworkingProfiler/ViewModels/PacketContentViewport.h b/Engine/Source/Developer/TraceInsights/Private/Insights/NetworkingProfiler/ViewModels/PacketContentViewport.h
new file mode 100644
--- /dev/null
+++ b/Engine/Source/Developer/TraceInsights/Private/Insights/NetworkingProfiler/ViewModels/PacketContentViewport.h
@@ -0,0 +1,54 @@
+// Copyright 1998-2019 Epic Games, Inc. All Rights Reserved.
+
+#pragma once
+
+#include "CoreMinimal.h"
+
+// Insights
+#include "Insights/ViewModels/AxisViewportDouble.h"
+
+////////////////////////////////////////////////////////////////////////////////////////////////////
+
+class FPacketContentViewport
+{
+public:
+	FPacketContentViewport()
+	{
+		Reset();
+	}
+
+	void Reset()
+	{
+		HorizontalAxisViewport.Reset();
+		Height = 0.0f;
+	}
+
+	const FAxisViewportDouble& GetHorizontalAxisViewport() const { return HorizontalAxisViewport; }
+	FAxisViewportDouble& GetHorizontalAxisViewport() { return HorizontalAxisViewport; }
+
+	float GetWidth() const { return HorizontalAxisViewport.GetSize(); }
+	float GetHeight() const { return Height; }
+
+	bool SetSize(const float InWidth, const float InHeight)
+	{
+		const bool bWidthChanged = HorizontalAxisViewport.SetSize(InWidth);
+		if (bWidthChanged || Height != InHeight)
+		{
+			Height = InHeight;
+			OnSizeChanged();
+			return true;
+		}
+		return false;
+	}
+
+private:
+	void OnSizeChanged()
+	{
+	}
+
+private:
+	FAxisViewportDouble HorizontalAxisViewport;
+	float Height;
+};
+
+////////////////////////////////////////////////////////////////////////////////////////////////////
diff --git a/Engine/Source/Developer/TraceInsights/Private/Insights/NetworkingProfiler/ViewModels/PacketSizesViewHelper.cpp b/Engine/Source/Developer/TraceInsights/Private/Insights/NetworkingProfiler/ViewModels/PacketSizesViewHelper.cpp
deleted file mode 100644
diff --git a/Engine/Source/Developer/TraceInsights/Private/Insights/NetworkingProfiler/ViewModels/PacketSizesViewHelper.h b/Engine/Source/Developer/TraceInsights/Private/Insights/NetworkingProfiler/ViewModels/PacketSizesViewHelper.h
deleted file mode 100644
diff --git a/Engine/Source/Developer/TraceInsights/Private/Insights/NetworkingProfiler/ViewModels/PacketSizesViewport.cpp b/Engine/Source/Developer/TraceInsights/Private/Insights/NetworkingProfiler/ViewModels/PacketSizesViewport.cpp
deleted file mode 100644
diff --git a/Engine/Source/Developer/TraceInsights/Private/Insights/NetworkingProfiler/ViewModels/PacketSizesViewport.h b/Engine/Source/Developer/TraceInsights/Private/Insights/NetworkingProfiler/ViewModels/PacketSizesViewport.h
deleted file mode 100644
diff --git a/Engine/Source/Developer/TraceInsights/Private/Insights/NetworkingProfiler/ViewModels/PacketViewDrawHelper.cpp b/Engine/Source/Developer/TraceInsights/Private/Insights/NetworkingProfiler/ViewModels/PacketViewDrawHelper.cpp
new file mode 100644
--- /dev/null
+++ b/Engine/Source/Developer/TraceInsights/Private/Insights/NetworkingProfiler/ViewModels/PacketViewDrawHelper.cpp
@@ -0,0 +1,247 @@
+// Copyright 1998-2019 Epic Games, Inc. All Rights Reserved.
+
+#include "PacketViewDrawHelper.h"
+
+#include "Framework/Application/SlateApplication.h"
+#include "Rendering/DrawElements.h"
+#include "TraceServices/AnalysisService.h"
+
+// Insights
+#include "Insights/Common/PaintUtils.h"
+#include "Insights/InsightsStyle.h"
+#include "Insights/NetworkingProfiler/ViewModels/PacketViewport.h"
+#include "Insights/ViewModels/DrawHelpers.h"
+
+#include <limits>
+
+////////////////////////////////////////////////////////////////////////////////////////////////////
+// FNetworkPacketSeriesBuilder
+////////////////////////////////////////////////////////////////////////////////////////////////////
+
+FNetworkPacketSeriesBuilder::FNetworkPacketSeriesBuilder(FNetworkPacketSeries& InSeries, const FPacketViewport& InViewport)
+	: Series(InSeries)
+	, Viewport(InViewport)
+	, NumAddedPackets(0)
+{
+	SampleW = Viewport.GetSampleWidth();
+	PacketsPerSample = Viewport.GetNumPacketsPerSample();
+	NumSamples = FMath::Max(0, FMath::CeilToInt(Viewport.GetWidth() / SampleW));
+	FirstPacketIndex = Viewport.GetFirstPacketIndex();
+
+	Series.NumAggregatedPackets = 0;
+	Series.Samples.Reset();
+	Series.Samples.AddDefaulted(NumSamples);
+}
+
+////////////////////////////////////////////////////////////////////////////////////////////////////
+
+void FNetworkPacketSeriesBuilder::AddPacket(int32 PacketIndex, const Trace::FNetProfilerPacket& Packet)
+{
+	NumAddedPackets++;
+
+	int32 SampleIndex = (PacketIndex - FirstPacketIndex) / PacketsPerSample;
+	if (SampleIndex >= 0 && SampleIndex < NumSamples)
+	{
+		FNetworkPacketAggregatedSample& Sample = Series.Samples[SampleIndex];
+		Sample.NumPackets++;
+
+		const double TimeStamp = static_cast<double>(Packet.TimeStamp);
+		if (TimeStamp < Sample.StartTime)
+		{
+			Sample.StartTime = TimeStamp;
+		}
+		if (TimeStamp > Sample.EndTime)
+		{
+			Sample.EndTime = TimeStamp;
+		}
+
+		if (Packet.TotalPacketSizeInBytes > Sample.LargestPacket.TotalSizeInBytes)
+		{
+			Sample.LargestPacket.Index = PacketIndex;
+			Sample.LargestPacket.SequenceNumber = Packet.SequenceNumber;
+			Sample.LargestPacket.ContentSizeInBits = Packet.ContentSizeInBits;
+			Sample.LargestPacket.TotalSizeInBytes = Packet.TotalPacketSizeInBytes;
+			Sample.LargestPacket.TimeStamp = TimeStamp;
+			Sample.LargestPacket.Status = Packet.DeliveryStatus;
+		}
+
+		switch (Sample.AggregatedStatus)
+		{
+			case Trace::ENetProfilerDeliveryStatus::Unknown:
+				Sample.AggregatedStatus = Packet.DeliveryStatus;
+				break;
+
+			case Trace::ENetProfilerDeliveryStatus::Dropped:
+				break;
+
+			case Trace::ENetProfilerDeliveryStatus::Delivered:
+				if (Packet.DeliveryStatus == Trace::ENetProfilerDeliveryStatus::Dropped)
+				{
+					Sample.AggregatedStatus = Trace::ENetProfilerDeliveryStatus::Dropped;
+				}
+				break;
+		}
+
+		Series.NumAggregatedPackets++;
+	}
+}
+
+////////////////////////////////////////////////////////////////////////////////////////////////////
+// FPacketViewDrawHelper
+////////////////////////////////////////////////////////////////////////////////////////////////////
+
+FPacketViewDrawHelper::FPacketViewDrawHelper(const FDrawContext& InDrawContext, const FPacketViewport& InViewport)
+	: DrawContext(InDrawContext)
+	, Viewport(InViewport)
+	, WhiteBrush(FInsightsStyle::Get().GetBrush("WhiteBrush"))
+	//, EventBorderBrush(FInsightsStyle::Get().GetBrush("EventBorder"))
+	, HoveredEventBorderBrush(FInsightsStyle::Get().GetBrush("HoveredEventBorder"))
+	, SelectedEventBorderBrush(FInsightsStyle::Get().GetBrush("SelectedEventBorder"))
+	, NumPackets(0)
+	, NumDrawSamples(0)
+{
+}
+
+////////////////////////////////////////////////////////////////////////////////////////////////////
+
+void FPacketViewDrawHelper::DrawBackground() const
+{
+	const FAxisViewportInt32& ViewportX = Viewport.GetHorizontalAxisViewport();
+
+	const float X0 = 0.0f;
+	const float X1 = ViewportX.GetMinPos() - ViewportX.GetPos();
+	const float X2 = ViewportX.GetMaxPos() - ViewportX.GetPos();
+	const float X3 = FMath::CeilToFloat(Viewport.GetWidth());
+
+	const float Y = 0.0f;
+	const float H = FMath::CeilToFloat(Viewport.GetHeight());
+
+	FDrawHelpers::DrawBackground(DrawContext, WhiteBrush, X0, X1, X2, X3, Y, H);
+}
+
+////////////////////////////////////////////////////////////////////////////////////////////////////
+
+FLinearColor FPacketViewDrawHelper::GetColorByStatus(Trace::ENetProfilerDeliveryStatus Status)
+{
+	constexpr float Alpha = 1.0f;
+	switch (Status)
+	{
+	case Trace::ENetProfilerDeliveryStatus::Unknown:
+		return FLinearColor(0.25f, 0.25f, 0.25f, Alpha);
+
+	case Trace::ENetProfilerDeliveryStatus::Delivered:
+		return FLinearColor(0.5f, 1.0f, 0.5f, Alpha);
+
+	case Trace::ENetProfilerDeliveryStatus::Dropped:
+		return FLinearColor(1.0f, 0.5f, 0.5f, Alpha);
+
+	default:
+		return FLinearColor(1.0f, 0.0f, 0.0f, Alpha);
+	}
+}
+
+////////////////////////////////////////////////////////////////////////////////////////////////////
+
+void FPacketViewDrawHelper::DrawCached(const FNetworkPacketSeries& Series) const
+{
+	if (Series.NumAggregatedPackets == 0)
+	{
+		return;
+	}
+
+	NumPackets += Series.NumAggregatedPackets;
+
+	const float SampleW = Viewport.GetSampleWidth();
+	const int32 NumSamples = Series.Samples.Num();
+
+	const FAxisViewportDouble& ViewportY = Viewport.GetVerticalAxisViewport();
+
+	const float ViewHeight = FMath::RoundToFloat(Viewport.GetHeight());
+	const float BaselineY = FMath::RoundToFloat(ViewportY.GetOffsetForValue(0.0));
+
+	for (int32 SampleIndex = 0; SampleIndex < NumSamples; SampleIndex++)
+	{
+		const FNetworkPacketAggregatedSample& Sample = Series.Samples[SampleIndex];
+		if (Sample.NumPackets == 0)
+		{
+			continue;
+		}
+
+		NumDrawSamples++;
+
+		const float X = SampleIndex * SampleW;
+
+		//const float ValueY = FMath::RoundToFloat(ViewportY.GetOffsetForValue(static_cast<double>(Sample.LargestPacket.ContentSizeInBits)));
+		const float ValueY = FMath::RoundToFloat(ViewportY.GetOffsetForValue(static_cast<double>(Sample.LargestPacket.TotalSizeInBytes * 8)));
+
+		const float H = ValueY - BaselineY;
+		const float Y = ViewHeight - H;
+
+		const FLinearColor ColorFill = GetColorByStatus(Sample.AggregatedStatus);
+		const FLinearColor ColorBorder(ColorFill.R * 0.75f, ColorFill.G * 0.75f, ColorFill.B * 0.75f, 1.0);
+
+		if (SampleW > 2.0f)
+		{
+			DrawContext.DrawBox(X + 1.0f, Y + 1.0f, SampleW - 2.0f, H - 2.0f, WhiteBrush, ColorFill);
+
+			// Draw border.
+			DrawContext.DrawBox(X, Y, 1.0, H, WhiteBrush, ColorBorder);
+			DrawContext.DrawBox(X + SampleW - 1.0f, Y, 1.0, H, WhiteBrush, ColorBorder);
+			DrawContext.DrawBox(X + 1.0f, Y, SampleW - 2.0f, 1.0f, WhiteBrush, ColorBorder);
+			DrawContext.DrawBox(X + 1.0f, Y + H - 1.0f, SampleW - 2.0f, 1.0f, WhiteBrush, ColorBorder);
+		}
+		else
+		{
+			DrawContext.DrawBox(X, Y, SampleW, H, WhiteBrush, ColorBorder);
+		}
+	}
+
+	DrawContext.LayerId++;
+}
+
+////////////////////////////////////////////////////////////////////////////////////////////////////
+
+void FPacketViewDrawHelper::DrawSampleHighlight(const FNetworkPacketAggregatedSample& Sample, EHighlightMode Mode) const
+{
+	const float SampleW = Viewport.GetSampleWidth();
+	const int32 PacketsPerSample = Viewport.GetNumPacketsPerSample();
+	const int32 FirstPacketIndex = Viewport.GetFirstPacketIndex();
+	const int32 SampleIndex = (Sample.LargestPacket.Index - FirstPacketIndex) / PacketsPerSample;
+	const float X = SampleIndex * SampleW;
+
+	const FAxisViewportDouble& ViewportY = Viewport.GetVerticalAxisViewport();
+
+	const float ViewHeight = FMath::RoundToFloat(Viewport.GetHeight());
+	const float BaselineY = FMath::RoundToFloat(ViewportY.GetOffsetForValue(0.0));
+
+	//const float ValueY = FMath::RoundToFloat(ViewportY.GetOffsetForValue(static_cast<double>(Sample.LargestPacket.ContentSizeInBits)));
+	const float ValueY = FMath::RoundToFloat(ViewportY.GetOffsetForValue(static_cast<double>(Sample.LargestPacket.TotalSizeInBytes * 8)));
+
+	const float H = ValueY - BaselineY;
+	const float Y = ViewHeight - H;
+
+	if (Mode == EHighlightMode::Hovered)
+	{
+		const FLinearColor Color(1.0f, 1.0f, 0.0f, 1.0f); // yellow
+
+		// Draw border around the timing event box.
+		DrawContext.DrawBox(X - 1.0f, Y - 1.0f, SampleW + 2.0f, H + 2.0f, HoveredEventBorderBrush, Color);
+	}
+	else // EHighlightMode::Selected or EHighlightMode::SelectedAndHovered
+	{
+		// Animate color from white (if selected and hovered) or yellow (if only selected) to black, using a squared sine function.
+		const double Time = static_cast<double>(FPlatformTime::Cycles64()) * FPlatformTime::GetSecondsPerCycle64();
+		float S = FMath::Sin(2.0 * Time);
+		S = S * S; // squared, to ensure only positive [0 - 1] values
+		const float Blue = (Mode == EHighlightMode::SelectedAndHovered) ? 0.0f : S;
+		const FLinearColor Color(S, S, Blue, 1.0f);
+
+		// Draw border around the timing event box.
+		DrawContext.DrawBox(X - 1.0f, Y - 1.0f, SampleW + 2.0f, H + 2.0f, SelectedEventBorderBrush, Color);
+	}
+	DrawContext.LayerId++;
+}
+
+
+
+////////////////////////////////////////////////////////////////////////////////////////////////////
diff --git a/Engine/Source/Developer/TraceInsights/Private/Insights/NetworkingProfiler/ViewModels/PacketViewDrawHelper.h b/Engine/Source/Developer/TraceInsights/Private/Insights/NetworkingProfiler/ViewModels/PacketViewDrawHelper.h
new file mode 100644
--- /dev/null
+++ b/Engine/Source/Developer/TraceInsights/Private/Insights/NetworkingProfiler/ViewModels/PacketViewDrawHelper.h
@@ -0,0 +1,205 @@
+// Copyright 1998-2019 Epic Games, Inc. All Rights Reserved.
+
+#pragma once
+
+#include "CoreMinimal.h"
+#include "Styling/WidgetStyle.h"
+#include "TraceServices/Model/NetProfiler.h"
+
+#include <limits>
+
+enum class ESlateDrawEffect : uint8;
+
+struct FDrawContext;
+struct FGeometry;
+struct FSlateBrush;
+
+class FPacketViewport;
+class FSlateWindowElementList;
+
+////////////////////////////////////////////////////////////////////////////////////////////////////
+
+struct FNetworkPacket
+{
+	int32 Index;
+	uint32 SequenceNumber;
+	uint32 ContentSizeInBits;
+	uint32 TotalSizeInBytes;
+	double TimeStamp;
+	Trace::ENetProfilerDeliveryStatus Status;
+
+	FNetworkPacket()
+		: Index(-1)
+		, SequenceNumber(0)
+		, ContentSizeInBits(0)
+		, TotalSizeInBytes(0)
+		, TimeStamp(std::numeric_limits<double>::infinity())
+		, Status(Trace::ENetProfilerDeliveryStatus::Unknown)
+	{}
+
+	FNetworkPacket(const FNetworkPacket&) = default;
+	FNetworkPacket& operator=(const FNetworkPacket&) = default;
+
+	FNetworkPacket(FNetworkPacket&&) = default;
+	FNetworkPacket& operator=(FNetworkPacket&&) = default;
+
+	bool Equals(const FNetworkPacket& Other) const
+	{
+		return Index == Other.Index
+			&& SequenceNumber == Other.SequenceNumber
+			//&& ContentSizeInBits == Other.ContentSizeInBits
+			&& TotalSizeInBytes == Other.TotalSizeInBytes
+			/*&& TimeStamp == Other.TimeStamp
+			&& Status == Other.Status*/;
+	}
+
+	static bool AreEquals(const FNetworkPacket& A, const FNetworkPacket& B)
+	{
+		return A.Equals(B);
+	}
+};
+
+////////////////////////////////////////////////////////////////////////////////////////////////////
+
+struct FNetworkPacketAggregatedSample
+{
+	int32 NumPackets;
+
+	double StartTime; // min TimeSent of all packets in this aggregated sample
+	double EndTime; // max TimeSent of all packets in this aggregated sample
+
+	/** Aggregated status for packets in this aggregated sample.
+	 *    Unknown  --> all aggregated packets are unknown
+	 *    Sent     --> all aggregated packets are sent; none is received or confirmed lost
+	 *    Received --> at least one packet in the sample set is confirmed received and none are confirmed lost
+	 *    Lost     --> at least one packet in the sample set is confirmed lost
+	**/
+	Trace::ENetProfilerDeliveryStatus AggregatedStatus;
+
+	FNetworkPacket LargestPacket;
+
+	FNetworkPacketAggregatedSample()
+		: NumPackets(0)
+		, StartTime(DBL_MAX)
+		, EndTime(-DBL_MAX)
+		, AggregatedStatus(Trace::ENetProfilerDeliveryStatus::Unknown)
+		, LargestPacket()
+	{}
+
+	FNetworkPacketAggregatedSample(const FNetworkPacketAggregatedSample&) = default;
+	FNetworkPacketAggregatedSample& operator=(const FNetworkPacketAggregatedSample&) = default;
+
+	FNetworkPacketAggregatedSample(FNetworkPacketAggregatedSample&&) = default;
+	FNetworkPacketAggregatedSample& operator=(FNetworkPacketAggregatedSample&&) = default;
+
+	bool Equals(const FNetworkPacketAggregatedSample& Other) const
+	{
+		return NumPackets == Other.NumPackets
+			&& StartTime == Other.StartTime
+			&& EndTime == Other.EndTime
+			&& AggregatedStatus == Other.AggregatedStatus
+			&& LargestPacket.Equals(Other.LargestPacket);
+	}
+
+	static bool AreEquals(const FNetworkPacketAggregatedSample& A, const FNetworkPacketAggregatedSample& B)
+	{
+		return A.Equals(B);
+	}
+};
+
+////////////////////////////////////////////////////////////////////////////////////////////////////
+
+struct FNetworkPacketSeries
+{
+	int32 NumAggregatedPackets; // total number of packets aggregated in samples; i.e. sum of all Sample.NumPackets
+
+	TArray<FNetworkPacketAggregatedSample> Samples;
+
+	FNetworkPacketSeries()
+		: NumAggregatedPackets(0)
+		, Samples()
+	{
+	}
+
+	void Reset()
+	{
+		NumAggregatedPackets = 0;
+		Samples.Reset();
+	}
+};
+
+////////////////////////////////////////////////////////////////////////////////////////////////////
+
+class FNetworkPacketSeriesBuilder
+{
+public:
+	explicit FNetworkPacketSeriesBuilder(FNetworkPacketSeries& InSeries, const FPacketViewport& InViewport);
+
+	/**
+	 * Non-copyable
+	 */
+	FNetworkPacketSeriesBuilder(const FNetworkPacketSeriesBuilder&) = delete;
+	FNetworkPacketSeriesBuilder& operator=(const FNetworkPacketSeriesBuilder&) = delete;
+
+	void AddPacket(int32 PacketIndex, const Trace::FNetProfilerPacket& Packet);
+
+	int32 GetNumAddedPackets() const { return NumAddedPackets; }
+
+private:
+	FNetworkPacketSeries& Series; // series to update
+	const FPacketViewport& Viewport;
+
+	float SampleW; // width of a sample, in Slate units
+	int32 PacketsPerSample; // number of packets in a sample
+	int32 FirstPacketIndex; // index of the first packet in the first sample; can be negative
+	int32 NumSamples; // total number of samples
+
+	// Debug stats.
+	int32 NumAddedPackets; // counts total number of added packets; i.e. number of AddPacket() calls
+};
+
+////////////////////////////////////////////////////////////////////////////////////////////////////
+
+class FPacketViewDrawHelper
+{
+public:
+	enum class EHighlightMode : uint32
+	{
+		Hovered = 1,
+		Selected = 2,
+		SelectedAndHovered = 3
+	};
+
+public:
+	explicit FPacketViewDrawHelper(const FDrawContext& InDrawContext, const FPacketViewport& InViewport);
+
+	/**
+	 * Non-copyable
+	 */
+	FPacketViewDrawHelper(const FPacketViewDrawHelper&) = delete;
+	FPacketViewDrawHelper& operator=(const FPacketViewDrawHelper&) = delete;
+
+	void DrawBackground() const;
+	void DrawCached(const FNetworkPacketSeries& Series) const;
+	void DrawSampleHighlight(const FNetworkPacketAggregatedSample& Sample, EHighlightMode Mode) const;
+
+	static FLinearColor GetColorByStatus(Trace::ENetProfilerDeliveryStatus Status);
+
+	int32 GetNumPackets() const { return NumPackets; }
+	int32 GetNumDrawSamples() const { return NumDrawSamples; }
+
+private:
+	const FDrawContext& DrawContext;
+	const FPacketViewport& Viewport;
+
+	const FSlateBrush* WhiteBrush;
+	//const FSlateBrush* EventBorderBrush;
+	const FSlateBrush* HoveredEventBorderBrush;
+	const FSlateBrush* SelectedEventBorderBrush;
+
+	// Debug stats.
+	mutable int32 NumPackets;
+	mutable int32 NumDrawSamples;
+};
+
+////////////////////////////////////////////////////////////////////////////////////////////////////
diff --git a/Engine/Source/Developer/TraceInsights/Private/Insights/NetworkingProfiler/ViewModels/PacketViewport.cpp b/Engine/Source/Developer/TraceInsights/Private/Insights/NetworkingProfiler/ViewModels/PacketViewport.cpp
new file mode 100644
--- /dev/null
+++ b/Engine/Source/Developer/TraceInsights/Private/Insights/NetworkingProfiler/ViewModels/PacketViewport.cpp
@@ -0,0 +1,5 @@
+// Copyright 1998-2019 Epic Games, Inc. All Rights Reserved.
+
+#include "PacketViewport.h"
+
+////////////////////////////////////////////////////////////////////////////////////////////////////
diff --git a/Engine/Source/Developer/TraceInsights/Private/Insights/NetworkingProfiler/ViewModels/PacketViewport.h b/Engine/Source/Developer/TraceInsights/Private/Insights/NetworkingProfiler/ViewModels/PacketViewport.h
new file mode 100644
--- /dev/null
+++ b/Engine/Source/Developer/TraceInsights/Private/Insights/NetworkingProfiler/ViewModels/PacketViewport.h
@@ -0,0 +1,65 @@
+// Copyright 1998-2019 Epic Games, Inc. All Rights Reserved.
+
+#pragma once
+
+#include "CoreMinimal.h"
+
+// Insights
+#include "Insights/ViewModels/AxisViewportInt32.h"
+#include "Insights/ViewModels/AxisViewportDouble.h"
+
+////////////////////////////////////////////////////////////////////////////////////////////////////
+
+class FPacketViewport
+{
+private:
+	static constexpr float SLATE_UNITS_TOLERANCE = 0.1f;
+
+public:
+	FPacketViewport()
+	{
+		Reset();
+	}
+
+	void Reset()
+	{
+		HorizontalAxisViewport.Reset();
+		VerticalAxisViewport.Reset();
+	}
+
+	const FAxisViewportInt32& GetHorizontalAxisViewport() const { return HorizontalAxisViewport; }
+	FAxisViewportInt32& GetHorizontalAxisViewport() { return HorizontalAxisViewport; }
+
+	const FAxisViewportDouble& GetVerticalAxisViewport() const { return VerticalAxisViewport; }
+	FAxisViewportDouble& GetVerticalAxisViewport() { return VerticalAxisViewport; }
+
+	float GetWidth() const { return HorizontalAxisViewport.GetSize(); }
+	float GetHeight() const { return VerticalAxisViewport.GetSize(); }
+
+	bool SetSize(const float InWidth, const float InHeight)
+	{
+		const bool bWidthChanged = HorizontalAxisViewport.SetSize(InWidth);
+		const bool bHeightChanged = VerticalAxisViewport.SetSize(InHeight);
+		if (bWidthChanged || bHeightChanged)
+		{
+			OnSizeChanged();
+			return true;
+		}
+		return false;
+	}
+
+	float GetSampleWidth() const { return HorizontalAxisViewport.GetSampleSize(); }
+	int32 GetNumPacketsPerSample() const { return HorizontalAxisViewport.GetNumSamplesPerPixel(); }
+	int32 GetFirstPacketIndex() const { return HorizontalAxisViewport.GetValueAtOffset(0.0f); }
+
+private:
+	void OnSizeChanged()
+	{
+	}
+
+private:
+	FAxisViewportInt32 HorizontalAxisViewport;
+	FAxisViewportDouble VerticalAxisViewport;
+};
+
+////////////////////////////////////////////////////////////////////////////////////////////////////
diff --git a/Engine/Source/Developer/TraceInsights/Private/Insights/NetworkingProfiler/Widgets/SNetworkingProfilerWindow.cpp b/Engine/Source/Developer/TraceInsights/Private/Insights/NetworkingProfiler/Widgets/SNetworkingProfilerWindow.cpp
--- a/Engine/Source/Developer/TraceInsights/Private/Insights/NetworkingProfiler/Widgets/SNetworkingProfilerWindow.cpp
+++ b/Engine/Source/Developer/TraceInsights/Private/Insights/NetworkingProfiler/Widgets/SNetworkingProfilerWindow.cpp
@@ -28,10 +28,8 @@
 #include "Insights/NetworkingProfiler/NetworkingProfilerManager.h"
 #include "Insights/NetworkingProfiler/Widgets/SNetStatsView.h"
 #include "Insights/NetworkingProfiler/Widgets/SNetworkingProfilerToolbar.h"
-#include "Insights/NetworkingProfiler/Widgets/SPacketBreakdownView.h"
-//#include "Insights/NetworkingProfiler/Widgets/SPacketContentView.h"
-#include "Insights/NetworkingProfiler/Widgets/SPacketSizesView.h"
-//#include "Insights/NetworkingProfiler/Widgets/SPacketView.h"
+#include "Insights/NetworkingProfiler/Widgets/SPacketContentView.h"
+#include "Insights/NetworkingProfiler/Widgets/SPacketView.h"
 #include "Insights/Version.h"
 #include "Insights/Widgets/SInsightsSettings.h"
 
diff --git a/Engine/Source/Developer/TraceInsights/Private/Insights/NetworkingProfiler/Widgets/SPacketBreakdownView.cpp b/Engine/Source/Developer/TraceInsights/Private/Insights/NetworkingProfiler/Widgets/SPacketBreakdownView.cpp
deleted file mode 100644
diff --git a/Engine/Source/Developer/TraceInsights/Private/Insights/NetworkingProfiler/Widgets/SPacketBreakdownView.h b/Engine/Source/Developer/TraceInsights/Private/Insights/NetworkingProfiler/Widgets/SPacketBreakdownView.h
deleted file mode 100644
diff --git a/Engine/Source/Developer/TraceInsights/Private/Insights/NetworkingProfiler/Widgets/SPacketContentView.cpp b/Engine/Source/Developer/TraceInsights/Private/Insights/NetworkingProfiler/Widgets/SPacketContentView.cpp
new file mode 100644
--- /dev/null
+++ b/Engine/Source/Developer/TraceInsights/Private/Insights/NetworkingProfiler/Widgets/SPacketContentView.cpp
@@ -0,0 +1,784 @@
+// Copyright 1998-2019 Epic Games, Inc. All Rights Reserved.
+
+#include "SPacketContentView.h"
+
+#include "Fonts/FontMeasure.h"
+#include "Fonts/SlateFontInfo.h"
+#include "Framework/Application/SlateApplication.h"
+#include "Framework/MultiBox/MultiBoxBuilder.h"
+#include "Rendering/DrawElements.h"
+#include "Styling/CoreStyle.h"
+#include "Widgets/Layout/SScrollBar.h"
+
+// Insights
+#include "Insights/Common/PaintUtils.h"
+#include "Insights/Common/TimeUtils.h"
+#include "Insights/Common/Stopwatch.h"
+#include "Insights/InsightsManager.h"
+#include "Insights/InsightsStyle.h"
+#include "Insights/NetworkingProfiler/Widgets/SNetworkingProfilerWindow.h"
+
+////////////////////////////////////////////////////////////////////////////////////////////////////
+
+#define LOCTEXT_NAMESPACE "SPacketContentView"
+
+////////////////////////////////////////////////////////////////////////////////////////////////////
+
+SPacketContentView::SPacketContentView()
+	: ProfilerWindow()
+	, DrawState(MakeShareable(new FPacketContentViewDrawState()))
+{
+	Reset();
+}
+
+////////////////////////////////////////////////////////////////////////////////////////////////////
+
+SPacketContentView::~SPacketContentView()
+{
+}
+
+////////////////////////////////////////////////////////////////////////////////////////////////////
+
+void SPacketContentView::Reset()
+{
+	//ProfilerWindow
+
+	Viewport.Reset();
+	//FAxisViewportDouble& ViewportX = Viewport.GetHorizontalAxisViewport();
+	//ViewportX.SetScaleLimits(0.000001, 100000.0);
+	//ViewportX.SetScale(1.0);
+	bIsViewportDirty = true;
+
+	GameInstanceIndex = 0;
+	ConnectionIndex = 0;
+	ConnectionMode = Trace::ENetProfilerConnectionMode::Outgoing;
+	PacketIndex = 0;
+	PacketBitSize = 0;
+
+	DrawState->Reset();
+	bIsStateDirty = true;
+
+	MousePosition = FVector2D::ZeroVector;
+
+	MousePositionOnButtonDown = FVector2D::ZeroVector;
+	ViewportPosXOnButtonDown = 0.0f;
+
+	MousePositionOnButtonUp = FVector2D::ZeroVector;
+
+	bIsLMB_Pressed = false;
+	bIsRMB_Pressed = false;
+
+	bIsScrolling = false;
+
+	HoveredEvent.Reset();
+	SelectedEvent.Reset();
+	Tooltip.Reset();
+
+	//ThisGeometry
+
+	CursorType = ECursorType::Default;
+
+	UpdateDurationHistory.Reset();
+	DrawDurationHistory.Reset();
+	OnPaintDurationHistory.Reset();
+	LastOnPaintTime = FPlatformTime::Cycles64();
+}
+
+////////////////////////////////////////////////////////////////////////////////////////////////////
+
+void SPacketContentView::Construct(const FArguments& InArgs, TSharedPtr<SNetworkingProfilerWindow> InProfilerWindow)
+{
+	ProfilerWindow = InProfilerWindow;
+
+	ChildSlot
+	[
+		SNew(SOverlay)
+		.Visibility(EVisibility::SelfHitTestInvisible)
+
+		+ SOverlay::Slot()
+		.VAlign(VAlign_Bottom)
+		.Padding(FMargin(0, 0, 0, 0))
+		[
+			SAssignNew(HorizontalScrollBar, SScrollBar)
+			.Orientation(Orient_Horizontal)
+			.AlwaysShowScrollbar(false)
+			.Visibility(EVisibility::Visible)
+			.Thickness(FVector2D(5.0f, 5.0f))
+			.RenderOpacity(0.75)
+			.OnUserScrolled(this, &SPacketContentView::HorizontalScrollBar_OnUserScrolled)
+		]
+	];
+
+	UpdateHorizontalScrollBar();
+
+	BindCommands();
+}
+
+////////////////////////////////////////////////////////////////////////////////////////////////////
+
+void SPacketContentView::Tick(const FGeometry& AllottedGeometry, const double InCurrentTime, const float InDeltaTime)
+{
+	if (ThisGeometry != AllottedGeometry || bIsViewportDirty)
+	{
+		bIsViewportDirty = false;
+		const float ViewWidth = AllottedGeometry.GetLocalSize().X;
+		const float ViewHeight = AllottedGeometry.GetLocalSize().Y;
+		Viewport.SetSize(ViewWidth, ViewHeight);
+		bIsStateDirty = true;
+	}
+
+	ThisGeometry = AllottedGeometry;
+
+	FAxisViewportDouble& ViewportX = Viewport.GetHorizontalAxisViewport();
+
+	if (!bIsScrolling)
+	{
+		// Elastic snap to horizontal limits.
+		if (ViewportX.UpdatePosWithinLimits())
+		{
+			bIsStateDirty = true;
+		}
+	}
+
+	if (bIsStateDirty)
+	{
+		bIsStateDirty = false;
+		UpdateState();
+	}
+
+	Tooltip.Update();
+	if (!MousePosition.IsZero())
+	{
+		Tooltip.SetPosition(MousePosition, 0.0f, Viewport.GetWidth(), 0.0f, Viewport.GetHeight() - 12.0f);
+	}
+}
+
+////////////////////////////////////////////////////////////////////////////////////////////////////
+
+void SPacketContentView::ResetPacket()
+{
+	GameInstanceIndex = 0;
+	ConnectionIndex = 0;
+	ConnectionMode = Trace::ENetProfilerConnectionMode::Outgoing;
+	PacketIndex = 0;
+	PacketBitSize = 0;
+
+	FAxisViewportDouble& ViewportX = Viewport.GetHorizontalAxisViewport();
+	ViewportX.SetMinMaxValueInterval(0.0, 0.0);
+	ViewportX.CenterOnValue(0.0);
+	UpdateHorizontalScrollBar();
+
+	DrawState->Reset();
+	bIsStateDirty = true;
+
+	HoveredEvent.Reset();
+	SelectedEvent.Reset();
+}
+
+////////////////////////////////////////////////////////////////////////////////////////////////////
+
+void SPacketContentView::SetPacket(uint32 InGameInstanceIndex, uint32 InConnectionIndex, Trace::ENetProfilerConnectionMode InConnectionMode, uint32 InPacketIndex, int64 InPacketBitSize)
+{
+	GameInstanceIndex = InGameInstanceIndex;
+	ConnectionIndex = InConnectionIndex;
+	ConnectionMode = InConnectionMode;
+	PacketIndex = InPacketIndex;
+	PacketBitSize = InPacketBitSize;
+
+	FAxisViewportDouble& ViewportX = Viewport.GetHorizontalAxisViewport();
+	ViewportX.SetMinMaxValueInterval(0.0, static_cast<double>(InPacketBitSize));
+	ViewportX.CenterOnValueInterval(0.0, static_cast<double>(InPacketBitSize));
+	UpdateHorizontalScrollBar();
+
+	DrawState->Reset();
+	bIsStateDirty = true;
+
+	HoveredEvent.Reset();
+	SelectedEvent.Reset();
+}
+
+////////////////////////////////////////////////////////////////////////////////////////////////////
+
+void SPacketContentView::UpdateState()
+{
+	FStopwatch Stopwatch;
+	Stopwatch.Start();
+
+	if (PacketBitSize > 0)
+	{
+		FPacketContentViewDrawStateBuilder Builder(*DrawState, Viewport);
+
+		TSharedPtr<const Trace::IAnalysisSession> Session = FInsightsManager::Get()->GetSession();
+		if (Session.IsValid())
+		{
+			Trace::FAnalysisSessionReadScope SessionReadScope(*Session.Get());
+			const Trace::INetProfilerProvider& NetProfilerProvider = Trace::ReadNetProfilerProvider(*Session.Get());
+
+			//uint32 StartEventIndex;
+			//uint32 EndEventIndex;
+			//NetProfilerProvider.ReadConnection(ConnectionIndex, [&StartEventIndex, &EndEventIndex](const Trace::FNetProfilerConnection& Connection)
+			//{
+			//	StartEventIndex = Connection.StartEventIndex;
+			//	EndEventIndex = Connection.EndEventIndex;
+			//});
+			//NetProfilerProvider.EnumeratePacketContentEventsByIndex(ConnectionIndex, ConnectionMode, StartEventIndex, uint32 EndEventIndex, [](const Trace::FNetProfilerContentEvent& Event)
+			//{
+			//});
+
+			const FAxisViewportDouble& ViewportX = Viewport.GetHorizontalAxisViewport();
+
+			const int64 MinPos = static_cast<int64>(FMath::FloorToDouble(ViewportX.GetValueAtOffset(0.0f)));
+			const int64 MaxPos = static_cast<int64>(FMath::CeilToDouble(ViewportX.GetValueAtOffset(ViewportX.GetSize())));
+
+			const uint32 StartPos = 0;
+			const uint32 EndPos = PacketBitSize;
+			NetProfilerProvider.EnumeratePacketContentEventsByPosition(ConnectionIndex, ConnectionMode, PacketIndex, StartPos, EndPos, [this, &Builder, &NetProfilerProvider](const Trace::FNetProfilerContentEvent& Event)
+			{
+				const TCHAR* Name = nullptr;
+				NetProfilerProvider.ReadName(Event.NameIndex, [&Name](const Trace::FNetProfilerName& NetProfilerName)
+				{
+					Name = NetProfilerName.Name;
+				});
+
+				uint32 NetId = 0;
+				if (Event.ObjectInstanceIndex != 0)
+				{
+					NetProfilerProvider.ReadObject(GameInstanceIndex, Event.ObjectInstanceIndex, [&NetId](const Trace::FNetProfilerObjectInstance& ObjectInstance)
+					{
+						NetId = ObjectInstance.NetId;
+					});
+				}
+
+				Builder.AddEvent(Event, Name, NetId);
+			});
+		}
+
+		Builder.Flush();
+	}
+
+	Stopwatch.Stop();
+	UpdateDurationHistory.AddValue(Stopwatch.AccumulatedTime);
+}
+
+////////////////////////////////////////////////////////////////////////////////////////////////////
+
+void SPacketContentView::UpdateHoveredEvent()
+{
+	HoveredEvent = GetEventAtMousePosition(MousePosition.X, MousePosition.Y);
+	//if (!HoveredEvent.IsValid())
+	//{
+	//	HoveredEvent = GetEventAtMousePosition(MousePosition.X - 1.0f, MousePosition.Y);
+	//}
+	//if (!HoveredEvent.IsValid())
+	//{
+	//	HoveredEvent = GetEventAtMousePosition(MousePosition.X + 1.0f, MousePosition.Y);
+	//}
+
+	if (HoveredEvent.IsValid())
+	{
+		// Init the tooltip's content.
+		Tooltip.ResetContent();
+
+		const FNetworkPacketEvent& Event = HoveredEvent.Event;
+		FString Name(TEXT("?"));
+		Trace::FNetProfilerEventType EventType;
+		Trace::FNetProfilerObjectInstance ObjectInstance;
+
+		TSharedPtr<const Trace::IAnalysisSession> Session = FInsightsManager::Get()->GetSession();
+		if (Session.IsValid())
+		{
+			Trace::FAnalysisSessionReadScope SessionReadScope(*Session.Get());
+			const Trace::INetProfilerProvider& NetProfilerProvider = Trace::ReadNetProfilerProvider(*Session.Get());
+
+			NetProfilerProvider.ReadEventType(Event.EventTypeIndex, [&EventType](const Trace::FNetProfilerEventType& InEventType)
+			{
+				EventType = InEventType;
+			});
+
+			NetProfilerProvider.ReadName(EventType.NameIndex, [&Name](const Trace::FNetProfilerName& NetProfilerName)
+			{
+				Name = NetProfilerName.Name;
+			});
+
+			if (Event.ObjectInstanceIndex != 0)
+			{
+				NetProfilerProvider.ReadObject(GameInstanceIndex, Event.ObjectInstanceIndex, [&ObjectInstance](const Trace::FNetProfilerObjectInstance& InObjectInstance)
+				{
+					ObjectInstance = InObjectInstance;
+				});
+			}
+		}
+
+		Tooltip.AddTitle(Name);
+
+		if (Event.ObjectInstanceIndex != 0)
+		{
+			Tooltip.AddNameValueTextLine(TEXT("Net Id:"), FText::AsNumber(ObjectInstance.NetId).ToString());
+			Tooltip.AddNameValueTextLine(TEXT("Type Id:"), FString::Printf(TEXT("0x%016X"), ObjectInstance.TypeId));
+			Tooltip.AddNameValueTextLine(TEXT("Obj. LifeTime:"), FString::Format(TEXT("from {0} to {1}"),
+				{ TimeUtils::FormatTimeAuto(ObjectInstance.LifeTime.Begin), TimeUtils::FormatTimeAuto(ObjectInstance.LifeTime.End) }));
+		}
+
+		Tooltip.AddNameValueTextLine(TEXT("Offset:"), FString::Format(TEXT("bit {0}"), { FText::AsNumber(Event.BitOffset).ToString() }));
+		if (Event.BitSize == 1)
+		{
+			Tooltip.AddNameValueTextLine(TEXT("Size:"), TEXT("1 bit"));
+		}
+		else
+		{
+			Tooltip.AddNameValueTextLine(TEXT("Size:"), FString::Format(TEXT("{0} bits"), { FText::AsNumber(Event.BitSize).ToString() }));
+		}
+
+		Tooltip.AddNameValueTextLine(TEXT("Level:"), FText::AsNumber(Event.Level).ToString());
+
+		Tooltip.UpdateLayout();
+
+		Tooltip.SetDesiredOpacity(1.0f);
+	}
+	else
+	{
+		Tooltip.SetDesiredOpacity(0.0f);
+	}
+}
+
+////////////////////////////////////////////////////////////////////////////////////////////////////
+
+void SPacketContentView::OnSelectedEventChanged()
+{
+	if (SelectedEvent.IsValid() && ProfilerWindow.IsValid())
+	{
+		// Select the node coresponding to net event type of selected net event instance.
+		const uint64 EventTypeId = static_cast<uint64>(SelectedEvent.Event.EventTypeIndex);
+		ProfilerWindow->SetSelectedEventType(EventTypeId);
+	}
+}
+
+////////////////////////////////////////////////////////////////////////////////////////////////////
+
+void SPacketContentView::SelectHoveredEvent()
+{
+	SelectedEvent = HoveredEvent;
+	//if (SelectedEvent.IsValid())
+	//{
+	//	LastSelectionType = ESelectionType::TimingEvent;
+	//	BringIntoView(SelectedEvent.StartPos, SelectedEvent.EndPos);
+	//}
+	OnSelectedEventChanged();
+}
+
+////////////////////////////////////////////////////////////////////////////////////////////////////
+
+FNetworkPacketEventRef SPacketContentView::GetEventAtMousePosition(float X, float Y)
+{
+	if (!bIsStateDirty)
+	{
+		for (const FNetworkPacketEvent& Event : DrawState->Events)
+		{
+			const FAxisViewportDouble& ViewportX = Viewport.GetHorizontalAxisViewport();
+
+			const float EventX1 = ViewportX.GetRoundedOffsetForValue(static_cast<double>(Event.BitOffset));
+			const float EventX2 = ViewportX.GetRoundedOffsetForValue(static_cast<double>(Event.BitOffset + Event.BitSize));
+
+			constexpr float Y0 = 0.0f;
+			constexpr float EventH = 14.0f;
+			constexpr float EventDY = 2.0f;
+			const float EventY = Y0 + (EventH + EventDY) * Event.Level;
+
+			constexpr float ToleranceX = 1.0f;
+
+			if (X >= EventX1 - ToleranceX && X <= EventX2 &&
+				Y >= EventY - EventDY / 2 && Y < EventY + EventH + EventDY / 2)
+			{
+				return FNetworkPacketEventRef(Event);
+			}
+		}
+	}
+	return FNetworkPacketEventRef();
+}
+
+////////////////////////////////////////////////////////////////////////////////////////////////////
+
+int32 SPacketContentView::OnPaint(const FPaintArgs& Args, const FGeometry& AllottedGeometry, const FSlateRect& MyCullingRect, FSlateWindowElementList& OutDrawElements, int32 LayerId, const FWidgetStyle& InWidgetStyle, bool bParentEnabled) const
+{
+	const bool bEnabled = ShouldBeEnabled(bParentEnabled);
+	const ESlateDrawEffect DrawEffects = bEnabled ? ESlateDrawEffect::None : ESlateDrawEffect::DisabledEffect;
+	FDrawContext DrawContext(AllottedGeometry, MyCullingRect, InWidgetStyle, DrawEffects, OutDrawElements, LayerId);
+
+	const float ViewWidth = AllottedGeometry.Size.X;
+	const float ViewHeight = AllottedGeometry.Size.Y;
+
+	//////////////////////////////////////////////////
+	{
+		FStopwatch Stopwatch;
+		Stopwatch.Start();
+
+		FPacketContentViewDrawHelper Helper(DrawContext, Viewport);
+
+		Helper.DrawBackground();
+
+		// Draw the cached state (the network events of the package breakdown).
+		Helper.Draw(*DrawState);
+
+		if (!FNetworkPacketEventRef::AreEquals(SelectedEvent, HoveredEvent))
+		{
+			// Highlight the selected event (if any).
+			if (SelectedEvent.IsValid())
+			{
+				Helper.DrawEventHighlight(SelectedEvent.Event, FPacketContentViewDrawHelper::EHighlightMode::Selected);
+			}
+
+			// Highlight the hovered event (if any).
+			if (HoveredEvent.IsValid())
+			{
+				Helper.DrawEventHighlight(HoveredEvent.Event, FPacketContentViewDrawHelper::EHighlightMode::Hovered);
+			}
+		}
+		else
+		{
+			// Highlight the selected and hovered event (if any).
+			if (SelectedEvent.IsValid())
+			{
+				Helper.DrawEventHighlight(SelectedEvent.Event, FPacketContentViewDrawHelper::EHighlightMode::SelectedAndHovered);
+			}
+		}
+
+		// Draw tooltip for hovered Event.
+		Tooltip.Draw(DrawContext);
+
+		Stopwatch.Stop();
+		DrawDurationHistory.AddValue(Stopwatch.AccumulatedTime);
+	}
+	//////////////////////////////////////////////////
+
+	const bool bShouldDisplayDebugInfo = FInsightsManager::Get()->IsDebugInfoEnabled();
+	if (bShouldDisplayDebugInfo)
+	{
+		const FSlateBrush* WhiteBrush = FInsightsStyle::Get().GetBrush("WhiteBrush");
+		FSlateFontInfo SummaryFont = FCoreStyle::GetDefaultFontStyle("Regular", 8);
+
+		const TSharedRef<FSlateFontMeasure> FontMeasureService = FSlateApplication::Get().GetRenderer()->GetFontMeasureService();
+		const float MaxFontCharHeight = FontMeasureService->Measure(TEXT("!"), SummaryFont).Y;
+		const float DbgDY = MaxFontCharHeight;
+
+		const float DbgW = 280.0f;
+		const float DbgH = DbgDY * 5 + 3.0f;
+		const float DbgX = ViewWidth - DbgW - 20.0f;
+		float DbgY = 7.0f;
+
+		DrawContext.LayerId++;
+
+		DrawContext.DrawBox(DbgX - 2.0f, DbgY - 2.0f, DbgW, DbgH, WhiteBrush, FLinearColor(1.0, 1.0, 1.0, 0.9));
+		DrawContext.LayerId++;
+
+		FLinearColor DbgTextColor(0.0, 0.0, 0.0, 0.9);
+
+		// Time interval since last OnPaint call.
+		const uint64 CurrentTime = FPlatformTime::Cycles64();
+		const uint64 OnPaintDuration = CurrentTime - LastOnPaintTime;
+		LastOnPaintTime = CurrentTime;
+		OnPaintDurationHistory.AddValue(OnPaintDuration);
+		const uint64 AvgOnPaintDuration = OnPaintDurationHistory.ComputeAverage();
+		const uint64 AvgOnPaintDurationMs = FStopwatch::Cycles64ToMilliseconds(AvgOnPaintDuration);
+		const double AvgOnPaintFps = AvgOnPaintDurationMs != 0 ? 1.0 / FStopwatch::Cycles64ToSeconds(AvgOnPaintDuration) : 0.0;
+
+		const uint64 AvgUpdateDurationMs = FStopwatch::Cycles64ToMilliseconds(UpdateDurationHistory.ComputeAverage());
+		const uint64 AvgDrawDurationMs = FStopwatch::Cycles64ToMilliseconds(DrawDurationHistory.ComputeAverage());
+
+		// Draw performance info.
+		DrawContext.DrawText
+		(
+			DbgX, DbgY,
+			FString::Printf(TEXT("U: %llu ms    D: %llu ms + %llu ms = %llu ms (%d fps)"),
+				AvgUpdateDurationMs, // average duration of UpdateState calls
+				AvgDrawDurationMs, // drawing time
+				AvgOnPaintDurationMs - AvgDrawDurationMs, // other overhead to OnPaint calls
+				AvgOnPaintDurationMs, // average time between two OnPaint calls
+				FMath::RoundToInt(AvgOnPaintFps)), // framerate of OnPaint calls
+			SummaryFont, DbgTextColor
+		);
+		DbgY += DbgDY;
+
+		// Draw "the update stats".
+		DrawContext.DrawText
+		(
+			DbgX, DbgY,
+			FString::Printf(TEXT("U: %s events"),
+				*FText::AsNumber(DrawState->Events.Num()).ToString()),
+			SummaryFont, DbgTextColor
+		);
+		DbgY += DbgDY;
+
+		// Draw "the draw stats".
+		DrawContext.DrawText
+		(
+			DbgX, DbgY,
+			FString::Printf(TEXT("D: %s boxes, %s borders (%s merged), %s texts"),
+				*FText::AsNumber(DrawState->Boxes.Num()).ToString(),
+				*FText::AsNumber(DrawState->Borders.Num()).ToString(),
+				*FText::AsNumber(DrawState->GetNumMergedBoxes()).ToString(),
+				*FText::AsNumber(DrawState->Texts.Num()).ToString()),
+			SummaryFont, DbgTextColor
+		);
+		DbgY += DbgDY;
+
+		// Draw viewport's horizontal info.
+		DrawContext.DrawText
+		(
+			DbgX, DbgY,
+			Viewport.GetHorizontalAxisViewport().ToDebugString(TEXT("X")),
+			SummaryFont, DbgTextColor
+		);
+		DbgY += DbgDY;
+
+		// Draw packet info.
+		DrawContext.DrawText
+		(
+			DbgX, DbgY,
+			FString::Printf(TEXT("Game Instance %d, Connection %d (%s), Packet %d"),
+				GameInstanceIndex,
+				ConnectionIndex,
+				(ConnectionMode == Trace::ENetProfilerConnectionMode::Outgoing) ? TEXT("Outgoing") : TEXT("Incoming"),
+				PacketIndex),
+			SummaryFont, DbgTextColor
+		);
+		DbgY += DbgDY;
+	}
+
+	return SCompoundWidget::OnPaint(Args, AllottedGeometry, MyCullingRect, OutDrawElements, LayerId, InWidgetStyle, bParentEnabled && IsEnabled());
+}
+
+////////////////////////////////////////////////////////////////////////////////////////////////////
+
+FReply SPacketContentView::OnMouseButtonDown(const FGeometry& MyGeometry, const FPointerEvent& MouseEvent)
+{
+	FReply Reply = FReply::Unhandled();
+
+	MousePosition = MyGeometry.AbsoluteToLocal(MouseEvent.GetScreenSpacePosition());
+	MousePositionOnButtonDown = MousePosition;
+
+	ViewportPosXOnButtonDown = Viewport.GetHorizontalAxisViewport().GetPos();
+
+	if (MouseEvent.GetEffectingButton() == EKeys::LeftMouseButton)
+	{
+		bIsLMB_Pressed = true;
+
+		// Capture mouse.
+		Reply = FReply::Handled().CaptureMouse(SharedThis(this));
+	}
+	else if (MouseEvent.GetEffectingButton() == EKeys::RightMouseButton)
+	{
+		bIsRMB_Pressed = true;
+
+		// Capture mouse, so we can drag outside this widget.
+		Reply = FReply::Handled().CaptureMouse(SharedThis(this));
+	}
+
+	return Reply;
+}
+
+////////////////////////////////////////////////////////////////////////////////////////////////////
+
+FReply SPacketContentView::OnMouseButtonUp(const FGeometry& MyGeometry, const FPointerEvent& MouseEvent)
+{
+	FReply Reply = FReply::Unhandled();
+
+	MousePosition = MyGeometry.AbsoluteToLocal(MouseEvent.GetScreenSpacePosition());
+	MousePositionOnButtonUp = MousePosition;
+
+	const bool bIsValidForMouseClick = MousePositionOnButtonUp.Equals(MousePositionOnButtonDown, MOUSE_SNAP_DISTANCE);
+
+	if (MouseEvent.GetEffectingButton() == EKeys::LeftMouseButton)
+	{
+		if (bIsLMB_Pressed)
+		{
+			if (bIsScrolling)
+			{
+				bIsScrolling = false;
+				CursorType = ECursorType::Default;
+			}
+			else if (bIsValidForMouseClick)
+			{
+				// Select the hovered timing event (if any).
+				UpdateHoveredEvent();
+				SelectHoveredEvent();
+			}
+
+			// Release mouse as we no longer drag.
+			Reply = FReply::Handled().ReleaseMouseCapture();
+
+			bIsLMB_Pressed = false;
+		}
+	}
+	else if (MouseEvent.GetEffectingButton() == EKeys::RightMouseButton)
+	{
+		if (bIsRMB_Pressed)
+		{
+			if (bIsScrolling)
+			{
+				bIsScrolling = false;
+				CursorType = ECursorType::Default;
+			}
+			else if (bIsValidForMouseClick)
+			{
+				//ShowContextMenu(MouseEvent);
+			}
+
+			// Release mouse as we no longer drag.
+			Reply = FReply::Handled().ReleaseMouseCapture();
+
+			bIsRMB_Pressed = false;
+		}
+	}
+
+	return Reply;
+}
+
+////////////////////////////////////////////////////////////////////////////////////////////////////
+
+FReply SPacketContentView::OnMouseMove(const FGeometry& MyGeometry, const FPointerEvent& MouseEvent)
+{
+	FReply Reply = FReply::Unhandled();
+
+	MousePosition = MyGeometry.AbsoluteToLocal(MouseEvent.GetScreenSpacePosition());
+
+	if (!MouseEvent.GetCursorDelta().IsZero())
+	{
+		if (MouseEvent.IsMouseButtonDown(EKeys::LeftMouseButton) ||
+			MouseEvent.IsMouseButtonDown(EKeys::RightMouseButton))
+		{
+			if (HasMouseCapture())
+			{
+				if (!bIsScrolling)
+				{
+					bIsScrolling = true;
+					CursorType = ECursorType::Hand;
+
+					HoveredEvent.Reset();
+					Tooltip.SetDesiredOpacity(0.0f);
+				}
+
+				FAxisViewportDouble& ViewportX = Viewport.GetHorizontalAxisViewport();
+				const float PosX = ViewportPosXOnButtonDown + (MousePositionOnButtonDown.X - MousePosition.X);
+				ViewportX.ScrollAtPos(PosX);
+				UpdateHorizontalScrollBar();
+				bIsStateDirty = true;
+			}
+		}
+		else
+		{
+			UpdateHoveredEvent();
+		}
+
+		Reply = FReply::Handled();
+	}
+
+	return Reply;
+}
+
+////////////////////////////////////////////////////////////////////////////////////////////////////
+
+void SPacketContentView::OnMouseEnter(const FGeometry& MyGeometry, const FPointerEvent& MouseEvent)
+{
+}
+
+////////////////////////////////////////////////////////////////////////////////////////////////////
+
+void SPacketContentView::OnMouseLeave(const FPointerEvent& MouseEvent)
+{
+	if (!HasMouseCapture())
+	{
+		// No longer dragging (unless we have mouse capture).
+		bIsScrolling = false;
+
+		bIsLMB_Pressed = false;
+		bIsRMB_Pressed = false;
+
+		MousePosition = FVector2D::ZeroVector;
+
+		HoveredEvent.Reset();
+		Tooltip.SetDesiredOpacity(0.0f);
+
+		CursorType = ECursorType::Default;
+	}
+}
+
+////////////////////////////////////////////////////////////////////////////////////////////////////
+
+FReply SPacketContentView::OnMouseWheel(const FGeometry& MyGeometry, const FPointerEvent& MouseEvent)
+{
+	MousePosition = MyGeometry.AbsoluteToLocal(MouseEvent.GetScreenSpacePosition());
+
+	//if (MouseEvent.GetModifierKeys().IsShiftDown())
+	//{
+	//}
+	//else //if (MouseEvent.GetModifierKeys().IsControlDown())
+	{
+		// Zoom in/out horizontally.
+		const float Delta = MouseEvent.GetWheelDelta();
+		ZoomHorizontally(Delta, MousePosition.X);
+	}
+
+	return FReply::Handled();
+}
+
+////////////////////////////////////////////////////////////////////////////////////////////////////
+
+FReply SPacketContentView::OnMouseButtonDoubleClick(const FGeometry& MyGeometry, const FPointerEvent& MouseEvent)
+{
+	return FReply::Unhandled();
+}
+
+////////////////////////////////////////////////////////////////////////////////////////////////////
+
+FCursorReply SPacketContentView::OnCursorQuery(const FGeometry& MyGeometry, const FPointerEvent& CursorEvent) const
+{
+	FCursorReply CursorReply = FCursorReply::Unhandled();
+
+	if (CursorType == ECursorType::Arrow)
+	{
+		CursorReply = FCursorReply::Cursor(EMouseCursor::ResizeLeftRight);
+	}
+	else if (CursorType == ECursorType::Hand)
+	{
+		CursorReply = FCursorReply::Cursor(EMouseCursor::GrabHand);
+	}
+
+	return CursorReply;
+}
+
+////////////////////////////////////////////////////////////////////////////////////////////////////
+
+void SPacketContentView::BindCommands()
+{
+}
+
+////////////////////////////////////////////////////////////////////////////////////////////////////
+
+void SPacketContentView::HorizontalScrollBar_OnUserScrolled(float ScrollOffset)
+{
+	FAxisViewportDouble& ViewportX = Viewport.GetHorizontalAxisViewport();
+	ViewportX.OnUserScrolled(HorizontalScrollBar, ScrollOffset);
+	bIsStateDirty = true;
+}
+
+////////////////////////////////////////////////////////////////////////////////////////////////////
+
+void SPacketContentView::UpdateHorizontalScrollBar()
+{
+	FAxisViewportDouble& ViewportX = Viewport.GetHorizontalAxisViewport();
+	ViewportX.UpdateScrollBar(HorizontalScrollBar);
+}
+
+////////////////////////////////////////////////////////////////////////////////////////////////////
+
+void SPacketContentView::ZoomHorizontally(const float Delta, const float X)
+{
+	FAxisViewportDouble& ViewportX = Viewport.GetHorizontalAxisViewport();
+	ViewportX.RelativeZoomWithFixedOffset(Delta, X);
+	UpdateHorizontalScrollBar();
+	bIsStateDirty = true;
+}
+
+////////////////////////////////////////////////////////////////////////////////////////////////////
+
+#undef LOCTEXT_NAMESPACE
diff --git a/Engine/Source/Developer/TraceInsights/Private/Insights/NetworkingProfiler/Widgets/SPacketContentView.h b/Engine/Source/Developer/TraceInsights/Private/Insights/NetworkingProfiler/Widgets/SPacketContentView.h
new file mode 100644
--- /dev/null
+++ b/Engine/Source/Developer/TraceInsights/Private/Insights/NetworkingProfiler/Widgets/SPacketContentView.h
@@ -0,0 +1,218 @@
+// Copyright 1998-2019 Epic Games, Inc. All Rights Reserved.
+
+#pragma once
+
+#include "CoreMinimal.h"
+#include "Input/CursorReply.h"
+#include "Input/Reply.h"
+#include "Layout/Geometry.h"
+#include "Rendering/RenderingCommon.h"
+#include "TraceServices/Model/NetProfiler.h"
+#include "Widgets/DeclarativeSyntaxSupport.h"
+#include "Widgets/SCompoundWidget.h"
+
+// Insights
+#include "Insights/Common/FixedCircularBuffer.h"
+#include "Insights/NetworkingProfiler/ViewModels/PacketContentViewDrawHelper.h"
+#include "Insights/NetworkingProfiler/ViewModels/PacketContentViewport.h"
+#include "Insights/ViewModels/TooltipDrawState.h"
+
+class SScrollBar;
+class SNetworkingProfilerWindow;
+
+////////////////////////////////////////////////////////////////////////////////////////////////////
+
+struct FNetworkPacketEventRef
+{
+	FNetworkPacketEvent Event;
+	bool bIsValid;
+
+	FNetworkPacketEventRef()
+		: Event()
+		, bIsValid(false)
+	{
+	}
+
+	FNetworkPacketEventRef(const FNetworkPacketEvent& InEvent)
+		: Event(InEvent)
+		, bIsValid(true)
+	{
+	}
+
+	FNetworkPacketEventRef(const FNetworkPacketEventRef& Other)
+		: Event(Other.Event)
+		, bIsValid(Other.bIsValid)
+	{
+	}
+
+	FNetworkPacketEventRef& operator=(const FNetworkPacketEventRef& Other)
+	{
+		Event = Other.Event;
+		bIsValid = Other.bIsValid;
+		return *this;
+	}
+
+	void Reset()
+	{
+		bIsValid = false;
+	}
+
+	bool IsValid() const
+	{
+		return bIsValid;
+	}
+
+	bool Equals(const FNetworkPacketEventRef& Other) const
+	{
+		return bIsValid == Other.bIsValid && Event.Equals(Other.Event);
+	}
+
+	static bool AreEquals(const FNetworkPacketEventRef& A, const FNetworkPacketEventRef& B)
+	{
+		return A.Equals(B);
+	}
+};
+
+////////////////////////////////////////////////////////////////////////////////////////////////////
+/**
+ * Widget used to present content of a network packet.
+ */
+class SPacketContentView : public SCompoundWidget
+{
+public:
+	/** Number of pixels. */
+	static constexpr float MOUSE_SNAP_DISTANCE = 2.0f;
+
+	enum class ECursorType
+	{
+		Default,
+		Arrow,
+		Hand,
+	};
+
+public:
+	/** Default constructor. */
+	SPacketContentView();
+
+	/** Virtual destructor. */
+	virtual ~SPacketContentView();
+
+	/** Resets internal widget's data to the default one. */
+	void Reset();
+
+	SLATE_BEGIN_ARGS(SPacketContentView) {}
+	SLATE_END_ARGS()
+
+	void Construct(const FArguments& InArgs, TSharedPtr<SNetworkingProfilerWindow> InProfilerWindow);
+
+	/**
+	 * Ticks this widget. Override in derived classes, but always call the parent implementation.
+	 * @param  AllottedGeometry The space allotted for this widget
+	 * @param  InCurrentTime  Current absolute real time
+	 * @param  InDeltaTime  Real time passed since last tick
+	 */
+	virtual void Tick(const FGeometry& AllottedGeometry, const double InCurrentTime, const float InDeltaTime) override;
+
+	virtual int32 OnPaint(const FPaintArgs& Args, const FGeometry& AllottedGeometry, const FSlateRect& MyCullingRect, FSlateWindowElementList& OutDrawElements, int32 LayerId, const FWidgetStyle& InWidgetStyle, bool bParentEnabled) const override;
+
+	virtual FReply OnMouseButtonDown(const FGeometry& MyGeometry, const FPointerEvent& MouseEvent) override;
+	virtual FReply OnMouseButtonUp(const FGeometry& MyGeometry, const FPointerEvent& MouseEvent) override;
+	virtual FReply OnMouseMove(const FGeometry& MyGeometry, const FPointerEvent& MouseEvent) override;
+	virtual void OnMouseEnter(const FGeometry& MyGeometry, const FPointerEvent& MouseEvent) override;
+	virtual void OnMouseLeave(const FPointerEvent& MouseEvent) override;
+	virtual FReply OnMouseWheel(const FGeometry& MyGeometry, const FPointerEvent& MouseEvent) override;
+	virtual FReply OnMouseButtonDoubleClick(const FGeometry& MyGeometry, const FPointerEvent& MouseEvent) override;
+
+	virtual FCursorReply OnCursorQuery(const FGeometry& MyGeometry, const FPointerEvent& CursorEvent) const override;
+
+	void ResetPacket();
+	void SetPacket(uint32 InGameInstanceIndex, uint32 InConnectionIndex, Trace::ENetProfilerConnectionMode InConnectionMode, uint32 InPacketIndex, int64 InPacketBitSize);
+
+private:
+	//void ShowContextMenu(const FPointerEvent& MouseEvent);
+
+	/** Binds our UI commands to delegates. */
+	void BindCommands();
+
+	////////////////////////////////////////////////////////////////////////////////////////////////////
+
+	/**
+	 * Called when the user scrolls the horizontal scrollbar.
+	 * @param ScrollOffset Scroll offset as a fraction between 0 and 1.
+	 */
+	void HorizontalScrollBar_OnUserScrolled(float ScrollOffset);
+	void UpdateHorizontalScrollBar();
+
+	void ZoomHorizontally(const float Delta, const float X);
+
+	////////////////////////////////////////////////////////////////////////////////////////////////////
+
+	void UpdateState();
+
+	void UpdateHoveredEvent();
+	FNetworkPacketEventRef GetEventAtMousePosition(float X, float Y);
+
+	void OnSelectedEventChanged();
+	void SelectHoveredEvent();
+
+private:
+	TSharedPtr<SNetworkingProfilerWindow> ProfilerWindow;
+
+	/** The track's viewport. Encapsulates info about position and scale. */
+	FPacketContentViewport Viewport;
+	bool bIsViewportDirty;
+
+	uint32 GameInstanceIndex;
+	uint32 ConnectionIndex;
+	Trace::ENetProfilerConnectionMode ConnectionMode;
+	uint32 PacketIndex;
+	int64 PacketBitSize; // total number of bits; [bit]
+
+	/** Cached draw state of the packet content (i.e. all it needs to render). */
+	TSharedPtr<FPacketContentViewDrawState> DrawState;
+	bool bIsStateDirty;
+	//////////////////////////////////////////////////
+
+	TSharedPtr<SScrollBar> HorizontalScrollBar;
+
+	//////////////////////////////////////////////////
+	// Panning and Zooming behaviors
+
+	/** The current mouse position. */
+	FVector2D MousePosition;
+
+	/** Mouse position during the call on mouse button down. */
+	FVector2D MousePositionOnButtonDown;
+	float ViewportPosXOnButtonDown;
+
+	/** Mouse position during the call on mouse button up. */
+	FVector2D MousePositionOnButtonUp;
+
+	bool bIsLMB_Pressed;
+	bool bIsRMB_Pressed;
+
+	/** True, if the user is currently interactively scrolling the view (ex.: by holding the left mouse button and dragging). */
+	bool bIsScrolling;
+
+	//////////////////////////////////////////////////
+	// Selection
+
+	FNetworkPacketEventRef HoveredEvent;
+	FNetworkPacketEventRef SelectedEvent;
+
+	FTooltipDrawState Tooltip;
+
+	//////////////////////////////////////////////////
+	// Misc
+
+	FGeometry ThisGeometry;
+
+	/** Cursor type. */
+	ECursorType CursorType;
+
+	// Debug stats
+	TFixedCircularBuffer<uint64, 32> UpdateDurationHistory;
+	mutable TFixedCircularBuffer<uint64, 32> DrawDurationHistory;
+	mutable TFixedCircularBuffer<uint64, 32> OnPaintDurationHistory;
+	mutable uint64 LastOnPaintTime;
+};
diff --git a/Engine/Source/Developer/TraceInsights/Private/Insights/NetworkingProfiler/Widgets/SPacketSizesView.cpp b/Engine/Source/Developer/TraceInsights/Private/Insights/NetworkingProfiler/Widgets/SPacketSizesView.cpp
deleted file mode 100644
diff --git a/Engine/Source/Developer/TraceInsights/Private/Insights/NetworkingProfiler/Widgets/SPacketSizesView.h b/Engine/Source/Developer/TraceInsights/Private/Insights/NetworkingProfiler/Widgets/SPacketSizesView.h
deleted file mode 100644
diff --git a/Engine/Source/Developer/TraceInsights/Private/Insights/NetworkingProfiler/Widgets/SPacketView.cpp b/Engine/Source/Developer/TraceInsights/Private/Insights/NetworkingProfiler/Widgets/SPacketView.cpp
new file mode 100644
--- /dev/null
+++ b/Engine/Source/Developer/TraceInsights/Private/Insights/NetworkingProfiler/Widgets/SPacketView.cpp
@@ -0,0 +1,1118 @@
+// Copyright 1998-2019 Epic Games, Inc. All Rights Reserved.
+
+#include "SPacketView.h"
+
+#include "Fonts/FontMeasure.h"
+#include "Fonts/SlateFontInfo.h"
+#include "Framework/Application/SlateApplication.h"
+#include "Framework/MultiBox/MultiBoxBuilder.h"
+#include "HAL/PlatformTime.h"
+#include "Rendering/DrawElements.h"
+#include "TraceServices/Model/NetProfiler.h"
+#include "Styling/CoreStyle.h"
+#include "Widgets/Layout/SScrollBar.h"
+
+// Insights
+#include "Insights/Common/PaintUtils.h"
+#include "Insights/Common/Stopwatch.h"
+#include "Insights/Common/TimeUtils.h"
+#include "Insights/InsightsManager.h"
+#include "Insights/InsightsStyle.h"
+#include "Insights/NetworkingProfiler/NetworkingProfilerManager.h"
+#include "Insights/NetworkingProfiler/Widgets/SNetworkingProfilerWindow.h"
+#include "Insights/NetworkingProfiler/Widgets/SPacketContentView.h"
+
+#include <limits>
+
+////////////////////////////////////////////////////////////////////////////////////////////////////
+
+#define LOCTEXT_NAMESPACE "SPacketView"
+
+////////////////////////////////////////////////////////////////////////////////////////////////////
+
+SPacketView::SPacketView()
+	: ProfilerWindow()
+	, PacketSeries(MakeShareable(new FNetworkPacketSeries()))
+{
+	Reset();
+}
+
+////////////////////////////////////////////////////////////////////////////////////////////////////
+
+SPacketView::~SPacketView()
+{
+}
+
+////////////////////////////////////////////////////////////////////////////////////////////////////
+
+void SPacketView::Reset()
+{
+	//ProfilerWindow
+
+	GameInstanceIndex = 0;
+	ConnectionIndex = 0;
+	ConnectionMode = Trace::ENetProfilerConnectionMode::Outgoing;
+
+	Viewport.Reset();
+	FAxisViewportInt32& ViewportX = Viewport.GetHorizontalAxisViewport();
+	ViewportX.SetScaleLimits(0.0001f, 16.0f); // 10000 [sample/px] to 16 [px/sample]
+	ViewportX.SetScale(5.0f);
+	FAxisViewportDouble& ViewportY = Viewport.GetVerticalAxisViewport();
+	ViewportY.SetScaleLimits(0.0001, 50.0);
+	ViewportY.SetScale(0.02);
+	bIsViewportDirty = true;
+
+	PacketSeries->Reset();
+
+	bIsStateDirty = true;
+
+	bIsAutoZoomEnabled = true;
+
+	AnalysisSyncNextTimestamp = 0;
+	ConnectionChangeCount = 0;
+
+	MousePosition = FVector2D::ZeroVector;
+
+	MousePositionOnButtonDown = FVector2D::ZeroVector;
+	ViewportPosXOnButtonDown = 0.0f;
+
+	MousePositionOnButtonUp = FVector2D::ZeroVector;
+
+	bIsLMB_Pressed = false;
+	bIsRMB_Pressed = false;
+
+	bIsScrolling = false;
+
+	SelectionStartFrameIndex = 0;
+	SelectionEndFrameIndex = 0;
+
+	SelectedSample.Reset();
+	HoveredSample.Reset();
+	TooltipDesiredOpacity = 0.9f;
+	TooltipOpacity = 0.0f;
+
+	//ThisGeometry
+
+	CursorType = ECursorType::Default;
+
+	NumUpdatedPackets = 0;
+	UpdateDurationHistory.Reset();
+	DrawDurationHistory.Reset();
+	OnPaintDurationHistory.Reset();
+	LastOnPaintTime = FPlatformTime::Cycles64();
+
+	if (ProfilerWindow.IsValid())
+	{
+		SetConnection(ProfilerWindow->GetSelectedGameInstanceIndex(), ProfilerWindow->GetSelectedConnectionIndex(), ProfilerWindow->GetSelectedConnectionMode());
+	}
+}
+
+////////////////////////////////////////////////////////////////////////////////////////////////////
+
+void SPacketView::SetConnection(uint32 InGameInstanceIndex, uint32 InConnectionIndex, Trace::ENetProfilerConnectionMode InConnectionMode)
+{
+	GameInstanceIndex = InGameInstanceIndex;
+	ConnectionIndex = InConnectionIndex;
+	ConnectionMode = InConnectionMode;
+
+	HoveredSample.Reset();
+
+	SelectedSample.Reset();
+	OnSelectedSampleChanged();
+
+	bIsStateDirty = true;
+}
+
+////////////////////////////////////////////////////////////////////////////////////////////////////
+
+void SPacketView::Construct(const FArguments& InArgs, TSharedPtr<SNetworkingProfilerWindow> InProfilerWindow)
+{
+	ProfilerWindow = InProfilerWindow;
+
+	ChildSlot
+	[
+		SNew(SOverlay)
+		.Visibility(EVisibility::SelfHitTestInvisible)
+
+		+ SOverlay::Slot()
+		.VAlign(VAlign_Top)
+		.Padding(FMargin(0, 0, 0, 0))
+		[
+			SAssignNew(HorizontalScrollBar, SScrollBar)
+			.Orientation(Orient_Horizontal)
+			.AlwaysShowScrollbar(false)
+			.Visibility(EVisibility::Visible)
+			.Thickness(FVector2D(5.0f, 5.0f))
+			.RenderOpacity(0.75)
+			.OnUserScrolled(this, &SPacketView::HorizontalScrollBar_OnUserScrolled)
+		]
+	];
+
+	UpdateHorizontalScrollBar();
+
+	BindCommands();
+
+	Reset();
+}
+
+////////////////////////////////////////////////////////////////////////////////////////////////////
+
+void SPacketView::Tick(const FGeometry& AllottedGeometry, const double InCurrentTime, const float InDeltaTime)
+{
+	if (ThisGeometry != AllottedGeometry || bIsViewportDirty)
+	{
+		bIsViewportDirty = false;
+		const float ViewWidth = AllottedGeometry.GetLocalSize().X;
+		const float ViewHeight = AllottedGeometry.GetLocalSize().Y;
+		Viewport.SetSize(ViewWidth, ViewHeight);
+		bIsStateDirty = true;
+	}
+
+	ThisGeometry = AllottedGeometry;
+
+	FAxisViewportInt32& ViewportX = Viewport.GetHorizontalAxisViewport();
+
+	if (!bIsScrolling)
+	{
+		// Elastic snap to horizontal limits.
+		if (ViewportX.UpdatePosWithinLimits())
+		{
+			bIsStateDirty = true;
+		}
+	}
+
+	uint64 Time = FPlatformTime::Cycles64();
+	if (Time > AnalysisSyncNextTimestamp)
+	{
+		const uint64 WaitTime = static_cast<uint64>(0.1 / FPlatformTime::GetSecondsPerCycle64()); // 100ms
+		AnalysisSyncNextTimestamp = Time + WaitTime;
+
+		TSharedPtr<const Trace::IAnalysisSession> Session = FInsightsManager::Get()->GetSession();
+		if (Session.IsValid())
+		{
+			Trace::FAnalysisSessionReadScope SessionReadScope(*Session.Get());
+			const Trace::INetProfilerProvider& NetProfilerProvider = Trace::ReadNetProfilerProvider(*Session.Get());
+
+			const uint32 NewConnectionChangeCount = NetProfilerProvider.GetConnectionChangeCount();
+			if (NewConnectionChangeCount != ConnectionChangeCount)
+			{
+				ConnectionChangeCount = NewConnectionChangeCount;
+				bIsStateDirty = true;
+			}
+		}
+	}
+
+	if (bIsStateDirty)
+	{
+		bIsStateDirty = false;
+		UpdateState();
+	}
+}
+
+////////////////////////////////////////////////////////////////////////////////////////////////////
+
+void SPacketView::UpdateState()
+{
+	FStopwatch Stopwatch;
+	Stopwatch.Start();
+
+	// Reset stats.
+	PacketSeries->NumAggregatedPackets = 0;
+	NumUpdatedPackets = 0;
+
+	FNetworkPacketSeriesBuilder Builder(*PacketSeries, Viewport);
+
+	FAxisViewportInt32& ViewportX = Viewport.GetHorizontalAxisViewport();
+	ViewportX.SetMinMaxInterval(0, 0);
+
+	TSharedPtr<const Trace::IAnalysisSession> Session = FInsightsManager::Get()->GetSession();
+	if (Session.IsValid())
+	{
+		Trace::FAnalysisSessionReadScope SessionReadScope(*Session.Get());
+		const Trace::INetProfilerProvider& NetProfilerProvider = Trace::ReadNetProfilerProvider(*Session.Get());
+
+		const uint32 GameInstanceCount = NetProfilerProvider.GetGameInstanceCount();
+		if (GameInstanceIndex < GameInstanceCount)
+		{
+			//const uint32 ConnectionCount = NetProfilerProvider.GetConnectionCount(GameInstanceIndex);
+			//if (ConnectionIndex < ConnectionCount) TODO: (if NetProfilerProvider.IsConnectionValid(GameInstanceIndex))
+			{
+				const uint32 NumPackets = NetProfilerProvider.GetPacketCount(ConnectionIndex, ConnectionMode);
+
+				ViewportX.SetMinMaxInterval(0, NumPackets);
+
+				if (NumPackets > 0)
+				{
+					const int32 MinIndex = ViewportX.GetValueAtOffset(0.0f);
+					const int32 MaxIndex = ViewportX.GetValueAtOffset(ViewportX.GetSize());
+
+					const uint32 PacketStartIndex = static_cast<uint32>(FMath::Max(0, MinIndex));
+					const uint32 PacketEndIndex = FMath::Min(NumPackets - 1, static_cast<uint32>(FMath::Max(0, MaxIndex)));
+
+					if (PacketStartIndex <= PacketEndIndex)
+					{
+						int32 PacketIndex = PacketStartIndex;
+						NetProfilerProvider.EnumeratePackets(ConnectionIndex, ConnectionMode, PacketStartIndex, PacketEndIndex, [&Builder, &PacketIndex](const Trace::FNetProfilerPacket& Packet)
+						{
+							Builder.AddPacket(PacketIndex++, Packet);
+						});
+					}
+				}
+			}
+		}
+	}
+
+	// Init series with mock data.
+	if (false)
+	{
+		constexpr int32 NumPackets = 100000;
+
+		ViewportX.SetMinMaxInterval(0, NumPackets);
+
+		const int32 StartIndex = FMath::Max(0, ViewportX.GetValueAtOffset(0.0f));
+		const int32 EndIndex = FMath::Min(NumPackets, ViewportX.GetValueAtOffset(ViewportX.GetSize()));
+
+		for (int32 PacketIndex = StartIndex; PacketIndex < EndIndex; ++PacketIndex)
+		{
+			FRandomStream RandomStream((PacketIndex * PacketIndex * PacketIndex) ^ 0x2c2c57ed);
+			int64 Size = RandomStream.RandRange(0, 2000);
+
+			Trace::ENetProfilerDeliveryStatus Status = Trace::ENetProfilerDeliveryStatus::Unknown;
+			const float Fraction = RandomStream.GetFraction();
+			if (Fraction < 0.01) // 1%
+			{
+				Status = Trace::ENetProfilerDeliveryStatus::Dropped;
+			}
+			else if (Fraction < 0.05) // 4%
+			{
+				Status = Trace::ENetProfilerDeliveryStatus::Delivered;
+			}
+
+			const double Timestamp = ((double)PacketIndex * 100.0) / (double)NumPackets + RandomStream.GetFraction() * 0.1;
+
+			Trace::FNetProfilerPacket Packet;
+			Packet.TimeStamp = static_cast<Trace::FNetProfilerTimeStamp>(Timestamp);
+			Packet.SequenceNumber = PacketIndex;
+			Packet.ContentSizeInBits = Size;
+			Packet.TotalPacketSizeInBytes = (Size + 7) / 8;
+			Packet.DeliveryStatus = Status;
+			Builder.AddPacket(PacketIndex, Packet);
+		}
+	}
+
+	NumUpdatedPackets += Builder.GetNumAddedPackets();
+
+	Stopwatch.Stop();
+	UpdateDurationHistory.AddValue(Stopwatch.AccumulatedTime);
+}
+
+////////////////////////////////////////////////////////////////////////////////////////////////////
+
+FNetworkPacketSampleRef SPacketView::GetSampleAtMousePosition(float X, float Y)
+{
+	if (!bIsStateDirty)
+	{
+		float SampleW = Viewport.GetSampleWidth();
+		int32 SampleIndex = FMath::FloorToInt(X / SampleW);
+		if (SampleIndex >= 0)
+		{
+			if (PacketSeries->NumAggregatedPackets > 0 &&
+				SampleIndex < PacketSeries->Samples.Num())
+			{
+				const FNetworkPacketAggregatedSample& Sample = PacketSeries->Samples[SampleIndex];
+				if (Sample.NumPackets > 0)
+				{
+					const FAxisViewportDouble& ViewportY = Viewport.GetVerticalAxisViewport();
+
+					const float ViewHeight = FMath::RoundToFloat(Viewport.GetHeight());
+					const float BaselineY = FMath::RoundToFloat(ViewportY.GetOffsetForValue(0.0));
+
+					//const float ValueY = FMath::RoundToFloat(ViewportY.GetOffsetForValue(static_cast<double>(Sample.LargestPacket.ContentSizeInBits)));
+					const float ValueY = FMath::RoundToFloat(ViewportY.GetOffsetForValue(static_cast<double>(Sample.LargestPacket.TotalSizeInBytes * 8)));
+
+					constexpr float ToleranceY = 3.0f; // [pixels]
+
+					const float BottomY = FMath::Min(ViewHeight, ViewHeight - BaselineY + ToleranceY);
+					const float TopY = FMath::Max(0.0f, ViewHeight - ValueY - ToleranceY);
+
+					if (Y >= TopY && Y < BottomY)
+					{
+						return FNetworkPacketSampleRef(PacketSeries, MakeShareable(new FNetworkPacketAggregatedSample(Sample)));
+					}
+				}
+			}
+		}
+	}
+	return FNetworkPacketSampleRef(nullptr, nullptr);
+}
+
+////////////////////////////////////////////////////////////////////////////////////////////////////
+
+void SPacketView::SelectSampleAtMousePosition(float X, float Y)
+{
+	FNetworkPacketSampleRef SampleRef = GetSampleAtMousePosition(X, Y);
+	if (!SampleRef.IsValid())
+	{
+		SampleRef = GetSampleAtMousePosition(X - 1.0f, Y);
+	}
+	if (!SampleRef.IsValid())
+	{
+		SampleRef = GetSampleAtMousePosition(X + 1.0f, Y);
+	}
+
+	if (!SelectedSample.Equals(SampleRef))
+	{
+		SelectedSample = SampleRef;
+		OnSelectedSampleChanged();
+	}
+}
+
+////////////////////////////////////////////////////////////////////////////////////////////////////
+
+void SPacketView::OnSelectedSampleChanged()
+{
+	if (ProfilerWindow.IsValid())
+	{
+		if (SelectedSample.IsValid())
+		{
+			const uint32 BitSize = SelectedSample.Sample->LargestPacket.TotalSizeInBytes * 8;
+			ProfilerWindow->SetSelectedPacket(SelectedSample.Sample->LargestPacket.Index, SelectedSample.Sample->LargestPacket.Index + 1, BitSize);
+			ProfilerWindow->SetSelectedBitRange(0, BitSize);
+		}
+		else
+		{
+			ProfilerWindow->SetSelectedPacket(0, 0);
+			ProfilerWindow->SetSelectedBitRange(0, 0);
+		}
+	}
+}
+
+////////////////////////////////////////////////////////////////////////////////////////////////////
+
+const TCHAR* StatusToString(Trace::ENetProfilerDeliveryStatus Status)
+{
+	switch (Status)
+	{
+		case Trace::ENetProfilerDeliveryStatus::Delivered:  return TEXT("Delivered");
+		case Trace::ENetProfilerDeliveryStatus::Dropped:    return TEXT("Dropped");
+		case Trace::ENetProfilerDeliveryStatus::Unknown:
+		default:                                            return TEXT("Unknown");
+	};
+}
+
+const TCHAR* AggregatedStatusToString(Trace::ENetProfilerDeliveryStatus Status)
+{
+	switch (Status)
+	{
+		case Trace::ENetProfilerDeliveryStatus::Delivered:  return TEXT("all packages are Delivered");
+		case Trace::ENetProfilerDeliveryStatus::Dropped:    return TEXT("at least one Dropped package");
+		case Trace::ENetProfilerDeliveryStatus::Unknown:
+		default:                                            return TEXT("Unknown");
+	}
+}
+
+////////////////////////////////////////////////////////////////////////////////////////////////////
+
+int32 SPacketView::OnPaint(const FPaintArgs& Args, const FGeometry& AllottedGeometry, const FSlateRect& MyCullingRect, FSlateWindowElementList& OutDrawElements, int32 LayerId, const FWidgetStyle& InWidgetStyle, bool bParentEnabled) const
+{
+	const bool bEnabled = ShouldBeEnabled(bParentEnabled);
+	const ESlateDrawEffect DrawEffects = bEnabled ? ESlateDrawEffect::None : ESlateDrawEffect::DisabledEffect;
+	FDrawContext DrawContext(AllottedGeometry, MyCullingRect, InWidgetStyle, DrawEffects, OutDrawElements, LayerId);
+
+	const TSharedRef<FSlateFontMeasure> FontMeasureService = FSlateApplication::Get().GetRenderer()->GetFontMeasureService();
+	FSlateFontInfo SummaryFont = FCoreStyle::GetDefaultFontStyle("Regular", 8);
+
+	const FSlateBrush* WhiteBrush = FInsightsStyle::Get().GetBrush("WhiteBrush");
+
+	const float ViewWidth = AllottedGeometry.Size.X;
+	const float ViewHeight = AllottedGeometry.Size.Y;
+
+	int32 NumDrawSamples = 0;
+
+	//////////////////////////////////////////////////
+	{
+		FStopwatch Stopwatch;
+		Stopwatch.Start();
+
+		FPacketViewDrawHelper Helper(DrawContext, Viewport);
+
+		Helper.DrawBackground();
+
+		// Draw the horizontal axis grid.
+		DrawHorizontalAxisGrid(DrawContext, WhiteBrush, SummaryFont);
+
+		Helper.DrawCached(*PacketSeries);
+
+		NumDrawSamples = Helper.GetNumDrawSamples();
+
+		// Highlight the selected and/or hovered sample.
+		bool bIsSelectedAndHovered = SelectedSample.Equals(HoveredSample);
+		if (SelectedSample.IsValid())
+		{
+			Helper.DrawSampleHighlight(*SelectedSample.Sample, bIsSelectedAndHovered ? FPacketViewDrawHelper::EHighlightMode::SelectedAndHovered : FPacketViewDrawHelper::EHighlightMode::Selected);
+		}
+		if (HoveredSample.IsValid() && !bIsSelectedAndHovered)
+		{
+			Helper.DrawSampleHighlight(*HoveredSample.Sample, FPacketViewDrawHelper::EHighlightMode::Hovered);
+		}
+
+		// Draw the vertical axis grid.
+		DrawVerticalAxisGrid(DrawContext, WhiteBrush, SummaryFont);
+
+		// Draw tooltip for hovered sample.
+		if (HoveredSample.IsValid())
+		{
+			if (TooltipOpacity < TooltipDesiredOpacity)
+			{
+				TooltipOpacity = TooltipOpacity * 0.9f + TooltipDesiredOpacity * 0.1f;
+			}
+			else
+			{
+				TooltipOpacity = TooltipDesiredOpacity;
+			}
+
+			const double Precision = 0.01; // 10ms
+			int NumLines;
+			FString Text;
+			uint32 UnusedBits = HoveredSample.Sample->LargestPacket.TotalSizeInBytes * 8 - HoveredSample.Sample->LargestPacket.ContentSizeInBits;
+			if (HoveredSample.Sample->NumPackets == 1)
+			{
+				NumLines = 6;
+				Text = FString::Format(TEXT("Packet Index: {0}\n"
+											"Sequence Number: {1}\n"
+											"Content Size: {2} bits\n"
+											"Total Size: {3} bytes ({4} unused bits)\n"
+											"Timestamp: {5}\n"
+											"Status: {6}"),
+					{
+						FText::AsNumber(HoveredSample.Sample->LargestPacket.Index).ToString(),
+						FText::AsNumber(HoveredSample.Sample->LargestPacket.SequenceNumber).ToString(),
+						FText::AsNumber(HoveredSample.Sample->LargestPacket.ContentSizeInBits).ToString(),
+						FText::AsNumber(HoveredSample.Sample->LargestPacket.TotalSizeInBytes).ToString(),
+						FText::AsNumber(UnusedBits).ToString(),
+						TimeUtils::FormatTimeHMS(HoveredSample.Sample->LargestPacket.TimeStamp, Precision),
+						::StatusToString(HoveredSample.Sample->LargestPacket.Status)
+					});
+			}
+			else
+			{
+				NumLines = 9;
+				Text = FString::Format(TEXT("{0} network packets\n"
+											"({1})\n"
+											"Largest Packet\n"
+											"    Index: {2}\n"
+											"    Sequance Number: {3}\n"
+											"    Content Size: {4} bits\n"
+											"    Total Size: {5} bytes ({6} unused bits)\n"
+											"    Timestamp: {7}\n"
+											"    Status: {8}"),
+					{
+						HoveredSample.Sample->NumPackets,
+						::AggregatedStatusToString(HoveredSample.Sample->AggregatedStatus),
+						FText::AsNumber(HoveredSample.Sample->LargestPacket.Index).ToString(),
+						FText::AsNumber(HoveredSample.Sample->LargestPacket.SequenceNumber).ToString(),
+						FText::AsNumber(HoveredSample.Sample->LargestPacket.ContentSizeInBits).ToString(),
+						FText::AsNumber(HoveredSample.Sample->LargestPacket.TotalSizeInBytes).ToString(),
+						FText::AsNumber(UnusedBits).ToString(),
+						TimeUtils::FormatTimeHMS(HoveredSample.Sample->LargestPacket.TimeStamp, Precision),
+						::StatusToString(HoveredSample.Sample->LargestPacket.Status)
+					});
+			}
+
+			FVector2D TextSize = FontMeasureService->Measure(Text, SummaryFont);
+
+			const float DX = 2.0f;
+			const float W2 = TextSize.X / 2 + DX;
+
+			const FAxisViewportInt32& ViewportX = Viewport.GetHorizontalAxisViewport();
+
+			float X1 = ViewportX.GetOffsetForValue(HoveredSample.Sample->LargestPacket.Index);
+			float CX = X1 + FMath::RoundToFloat(Viewport.GetSampleWidth() / 2);
+			if (CX + W2 > ViewportX.GetSize())
+			{
+				CX = FMath::RoundToFloat(ViewportX.GetSize() - W2);
+			}
+			if (CX - W2 < 0)
+			{
+				CX = W2;
+			}
+
+			const float Y = 10.0f;
+			const float H = 2.0f + 13.0f * NumLines;
+			DrawContext.DrawBox(CX - W2, Y, 2 * W2, H, WhiteBrush, FLinearColor(0.7, 0.7, 0.7, TooltipOpacity));
+			DrawContext.LayerId++;
+			DrawContext.DrawText(CX - W2 + DX, Y + 1.0f, Text, SummaryFont, FLinearColor(0.0, 0.0, 0.0, TooltipOpacity));
+			DrawContext.LayerId++;
+		}
+
+		Stopwatch.Stop();
+		DrawDurationHistory.AddValue(Stopwatch.AccumulatedTime);
+	}
+	//////////////////////////////////////////////////
+
+	const bool bShouldDisplayDebugInfo = FInsightsManager::Get()->IsDebugInfoEnabled();
+	if (bShouldDisplayDebugInfo)
+	{
+		const float MaxFontCharHeight = FontMeasureService->Measure(TEXT("!"), SummaryFont).Y;
+		const float DbgDY = MaxFontCharHeight;
+
+		const float DbgW = 280.0f;
+		const float DbgH = DbgDY * 5 + 3.0f;
+		const float DbgX = ViewWidth - DbgW - 20.0f;
+		float DbgY = 7.0f;
+
+		DrawContext.LayerId++;
+
+		DrawContext.DrawBox(DbgX - 2.0f, DbgY - 2.0f, DbgW, DbgH, WhiteBrush, FLinearColor(1.0, 1.0, 1.0, 0.9));
+		DrawContext.LayerId++;
+
+		FLinearColor DbgTextColor(0.0, 0.0, 0.0, 0.9);
+
+		// Time interval since last OnPaint call.
+		const uint64 CurrentTime = FPlatformTime::Cycles64();
+		const uint64 OnPaintDuration = CurrentTime - LastOnPaintTime;
+		LastOnPaintTime = CurrentTime;
+		OnPaintDurationHistory.AddValue(OnPaintDuration); // saved for last 32 OnPaint calls
+		const uint64 AvgOnPaintDuration = OnPaintDurationHistory.ComputeAverage();
+		const uint64 AvgOnPaintDurationMs = FStopwatch::Cycles64ToMilliseconds(AvgOnPaintDuration);
+		const double AvgOnPaintFps = AvgOnPaintDurationMs != 0 ? 1.0 / FStopwatch::Cycles64ToSeconds(AvgOnPaintDuration) : 0.0;
+
+		const uint64 AvgUpdateDurationMs = FStopwatch::Cycles64ToMilliseconds(UpdateDurationHistory.ComputeAverage());
+		const uint64 AvgDrawDurationMs = FStopwatch::Cycles64ToMilliseconds(DrawDurationHistory.ComputeAverage());
+
+		// Draw performance info.
+		DrawContext.DrawText
+		(
+			DbgX, DbgY,
+			FString::Printf(TEXT("U: %llu ms, D: %llu ms + %llu ms = %llu ms (%d fps)"),
+				AvgUpdateDurationMs, // caching time
+				AvgDrawDurationMs, // drawing time
+				AvgOnPaintDurationMs - AvgDrawDurationMs, // other overhead to OnPaint calls
+				AvgOnPaintDurationMs, // average time between two OnPaint calls
+				FMath::RoundToInt(AvgOnPaintFps)), // framerate of OnPaint calls
+			SummaryFont, DbgTextColor
+		);
+		DbgY += DbgDY;
+
+		// Draw number of draw calls.
+		DrawContext.DrawText
+		(
+			DbgX, DbgY,
+			FString::Printf(TEXT("U: %s packets, D: %s samples"),
+				*FText::AsNumber(NumUpdatedPackets).ToString(),
+				*FText::AsNumber(NumDrawSamples).ToString()),
+			SummaryFont, DbgTextColor
+		);
+		DbgY += DbgDY;
+
+		// Draw viewport's horizontal info.
+		DrawContext.DrawText
+		(
+			DbgX, DbgY,
+			Viewport.GetHorizontalAxisViewport().ToDebugString(TEXT("X"), TEXT("packet")),
+			SummaryFont, DbgTextColor
+		);
+		DbgY += DbgDY;
+
+		// Draw viewport's vertical info.
+		DrawContext.DrawText
+		(
+			DbgX, DbgY,
+			Viewport.GetVerticalAxisViewport().ToDebugString(TEXT("Y")),
+			SummaryFont, DbgTextColor
+		);
+		DbgY += DbgDY;
+
+		// Draw connection info.
+		DrawContext.DrawText
+		(
+			DbgX, DbgY,
+			FString::Printf(TEXT("Game Instance %d, Connection %d (%s)"),
+				GameInstanceIndex,
+				ConnectionIndex,
+				(ConnectionMode == Trace::ENetProfilerConnectionMode::Outgoing) ? TEXT("Outgoing") : TEXT("Incoming")),
+			SummaryFont, DbgTextColor
+		);
+		DbgY += DbgDY;
+	}
+
+	return SCompoundWidget::OnPaint(Args, AllottedGeometry, MyCullingRect, OutDrawElements, LayerId, InWidgetStyle, bParentEnabled && IsEnabled());
+}
+
+////////////////////////////////////////////////////////////////////////////////////////////////////
+
+void SPacketView::DrawVerticalAxisGrid(FDrawContext& DrawContext, const FSlateBrush* Brush, const FSlateFontInfo& Font) const
+{
+	const FAxisViewportDouble& ViewportY = Viewport.GetVerticalAxisViewport();
+
+	const float RoundedViewHeight = FMath::RoundToFloat(ViewportY.GetSize());
+
+	constexpr float MinDY = 32.0f; // min vertical distance between horizontal grid lines
+
+	const double TopValue = ViewportY.GetValueAtOffset(RoundedViewHeight);
+	const double GridValue = ViewportY.GetValueAtOffset(MinDY);
+	const double BottomValue = ViewportY.GetValueAtOffset(0.0f);
+	const double Delta = GridValue - BottomValue;
+
+	if (Delta > 0.0)
+	{
+		const int64 DeltaBits = static_cast<int64>(Delta);
+
+		// Compute rounding based on magnitude of visible range of values (Delta).
+		int64 Power10 = 1;
+		int64 Delta10 = DeltaBits;
+		while (Delta10 > 0)
+		{
+			Delta10 /= 10;
+			Power10 *= 10;
+		}
+		if (Power10 >= 100)
+		{
+			Power10 /= 100;
+		}
+		else
+		{
+			Power10 = 1;
+		}
+
+		const double Grid = static_cast<double>(((DeltaBits + Power10 - 1) / Power10) * Power10); // next value divisible with a multiple of 10
+
+		const double StartValue = FMath::GridSnap(BottomValue, Grid);
+
+		const float ViewWidth = Viewport.GetWidth();
+
+		const FLinearColor GridColor(0.0f, 0.0f, 0.0f, 0.1f);
+		const FLinearColor TextBgColor(0.05f, 0.05f, 0.05f, 1.0f);
+		const FLinearColor TextColor(1.0f, 1.0f, 1.0f, 1.0f);
+
+		const TSharedRef<FSlateFontMeasure> FontMeasureService = FSlateApplication::Get().GetRenderer()->GetFontMeasureService();
+
+		for (double Value = StartValue; Value < TopValue; Value += Grid)
+		{
+			const float Y = RoundedViewHeight - FMath::RoundToFloat(ViewportY.GetOffsetForValue(Value));
+
+			// Draw horizontal grid line.
+			DrawContext.DrawBox(0, Y, ViewWidth, 1, Brush, GridColor);
+
+			const int64 ValueBits = static_cast<int64>(Value);
+			const FString Text = (ValueBits == 0) ? TEXT("0") : FString::Format(TEXT("{0} bits"), { FText::AsNumber(ValueBits).ToString() });
+			const FVector2D TextSize = FontMeasureService->Measure(Text, Font);
+			constexpr float TextH = 14.0f;
+
+			// Draw background for value text.
+			DrawContext.DrawBox(ViewWidth - TextSize.X - 4.0f, Y - TextH, TextSize.X + 4.0f, TextH, Brush, TextBgColor);
+
+			// Draw value text.
+			DrawContext.DrawText(ViewWidth - TextSize.X - 2.0f, Y - TextH + 1.0f, Text, Font, TextColor);
+		}
+		DrawContext.LayerId++;
+	}
+}
+
+////////////////////////////////////////////////////////////////////////////////////////////////////
+
+void SPacketView::DrawHorizontalAxisGrid(FDrawContext& DrawContext, const FSlateBrush* Brush, const FSlateFontInfo& Font) const
+{
+	const FAxisViewportInt32& ViewportX = Viewport.GetHorizontalAxisViewport();
+
+	const float RoundedViewWidth = FMath::RoundToFloat(ViewportX.GetSize());
+
+	constexpr float MinDX = 100.0f; // min horizontal distance between vertical grid lines
+
+	const int32 LeftIndex = ViewportX.GetValueAtOffset(0.0f);
+	const int32 GridIndex = ViewportX.GetValueAtOffset(MinDX);
+	const int32 RightIndex = ViewportX.GetValueAtOffset(RoundedViewWidth);
+	const int32 Delta = GridIndex - LeftIndex;
+
+	if (Delta > 0)
+	{
+		// Compute rounding based on magnitude of visible range of samples (Delta).
+		int32 Power10 = 1;
+		int32 Delta10 = Delta;
+		while (Delta10 > 0)
+		{
+			Delta10 /= 10;
+			Power10 *= 10;
+		}
+		if (Power10 >= 100)
+		{
+			Power10 /= 100;
+		}
+		else
+		{
+			Power10 = 1;
+		}
+
+		const int32 Grid = ((Delta + Power10 - 1) / Power10) * Power10; // next value divisible with a multiple of 10
+
+		// Skip grid lines for negative indices.
+		double StartIndex = ((LeftIndex + Grid - 1) / Grid) * Grid;
+		while (StartIndex < 0)
+		{
+			StartIndex += Grid;
+		}
+
+		const float ViewHeight = Viewport.GetHeight();
+
+		const FLinearColor GridColor(0.0f, 0.0f, 0.0f, 0.1f);
+		//const FLinearColor TextBgColor(0.05f, 0.05f, 0.05f, 1.0f);
+		//const FLinearColor TextColor(1.0f, 1.0f, 1.0f, 1.0f);
+		const FLinearColor TopTextColor(1.0f, 1.0f, 1.0f, 0.7f);
+
+		//const TSharedRef<FSlateFontMeasure> FontMeasureService = FSlateApplication::Get().GetRenderer()->GetFontMeasureService();
+
+		for (int32 Index = StartIndex; Index < RightIndex; Index += Grid)
+		{
+			const float X = FMath::RoundToFloat(ViewportX.GetOffsetForValue(Index));
+
+			// Draw vertical grid line.
+			DrawContext.DrawBox(X, 0, 1, ViewHeight, Brush, GridColor);
+
+			const FString Text = FText::AsNumber(Index).ToString();
+			//const FVector2D TextSize = FontMeasureService->Measure(Text, Font);
+			//constexpr float TextH = 14.0f;
+
+			// Draw background for index text.
+			//DrawContext.DrawBox(X, ViewHeight - TextH, TextSize.X + 4.0f, TextH, Brush, TextBgColor);
+
+			// Draw index text.
+			//DrawContext.DrawText(X + 2.0f, ViewHeight - TextH + 1.0f, Text, Font, TextColor);
+
+			DrawContext.DrawText(X + 2.0f, 10.0f, Text, Font, TopTextColor);
+		}
+		DrawContext.LayerId++;
+	}
+}
+
+////////////////////////////////////////////////////////////////////////////////////////////////////
+
+FReply SPacketView::OnMouseButtonDown(const FGeometry& MyGeometry, const FPointerEvent& MouseEvent)
+{
+	FReply Reply = FReply::Unhandled();
+
+	MousePositionOnButtonDown = MyGeometry.AbsoluteToLocal(MouseEvent.GetScreenSpacePosition());
+	ViewportPosXOnButtonDown = Viewport.GetHorizontalAxisViewport().GetPos();
+
+	if (MouseEvent.GetEffectingButton() == EKeys::LeftMouseButton)
+	{
+		bIsLMB_Pressed = true;
+
+		// Capture mouse.
+		Reply = FReply::Handled().CaptureMouse(SharedThis(this));
+	}
+	else if (MouseEvent.GetEffectingButton() == EKeys::RightMouseButton)
+	{
+		bIsRMB_Pressed = true;
+
+		// Capture mouse, so we can scroll outside this widget.
+		Reply = FReply::Handled().CaptureMouse(SharedThis(this));
+	}
+
+	return Reply;
+}
+
+////////////////////////////////////////////////////////////////////////////////////////////////////
+
+FReply SPacketView::OnMouseButtonUp(const FGeometry& MyGeometry, const FPointerEvent& MouseEvent)
+{
+	FReply Reply = FReply::Unhandled();
+
+	MousePositionOnButtonUp = MyGeometry.AbsoluteToLocal(MouseEvent.GetScreenSpacePosition());
+
+	const bool bIsValidForMouseClick = MousePositionOnButtonUp.Equals(MousePositionOnButtonDown, MOUSE_SNAP_DISTANCE);
+
+	if (MouseEvent.GetEffectingButton() == EKeys::LeftMouseButton)
+	{
+		if (bIsLMB_Pressed)
+		{
+			if (bIsScrolling)
+			{
+				bIsScrolling = false;
+				CursorType = ECursorType::Default;
+			}
+			else if (bIsValidForMouseClick)
+			{
+				SelectSampleAtMousePosition(MousePositionOnButtonUp.X, MousePositionOnButtonUp.Y);
+			}
+
+			bIsLMB_Pressed = false;
+
+			// Release the mouse.
+			Reply = FReply::Handled().ReleaseMouseCapture();
+		}
+	}
+	else if (MouseEvent.GetEffectingButton() == EKeys::RightMouseButton)
+	{
+		if (bIsRMB_Pressed)
+		{
+			if (bIsScrolling)
+			{
+				bIsScrolling = false;
+				CursorType = ECursorType::Default;
+			}
+			else if (bIsValidForMouseClick)
+			{
+				ShowContextMenu(MouseEvent);
+			}
+
+			bIsRMB_Pressed = false;
+
+			// Release mouse as we no longer scroll.
+			Reply = FReply::Handled().ReleaseMouseCapture();
+		}
+	}
+
+	return Reply;
+}
+
+////////////////////////////////////////////////////////////////////////////////////////////////////
+
+FReply SPacketView::OnMouseMove(const FGeometry& MyGeometry, const FPointerEvent& MouseEvent)
+{
+	FReply Reply = FReply::Unhandled();
+
+	MousePosition = MyGeometry.AbsoluteToLocal(MouseEvent.GetScreenSpacePosition());
+
+	if (!MouseEvent.GetCursorDelta().IsZero())
+	{
+		if (MouseEvent.IsMouseButtonDown(EKeys::LeftMouseButton) ||
+			MouseEvent.IsMouseButtonDown(EKeys::RightMouseButton))
+		{
+			if (HasMouseCapture())
+			{
+				if (!bIsScrolling)
+				{
+					bIsScrolling = true;
+					CursorType = ECursorType::Hand;
+
+					HoveredSample.Reset();
+				}
+
+				FAxisViewportInt32& ViewportX = Viewport.GetHorizontalAxisViewport();
+				const float PosX = ViewportPosXOnButtonDown + (MousePositionOnButtonDown.X - MousePosition.X);
+				ViewportX.ScrollAtValue(ViewportX.GetValueAtPos(PosX)); // align viewport position with sample
+				UpdateHorizontalScrollBar();
+				bIsStateDirty = true;
+			}
+		}
+		else
+		{
+			if (!HoveredSample.IsValid())
+			{
+				TooltipOpacity = 0.0f;
+			}
+			HoveredSample = GetSampleAtMousePosition(MousePosition.X, MousePosition.Y);
+			if (!HoveredSample.IsValid())
+			{
+				HoveredSample = GetSampleAtMousePosition(MousePosition.X - 1.0f, MousePosition.Y);
+			}
+			if (!HoveredSample.IsValid())
+			{
+				HoveredSample = GetSampleAtMousePosition(MousePosition.X + 1.0f, MousePosition.Y);
+			}
+		}
+
+		Reply = FReply::Handled();
+	}
+
+	return Reply;
+}
+
+////////////////////////////////////////////////////////////////////////////////////////////////////
+
+void SPacketView::OnMouseEnter(const FGeometry& MyGeometry, const FPointerEvent& MouseEvent)
+{
+}
+
+////////////////////////////////////////////////////////////////////////////////////////////////////
+
+void SPacketView::OnMouseLeave(const FPointerEvent& MouseEvent)
+{
+	if (!HasMouseCapture())
+	{
+		bIsLMB_Pressed = false;
+		bIsRMB_Pressed = false;
+
+		HoveredSample.Reset();
+
+		CursorType = ECursorType::Default;
+	}
+}
+
+////////////////////////////////////////////////////////////////////////////////////////////////////
+
+FReply SPacketView::OnMouseWheel(const FGeometry& MyGeometry, const FPointerEvent& MouseEvent)
+{
+	MousePosition = MyGeometry.AbsoluteToLocal(MouseEvent.GetScreenSpacePosition());
+
+	if (MouseEvent.GetModifierKeys().IsShiftDown())
+	{
+		FAxisViewportDouble& ViewportY = Viewport.GetVerticalAxisViewport();
+
+		// Zoom in/out vertically.
+		const float Delta = MouseEvent.GetWheelDelta();
+		constexpr float ZoomStep = 0.25f; // as percent
+		float ScaleY;
+
+		if (Delta > 0)
+		{
+			ScaleY = ViewportY.GetScale() * FMath::Pow(1.0f + ZoomStep, Delta);
+		}
+		else
+		{
+			ScaleY = ViewportY.GetScale() * FMath::Pow(1.0f / (1.0f + ZoomStep), -Delta);
+		}
+
+		ViewportY.SetScale(ScaleY);
+		//UpdateVerticalScrollBar();
+	}
+	else //if (MouseEvent.GetModifierKeys().IsControlDown())
+	{
+		// Zoom in/out horizontally.
+		const float Delta = MouseEvent.GetWheelDelta();
+		ZoomHorizontally(Delta, MousePosition.X);
+	}
+
+	return FReply::Handled();
+}
+
+////////////////////////////////////////////////////////////////////////////////////////////////////
+
+FReply SPacketView::OnMouseButtonDoubleClick(const FGeometry& MyGeometry, const FPointerEvent& MouseEvent)
+{
+	return FReply::Unhandled();
+}
+
+////////////////////////////////////////////////////////////////////////////////////////////////////
+
+FCursorReply SPacketView::OnCursorQuery(const FGeometry& MyGeometry, const FPointerEvent& CursorEvent) const
+{
+	FCursorReply CursorReply = FCursorReply::Unhandled();
+
+	if (CursorType == ECursorType::Arrow)
+	{
+		CursorReply = FCursorReply::Cursor(EMouseCursor::ResizeLeftRight);
+	}
+	else if (CursorType == ECursorType::Hand)
+	{
+		CursorReply = FCursorReply::Cursor(EMouseCursor::GrabHand);
+	}
+
+	return CursorReply;
+}
+
+////////////////////////////////////////////////////////////////////////////////////////////////////
+
+void SPacketView::ShowContextMenu(const FPointerEvent& MouseEvent)
+{
+	const bool bShouldCloseWindowAfterMenuSelection = true;
+	FMenuBuilder MenuBuilder(bShouldCloseWindowAfterMenuSelection, NULL);
+
+	MenuBuilder.BeginSection("Misc");
+	{
+		FUIAction Action_AutoZoom
+		(
+			FExecuteAction::CreateSP(this, &SPacketView::ContextMenu_AutoZoom_Execute),
+			FCanExecuteAction::CreateSP(this, &SPacketView::ContextMenu_AutoZoom_CanExecute),
+			FIsActionChecked::CreateSP(this, &SPacketView::ContextMenu_AutoZoom_IsChecked)
+		);
+		MenuBuilder.AddMenuEntry
+		(
+			LOCTEXT("ContextMenu_AutoZoom", "Auto Zoom"),
+			LOCTEXT("ContextMenu_AutoZoom_Desc", "Enable auto zoom. Makes entire graph series to fit into view."),
+			FSlateIcon(),
+			Action_AutoZoom,
+			NAME_None,
+			EUserInterfaceActionType::ToggleButton
+		);
+	}
+	MenuBuilder.EndSection();
+
+	TSharedRef<SWidget> MenuWidget = MenuBuilder.MakeWidget();
+
+	FWidgetPath EventPath = MouseEvent.GetEventPath() != nullptr ? *MouseEvent.GetEventPath() : FWidgetPath();
+	const FVector2D ScreenSpacePosition = MouseEvent.GetScreenSpacePosition();
+	FSlateApplication::Get().PushMenu(SharedThis(this), EventPath, MenuWidget, ScreenSpacePosition, FPopupTransitionEffect::ContextMenu);
+}
+
+////////////////////////////////////////////////////////////////////////////////////////////////////
+
+void SPacketView::ContextMenu_AutoZoom_Execute()
+{
+	FAxisViewportInt32& ViewportX = Viewport.GetHorizontalAxisViewport();
+
+	bIsAutoZoomEnabled = !(bIsAutoZoomEnabled && ViewportX.GetPos() == 0.0f);
+	if (bIsAutoZoomEnabled)
+	{
+		ViewportX.ScrollAtPos(0.0f);
+
+		// Auto zoom in.
+		while (ViewportX.GetMaxPos() - ViewportX.GetMinPos() < ViewportX.GetSize())
+		{
+			ZoomHorizontally(+0.1f, 0.0f);
+			ViewportX.ScrollAtPos(0.0f);
+		}
+
+		// Auto zoom out (until entire session time range fits into view).
+		while (ViewportX.GetMaxPos() - ViewportX.GetMinPos() > ViewportX.GetSize())
+		{
+			ZoomHorizontally(-0.1f, 0.0f);
+			ViewportX.ScrollAtPos(0.0f);
+		}
+
+		UpdateHorizontalScrollBar();
+		bIsStateDirty = true;
+	}
+}
+
+////////////////////////////////////////////////////////////////////////////////////////////////////
+
+bool SPacketView::ContextMenu_AutoZoom_CanExecute()
+{
+	return true;
+}
+
+////////////////////////////////////////////////////////////////////////////////////////////////////
+
+bool SPacketView::ContextMenu_AutoZoom_IsChecked()
+{
+	return bIsAutoZoomEnabled && Viewport.GetHorizontalAxisViewport().GetPos() == 0.0f;
+}
+
+////////////////////////////////////////////////////////////////////////////////////////////////////
+
+void SPacketView::BindCommands()
+{
+}
+
+////////////////////////////////////////////////////////////////////////////////////////////////////
+
+void SPacketView::HorizontalScrollBar_OnUserScrolled(float ScrollOffset)
+{
+	FAxisViewportInt32& ViewportX = Viewport.GetHorizontalAxisViewport();
+	ViewportX.OnUserScrolled(HorizontalScrollBar, ScrollOffset);
+	bIsStateDirty = true;
+}
+
+////////////////////////////////////////////////////////////////////////////////////////////////////
+
+void SPacketView::UpdateHorizontalScrollBar()
+{
+	FAxisViewportInt32& ViewportX = Viewport.GetHorizontalAxisViewport();
+	ViewportX.UpdateScrollBar(HorizontalScrollBar);
+}
+
+////////////////////////////////////////////////////////////////////////////////////////////////////
+
+void SPacketView::ZoomHorizontally(const float Delta, const float X)
+{
+	FAxisViewportInt32& ViewportX = Viewport.GetHorizontalAxisViewport();
+	ViewportX.RelativeZoomWithFixedOffset(Delta, X);
+	ViewportX.ScrollAtValue(ViewportX.GetValueAtPos(ViewportX.GetPos())); // align viewport position with sample
+	UpdateHorizontalScrollBar();
+	bIsStateDirty = true;
+}
+
+////////////////////////////////////////////////////////////////////////////////////////////////////
+
+#undef LOCTEXT_NAMESPACE
diff --git a/Engine/Source/Developer/TraceInsights/Private/Insights/NetworkingProfiler/Widgets/SPacketView.h b/Engine/Source/Developer/TraceInsights/Private/Insights/NetworkingProfiler/Widgets/SPacketView.h
new file mode 100644
--- /dev/null
+++ b/Engine/Source/Developer/TraceInsights/Private/Insights/NetworkingProfiler/Widgets/SPacketView.h
@@ -0,0 +1,228 @@
+// Copyright 1998-2019 Epic Games, Inc. All Rights Reserved.
+
+#pragma once
+
+#include "CoreMinimal.h"
+#include "Input/CursorReply.h"
+#include "Input/Reply.h"
+#include "Layout/Geometry.h"
+#include "Rendering/RenderingCommon.h"
+#include "TraceServices/Model/NetProfiler.h"
+#include "Widgets/DeclarativeSyntaxSupport.h"
+#include "Widgets/SCompoundWidget.h"
+
+// Insights
+#include "Insights/Common/FixedCircularBuffer.h"
+#include "Insights/NetworkingProfiler/ViewModels/PacketViewDrawHelper.h"
+#include "Insights/NetworkingProfiler/ViewModels/PacketViewport.h"
+
+class SScrollBar;
+class SNetworkingProfilerWindow;
+
+////////////////////////////////////////////////////////////////////////////////////////////////////
+
+struct FNetworkPacketSampleRef
+{
+	TSharedPtr<FNetworkPacketSeries> Series;
+	TSharedPtr<FNetworkPacketAggregatedSample> Sample;
+
+	FNetworkPacketSampleRef()
+		: Series(), Sample()
+	{
+	}
+
+	FNetworkPacketSampleRef(TSharedPtr<FNetworkPacketSeries> InSeries, TSharedPtr<FNetworkPacketAggregatedSample> InSample)
+		: Series(InSeries), Sample(InSample)
+	{
+	}
+
+	FNetworkPacketSampleRef(const FNetworkPacketSampleRef& Other)
+		: Series(Other.Series), Sample(Other.Sample)
+	{
+	}
+
+	FNetworkPacketSampleRef& operator=(const FNetworkPacketSampleRef& Other)
+	{
+		Series = Other.Series;
+		Sample = Other.Sample;
+		return *this;
+	}
+
+	void Reset()
+	{
+		Series.Reset();
+		Sample.Reset();
+	}
+
+	bool IsValid() const
+	{
+		return Series.IsValid() && Sample.IsValid();
+	}
+
+	bool Equals(const FNetworkPacketSampleRef& Other) const
+	{
+		return Series == Other.Series
+			&& ((Sample == Other.Sample) || (Sample.IsValid() && Other.Sample.IsValid() && Sample->Equals(*Other.Sample)));
+	}
+
+	static bool AreEquals(const FNetworkPacketSampleRef& A, const FNetworkPacketSampleRef& B)
+	{
+		return A.Equals(B);
+	}
+};
+
+////////////////////////////////////////////////////////////////////////////////////////////////////
+/**
+ * Widget used to present the network packets as a bar track.
+ */
+class SPacketView : public SCompoundWidget
+{
+public:
+	/** Number of pixels. */
+	static constexpr float MOUSE_SNAP_DISTANCE = 2.0f;
+
+	enum class ECursorType
+	{
+		Default,
+		Arrow,
+		Hand,
+	};
+
+public:
+	/** Default constructor. */
+	SPacketView();
+
+	/** Virtual destructor. */
+	virtual ~SPacketView();
+
+	/** Resets internal widget's data to the default one. */
+	void Reset();
+
+	void SetConnection(uint32 GameInstanceIndex, uint32 ConnectionIndex, Trace::ENetProfilerConnectionMode ConnectionMode);
+
+	SLATE_BEGIN_ARGS(SPacketView)
+	{
+		_Clipping = EWidgetClipping::ClipToBounds;
+	}
+	SLATE_END_ARGS()
+
+	void Construct(const FArguments& InArgs, TSharedPtr<SNetworkingProfilerWindow> InProfilerWindow);
+
+	/**
+	 * Ticks this widget. Override in derived classes, but always call the parent implementation.
+	 * @param  AllottedGeometry The space allotted for this widget
+	 * @param  InCurrentTime  Current absolute real time
+	 * @param  InDeltaTime  Real time passed since last tick
+	 */
+	virtual void Tick(const FGeometry& AllottedGeometry, const double InCurrentTime, const float InDeltaTime) override;
+
+	virtual int32 OnPaint(const FPaintArgs& Args, const FGeometry& AllottedGeometry, const FSlateRect& MyCullingRect, FSlateWindowElementList& OutDrawElements, int32 LayerId, const FWidgetStyle& InWidgetStyle, bool bParentEnabled) const override;
+
+	virtual FReply OnMouseButtonDown(const FGeometry& MyGeometry, const FPointerEvent& MouseEvent) override;
+	virtual FReply OnMouseButtonUp(const FGeometry& MyGeometry, const FPointerEvent& MouseEvent) override;
+	virtual FReply OnMouseMove(const FGeometry& MyGeometry, const FPointerEvent& MouseEvent) override;
+	virtual void OnMouseEnter(const FGeometry& MyGeometry, const FPointerEvent& MouseEvent) override;
+	virtual void OnMouseLeave(const FPointerEvent& MouseEvent) override;
+	virtual FReply OnMouseWheel(const FGeometry& MyGeometry, const FPointerEvent& MouseEvent) override;
+	virtual FReply OnMouseButtonDoubleClick(const FGeometry& MyGeometry, const FPointerEvent& MouseEvent) override;
+
+	virtual FCursorReply OnCursorQuery(const FGeometry& MyGeometry, const FPointerEvent& CursorEvent) const override;
+
+private:
+	void UpdateState();
+
+	void DrawHorizontalAxisGrid(FDrawContext& DrawContext, const FSlateBrush* Brush, const FSlateFontInfo& Font) const;
+	void DrawVerticalAxisGrid(FDrawContext& DrawContext, const FSlateBrush* Brush, const FSlateFontInfo& Font) const;
+
+	FNetworkPacketSampleRef GetSampleAtMousePosition(float X, float Y);
+	void SelectSampleAtMousePosition(float X, float Y);
+	void OnSelectedSampleChanged();
+
+	void ShowContextMenu(const FPointerEvent& MouseEvent);
+
+	void ContextMenu_AutoZoom_Execute();
+	bool ContextMenu_AutoZoom_CanExecute();
+	bool ContextMenu_AutoZoom_IsChecked();
+
+	/** Binds our UI commands to delegates. */
+	void BindCommands();
+
+	/**
+	 * Called when the user scrolls the horizontal scrollbar.
+	 * @param ScrollOffset Scroll offset as a fraction between 0 and 1.
+	 */
+	void HorizontalScrollBar_OnUserScrolled(float ScrollOffset);
+	void UpdateHorizontalScrollBar();
+
+	void ZoomHorizontally(const float Delta, const float X);
+
+private:
+	TSharedPtr<SNetworkingProfilerWindow> ProfilerWindow;
+
+	uint32 GameInstanceIndex;
+	uint32 ConnectionIndex;
+	Trace::ENetProfilerConnectionMode ConnectionMode;
+
+	/** The track's viewport. Encapsulates info about position and scale. */
+	FPacketViewport Viewport;
+	bool bIsViewportDirty;
+
+	/** Cached info for the packet series. */
+	TSharedPtr<FNetworkPacketSeries> PacketSeries;
+	bool bIsStateDirty;
+
+	bool bIsAutoZoomEnabled;
+
+	uint64 AnalysisSyncNextTimestamp;
+	uint32 ConnectionChangeCount;
+
+	//////////////////////////////////////////////////
+
+	TSharedPtr<SScrollBar> HorizontalScrollBar;
+
+	//////////////////////////////////////////////////
+	// Panning and Zooming behaviors
+
+	/** The current mouse position. */
+	FVector2D MousePosition;
+
+	/** Mouse position during the call on mouse button down. */
+	FVector2D MousePositionOnButtonDown;
+	float ViewportPosXOnButtonDown;
+
+	/** Mouse position during the call on mouse button up. */
+	FVector2D MousePositionOnButtonUp;
+
+	bool bIsLMB_Pressed;
+	bool bIsRMB_Pressed;
+
+	/** True, if the user is currently interactively scrolling the view (ex.: by holding the left mouse button and dragging). */
+	bool bIsScrolling;
+
+	//////////////////////////////////////////////////
+	// Selection
+
+	int32 SelectionStartFrameIndex;
+	int32 SelectionEndFrameIndex;
+
+	FNetworkPacketSampleRef SelectedSample;
+	FNetworkPacketSampleRef HoveredSample;
+
+	float TooltipDesiredOpacity;
+	mutable float TooltipOpacity;
+
+	//////////////////////////////////////////////////
+	// Misc
+
+	FGeometry ThisGeometry;
+
+	/** Cursor type. */
+	ECursorType CursorType;
+
+	// Debug stats
+	int32 NumUpdatedPackets;
+	TFixedCircularBuffer<uint64, 32> UpdateDurationHistory;
+	mutable TFixedCircularBuffer<uint64, 32> DrawDurationHistory;
+	mutable TFixedCircularBuffer<uint64, 32> OnPaintDurationHistory;
+	mutable uint64 LastOnPaintTime;
+};
