[Insights]
  - Timers view: Added "!0" toggle button to filter out the timers having zero total instance count (aggregated stats).
  - Counters view: Added "!0" toggle button to filter out the stats counters having zero total instance count (aggregated stats).
  - Networking / Stats view: Added "!0" toggle button to filter out the net event types having zero total instance count (aggregated stats).

diff --git a/Engine/Source/Developer/TraceInsights/Private/Insights/NetworkingProfiler/Widgets/SNetStatsView.cpp b/Engine/Source/Developer/TraceInsights/Private/Insights/NetworkingProfiler/Widgets/SNetStatsView.cpp
--- a/Engine/Source/Developer/TraceInsights/Private/Insights/NetworkingProfiler/Widgets/SNetStatsView.cpp
+++ b/Engine/Source/Developer/TraceInsights/Private/Insights/NetworkingProfiler/Widgets/SNetStatsView.cpp
@@ -32,6 +32,7 @@
 SNetStatsView::SNetStatsView()
 	: Table(MakeShareable(new Insights::FTable()))
 	, bExpansionSaved(false)
+	, bFilterOutZeroCountEvents(false)
 	, GroupingMode(ENetEventGroupingMode::Flat)
 	, AvailableSorters()
 	, CurrentSorter(nullptr)
@@ -73,7 +74,7 @@ void SNetStatsView::Construct(const FArguments& InArgs)
 	[
 		SNew(SVerticalBox)
 
-		+SVerticalBox::Slot()
+		+ SVerticalBox::Slot()
 		.VAlign(VAlign_Center)
 		.AutoHeight()
 		[
@@ -83,28 +84,57 @@ void SNetStatsView::Construct(const FArguments& InArgs)
 			[
 				SNew(SVerticalBox)
 
-				// Search box
-				+SVerticalBox::Slot()
+				+ SVerticalBox::Slot()
 				.VAlign(VAlign_Center)
 				.Padding(2.0f)
 				.AutoHeight()
 				[
-					SAssignNew(SearchBox, SSearchBox)
-					.HintText(LOCTEXT("SearchBoxHint", "Search net events or groups"))
-					.OnTextChanged(this, &SNetStatsView::SearchBox_OnTextChanged)
-					.IsEnabled(this, &SNetStatsView::SearchBox_IsEnabled)
-					.ToolTipText(LOCTEXT("FilterSearchHint", "Type here to search net events or groups"))
+					SNew(SHorizontalBox)
+
+					// Search box
+					+ SHorizontalBox::Slot()
+					.VAlign(VAlign_Center)
+					.Padding(2.0f)
+					.FillWidth(1.0f)
+					[
+						SAssignNew(SearchBox, SSearchBox)
+						.HintText(LOCTEXT("SearchBoxHint", "Search net events or groups"))
+						.OnTextChanged(this, &SNetStatsView::SearchBox_OnTextChanged)
+						.IsEnabled(this, &SNetStatsView::SearchBox_IsEnabled)
+						.ToolTipText(LOCTEXT("FilterSearchHint", "Type here to search net events or groups"))
+					]
+
+					// Filter out net event types with zero instance count
+					+ SHorizontalBox::Slot()
+					.VAlign(VAlign_Center)
+					.Padding(2.0f)
+					.AutoWidth()
+					[
+						SNew(SCheckBox)
+						.Style(FEditorStyle::Get(), "ToggleButtonCheckbox")
+						.HAlign(HAlign_Center)
+						.Padding(2.0f)
+						.OnCheckStateChanged(this, &SNetStatsView::FilterOutZeroCountEvents_OnCheckStateChanged)
+						.IsChecked(this, &SNetStatsView::FilterOutZeroCountEvents_IsChecked)
+						.ToolTipText(LOCTEXT("FilterOutZeroCountEvents_Tooltip", "Filter out the net event types having zero total instance count (aggregated stats)."))
+						[
+							//TODO: SNew(SImage)
+							SNew(STextBlock)
+							.Text(LOCTEXT("FilterOutZeroCountEvents_Button", " !0 "))
+							.TextStyle(FEditorStyle::Get(), TEXT("Profiler.Caption"))
+						]
+					]
 				]
 
 				// Group by
-				+SVerticalBox::Slot()
+				+ SVerticalBox::Slot()
 				.VAlign(VAlign_Center)
 				.Padding(2.0f)
 				.AutoHeight()
 				[
 					SNew(SHorizontalBox)
 
-					+SHorizontalBox::Slot()
+					+ SHorizontalBox::Slot()
 					.FillWidth(1.0f)
 					.VAlign(VAlign_Center)
 					[
@@ -112,7 +142,7 @@ void SNetStatsView::Construct(const FArguments& InArgs)
 						.Text(LOCTEXT("GroupByText", "Group by"))
 					]
 
-					+SHorizontalBox::Slot()
+					+ SHorizontalBox::Slot()
 					.FillWidth(2.0f)
 					.VAlign(VAlign_Center)
 					[
@@ -137,21 +167,21 @@ void SNetStatsView::Construct(const FArguments& InArgs)
 				[
 					SNew(SHorizontalBox)
 
-					+SHorizontalBox::Slot()
+					+ SHorizontalBox::Slot()
 					.Padding(FMargin(0.0f,0.0f,1.0f,0.0f))
 					.FillWidth(1.0f)
 					[
 						GetToggleButtonForNetEventType(ENetEventNodeType::NetEvent)
 					]
 
-					//+SHorizontalBox::Slot()
+					//+ SHorizontalBox::Slot()
 					//.Padding(FMargin(1.0f,0.0f,1.0f,0.0f))
 					//.FillWidth(1.0f)
 					//[
 					//	GetToggleButtonForNetEventType(ENetEventNodeType::NetEvent1)
 					//]
 
-					//+SHorizontalBox::Slot()
+					//+ SHorizontalBox::Slot()
 					.//Padding(FMargin(1.0f,0.0f,1.0f,0.0f))
 					//.FillWidth(1.0f)
 					//[
@@ -163,7 +193,7 @@ void SNetStatsView::Construct(const FArguments& InArgs)
 		]
 
 		// Tree view
-		+SVerticalBox::Slot()
+		+ SVerticalBox::Slot()
 		.FillHeight(1.0f)
 		.Padding(0.0f, 6.0f, 0.0f, 0.0f)
 		[
@@ -644,7 +674,9 @@ void SNetStatsView::ApplyFiltering()
 		{
 			// Add a child.
 			const FNetEventNodePtr& NodePtr = StaticCastSharedPtr<FNetEventNode, Insights::FBaseTreeNode>(GroupChildren[Cx]);
-			const bool bIsChildVisible = Filters->PassesAllFilters(NodePtr) && bNetEventTypeIsVisible[static_cast<int>(NodePtr->GetType())];
+			const bool bIsChildVisible = (!bFilterOutZeroCountEvents || NodePtr->GetAggregatedStats().InstanceCount > 0)
+									  && bNetEventTypeIsVisible[static_cast<int>(NodePtr->GetType())]
+									  && Filters->PassesAllFilters(NodePtr);
 			if (bIsChildVisible)
 			{
 				GroupPtr->AddFilteredChild(NodePtr);
@@ -728,7 +760,7 @@ TSharedRef<SWidget> SNetStatsView::GetToggleButtonForNetEventType(const ENetEven
 						.Image(NetEventNodeTypeHelper::GetIconForNetEventNodeType(NodeType))
 				]
 
-			+SHorizontalBox::Slot()
+			+ SHorizontalBox::Slot()
 				.Padding(2.0f, 0.0f, 0.0f, 0.0f)
 				.VAlign(VAlign_Center)
 				[
@@ -741,6 +773,21 @@ TSharedRef<SWidget> SNetStatsView::GetToggleButtonForNetEventType(const ENetEven
 
 ////////////////////////////////////////////////////////////////////////////////////////////////////
 
+void SNetStatsView::FilterOutZeroCountEvents_OnCheckStateChanged(ECheckBoxState NewRadioState)
+{
+	bFilterOutZeroCountEvents = (NewRadioState == ECheckBoxState::Checked);
+	ApplyFiltering();
+}
+
+////////////////////////////////////////////////////////////////////////////////////////////////////
+
+ECheckBoxState SNetStatsView::FilterOutZeroCountEvents_IsChecked() const
+{
+	return bFilterOutZeroCountEvents ? ECheckBoxState::Checked : ECheckBoxState::Unchecked;
+}
+
+////////////////////////////////////////////////////////////////////////////////////////////////////
+
 void SNetStatsView::FilterByNetEventType_OnCheckStateChanged(ECheckBoxState NewRadioState, const ENetEventNodeType InStatType)
 {
 	bNetEventTypeIsVisible[static_cast<int>(InStatType)] = (NewRadioState == ECheckBoxState::Checked);
diff --git a/Engine/Source/Developer/TraceInsights/Private/Insights/NetworkingProfiler/Widgets/SNetStatsView.h b/Engine/Source/Developer/TraceInsights/Private/Insights/NetworkingProfiler/Widgets/SNetStatsView.h
--- a/Engine/Source/Developer/TraceInsights/Private/Insights/NetworkingProfiler/Widgets/SNetStatsView.h
+++ b/Engine/Source/Developer/TraceInsights/Private/Insights/NetworkingProfiler/Widgets/SNetStatsView.h
@@ -162,6 +162,9 @@ protected:
 	/** Populates the group and stat tree with items based on the current data. */
 	void ApplyFiltering();
 
+	void FilterOutZeroCountEvents_OnCheckStateChanged(ECheckBoxState NewRadioState);
+	ECheckBoxState FilterOutZeroCountEvents_IsChecked() const;
+
 	TSharedRef<SWidget> GetToggleButtonForNetEventType(const ENetEventNodeType InNetEventType);
 	void FilterByNetEventType_OnCheckStateChanged(ECheckBoxState NewRadioState, const ENetEventNodeType InNetEventType);
 	ECheckBoxState FilterByNetEventType_IsChecked(const ENetEventNodeType InNetEventType) const;
@@ -316,6 +319,9 @@ protected:
 	/** Holds the visibility of each net event type. */
 	bool bNetEventTypeIsVisible[static_cast<int>(ENetEventNodeType::InvalidOrMax)];
 
+	/** Filter out the net event types having zero total instance count (aggregated stats). */
+	bool bFilterOutZeroCountEvents;
+
 	//////////////////////////////////////////////////
 	// Grouping
 
diff --git a/Engine/Source/Developer/TraceInsights/Private/Insights/Widgets/SStatsView.cpp b/Engine/Source/Developer/TraceInsights/Private/Insights/Widgets/SStatsView.cpp
--- a/Engine/Source/Developer/TraceInsights/Private/Insights/Widgets/SStatsView.cpp
+++ b/Engine/Source/Developer/TraceInsights/Private/Insights/Widgets/SStatsView.cpp
@@ -46,6 +46,7 @@ const FName SStatsView::GetDefaultColumnBeingSorted()
 
 SStatsView::SStatsView()
 	: bExpansionSaved(false)
+	, bFilterOutZeroCountStats(false)
 	, GroupingMode(EStatsGroupingMode::Flat)
 	, ColumnSortMode(GetDefaultColumnSortMode())
 	, ColumnBeingSorted(GetDefaultColumnBeingSorted())
@@ -76,7 +77,7 @@ void SStatsView::Construct(const FArguments& InArgs)
 	[
 		SNew(SVerticalBox)
 
-		+SVerticalBox::Slot()
+		+ SVerticalBox::Slot()
 		.VAlign(VAlign_Center)
 		.AutoHeight()
 		[
@@ -86,28 +87,57 @@ void SStatsView::Construct(const FArguments& InArgs)
 			[
 				SNew(SVerticalBox)
 
-				// Search box
-				+SVerticalBox::Slot()
+				+ SVerticalBox::Slot()
 				.VAlign(VAlign_Center)
 				.Padding(2.0f)
 				.AutoHeight()
 				[
-					SAssignNew(SearchBox, SSearchBox)
-					.HintText(LOCTEXT("SearchBoxHint", "Search stats or groups"))
-					.OnTextChanged(this, &SStatsView::SearchBox_OnTextChanged)
-					.IsEnabled(this, &SStatsView::SearchBox_IsEnabled)
-					.ToolTipText(LOCTEXT("FilterSearchHint", "Type here to search stats or group"))
+					SNew(SHorizontalBox)
+
+					// Search box
+					+ SHorizontalBox::Slot()
+					.VAlign(VAlign_Center)
+					.Padding(2.0f)
+					.FillWidth(1.0f)
+					[
+						SAssignNew(SearchBox, SSearchBox)
+						.HintText(LOCTEXT("SearchBoxHint", "Search stats counters or groups"))
+						.OnTextChanged(this, &SStatsView::SearchBox_OnTextChanged)
+						.IsEnabled(this, &SStatsView::SearchBox_IsEnabled)
+						.ToolTipText(LOCTEXT("FilterSearchHint", "Type here to search stats counter or group"))
+					]
+
+					// Filter out timers with zero instance count
+					+ SHorizontalBox::Slot()
+					.VAlign(VAlign_Center)
+					.Padding(2.0f)
+					.AutoWidth()
+					[
+						SNew(SCheckBox)
+						.Style(FEditorStyle::Get(), "ToggleButtonCheckbox")
+						.HAlign(HAlign_Center)
+						.Padding(2.0f)
+						.OnCheckStateChanged(this, &SStatsView::FilterOutZeroCountStats_OnCheckStateChanged)
+						.IsChecked(this, &SStatsView::FilterOutZeroCountStats_IsChecked)
+						.ToolTipText(LOCTEXT("FilterOutZeroCountStats_Tooltip", "Filter out the stats counters having zero total instance count (aggregated stats)."))
+						[
+							//TODO: SNew(SImage)
+							SNew(STextBlock)
+							.Text(LOCTEXT("FilterOutZeroCountStats_Button", " !0 "))
+							.TextStyle(FEditorStyle::Get(), TEXT("Profiler.Caption"))
+						]
+					]
 				]
 
 				// Group by
-				+SVerticalBox::Slot()
+				+ SVerticalBox::Slot()
 				.VAlign(VAlign_Center)
 				.Padding(2.0f)
 				.AutoHeight()
 				[
 					SNew(SHorizontalBox)
 
-					+SHorizontalBox::Slot()
+					+ SHorizontalBox::Slot()
 					.FillWidth(1.0f)
 					.VAlign(VAlign_Center)
 					[
@@ -115,7 +145,7 @@ void SStatsView::Construct(const FArguments& InArgs)
 						.Text(LOCTEXT("GroupByText", "Group by"))
 					]
 
-					+SHorizontalBox::Slot()
+					+ SHorizontalBox::Slot()
 					.FillWidth(2.0f)
 					.VAlign(VAlign_Center)
 					[
@@ -139,14 +169,14 @@ void SStatsView::Construct(const FArguments& InArgs)
 				[
 					SNew(SHorizontalBox)
 
-					+SHorizontalBox::Slot()
+					+ SHorizontalBox::Slot()
 					.Padding(FMargin(0.0f,0.0f,1.0f,0.0f))
 					.FillWidth(1.0f)
 					[
 						GetToggleButtonForStatsType(EStatsNodeType::Int64)
 					]
 
-					+SHorizontalBox::Slot()
+					+ SHorizontalBox::Slot()
 					.Padding(FMargin(1.0f,0.0f,1.0f,0.0f))
 					.FillWidth(1.0f)
 					[
@@ -157,7 +187,7 @@ void SStatsView::Construct(const FArguments& InArgs)
 		]
 
 		// Tree view
-		+SVerticalBox::Slot()
+		+ SVerticalBox::Slot()
 		.FillHeight(1.0f)
 		.Padding(0.0f, 6.0f, 0.0f, 0.0f)
 		[
@@ -666,7 +696,9 @@ void SStatsView::ApplyFiltering()
 		{
 			// Add a child.
 			const FStatsNodePtr& NodePtr = StaticCastSharedPtr<FStatsNode, Insights::FBaseTreeNode>(GroupChildren[Cx]);
-			const bool bIsChildVisible = Filters->PassesAllFilters(NodePtr) && bStatsNodeIsVisible[static_cast<int>(NodePtr->GetType())];
+			const bool bIsChildVisible = (!bFilterOutZeroCountStats || NodePtr->GetAggregatedStats().Count > 0)
+									  && bStatsNodeIsVisible[static_cast<int>(NodePtr->GetType())]
+									  && Filters->PassesAllFilters(NodePtr);
 			if (bIsChildVisible)
 			{
 				GroupPtr->AddFilteredChild(NodePtr);
@@ -750,7 +782,7 @@ TSharedRef<SWidget> SStatsView::GetToggleButtonForStatsType(const EStatsNodeType
 						.Image(StatsNodeTypeHelper::GetIconForStatsNodeType(NodeType))
 				]
 
-			+SHorizontalBox::Slot()
+			+ SHorizontalBox::Slot()
 				.Padding(2.0f, 0.0f, 0.0f, 0.0f)
 				.VAlign(VAlign_Center)
 				[
@@ -763,6 +795,21 @@ TSharedRef<SWidget> SStatsView::GetToggleButtonForStatsType(const EStatsNodeType
 
 ////////////////////////////////////////////////////////////////////////////////////////////////////
 
+void SStatsView::FilterOutZeroCountStats_OnCheckStateChanged(ECheckBoxState NewRadioState)
+{
+	bFilterOutZeroCountStats = (NewRadioState == ECheckBoxState::Checked);
+	ApplyFiltering();
+}
+
+////////////////////////////////////////////////////////////////////////////////////////////////////
+
+ECheckBoxState SStatsView::FilterOutZeroCountStats_IsChecked() const
+{
+	return bFilterOutZeroCountStats ? ECheckBoxState::Checked : ECheckBoxState::Unchecked;
+}
+
+////////////////////////////////////////////////////////////////////////////////////////////////////
+
 void SStatsView::FilterByStatsType_OnCheckStateChanged(ECheckBoxState NewRadioState, const EStatsNodeType InStatType)
 {
 	bStatsNodeIsVisible[static_cast<int>(InStatType)] = (NewRadioState == ECheckBoxState::Checked);
diff --git a/Engine/Source/Developer/TraceInsights/Private/Insights/Widgets/SStatsView.h b/Engine/Source/Developer/TraceInsights/Private/Insights/Widgets/SStatsView.h
--- a/Engine/Source/Developer/TraceInsights/Private/Insights/Widgets/SStatsView.h
+++ b/Engine/Source/Developer/TraceInsights/Private/Insights/Widgets/SStatsView.h
@@ -163,6 +163,9 @@ protected:
 	/** Populates the group and stat tree with items based on the current data. */
 	void ApplyFiltering();
 
+	void FilterOutZeroCountStats_OnCheckStateChanged(ECheckBoxState NewRadioState);
+	ECheckBoxState FilterOutZeroCountStats_IsChecked() const;
+
 	TSharedRef<SWidget> GetToggleButtonForStatsType(const EStatsNodeType StatType);
 	void FilterByStatsType_OnCheckStateChanged(ECheckBoxState NewRadioState, const EStatsNodeType InStatType);
 	ECheckBoxState FilterByStatsType_IsChecked(const EStatsNodeType InStatType) const;
@@ -317,6 +320,9 @@ protected:
 	/** Holds the visibility of each stats type. */
 	bool bStatsNodeIsVisible[static_cast<int>(EStatsNodeType::InvalidOrMax)];
 
+	/** Filter out the stats counters having zero total instance count (aggregated stats). */
+	bool bFilterOutZeroCountStats;
+
 	//////////////////////////////////////////////////
 	// Grouping
 
diff --git a/Engine/Source/Developer/TraceInsights/Private/Insights/Widgets/STimersView.cpp b/Engine/Source/Developer/TraceInsights/Private/Insights/Widgets/STimersView.cpp
--- a/Engine/Source/Developer/TraceInsights/Private/Insights/Widgets/STimersView.cpp
+++ b/Engine/Source/Developer/TraceInsights/Private/Insights/Widgets/STimersView.cpp
@@ -35,6 +35,7 @@
 STimersView::STimersView()
 	: Table(MakeShareable(new Insights::FTable()))
 	, bExpansionSaved(false)
+	, bFilterOutZeroCountTimers(false)
 	, GroupingMode(ETimerGroupingMode::ByType)
 	, AvailableSorters()
 	, CurrentSorter(nullptr)
@@ -69,7 +70,7 @@ void STimersView::Construct(const FArguments& InArgs)
 	[
 		SNew(SVerticalBox)
 
-		+SVerticalBox::Slot()
+		+ SVerticalBox::Slot()
 		.VAlign(VAlign_Center)
 		.AutoHeight()
 		[
@@ -79,28 +80,57 @@ void STimersView::Construct(const FArguments& InArgs)
 			[
 				SNew(SVerticalBox)
 
-				// Search box
-				+SVerticalBox::Slot()
+				+ SVerticalBox::Slot()
 				.VAlign(VAlign_Center)
 				.Padding(2.0f)
 				.AutoHeight()
 				[
-					SAssignNew(SearchBox, SSearchBox)
-					.HintText(LOCTEXT("SearchBoxHint", "Search timers or groups"))
-					.OnTextChanged(this, &STimersView::SearchBox_OnTextChanged)
-					.IsEnabled(this, &STimersView::SearchBox_IsEnabled)
-					.ToolTipText(LOCTEXT("FilterSearchHint", "Type here to search timer or group"))
+					SNew(SHorizontalBox)
+
+					// Search box
+					+ SHorizontalBox::Slot()
+					.VAlign(VAlign_Center)
+					.Padding(2.0f)
+					.FillWidth(1.0f)
+					[
+						SAssignNew(SearchBox, SSearchBox)
+						.HintText(LOCTEXT("SearchBoxHint", "Search timers or groups"))
+						.OnTextChanged(this, &STimersView::SearchBox_OnTextChanged)
+						.IsEnabled(this, &STimersView::SearchBox_IsEnabled)
+						.ToolTipText(LOCTEXT("FilterSearchHint", "Type here to search timer or group"))
+					]
+
+					// Filter out timers with zero instance count
+					+ SHorizontalBox::Slot()
+					.VAlign(VAlign_Center)
+					.Padding(2.0f)
+					.AutoWidth()
+					[
+						SNew(SCheckBox)
+						.Style(FEditorStyle::Get(), "ToggleButtonCheckbox")
+						.HAlign(HAlign_Center)
+						.Padding(2.0f)
+						.OnCheckStateChanged(this, &STimersView::FilterOutZeroCountTimers_OnCheckStateChanged)
+						.IsChecked(this, &STimersView::FilterOutZeroCountTimers_IsChecked)
+						.ToolTipText(LOCTEXT("FilterOutZeroCountTimers_Tooltip", "Filter out the timers having zero total instance count (aggregated stats)."))
+						[
+							//TODO: SNew(SImage)
+							SNew(STextBlock)
+							.Text(LOCTEXT("FilterOutZeroCountTimers_Button", " !0 "))
+							.TextStyle(FEditorStyle::Get(), TEXT("Profiler.Caption"))
+						]
+					]
 				]
 
 				// Group by
-				+SVerticalBox::Slot()
+				+ SVerticalBox::Slot()
 				.VAlign(VAlign_Center)
 				.Padding(2.0f)
 				.AutoHeight()
 				[
 					SNew(SHorizontalBox)
 
-					+SHorizontalBox::Slot()
+					+ SHorizontalBox::Slot()
 					.FillWidth(1.0f)
 					.VAlign(VAlign_Center)
 					[
@@ -108,7 +138,7 @@ void STimersView::Construct(const FArguments& InArgs)
 						.Text(LOCTEXT("GroupByText", "Group by"))
 					]
 
-					+SHorizontalBox::Slot()
+					+ SHorizontalBox::Slot()
 					.FillWidth(2.0f)
 					.VAlign(VAlign_Center)
 					[
@@ -132,21 +162,21 @@ void STimersView::Construct(const FArguments& InArgs)
 				[
 					SNew(SHorizontalBox)
 
-					+SHorizontalBox::Slot()
+					+ SHorizontalBox::Slot()
 					.Padding(FMargin(0.0f,0.0f,1.0f,0.0f))
 					.FillWidth(1.0f)
 					[
 						GetToggleButtonForTimerType(ETimerNodeType::GpuScope)
 					]
 
-					//+SHorizontalBox::Slot()
+					//+ SHorizontalBox::Slot()
 					//.Padding(FMargin(1.0f,0.0f,1.0f,0.0f))
 					//.FillWidth(1.0f)
 					//[
 					//	GetToggleButtonForTimerType(ETimerNodeType::ComputeScope)
 					//]
 
-					+SHorizontalBox::Slot()
+					+ SHorizontalBox::Slot()
 					.Padding(FMargin(1.0f,0.0f,1.0f,0.0f))
 					.FillWidth(1.0f)
 					[
@@ -157,7 +187,7 @@ void STimersView::Construct(const FArguments& InArgs)
 		]
 
 		// Tree view
-		+SVerticalBox::Slot()
+		+ SVerticalBox::Slot()
 		.FillHeight(1.0f)
 		.Padding(0.0f, 6.0f, 0.0f, 0.0f)
 		[
@@ -612,7 +642,9 @@ void STimersView::ApplyFiltering()
 		{
 			// Add a child.
 			const FTimerNodePtr& NodePtr = StaticCastSharedPtr<FTimerNode, Insights::FBaseTreeNode>(GroupChildren[Cx]);
-			const bool bIsChildVisible = Filters->PassesAllFilters(NodePtr) && bTimerTypeIsVisible[static_cast<int>(NodePtr->GetType())];
+			const bool bIsChildVisible = (!bFilterOutZeroCountTimers || NodePtr->GetAggregatedStats().InstanceCount > 0)
+									  && bTimerTypeIsVisible[static_cast<int>(NodePtr->GetType())]
+									  && Filters->PassesAllFilters(NodePtr);
 			if (bIsChildVisible)
 			{
 				GroupPtr->AddFilteredChild(NodePtr);
@@ -696,7 +728,7 @@ TSharedRef<SWidget> STimersView::GetToggleButtonForTimerType(const ETimerNodeTyp
 						.Image(TimerNodeTypeHelper::GetIconForTimerNodeType(NodeType))
 				]
 
-			+SHorizontalBox::Slot()
+			+ SHorizontalBox::Slot()
 				.Padding(2.0f, 0.0f, 0.0f, 0.0f)
 				.VAlign(VAlign_Center)
 				[
@@ -709,6 +741,21 @@ TSharedRef<SWidget> STimersView::GetToggleButtonForTimerType(const ETimerNodeTyp
 
 ////////////////////////////////////////////////////////////////////////////////////////////////////
 
+void STimersView::FilterOutZeroCountTimers_OnCheckStateChanged(ECheckBoxState NewRadioState)
+{
+	bFilterOutZeroCountTimers = (NewRadioState == ECheckBoxState::Checked);
+	ApplyFiltering();
+}
+
+////////////////////////////////////////////////////////////////////////////////////////////////////
+
+ECheckBoxState STimersView::FilterOutZeroCountTimers_IsChecked() const
+{
+	return bFilterOutZeroCountTimers ? ECheckBoxState::Checked : ECheckBoxState::Unchecked;
+}
+
+////////////////////////////////////////////////////////////////////////////////////////////////////
+
 void STimersView::FilterByTimerType_OnCheckStateChanged(ECheckBoxState NewRadioState, const ETimerNodeType InStatType)
 {
 	bTimerTypeIsVisible[static_cast<int>(InStatType)] = (NewRadioState == ECheckBoxState::Checked);
diff --git a/Engine/Source/Developer/TraceInsights/Private/Insights/Widgets/STimersView.h b/Engine/Source/Developer/TraceInsights/Private/Insights/Widgets/STimersView.h
--- a/Engine/Source/Developer/TraceInsights/Private/Insights/Widgets/STimersView.h
+++ b/Engine/Source/Developer/TraceInsights/Private/Insights/Widgets/STimersView.h
@@ -149,6 +149,9 @@ protected:
 	/** Populates the group and stat tree with items based on the current data. */
 	void ApplyFiltering();
 
+	void FilterOutZeroCountTimers_OnCheckStateChanged(ECheckBoxState NewRadioState);
+	ECheckBoxState FilterOutZeroCountTimers_IsChecked() const;
+
 	TSharedRef<SWidget> GetToggleButtonForTimerType(const ETimerNodeType InTimerType);
 	void FilterByTimerType_OnCheckStateChanged(ECheckBoxState NewRadioState, const ETimerNodeType InTimerType);
 	ECheckBoxState FilterByTimerType_IsChecked(const ETimerNodeType InTimerType) const;
@@ -303,6 +306,9 @@ protected:
 	/** Holds the visibility of each timer type. */
 	bool bTimerTypeIsVisible[static_cast<int>(ETimerNodeType::InvalidOrMax)];
 
+	/** Filter out the timers having zero total instance count (aggregated stats). */
+	bool bFilterOutZeroCountTimers;
+
 	//////////////////////////////////////////////////
 	// Grouping
 
