[Insights]
  - Networking / Packet view:  Added support for multiple selection. Click on a packet selects only respective packet. Shift+Click on a different packet selects a range of packets. Aggregation stats are computed for entire selection.
  - Networking / Packet view:  Added Ctrl+A shortcut to select all packets.

diff --git a/Engine/Source/Developer/TraceInsights/Private/Insights/NetworkingProfiler/ViewModels/PacketViewDrawHelper.cpp b/Engine/Source/Developer/TraceInsights/Private/Insights/NetworkingProfiler/ViewModels/PacketViewDrawHelper.cpp
--- a/Engine/Source/Developer/TraceInsights/Private/Insights/NetworkingProfiler/ViewModels/PacketViewDrawHelper.cpp
+++ b/Engine/Source/Developer/TraceInsights/Private/Insights/NetworkingProfiler/ViewModels/PacketViewDrawHelper.cpp
@@ -97,6 +97,7 @@ FPacketViewDrawHelper::FPacketViewDrawHelper(const FDrawContext& InDrawContext,
 	//, EventBorderBrush(FInsightsStyle::Get().GetBrush("EventBorder"))
 	, HoveredEventBorderBrush(FInsightsStyle::Get().GetBrush("HoveredEventBorder"))
 	, SelectedEventBorderBrush(FInsightsStyle::Get().GetBrush("SelectedEventBorder"))
+	, SelectionFont(FCoreStyle::GetDefaultFontStyle("Regular", 8))
 	, NumPackets(0)
 	, NumDrawSamples(0)
 {
@@ -214,8 +215,9 @@ void FPacketViewDrawHelper::DrawSampleHighlight(const FNetworkPacketAggregatedSa
 	const float ViewHeight = FMath::RoundToFloat(Viewport.GetHeight());
 	const float BaselineY = FMath::RoundToFloat(ViewportY.GetOffsetForValue(0.0));
 
-	//const float ValueY = FMath::RoundToFloat(ViewportY.GetOffsetForValue(static_cast<double>(Sample.LargestPacket.ContentSizeInBits)));
-	const float ValueY = FMath::RoundToFloat(ViewportY.GetOffsetForValue(static_cast<double>(Sample.LargestPacket.TotalSizeInBytes * 8)));
+	//const double Value = static_cast<double>(Sample.LargestPacket.ContentSizeInBits);
+	const double Value = static_cast<double>(Sample.LargestPacket.TotalSizeInBytes * 8);
+	const float ValueY = FMath::RoundToFloat(ViewportY.GetOffsetForValue(Value));
 
 	const float H = ValueY - BaselineY;
 	const float Y = ViewHeight - H;
@@ -224,7 +226,7 @@ void FPacketViewDrawHelper::DrawSampleHighlight(const FNetworkPacketAggregatedSa
 	{
 		const FLinearColor Color(1.0f, 1.0f, 0.0f, 1.0f); // yellow
 
-		// Draw border around the timing event box.
+		// Draw border around the hovered box.
 		DrawContext.DrawBox(X - 1.0f, Y - 1.0f, SampleW + 2.0f, H + 2.0f, HoveredEventBorderBrush, Color);
 	}
 	else // EHighlightMode::Selected or EHighlightMode::SelectedAndHovered
@@ -236,12 +238,54 @@ void FPacketViewDrawHelper::DrawSampleHighlight(const FNetworkPacketAggregatedSa
 		const float Blue = (Mode == EHighlightMode::SelectedAndHovered) ? 0.0f : S;
 		const FLinearColor Color(S, S, Blue, 1.0f);
 
-		// Draw border around the timing event box.
+		// Draw border around the selected box.
 		DrawContext.DrawBox(X - 1.0f, Y - 1.0f, SampleW + 2.0f, H + 2.0f, SelectedEventBorderBrush, Color);
 	}
 	DrawContext.LayerId++;
 }
 
+////////////////////////////////////////////////////////////////////////////////////////////////////
+
+void FPacketViewDrawHelper::DrawSelection(int32 StartPacketIndex, int32 EndPacketIndex) const
+{
+	const float SampleW = Viewport.GetSampleWidth();
+	const int32 PacketsPerSample = Viewport.GetNumPacketsPerSample();
+	const int32 FirstPacketIndex = Viewport.GetFirstPacketIndex();
+
+	const int32 StartSampleIndex = (StartPacketIndex - FirstPacketIndex) / PacketsPerSample;
+	const float X1 = FMath::Max(-2.0f, StartSampleIndex * SampleW);
+
+	const int32 EndSampleIndex = (EndPacketIndex - FirstPacketIndex) / PacketsPerSample;
+	const float X2 = FMath::Min(Viewport.GetWidth() + 2.0f, EndSampleIndex * SampleW);
+
+	const FAxisViewportDouble& ViewportY = Viewport.GetVerticalAxisViewport();
+
+	const float ViewHeight = FMath::RoundToFloat(Viewport.GetHeight());
+	const float Y = 0.0f;
+	const float H = ViewHeight - Y;
+
+	//// Animate color from white (if selected and hovered) or yellow (if only selected) to black, using a squared sine function.
+	//const double Time = static_cast<double>(FPlatformTime::Cycles64()) * FPlatformTime::GetSecondsPerCycle64();
+	//float S = FMath::Sin(2.0 * Time);
+	//S = S * S; // squared, to ensure only positive [0 - 1] values
+	//const FLinearColor Color(S, S, S, 1.0f);
+	//
+	//// Draw border around the selected box.
+	//DrawContext.DrawBox(X1 - 1.0f, Y - 1.0f, X2 - X1 + 2.0f, H + 2.0f, SelectedEventBorderBrush, Color);
+
+	// Fill the selected area.
+	const int32 PacketCount = EndPacketIndex - StartPacketIndex;
+	if (PacketCount > 1)
+	{
+		const float MinX = 0.0f;
+		const float MaxX = Viewport.GetWidth();
 
+		if (X1 <= MaxX && X2 >= MinX)
+		{
+			const FString Text = FString::Printf(TEXT("%d packets"), PacketCount);
+			FDrawHelpers::DrawSelection(DrawContext, MinX, MaxX, X1, X2, Y, H, 30.0f, Text, WhiteBrush, SelectionFont);
+		}
+	}
+}
 
 ////////////////////////////////////////////////////////////////////////////////////////////////////
diff --git a/Engine/Source/Developer/TraceInsights/Private/Insights/NetworkingProfiler/ViewModels/PacketViewDrawHelper.h b/Engine/Source/Developer/TraceInsights/Private/Insights/NetworkingProfiler/ViewModels/PacketViewDrawHelper.h
--- a/Engine/Source/Developer/TraceInsights/Private/Insights/NetworkingProfiler/ViewModels/PacketViewDrawHelper.h
+++ b/Engine/Source/Developer/TraceInsights/Private/Insights/NetworkingProfiler/ViewModels/PacketViewDrawHelper.h
@@ -3,6 +3,7 @@
 #pragma once
 
 #include "CoreMinimal.h"
+#include "Fonts/SlateFontInfo.h"
 #include "Styling/WidgetStyle.h"
 #include "TraceServices/Model/NetProfiler.h"
 
@@ -182,6 +183,7 @@ public:
 	void DrawBackground() const;
 	void DrawCached(const FNetworkPacketSeries& Series) const;
 	void DrawSampleHighlight(const FNetworkPacketAggregatedSample& Sample, EHighlightMode Mode) const;
+	void DrawSelection(int32 StartPacketIndex, int32 EndPacketIndex) const;
 
 	static FLinearColor GetColorByStatus(Trace::ENetProfilerDeliveryStatus Status);
 
@@ -196,6 +198,7 @@ private:
 	//const FSlateBrush* EventBorderBrush;
 	const FSlateBrush* HoveredEventBorderBrush;
 	const FSlateBrush* SelectedEventBorderBrush;
+	const FSlateFontInfo SelectionFont;
 
 	// Debug stats.
 	mutable int32 NumPackets;
diff --git a/Engine/Source/Developer/TraceInsights/Private/Insights/NetworkingProfiler/Widgets/SPacketView.cpp b/Engine/Source/Developer/TraceInsights/Private/Insights/NetworkingProfiler/Widgets/SPacketView.cpp
--- a/Engine/Source/Developer/TraceInsights/Private/Insights/NetworkingProfiler/Widgets/SPacketView.cpp
+++ b/Engine/Source/Developer/TraceInsights/Private/Insights/NetworkingProfiler/Widgets/SPacketView.cpp
@@ -83,8 +83,9 @@ void SPacketView::Reset()
 
 	bIsScrolling = false;
 
-	SelectionStartFrameIndex = 0;
-	SelectionEndFrameIndex = 0;
+	SelectionStartPacketIndex = 0;
+	SelectionEndPacketIndex = 0;
+	LastSelectedPacketIndex = 0;
 
 	SelectedSample.Reset();
 	HoveredSample.Reset();
@@ -118,7 +119,10 @@ void SPacketView::SetConnection(uint32 InGameInstanceIndex, uint32 InConnectionI
 	HoveredSample.Reset();
 
 	SelectedSample.Reset();
-	OnSelectedSampleChanged();
+	SelectionStartPacketIndex = 0;
+	SelectionEndPacketIndex = 0;
+	LastSelectedPacketIndex = 0;
+	OnSelectionChanged();
 
 	bIsStateDirty = true;
 }
@@ -348,7 +352,7 @@ FNetworkPacketSampleRef SPacketView::GetSampleAtMousePosition(float X, float Y)
 
 ////////////////////////////////////////////////////////////////////////////////////////////////////
 
-void SPacketView::SelectSampleAtMousePosition(float X, float Y)
+void SPacketView::SelectSampleAtMousePosition(float X, float Y, const FPointerEvent& MouseEvent)
 {
 	FNetworkPacketSampleRef SampleRef = GetSampleAtMousePosition(X, Y);
 	if (!SampleRef.IsValid())
@@ -360,16 +364,83 @@ void SPacketView::SelectSampleAtMousePosition(float X, float Y)
 		SampleRef = GetSampleAtMousePosition(X + 1.0f, Y);
 	}
 
-	if (!SelectedSample.Equals(SampleRef))
+	bool bRaiseSelectionChanged = false;
+
+	if (SampleRef.IsValid())
 	{
-		SelectedSample = SampleRef;
-		OnSelectedSampleChanged();
+		const int32 SelectedPacketIndex = SampleRef.Sample->LargestPacket.Index;
+
+		if (MouseEvent.GetModifierKeys().IsShiftDown())
+		{
+			if (SelectedPacketIndex >= LastSelectedPacketIndex)
+			{
+				// Extend selection toward right.
+				if (SelectionStartPacketIndex != LastSelectedPacketIndex ||
+					SelectionEndPacketIndex != SelectedPacketIndex + 1)
+				{
+					SelectionStartPacketIndex = LastSelectedPacketIndex;
+					SelectionEndPacketIndex = SelectedPacketIndex + 1;
+					bRaiseSelectionChanged = true;
+				}
+				LastSelectedPacketIndex = SelectionStartPacketIndex;
+			}
+			else
+			{
+				// Extend selection toward left.
+				if (SelectionEndPacketIndex != SelectedPacketIndex + 1)
+				{
+					SelectionStartPacketIndex = SelectedPacketIndex;
+					SelectionEndPacketIndex = LastSelectedPacketIndex + 1;
+					bRaiseSelectionChanged = true;
+				}
+				LastSelectedPacketIndex = SelectionEndPacketIndex - 1;
+			}
+		}
+		else
+		{
+			if (SelectionStartPacketIndex != SelectedPacketIndex ||
+				SelectionEndPacketIndex != SelectedPacketIndex + 1)
+			{
+				SelectionStartPacketIndex = SelectedPacketIndex;
+				SelectionEndPacketIndex = SelectedPacketIndex + 1;
+				bRaiseSelectionChanged = true;
+			}
+			LastSelectedPacketIndex = SelectedPacketIndex;
+		}
+	}
+	else
+	{
+		if (SelectionStartPacketIndex != 0 ||
+			SelectionEndPacketIndex != 0)
+		{
+			SelectionStartPacketIndex = 0;
+			SelectionEndPacketIndex = 0;
+			LastSelectedPacketIndex = 0;
+			bRaiseSelectionChanged = true;
+		}
+	}
+
+	if (SelectionEndPacketIndex == SelectionStartPacketIndex + 1)
+	{
+		if (!SelectedSample.Equals(SampleRef))
+		{
+			SelectedSample = SampleRef;
+		}
+	}
+	else
+	{
+		SelectedSample.Reset();
+	}
+
+	if (bRaiseSelectionChanged)
+	{
+		OnSelectionChanged();
 	}
 }
 
 ////////////////////////////////////////////////////////////////////////////////////////////////////
 
-void SPacketView::OnSelectedSampleChanged()
+void SPacketView::OnSelectionChanged()
 {
 	if (ProfilerWindow.IsValid())
 	{
@@ -381,7 +452,7 @@ void SPacketView::OnSelectedSampleChanged()
 		}
 		else
 		{
-			ProfilerWindow->SetSelectedPacket(0, 0);
+			ProfilerWindow->SetSelectedPacket(SelectionStartPacketIndex, SelectionEndPacketIndex);
 			ProfilerWindow->SetSelectedBitRange(0, 0);
 		}
 	}
@@ -455,6 +526,10 @@ int32 SPacketView::OnPaint(const FPaintArgs& Args, const FGeometry& AllottedGeom
 		{
 			Helper.DrawSampleHighlight(*HoveredSample.Sample, FPacketViewDrawHelper::EHighlightMode::Hovered);
 		}
+		if (SelectionEndPacketIndex > SelectionStartPacketIndex + 1)
+		{
+			Helper.DrawSelection(SelectionStartPacketIndex, SelectionEndPacketIndex);
+		}
 
 		// Draw the vertical axis grid.
 		DrawVerticalAxisGrid(DrawContext, WhiteBrush, SummaryFont);
@@ -832,7 +907,7 @@ FReply SPacketView::OnMouseButtonUp(const FGeometry& MyGeometry, const FPointerE
 			}
 			else if (bIsValidForMouseClick)
 			{
-				SelectSampleAtMousePosition(MousePositionOnButtonUp.X, MousePositionOnButtonUp.Y);
+				SelectSampleAtMousePosition(MousePositionOnButtonUp.X, MousePositionOnButtonUp.Y, MouseEvent);
 			}
 
 			bIsLMB_Pressed = false;
@@ -1003,6 +1078,28 @@ FCursorReply SPacketView::OnCursorQuery(const FGeometry& MyGeometry, const FPoin
 
 ////////////////////////////////////////////////////////////////////////////////////////////////////
 
+FReply SPacketView::OnKeyDown(const FGeometry& MyGeometry, const FKeyEvent& InKeyEvent)
+{
+	if (InKeyEvent.GetKey() == EKeys::A)
+	{
+		if (InKeyEvent.GetModifierKeys().IsControlDown())
+		{
+			// Select all.
+			SelectedSample.Reset();
+			SelectionStartPacketIndex = 0;
+			const FAxisViewportInt32& ViewportX = Viewport.GetHorizontalAxisViewport();
+			SelectionEndPacketIndex = ViewportX.GetMaxValue();
+			LastSelectedPacketIndex = 0;
+			OnSelectionChanged();
+			return FReply::Handled();
+		}
+	}
+
+	return SCompoundWidget::OnKeyDown(MyGeometry, InKeyEvent);
+}
+
+////////////////////////////////////////////////////////////////////////////////////////////////////
+
 void SPacketView::ShowContextMenu(const FPointerEvent& MouseEvent)
 {
 	const bool bShouldCloseWindowAfterMenuSelection = true;
diff --git a/Engine/Source/Developer/TraceInsights/Private/Insights/NetworkingProfiler/Widgets/SPacketView.h b/Engine/Source/Developer/TraceInsights/Private/Insights/NetworkingProfiler/Widgets/SPacketView.h
--- a/Engine/Source/Developer/TraceInsights/Private/Insights/NetworkingProfiler/Widgets/SPacketView.h
+++ b/Engine/Source/Developer/TraceInsights/Private/Insights/NetworkingProfiler/Widgets/SPacketView.h
@@ -128,6 +128,10 @@ public:
 
 	virtual FCursorReply OnCursorQuery(const FGeometry& MyGeometry, const FPointerEvent& CursorEvent) const override;
 
+	virtual bool SupportsKeyboardFocus() const override { return true; }
+	virtual FReply OnKeyDown(const FGeometry& MyGeometry, const FKeyEvent& InKeyEvent) override;
+	//virtual FReply OnKeyUp(const FGeometry& MyGeometry, const FKeyEvent& InKeyEvent) override;
+
 private:
 	void UpdateState();
 
@@ -135,8 +139,8 @@ private:
 	void DrawVerticalAxisGrid(FDrawContext& DrawContext, const FSlateBrush* Brush, const FSlateFontInfo& Font) const;
 
 	FNetworkPacketSampleRef GetSampleAtMousePosition(float X, float Y);
-	void SelectSampleAtMousePosition(float X, float Y);
-	void OnSelectedSampleChanged();
+	void SelectSampleAtMousePosition(float X, float Y, const FPointerEvent& MouseEvent);
+	void OnSelectionChanged();
 
 	void ShowContextMenu(const FPointerEvent& MouseEvent);
 
@@ -202,8 +206,9 @@ private:
 	//////////////////////////////////////////////////
 	// Selection
 
-	int32 SelectionStartFrameIndex;
-	int32 SelectionEndFrameIndex;
+	int32 SelectionStartPacketIndex;
+	int32 SelectionEndPacketIndex;
+	int32 LastSelectedPacketIndex;
 
 	FNetworkPacketSampleRef SelectedSample;
 	FNetworkPacketSampleRef HoveredSample;
diff --git a/Engine/Source/Developer/TraceInsights/Private/Insights/ViewModels/DrawHelpers.cpp b/Engine/Source/Developer/TraceInsights/Private/Insights/ViewModels/DrawHelpers.cpp
--- a/Engine/Source/Developer/TraceInsights/Private/Insights/ViewModels/DrawHelpers.cpp
+++ b/Engine/Source/Developer/TraceInsights/Private/Insights/ViewModels/DrawHelpers.cpp
@@ -172,7 +172,7 @@ void FDrawHelpers::DrawTimeRangeSelection(
 			const double Duration = EndTime - StartTime;
 			const FString Text = TimeUtils::FormatTimeAuto(Duration);
 
-			DrawSelection(DrawContext, MinX, MaxX, SelectionX1, SelectionX2, Y, H, Text, Brush, Font);
+			DrawSelection(DrawContext, MinX, MaxX, SelectionX1, SelectionX2, Y, H, 6.0f, Text, Brush, Font);
 		}
 	}
 }
@@ -187,6 +187,7 @@ void FDrawHelpers::DrawSelection(
 	float SelectionX2,
 	const float Y,
 	const float H,
+	const float TextY,
 	const FString Text,
 	const FSlateBrush* Brush,
 	const FSlateFontInfo& Font)
@@ -228,7 +229,7 @@ void FDrawHelpers::DrawSelection(
 	DrawContext.LayerId++;
 
 	const float ArrowSize = 6.0f;
-	const float ArrowY = Y + 6.0f;
+	const float ArrowY = Y + TextY;
 
 	if (SelectionX2 - SelectionX1 > 2 * ArrowSize)
 	{
diff --git a/Engine/Source/Developer/TraceInsights/Private/Insights/ViewModels/DrawHelpers.h b/Engine/Source/Developer/TraceInsights/Private/Insights/ViewModels/DrawHelpers.h
--- a/Engine/Source/Developer/TraceInsights/Private/Insights/ViewModels/DrawHelpers.h
+++ b/Engine/Source/Developer/TraceInsights/Private/Insights/ViewModels/DrawHelpers.h
@@ -81,6 +81,7 @@ public:
 							  float SelectionX2,
 							  const float Y,
 							  const float H,
+							  const float TextY,
 							  const FString Text,
 							  const FSlateBrush* Brush,
 							  const FSlateFontInfo& Font);
