[Insights] Start Page: Added extra info (if available) for available sessions (AppName, Platform, Build Config and Build Target). In tooltip was also added the Command Line.

diff --git a/Engine/Source/Developer/TraceInsights/Private/Insights/Widgets/SStartPageWindow.cpp b/Engine/Source/Developer/TraceInsights/Private/Insights/Widgets/SStartPageWindow.cpp
--- a/Engine/Source/Developer/TraceInsights/Private/Insights/Widgets/SStartPageWindow.cpp
+++ b/Engine/Source/Developer/TraceInsights/Private/Insights/Widgets/SStartPageWindow.cpp
@@ -9,6 +9,7 @@
 #include "Framework/MultiBox/MultiBoxBuilder.h"
 #include "HAL/FileManagerGeneric.h"
 #include "Input/Events.h"
+#include "Internationalization/Text.h"
 #include "SlateOptMacros.h"
 #include "Styling/CoreStyle.h"
 #include "Widgets/Docking/SDockTab.h"
@@ -91,6 +92,46 @@ public:
 					.ToolTipText(this, &STraceSessionRow::GetTraceSessionTooltip)
 				];
 		}
+		else if (ColumnName == FName(TEXT("Platform")))
+		{
+			return SNew(SBox)
+				.Padding(FMargin(4.0, 0.0))
+				[
+					SNew(STextBlock)
+					.Text(this, &STraceSessionRow::GetTraceSessionPlatform)
+					.ToolTipText(this, &STraceSessionRow::GetTraceSessionTooltip)
+				];
+		}
+		else if (ColumnName == FName(TEXT("AppName")))
+		{
+			return SNew(SBox)
+				.Padding(FMargin(4.0, 0.0))
+				[
+					SNew(STextBlock)
+					.Text(this, &STraceSessionRow::GetTraceSessionAppName)
+					.ToolTipText(this, &STraceSessionRow::GetTraceSessionTooltip)
+				];
+		}
+		else if (ColumnName == FName(TEXT("BuildConfig")))
+		{
+			return SNew(SBox)
+				.Padding(FMargin(4.0, 0.0))
+				[
+					SNew(STextBlock)
+					.Text(this, &STraceSessionRow::GetTraceSessionBuildConfiguration)
+					.ToolTipText(this, &STraceSessionRow::GetTraceSessionTooltip)
+				];
+		}
+		else if (ColumnName == FName(TEXT("BuildTarget")))
+		{
+			return SNew(SBox)
+				.Padding(FMargin(4.0, 0.0))
+				[
+					SNew(STextBlock)
+					.Text(this, &STraceSessionRow::GetTraceSessionBuildTarget)
+					.ToolTipText(this, &STraceSessionRow::GetTraceSessionTooltip)
+				];
+		}
 		else if (ColumnName == FName(TEXT("Size")))
 		{
 			return SNew(SBox)
@@ -108,7 +149,7 @@ public:
 				.Padding(FMargin(4.0, 0.0))
 				[
 					SNew(STextBlock)
-					.Text(this, &STraceSessionRow::GetTraceSessionstatus)
+					.Text(this, &STraceSessionRow::GetTraceSessionStatus)
 					.ToolTipText(this, &STraceSessionRow::GetTraceSessionTooltip)
 				];
 		}
@@ -127,7 +168,7 @@ public:
 		}
 		else
 		{
-			return FText();
+			return FText::GetEmpty();
 		}
 	}
 
@@ -140,10 +181,62 @@ public:
 		}
 		else
 		{
-			return FText();
+			return FText::GetEmpty();
 		}
 	}
 
+	FText GetTraceSessionPlatform() const
+	{
+		TSharedPtr<FTraceSession> TraceSessionPin = WeakTraceSession.Pin();
+		if (TraceSessionPin.IsValid())
+		{
+			return TraceSessionPin->Platform;
+		}
+		else
+		{
+			return FText::GetEmpty();
+		}
+	}
+
+	FText GetTraceSessionAppName() const
+	{
+		TSharedPtr<FTraceSession> TraceSessionPin = WeakTraceSession.Pin();
+		if (TraceSessionPin.IsValid())
+		{
+			return TraceSessionPin->AppName;
+		}
+		else
+		{
+			return FText::GetEmpty();
+		}
+	}
+
+	FText GetTraceSessionBuildConfiguration() const
+	{
+		TSharedPtr<FTraceSession> TraceSessionPin = WeakTraceSession.Pin();
+		if (TraceSessionPin.IsValid())
+		{
+			if (TraceSessionPin->ConfigurationType != EBuildConfiguration::Unknown)
+			{
+				return EBuildConfigurations::ToText(TraceSessionPin->ConfigurationType);
+			}
+		}
+		return FText::GetEmpty();
+	}
+
+	FText GetTraceSessionBuildTarget() const
+	{
+		TSharedPtr<FTraceSession> TraceSessionPin = WeakTraceSession.Pin();
+		if (TraceSessionPin.IsValid())
+		{
+			if (TraceSessionPin->TargetType != EBuildTargetType::Unknown)
+			{
+				return FText::FromString(LexToString(TraceSessionPin->TargetType));
+			}
+		}
+		return FText::GetEmpty();
+	}
+
 	FText GetTraceSessionTimeStamp() const
 	{
 		TSharedPtr<FTraceSession> TraceSessionPin = WeakTraceSession.Pin();
@@ -153,7 +246,7 @@ public:
 		}
 		else
 		{
-			return FText();
+			return FText::GetEmpty();
 		}
 	}
 
@@ -170,7 +263,7 @@ public:
 		}
 		else
 		{
-			return FText();
+			return FText::GetEmpty();
 		}
 	}
 
@@ -209,7 +302,7 @@ public:
 		}
 	}
 
-	FText GetTraceSessionstatus() const
+	FText GetTraceSessionStatus() const
 	{
 		TSharedPtr<FTraceSession> TraceSessionPin = WeakTraceSession.Pin();
 		if (TraceSessionPin.IsValid())
@@ -219,7 +312,7 @@ public:
 				return LOCTEXT("LiveTraceSessionStatus", "LIVE");
 			}
 		}
-		return FText();
+		return FText::GetEmpty();
 	}
 
 	FText GetTraceSessionTooltip() const
@@ -227,32 +320,34 @@ public:
 		TSharedPtr<FTraceSession> TraceSessionPin = WeakTraceSession.Pin();
 		if (TraceSessionPin.IsValid())
 		{
-			const FText Status = TraceSessionPin->bIsLive ? LOCTEXT("LiveTraceSessionStatus", "LIVE") : LOCTEXT("OfflineTraceSessionStatus", "Offline");
+			FTextBuilder TextBuilder;
+
+			TextBuilder.AppendLineFormat(LOCTEXT("TraceSessionTooltip_Name", "{0}"), TraceSessionPin->Name);
+			TextBuilder.AppendLineFormat(LOCTEXT("TraceSessionTooltip_Uri", "{0}"), TraceSessionPin->Uri);
+			TextBuilder.AppendLineFormat(LOCTEXT("TraceSessionTooltip_Platform", "Platform: {0}"), TraceSessionPin->Platform);
+			TextBuilder.AppendLineFormat(LOCTEXT("TraceSessionTooltip_AppName", "App Name: {0}"), TraceSessionPin->AppName);
+			TextBuilder.AppendLineFormat(LOCTEXT("TraceSessionTooltip_CommandLine", "Command Line: {0}"), TraceSessionPin->CommandLine);
+			TextBuilder.AppendLineFormat(LOCTEXT("TraceSessionTooltip_BuildConfig", "Build Configuration: {0}"),
+				TraceSessionPin->ConfigurationType == EBuildConfiguration::Unknown ? FText::GetEmpty() : EBuildConfigurations::ToText(TraceSessionPin->ConfigurationType));
+			TextBuilder.AppendLineFormat(LOCTEXT("TraceSessionTooltip_BuildTarget", "Build Target: {0}"),
+				TraceSessionPin->TargetType == EBuildTargetType::Unknown ? FText::GetEmpty() : FText::FromString(LexToString(TraceSessionPin->TargetType)));
+			TextBuilder.AppendLineFormat(LOCTEXT("TraceSessionTooltip_Timestamp", "Timestamp: {0}"), FText::AsDateTime(TraceSessionPin->TimeStamp));
 			if (TraceSessionPin->Size > 1024)
 			{
-				return FText::Format(LOCTEXT("TraceSessionTooltipFormat1", "{0}\n{1}\nTimestamp: {2}\nFile Size: {3} bytes ({4})\nStatus: {5}"),
-					TraceSessionPin->Name,
-					TraceSessionPin->Uri,
-					FText::AsDateTime(TraceSessionPin->TimeStamp),
-					FText::AsNumber(TraceSessionPin->Size),
-					FText::AsMemory(TraceSessionPin->Size),
-					Status
-				);
+				TextBuilder.AppendLineFormat(LOCTEXT("TraceSessionTooltip_FileSize2", "File Size: {0} bytes ({1})"), FText::AsNumber(TraceSessionPin->Size), FText::AsMemory(TraceSessionPin->Size));
 			}
 			else
 			{
-				return FText::Format(LOCTEXT("TraceSessionTooltipFormat2", "{0}\n{1}\nTimestamp: {2}\nFile Size: {3} bytes\nStatus: {4}"),
-					TraceSessionPin->Name,
-					TraceSessionPin->Uri,
-					FText::AsDateTime(TraceSessionPin->TimeStamp),
-					FText::AsNumber(TraceSessionPin->Size),
-					Status
-				);
+				TextBuilder.AppendLineFormat(LOCTEXT("TraceSessionTooltip_FileSize1", "File Size: {0} bytes"), FText::AsNumber(TraceSessionPin->Size));
 			}
+			const FText Status = TraceSessionPin->bIsLive ? LOCTEXT("LiveTraceSessionStatus", "LIVE") : LOCTEXT("OfflineTraceSessionStatus", "Offline");
+			TextBuilder.AppendLineFormat(LOCTEXT("TraceSessionTooltip_Status", "Status: {0}"), Status);
+
+			return TextBuilder.ToText();
 		}
 		else
 		{
-			return FText();
+			return FText::GetEmpty();
 		}
 	}
 
@@ -319,7 +414,7 @@ void SStartPageWindow::Construct(const FArguments& InArgs)
 						.Padding(3.0f, 3.0f)
 						[
 							SNew(SBox)
-							.WidthOverride(512.0f)
+							.WidthOverride(1024.0f)
 							[
 								SNew(SBorder)
 								.BorderImage(FEditorStyle::GetBrush("NotificationList.ItemBackground"))
@@ -337,7 +432,7 @@ void SStartPageWindow::Construct(const FArguments& InArgs)
 						.Padding(3.0f, 3.0f)
 						[
 							SNew(SBox)
-							.WidthOverride(512.0f)
+							.WidthOverride(256.0f)
 							[
 								SNew(SBorder)
 								.BorderImage(FEditorStyle::GetBrush("NotificationList.ItemBackground"))
@@ -410,7 +505,7 @@ TSharedRef<SWidget> SStartPageWindow::ConstructSessionsPanel()
 		.AutoHeight()
 		.HAlign(HAlign_Fill)
 		.Padding(0.0f, 1.0f, 0.0f, 2.0f)
-		.MaxHeight(134.0f)
+		.MaxHeight(22.0f + 20 * 14.0f) // max 20 rows
 		[
 			SAssignNew(TraceSessionsListView, SListView<TSharedPtr<FTraceSession>>)
 			.IsFocusable(true)
@@ -427,11 +522,27 @@ TSharedRef<SWidget> SStartPageWindow::ConstructSessionsPanel()
 				SNew(SHeaderRow)
 
 				+ SHeaderRow::Column(FName(TEXT("Name")))
-				.FillWidth(0.8f)
+				.FillWidth(0.4f)
 				.DefaultLabel(LOCTEXT("NameColumn", "Name"))
 
+				+ SHeaderRow::Column(FName(TEXT("Platform")))
+				.FillWidth(0.1f)
+				.DefaultLabel(LOCTEXT("PlatformColumn", "Platform"))
+
+				+ SHeaderRow::Column(FName(TEXT("AppName")))
+				.FillWidth(0.1f)
+				.DefaultLabel(LOCTEXT("AppNameColumn", "App Name"))
+
+				+ SHeaderRow::Column(FName(TEXT("BuildConfig")))
+				.FillWidth(0.1f)
+				.DefaultLabel(LOCTEXT("BuildConfigColumn", "Build Config"))
+
+				+ SHeaderRow::Column(FName(TEXT("BuildTarget")))
+				.FillWidth(0.1f)
+				.DefaultLabel(LOCTEXT("BuildTargetColumn", "Build Target"))
+
 				+ SHeaderRow::Column(FName(TEXT("Size")))
-				.FillWidth(0.2f)
+				.FixedWidth(100.0f)
 				.HAlignHeader(HAlign_Right)
 				.HAlignCell(HAlign_Right)
 				.DefaultLabel(LOCTEXT("SizeColumn", "File Size"))
diff --git a/Engine/Source/Developer/TraceInsights/Private/Insights/Widgets/SStartPageWindow.h b/Engine/Source/Developer/TraceInsights/Private/Insights/Widgets/SStartPageWindow.h
--- a/Engine/Source/Developer/TraceInsights/Private/Insights/Widgets/SStartPageWindow.h
+++ b/Engine/Source/Developer/TraceInsights/Private/Insights/Widgets/SStartPageWindow.h
@@ -4,6 +4,7 @@
 
 #include "CoreMinimal.h"
 #include "Framework/Docking/TabManager.h"
+#include "GenericPlatform/GenericPlatformMisc.h"
 #include "Input/Reply.h"
 #include "Layout/Visibility.h"
 #include "Misc/Guid.h"
@@ -44,6 +45,11 @@ struct FTraceSession
 	Trace::FSessionHandle SessionHandle;
 	FText Name;
 	FText Uri;
+	FText Platform;
+	FText AppName;
+	FText CommandLine;
+	EBuildConfiguration ConfigurationType;
+	EBuildTargetType TargetType;
 	FDateTime TimeStamp;
 	uint64 Size;
 	bool bIsLive;
@@ -52,6 +58,13 @@ struct FTraceSession
 		: SessionHandle(InSessionHandle)
 		, Name(FText::FromString(InSessionInfo.Name))
 		, Uri(FText::FromString(InSessionInfo.Uri))
+		, Platform(FText::FromString(InSessionInfo.Platform))
+		, AppName(FText::FromString(InSessionInfo.AppName))
+		, CommandLine(FText::FromString(InSessionInfo.CommandLine))
+		//, ConfigurationType(EBuildConfiguration::Unknown)
+		, ConfigurationType(InSessionInfo.ConfigurationType)
+		//, TargetType(EBuildTargetType::Unknown)
+		, TargetType(InSessionInfo.TargetType)
 		, TimeStamp(InSessionInfo.TimeStamp)
 		, Size(InSessionInfo.Size)
 		, bIsLive(InSessionInfo.bIsLive)
diff --git a/Engine/Source/Developer/TraceServices/Private/SessionService.cpp b/Engine/Source/Developer/TraceServices/Private/SessionService.cpp
--- a/Engine/Source/Developer/TraceServices/Private/SessionService.cpp
+++ b/Engine/Source/Developer/TraceServices/Private/SessionService.cpp
@@ -280,7 +280,7 @@ bool FSessionService::Tick(float DeltaTime)
 
 void FSessionService::UpdateSessionContext(FStoreSessionHandle StoreHandle, FSessionInfoInternal& Info)
 {
-	if (Info.Platform.Len())
+	if (Info.bIsUpdated)
 	{
 		return;
 	}
@@ -311,11 +311,7 @@ void FSessionService::UpdateSessionContext(FStoreSessionHandle StoreHandle, FSes
 
 	delete Stream;
 
-	if (Analyzer.Platform.Len() == 0)
-	{
-		Info.Platform = TEXT("Unknown");
-	}
-	else
+	if (Analyzer.Platform.Len() != 0)
 	{
 		Info.Platform = MoveTemp(Analyzer.Platform);
 		Info.AppName = MoveTemp(Analyzer.AppName);
@@ -323,6 +319,8 @@ void FSessionService::UpdateSessionContext(FStoreSessionHandle StoreHandle, FSes
 		Info.ConfigurationType = Analyzer.ConfigurationType;
 		Info.TargetType = Analyzer.TargetType;
 	}
+
+	Info.bIsUpdated = true;
 }
 
 void FSessionService::UpdateSessions()
diff --git a/Engine/Source/Developer/TraceServices/Private/SessionServicePrivate.h b/Engine/Source/Developer/TraceServices/Private/SessionServicePrivate.h
--- a/Engine/Source/Developer/TraceServices/Private/SessionServicePrivate.h
+++ b/Engine/Source/Developer/TraceServices/Private/SessionServicePrivate.h
@@ -46,6 +46,7 @@ private:
 		bool bIsLive;
 		int8 ConfigurationType;
 		int8 TargetType;
+		bool bIsUpdated = false;
 	};
 
 	bool Tick(float DeltaTime);
