[Insights]
  - Timing Insights: Fixed crash occuring when clicking on a tree view (Timers, Stats, Callers, Callees, etc.) after selecting a time range to compute aggregated stats (when computation of aggregated states takes longer than usual).
  - Timing Insights: Fixed analysis being slow when opening a new session after in previous session a time range was left selected.

diff --git a/Engine/Source/Developer/TraceInsights/Private/Insights/LoadingProfiler/Widgets/SLoadingProfilerWindow.cpp b/Engine/Source/Developer/TraceInsights/Private/Insights/LoadingProfiler/Widgets/SLoadingProfilerWindow.cpp
--- a/Engine/Source/Developer/TraceInsights/Private/Insights/LoadingProfiler/Widgets/SLoadingProfilerWindow.cpp
+++ b/Engine/Source/Developer/TraceInsights/Private/Insights/LoadingProfiler/Widgets/SLoadingProfilerWindow.cpp
@@ -74,22 +74,22 @@ void SLoadingProfilerWindow::Reset()
 
 	if (EventAggregationTreeView)
 	{
-		EventAggregationTreeView->RebuildTree(true);
+		EventAggregationTreeView->Reset();
 	}
 
 	if (ObjectTypeAggregationTreeView)
 	{
-		ObjectTypeAggregationTreeView->RebuildTree(true);
+		ObjectTypeAggregationTreeView->Reset();
 	}
 
 	if (PackageDetailsTreeView)
 	{
-		PackageDetailsTreeView->RebuildTree(true);
+		PackageDetailsTreeView->Reset();
 	}
 
 	if (ExportDetailsTreeView)
 	{
-		ExportDetailsTreeView->RebuildTree(true);
+		ExportDetailsTreeView->Reset();
 	}
 }
 
@@ -170,7 +170,7 @@ TSharedRef<SDockTab> SLoadingProfilerWindow::SpawnTab_EventAggregationTreeView(c
 			SAssignNew(EventAggregationTreeView, Insights::STableTreeView, TablePtr)
 		];
 
-	EventAggregationTreeView->RebuildTree(true);
+	EventAggregationTreeView->Reset();
 
 	DockTab->SetOnTabClosed(SDockTab::FOnTabClosedCallback::CreateRaw(this, &SLoadingProfilerWindow::OnEventAggregationTreeViewTabClosed));
 
@@ -210,7 +210,7 @@ TSharedRef<SDockTab> SLoadingProfilerWindow::SpawnTab_ObjectTypeAggregationTreeV
 			SAssignNew(ObjectTypeAggregationTreeView, Insights::STableTreeView, TablePtr)
 		];
 
-	ObjectTypeAggregationTreeView->RebuildTree(true);
+	ObjectTypeAggregationTreeView->Reset();
 
 	DockTab->SetOnTabClosed(SDockTab::FOnTabClosedCallback::CreateRaw(this, &SLoadingProfilerWindow::OnObjectTypeAggregationTreeViewTabClosed));
 
@@ -250,7 +250,7 @@ TSharedRef<SDockTab> SLoadingProfilerWindow::SpawnTab_PackageDetailsTreeView(con
 			SAssignNew(PackageDetailsTreeView, Insights::STableTreeView, TablePtr)
 		];
 
-	PackageDetailsTreeView->RebuildTree(true);
+	PackageDetailsTreeView->Reset();
 
 	DockTab->SetOnTabClosed(SDockTab::FOnTabClosedCallback::CreateRaw(this, &SLoadingProfilerWindow::OnPackageDetailsTreeViewTabClosed));
 
@@ -290,7 +290,7 @@ TSharedRef<SDockTab> SLoadingProfilerWindow::SpawnTab_ExportDetailsTreeView(cons
 			SAssignNew(ExportDetailsTreeView, Insights::STableTreeView, TablePtr)
 		];
 
-	ExportDetailsTreeView->RebuildTree(true);
+	ExportDetailsTreeView->Reset();
 
 	DockTab->SetOnTabClosed(SDockTab::FOnTabClosedCallback::CreateRaw(this, &SLoadingProfilerWindow::OnExportDetailsTreeViewTabClosed));
 
diff --git a/Engine/Source/Developer/TraceInsights/Private/Insights/NetworkingProfiler/Widgets/SNetStatsView.cpp b/Engine/Source/Developer/TraceInsights/Private/Insights/NetworkingProfiler/Widgets/SNetStatsView.cpp
--- a/Engine/Source/Developer/TraceInsights/Private/Insights/NetworkingProfiler/Widgets/SNetStatsView.cpp
+++ b/Engine/Source/Developer/TraceInsights/Private/Insights/NetworkingProfiler/Widgets/SNetStatsView.cpp
@@ -627,16 +627,7 @@ void SNetStatsView::InsightsManager_OnSessionChanged()
 	if (NewSession != Session)
 	{
 		Session = NewSession;
-
-		GameInstanceIndex = 0;
-		ConnectionIndex = 0;
-		ConnectionMode = Trace::ENetProfilerConnectionMode::Outgoing;
-		StatsPacketStartIndex = 0;
-		StatsPacketEndIndex = 0;
-		StatsStartPosition = 0;
-		StatsEndPosition = 0;
-
-		RebuildTree();
+		Reset();
 	}
 	else
 	{
@@ -1493,6 +1484,22 @@ void SNetStatsView::ContextMenu_ResetColumns_Execute()
 
 ////////////////////////////////////////////////////////////////////////////////////////////////////
 
+void SNetStatsView::Reset()
+{
+	GameInstanceIndex = 0;
+	ConnectionIndex = 0;
+	ConnectionMode = Trace::ENetProfilerConnectionMode::Outgoing;
+
+	StatsPacketStartIndex = 0;
+	StatsPacketEndIndex = 0;
+	StatsStartPosition = 0;
+	StatsEndPosition = 0;
+
+	RebuildTree(true);
+}
+
+////////////////////////////////////////////////////////////////////////////////////////////////////
+
 void SNetStatsView::RebuildTree(bool bResync)
 {
 	TArray<FNetEventNodePtr> SelectedItems;
@@ -1552,9 +1559,11 @@ void SNetStatsView::RebuildTree(bool bResync)
 
 	if (bListHasChanged)
 	{
-		//UpdateTree();
+		UpdateTree();
 		UpdateStatsInternal();
 
+		TreeView->RebuildList();
+
 		// Restore selection.
 		if (SelectedItems.Num() > 0)
 		{
@@ -1584,6 +1593,7 @@ void SNetStatsView::ResetStats()
 	GameInstanceIndex = 0;
 	ConnectionIndex = 0;
 	ConnectionMode = Trace::ENetProfilerConnectionMode::Outgoing;
+
 	StatsPacketStartIndex = 0;
 	StatsPacketEndIndex = 0;
 	StatsStartPosition = 0;
@@ -1599,6 +1609,7 @@ void SNetStatsView::UpdateStats(uint32 InGameInstanceIndex, uint32 InConnectionI
 	GameInstanceIndex = InGameInstanceIndex;
 	ConnectionIndex = InConnectionIndex;
 	ConnectionMode = InConnectionMode;
+
 	StatsPacketStartIndex = InStatsPacketStartIndex;
 	StatsPacketEndIndex = InStatsPacketEndIndex;
 	StatsStartPosition = InStatsStartPosition;
@@ -1611,12 +1622,19 @@ void SNetStatsView::UpdateStats(uint32 InGameInstanceIndex, uint32 InConnectionI
 
 void SNetStatsView::UpdateStatsInternal()
 {
+	if (StatsPacketStartIndex >= StatsPacketEndIndex ||
+		StatsStartPosition >= StatsEndPosition)
+	{
+		// keep previous aggregated stats
+		return;
+	}
+
 	for (const FNetEventNodePtr& NetEventNodePtr : NetEventNodes)
 	{
 		NetEventNodePtr->ResetAggregatedStats();
 	}
 
-	if (Session.IsValid() && StatsStartPosition < StatsEndPosition)
+	if (Session.IsValid())
 	{
 		TUniquePtr<Trace::ITable<Trace::FNetProfilerAggregatedStats>> AggregatedStatsTable;
 
@@ -1645,17 +1663,17 @@ void SNetStatsView::UpdateStatsInternal()
 	UpdateTree();
 
 	// Save selection.
-	TArray<FNetEventNodePtr> SelectedItems = TreeView->GetSelectedItems();
+	//TArray<FNetEventNodePtr> SelectedItems = TreeView->GetSelectedItems();
 
-	TreeView->RebuildList();
+	//TreeView->RebuildList();
 
 	// Restore selection.
-	if (SelectedItems.Num() > 0)
-	{
-		TreeView->ClearSelection();
-		TreeView->SetItemSelection(SelectedItems, true);
-		TreeView->RequestScrollIntoView(SelectedItems[0]);
-	}
+	//if (SelectedItems.Num() > 0)
+	//{
+	//	TreeView->ClearSelection();
+	//	TreeView->SetItemSelection(SelectedItems, true);
+	//	TreeView->RequestScrollIntoView(SelectedItems[0]);
+	//}
 }
 
 ////////////////////////////////////////////////////////////////////////////////////////////////////
diff --git a/Engine/Source/Developer/TraceInsights/Private/Insights/NetworkingProfiler/Widgets/SNetStatsView.h b/Engine/Source/Developer/TraceInsights/Private/Insights/NetworkingProfiler/Widgets/SNetStatsView.h
--- a/Engine/Source/Developer/TraceInsights/Private/Insights/NetworkingProfiler/Widgets/SNetStatsView.h
+++ b/Engine/Source/Developer/TraceInsights/Private/Insights/NetworkingProfiler/Widgets/SNetStatsView.h
@@ -76,11 +76,13 @@ public:
 	 */
 	virtual void Tick(const FGeometry& AllottedGeometry, const double InCurrentTime, const float InDeltaTime) override;
 
+	void Reset();
+
 	/**
 	 * Rebuilds the tree (if necessary).
 	 * @param bResync - If true, it forces a resync with list of net events from Analysis, even if the list did not changed since last sync.
 	 */
-	void RebuildTree(bool bResync = true);
+	void RebuildTree(bool bResync);
 
 	void ResetStats();
 	void UpdateStats(uint32 InGameInstanceIndex, uint32 InConnectionIndex, Trace::ENetProfilerConnectionMode InConnectionMode, uint32 InStatsPacketStartIndex, uint32 InStatsPacketEndIndex, uint32 InStatsStartPosition, uint32 InStatsEndPosition);
diff --git a/Engine/Source/Developer/TraceInsights/Private/Insights/NetworkingProfiler/Widgets/SNetworkingProfilerWindow.cpp b/Engine/Source/Developer/TraceInsights/Private/Insights/NetworkingProfiler/Widgets/SNetworkingProfilerWindow.cpp
--- a/Engine/Source/Developer/TraceInsights/Private/Insights/NetworkingProfiler/Widgets/SNetworkingProfilerWindow.cpp
+++ b/Engine/Source/Developer/TraceInsights/Private/Insights/NetworkingProfiler/Widgets/SNetworkingProfilerWindow.cpp
@@ -85,8 +85,7 @@ void SNetworkingProfilerWindow::Reset()
 
 	if (NetStatsView)
 	{
-		//TODO: NetStatsView->Reset();
-		NetStatsView->RebuildTree(true);
+		NetStatsView->Reset();
 	}
 
 	AvailableGameInstances.Reset();
@@ -968,7 +967,7 @@ void SNetworkingProfilerWindow::UpdateAggregatedNetStats()
 		}
 		else
 		{
-			//NetStatsView->ResetStats();
+			NetStatsView->ResetStats();
 		}
 	}
 }
diff --git a/Engine/Source/Developer/TraceInsights/Private/Insights/Table/Widgets/STableTreeView.cpp b/Engine/Source/Developer/TraceInsights/Private/Insights/Table/Widgets/STableTreeView.cpp
--- a/Engine/Source/Developer/TraceInsights/Private/Insights/Table/Widgets/STableTreeView.cpp
+++ b/Engine/Source/Developer/TraceInsights/Private/Insights/Table/Widgets/STableTreeView.cpp
@@ -539,7 +539,7 @@ void STableTreeView::InsightsManager_OnSessionChanged()
 	if (NewSession != Session)
 	{
 		Session = NewSession;
-		RebuildTree();
+		Reset();
 	}
 	else
 	{
@@ -1812,6 +1812,16 @@ void STableTreeView::ContextMenu_ResetColumns_Execute()
 
 ////////////////////////////////////////////////////////////////////////////////////////////////////
 
+void STableTreeView::Reset()
+{
+	StatsStartTime = 0.0;
+	StatsEndTime = 0.0;
+
+	RebuildTree(true);
+}
+
+////////////////////////////////////////////////////////////////////////////////////////////////////
+
 void STableTreeView::RebuildTree(bool bResync)
 {
 	bool bListHasChanged = false;
diff --git a/Engine/Source/Developer/TraceInsights/Private/Insights/Table/Widgets/STableTreeView.h b/Engine/Source/Developer/TraceInsights/Private/Insights/Table/Widgets/STableTreeView.h
--- a/Engine/Source/Developer/TraceInsights/Private/Insights/Table/Widgets/STableTreeView.h
+++ b/Engine/Source/Developer/TraceInsights/Private/Insights/Table/Widgets/STableTreeView.h
@@ -66,11 +66,13 @@ public:
 
 	TSharedPtr<Insights::FTable> GetTable() const { return Table; }
 
+	virtual void Reset();
+
 	/**
 	 * Rebuilds the tree (if necessary).
 	 * @param bResync - If true, it forces a resync even if the list did not changed since last sync.
 	 */
-	virtual void RebuildTree(bool bResync = true);
+	virtual void RebuildTree(bool bResync);
 
 	const FTableTreeNodePtr* FindNode(uint64 Id) const { return TableTreeNodesIdMap.Find(Id); }
 
diff --git a/Engine/Source/Developer/TraceInsights/Private/Insights/TimingProfilerManager.cpp b/Engine/Source/Developer/TraceInsights/Private/Insights/TimingProfilerManager.cpp
--- a/Engine/Source/Developer/TraceInsights/Private/Insights/TimingProfilerManager.cpp
+++ b/Engine/Source/Developer/TraceInsights/Private/Insights/TimingProfilerManager.cpp
@@ -427,16 +427,13 @@ void FTimingProfilerManager::UpdateCallersAndCallees()
 
 void FTimingProfilerManager::UpdateAggregatedTimerStats()
 {
-	if (SelectionStartTime < SelectionEndTime)
+	TSharedPtr<STimingProfilerWindow> Wnd = GetProfilerWindow();
+	if (Wnd)
 	{
-		TSharedPtr<STimingProfilerWindow> Wnd = GetProfilerWindow();
-		if (Wnd)
+		TSharedPtr<STimersView> TimersView = Wnd->GetTimersView();
+		if (TimersView)
 		{
-			TSharedPtr<STimersView> TimersView = Wnd->GetTimersView();
-			if (TimersView)
-			{
-				TimersView->UpdateStats(SelectionStartTime, SelectionEndTime);
-			}
+			TimersView->UpdateStats(SelectionStartTime, SelectionEndTime);
 		}
 	}
 }
@@ -445,16 +442,13 @@ void FTimingProfilerManager::UpdateAggregatedTimerStats()
 
 void FTimingProfilerManager::UpdateAggregatedCounterStats()
 {
-	if (SelectionStartTime < SelectionEndTime)
+	TSharedPtr<STimingProfilerWindow> Wnd = GetProfilerWindow();
+	if (Wnd)
 	{
-		TSharedPtr<STimingProfilerWindow> Wnd = GetProfilerWindow();
-		if (Wnd)
+		TSharedPtr<SStatsView> StatsView = Wnd->GetStatsView();
+		if (StatsView)
 		{
-			TSharedPtr<SStatsView> StatsView = Wnd->GetStatsView();
-			if (StatsView)
-			{
-				StatsView->UpdateStats(SelectionStartTime, SelectionEndTime);
-			}
+			StatsView->UpdateStats(SelectionStartTime, SelectionEndTime);
 		}
 	}
 }
diff --git a/Engine/Source/Developer/TraceInsights/Private/Insights/Widgets/SStatsView.cpp b/Engine/Source/Developer/TraceInsights/Private/Insights/Widgets/SStatsView.cpp
--- a/Engine/Source/Developer/TraceInsights/Private/Insights/Widgets/SStatsView.cpp
+++ b/Engine/Source/Developer/TraceInsights/Private/Insights/Widgets/SStatsView.cpp
@@ -658,7 +658,7 @@ void SStatsView::InsightsManager_OnSessionChanged()
 	if (NewSession != Session)
 	{
 		Session = NewSession;
-		RebuildTree();
+		Reset();
 	}
 	else
 	{
@@ -1501,6 +1501,16 @@ void SStatsView::ContextMenu_ResetColumns_Execute()
 
 ////////////////////////////////////////////////////////////////////////////////////////////////////
 
+void SStatsView::Reset()
+{
+	StatsStartTime = 0.0;
+	StatsEndTime = 0.0;
+
+	RebuildTree(true);
+}
+
+////////////////////////////////////////////////////////////////////////////////////////////////////
+
 void SStatsView::RebuildTree(bool bResync)
 {
 	TArray<FStatsNodePtr> SelectedItems;
@@ -1557,6 +1567,8 @@ void SStatsView::RebuildTree(bool bResync)
 		UpdateTree();
 		UpdateStats(StatsStartTime, StatsEndTime);
 
+		TreeView->RebuildList();
+
 		// Restore selection.
 		if (SelectedItems.Num() > 0)
 		{
@@ -1849,12 +1861,21 @@ void TTimeCalculationHelper<int64>::PostProcess(TMap<uint64, FStatsNodePtr>& Sta
 
 void SStatsView::UpdateStats(double StartTime, double EndTime)
 {
+	StatsStartTime = StartTime;
+	StatsEndTime = EndTime;
+
+	if (StartTime >= EndTime)
+	{
+		// keep previous aggregated stats
+		return;
+	}
+
 	for (const FStatsNodePtr& StatsNodePtr : StatsNodes)
 	{
 		StatsNodePtr->ResetAggregatedStats();
 	}
 
-	if (Session.IsValid() && StartTime < EndTime)
+	if (Session.IsValid())
 	{
 		const bool bComputeMedian = true;
 
@@ -1907,23 +1928,20 @@ void SStatsView::UpdateStats(double StartTime, double EndTime)
 		CalculationHelperDbl.PostProcess(StatsNodesIdMap, bComputeMedian);
 		CalculationHelperInt.PostProcess(StatsNodesIdMap, bComputeMedian);
 
-		StatsStartTime = StartTime;
-		StatsEndTime = EndTime;
-
 		UpdateTree();
 
 		// Save selection.
-		TArray<FStatsNodePtr> SelectedItems = TreeView->GetSelectedItems();
+		//TArray<FStatsNodePtr> SelectedItems = TreeView->GetSelectedItems();
 
-		TreeView->RebuildList();
+		//TreeView->RebuildList();
 
 		// Restore selection.
-		if (SelectedItems.Num() > 0)
-		{
-			TreeView->ClearSelection();
-			TreeView->SetItemSelection(SelectedItems, true);
-			TreeView->RequestScrollIntoView(SelectedItems[0]);
-		}
+		//if (SelectedItems.Num() > 0)
+		//{
+		//	TreeView->ClearSelection();
+		//	TreeView->SetItemSelection(SelectedItems, true);
+		//	TreeView->RequestScrollIntoView(SelectedItems[0]);
+		//}
 	}
 }
 
diff --git a/Engine/Source/Developer/TraceInsights/Private/Insights/Widgets/SStatsView.h b/Engine/Source/Developer/TraceInsights/Private/Insights/Widgets/SStatsView.h
--- a/Engine/Source/Developer/TraceInsights/Private/Insights/Widgets/SStatsView.h
+++ b/Engine/Source/Developer/TraceInsights/Private/Insights/Widgets/SStatsView.h
@@ -80,11 +80,13 @@ public:
 	 */
 	void Construct(const FArguments& InArgs);
 
+	void Reset();
+
 	/**
 	 * Rebuilds the tree (if necessary).
 	 * @param bResync - If true, it forces a resync with list of stats counters from Analysis, even if the list did not changed since last sync.
 	 */
-	void RebuildTree(bool bResync = true);
+	void RebuildTree(bool bResync);
 	void UpdateStats(double StartTime, double EndTime);
 
 	void SelectStatsNode(uint64 Id);
diff --git a/Engine/Source/Developer/TraceInsights/Private/Insights/Widgets/STimersView.cpp b/Engine/Source/Developer/TraceInsights/Private/Insights/Widgets/STimersView.cpp
--- a/Engine/Source/Developer/TraceInsights/Private/Insights/Widgets/STimersView.cpp
+++ b/Engine/Source/Developer/TraceInsights/Private/Insights/Widgets/STimersView.cpp
@@ -604,7 +604,7 @@ void STimersView::InsightsManager_OnSessionChanged()
 	if (NewSession != Session)
 	{
 		Session = NewSession;
-		RebuildTree();
+		Reset();
 	}
 	else
 	{
@@ -1493,6 +1493,16 @@ void STimersView::ContextMenu_ResetColumns_Execute()
 
 ////////////////////////////////////////////////////////////////////////////////////////////////////
 
+void STimersView::Reset()
+{
+	StatsStartTime = 0.0;
+	StatsEndTime = 0.0;
+
+	RebuildTree(true);
+}
+
+////////////////////////////////////////////////////////////////////////////////////////////////////
+
 void STimersView::RebuildTree(bool bResync)
 {
 	TArray<FTimerNodePtr> SelectedItems;
@@ -1548,6 +1558,8 @@ void STimersView::RebuildTree(bool bResync)
 		UpdateTree();
 		UpdateStats(StatsStartTime, StatsEndTime);
 
+		TreeView->RebuildList();
+
 		// Restore selection.
 		if (SelectedItems.Num() > 0)
 		{
@@ -1574,12 +1586,21 @@ void STimersView::RebuildTree(bool bResync)
 
 void STimersView::UpdateStats(double StartTime, double EndTime)
 {
+	StatsStartTime = StartTime;
+	StatsEndTime = EndTime;
+
+	if (StartTime >= EndTime)
+	{
+		// keep previous aggregated stats
+		return;
+	}
+
 	for (const FTimerNodePtr& TimerNodePtr : TimerNodes)
 	{
 		TimerNodePtr->ResetAggregatedStats();
 	}
 
-	if (Session.IsValid() && StartTime < EndTime && Trace::ReadTimingProfilerProvider(*Session.Get()))
+	if (Session.IsValid() && Trace::ReadTimingProfilerProvider(*Session.Get()))
 	{
 		TSharedPtr<STimingProfilerWindow> Wnd = FTimingProfilerManager::Get()->GetProfilerWindow();
 		TSharedPtr<STimingView> TimingView = Wnd.IsValid() ? Wnd->GetTimingView() : nullptr;
@@ -1610,23 +1631,20 @@ void STimersView::UpdateStats(double StartTime, double EndTime)
 			AggregationResultTableReader->NextRow();
 		}
 
-		StatsStartTime = StartTime;
-		StatsEndTime = EndTime;
-
 		UpdateTree();
 
 		// Save selection.
-		TArray<FTimerNodePtr> SelectedItems = TreeView->GetSelectedItems();
+		//TArray<FTimerNodePtr> SelectedItems = TreeView->GetSelectedItems();
 
-		TreeView->RebuildList();
+		//TreeView->RebuildList();
 
 		// Restore selection.
-		if (SelectedItems.Num() > 0)
-		{
-			TreeView->ClearSelection();
-			TreeView->SetItemSelection(SelectedItems, true);
-			TreeView->RequestScrollIntoView(SelectedItems[0]);
-		}
+		//if (SelectedItems.Num() > 0)
+		//{
+		//	TreeView->ClearSelection();
+		//	TreeView->SetItemSelection(SelectedItems, true);
+		//	TreeView->RequestScrollIntoView(SelectedItems[0]);
+		//}
 	}
 }
 
diff --git a/Engine/Source/Developer/TraceInsights/Private/Insights/Widgets/STimersView.h b/Engine/Source/Developer/TraceInsights/Private/Insights/Widgets/STimersView.h
--- a/Engine/Source/Developer/TraceInsights/Private/Insights/Widgets/STimersView.h
+++ b/Engine/Source/Developer/TraceInsights/Private/Insights/Widgets/STimersView.h
@@ -67,11 +67,13 @@ public:
 
 	TSharedPtr<Insights::FTable> GetTable() const { return Table; }
 
+	void Reset();
+
 	/**
 	 * Rebuilds the tree (if necessary).
 	 * @param bResync - If true, it forces a resync with list of timers from Analysis, even if the list did not changed since last sync.
 	 */
-	void RebuildTree(bool bResync = true);
+	void RebuildTree(bool bResync);
 	void UpdateStats(double StartTime, double EndTime);
 
 	void SelectTimerNode(uint64 Id);
diff --git a/Engine/Source/Developer/TraceInsights/Private/Insights/Widgets/STimingProfilerWindow.cpp b/Engine/Source/Developer/TraceInsights/Private/Insights/Widgets/STimingProfilerWindow.cpp
--- a/Engine/Source/Developer/TraceInsights/Private/Insights/Widgets/STimingProfilerWindow.cpp
+++ b/Engine/Source/Developer/TraceInsights/Private/Insights/Widgets/STimingProfilerWindow.cpp
@@ -79,30 +79,37 @@ void STimingProfilerWindow::Reset()
 	{
 		FrameTrack->Reset();
 	}
+
 	if (GraphTrack)
 	{
 		GraphTrack->Reset();
 	}
+
 	if (TimingView)
 	{
 		TimingView->Reset();
 	}
+
 	if (TimersView)
 	{
-		TimersView->RebuildTree(true);
+		TimersView->Reset();
 	}
+
 	if (CallersTreeView)
 	{
 		CallersTreeView->Reset();
 	}
+
 	if (CalleesTreeView)
 	{
 		CalleesTreeView->Reset();
 	}
+
 	if (StatsView)
 	{
-		StatsView->RebuildTree(true);
+		StatsView->Reset();
 	}
+
 	if (LogView)
 	{
 		LogView->Reset();
