[Insights] More fixes re tree refresh.

diff --git a/Engine/Source/Developer/TraceInsights/Private/Insights/NetworkingProfiler/Widgets/SNetStatsView.cpp b/Engine/Source/Developer/TraceInsights/Private/Insights/NetworkingProfiler/Widgets/SNetStatsView.cpp
--- a/Engine/Source/Developer/TraceInsights/Private/Insights/NetworkingProfiler/Widgets/SNetStatsView.cpp
+++ b/Engine/Source/Developer/TraceInsights/Private/Insights/NetworkingProfiler/Widgets/SNetStatsView.cpp
@@ -1661,19 +1661,6 @@ void SNetStatsView::UpdateStatsInternal()
 	}
 
 	UpdateTree();
-
-	// Save selection.
-	//TArray<FNetEventNodePtr> SelectedItems = TreeView->GetSelectedItems();
-
-	//TreeView->RebuildList();
-
-	// Restore selection.
-	//if (SelectedItems.Num() > 0)
-	//{
-	//	TreeView->ClearSelection();
-	//	TreeView->SetItemSelection(SelectedItems, true);
-	//	TreeView->RequestScrollIntoView(SelectedItems[0]);
-	//}
 }
 
 ////////////////////////////////////////////////////////////////////////////////////////////////////
diff --git a/Engine/Source/Developer/TraceInsights/Private/Insights/Table/Widgets/STableTreeView.cpp b/Engine/Source/Developer/TraceInsights/Private/Insights/Table/Widgets/STableTreeView.cpp
--- a/Engine/Source/Developer/TraceInsights/Private/Insights/Table/Widgets/STableTreeView.cpp
+++ b/Engine/Source/Developer/TraceInsights/Private/Insights/Table/Widgets/STableTreeView.cpp
@@ -1868,6 +1868,8 @@ void STableTreeView::RebuildTree(bool bResync)
 	if (bListHasChanged)
 	{
 		UpdateTree();
+
+		TreeView->RebuildList();
 	}
 }
 
diff --git a/Engine/Source/Developer/TraceInsights/Private/Insights/Widgets/SStatsView.cpp b/Engine/Source/Developer/TraceInsights/Private/Insights/Widgets/SStatsView.cpp
--- a/Engine/Source/Developer/TraceInsights/Private/Insights/Widgets/SStatsView.cpp
+++ b/Engine/Source/Developer/TraceInsights/Private/Insights/Widgets/SStatsView.cpp
@@ -1927,22 +1927,9 @@ void SStatsView::UpdateStats(double StartTime, double EndTime)
 		// Compute average and median inclusive/exclusive times.
 		CalculationHelperDbl.PostProcess(StatsNodesIdMap, bComputeMedian);
 		CalculationHelperInt.PostProcess(StatsNodesIdMap, bComputeMedian);
-
-		UpdateTree();
-
-		// Save selection.
-		//TArray<FStatsNodePtr> SelectedItems = TreeView->GetSelectedItems();
-
-		//TreeView->RebuildList();
-
-		// Restore selection.
-		//if (SelectedItems.Num() > 0)
-		//{
-		//	TreeView->ClearSelection();
-		//	TreeView->SetItemSelection(SelectedItems, true);
-		//	TreeView->RequestScrollIntoView(SelectedItems[0]);
-		//}
 	}
+
+	UpdateTree();
 }
 
 ////////////////////////////////////////////////////////////////////////////////////////////////////
diff --git a/Engine/Source/Developer/TraceInsights/Private/Insights/Widgets/STimersView.cpp b/Engine/Source/Developer/TraceInsights/Private/Insights/Widgets/STimersView.cpp
--- a/Engine/Source/Developer/TraceInsights/Private/Insights/Widgets/STimersView.cpp
+++ b/Engine/Source/Developer/TraceInsights/Private/Insights/Widgets/STimersView.cpp
@@ -1630,22 +1630,9 @@ void STimersView::UpdateStats(double StartTime, double EndTime)
 			}
 			AggregationResultTableReader->NextRow();
 		}
-
-		UpdateTree();
-
-		// Save selection.
-		//TArray<FTimerNodePtr> SelectedItems = TreeView->GetSelectedItems();
-
-		//TreeView->RebuildList();
-
-		// Restore selection.
-		//if (SelectedItems.Num() > 0)
-		//{
-		//	TreeView->ClearSelection();
-		//	TreeView->SetItemSelection(SelectedItems, true);
-		//	TreeView->RequestScrollIntoView(SelectedItems[0]);
-		//}
 	}
+
+	UpdateTree();
 }
 
 ////////////////////////////////////////////////////////////////////////////////////////////////////
