[Insights] Disabled the functionality for controlling the trace modules enabled (for live sessions) from Timing Insights toolbar.

diff --git a/Engine/Source/Developer/TraceInsights/Private/Insights/Widgets/STimingProfilerToolbar.cpp b/Engine/Source/Developer/TraceInsights/Private/Insights/Widgets/STimingProfilerToolbar.cpp
--- a/Engine/Source/Developer/TraceInsights/Private/Insights/Widgets/STimingProfilerToolbar.cpp
+++ b/Engine/Source/Developer/TraceInsights/Private/Insights/Widgets/STimingProfilerToolbar.cpp
@@ -71,7 +71,7 @@ void STimingProfilerToolbar::Construct(const FArguments& InArgs)
 
 	FToolBarBuilder ToolbarBuilder(CommandList.ToSharedRef(), FMultiBoxCustomization::None);
 	Local::FillViewToolbar(ToolbarBuilder);
-	FillModulesToolbar(ToolbarBuilder);
+	//FillModulesToolbar(ToolbarBuilder);
 
 	FToolBarBuilder RightSideToolbarBuilder(CommandList.ToSharedRef(), FMultiBoxCustomization::None);
 	Local::FillRightSideToolbar(RightSideToolbarBuilder);
@@ -218,7 +218,7 @@ void STimingProfilerToolbar::FillModulesToolbar(FToolBarBuilder& ToolbarBuilder)
 			//	EUserInterfaceActionType::ToggleButton,
 			//	NAME_None); // TutorialHighlightName
 
-			// Workaround for having toogle buttons without icons.
+			// Workaround for having toggle buttons without icons.
 			ToolbarBuilder.AddWidget(SNew(SCheckBox)
 				.IsEnabled(this, &STimingProfilerToolbar::ToggleModule_CanExecute, Module.Name)
 				.IsChecked(this, &STimingProfilerToolbar::ToggleModule_IsChecked2, Module.Name)
