Unreal Insights in UnrealEd

This change compiles and loads the insights module, exposes Insights to the workspace menu structure in the editor and allows it to connect to itself.
Shutdown needs to happen slightly differently in editoir due to module lifetime differences, so we kill the TraceSessionService at a different point to avoid a hang on shutdown.

diff --git a/Engine/Source/Developer/TraceInsights/Private/Insights/TraceInsightsModule.cpp b/Engine/Source/Developer/TraceInsights/Private/Insights/TraceInsightsModule.cpp
--- a/Engine/Source/Developer/TraceInsights/Private/Insights/TraceInsightsModule.cpp
+++ b/Engine/Source/Developer/TraceInsights/Private/Insights/TraceInsightsModule.cpp
@@ -7,7 +7,11 @@
 #include "TraceServices/SessionService.h"
 #include "Widgets/DeclarativeSyntaxSupport.h"
 #include "Widgets/Docking/SDockTab.h"
+#include "Framework/Docking/LayoutService.h"
 #include "Widgets/SWidget.h"
+#include "Widgets/Text/STextBlock.h"
+#include "WorkspaceMenuStructure.h"
+#include "WorkspaceMenuStructureModule.h"
 
 // Insights
 #include "Insights/InsightsManager.h"
@@ -64,10 +68,19 @@ protected:
 	/** Networking Profiler */
 	TSharedRef<SDockTab> SpawnNetworkingProfilerTab(const FSpawnTabArgs& Args);
 
+#if WITH_EDITOR
+	/** Handle exit */
+	void HandleExit();
+#endif
+
 protected:
 	TSharedPtr<Trace::IAnalysisService> TraceAnalysisService;
 	TSharedPtr<Trace::ISessionService> TraceSessionService;
 	TSharedPtr<Trace::IModuleService> TraceModuleService;
+
+#if WITH_EDITOR
+	TSharedPtr<FTabManager::FLayout> PersistentLayout;
+#endif
 };
 
 ////////////////////////////////////////////////////////////////////////////////////////////////////
@@ -109,11 +122,11 @@ void FTraceInsightsModule::StartupModule()
 		.SetTooltipText(NSLOCTEXT("FTraceInsightsModule", "StartPageTooltipText", "Open the start page for Unreal Insights."))
 		.SetIcon(FSlateIcon(FInsightsStyle::GetStyleSetName(), "StartPage.Icon.Small"));
 
-//#if WITH_EDITOR
-//	StartPageTabSpawnerEntry.SetGroup(WorkspaceMenu::GetMenuStructure().GetDeveloperToolsMiscCategory());
-//#else
-//	StartPageTabSpawnerEntry.SetGroup(WorkspaceMenu::GetMenuStructure().GetToolsCategory());
-//#endif
+#if WITH_EDITOR
+	StartPageTabSpawnerEntry.SetGroup(WorkspaceMenu::GetMenuStructure().GetDeveloperToolsProfilingCategory());
+#else
+	StartPageTabSpawnerEntry.SetGroup(WorkspaceMenu::GetMenuStructure().GetToolsCategory());
+#endif
 
 	// Register tab spawner for the Timing Insights.
 	auto& TimingProfilerTabSpawnerEntry = FGlobalTabmanager::Get()->RegisterNomadTabSpawner(FInsightsManagerTabs::TimingProfilerTabId,
@@ -122,11 +135,11 @@ void FTraceInsightsModule::StartupModule()
 		.SetTooltipText(NSLOCTEXT("FTraceInsightsModule", "TimingProfilerTooltipText", "Open the Timing Insights tab."))
 		.SetIcon(FSlateIcon(FInsightsStyle::GetStyleSetName(), "TimingProfiler.Icon.Small"));
 
-//#if WITH_EDITOR
-//	TimingProfilerTabSpawnerEntry.SetGroup(WorkspaceMenu::GetMenuStructure().GetDeveloperToolsMiscCategory());
-//#else
-//	TimingProfilerTabSpawnerEntry.SetGroup(WorkspaceMenu::GetMenuStructure().GetToolsCategory());
-//#endif
+#if WITH_EDITOR
+	TimingProfilerTabSpawnerEntry.SetGroup(WorkspaceMenu::GetMenuStructure().GetDeveloperToolsProfilingCategory());
+#else
+	TimingProfilerTabSpawnerEntry.SetGroup(WorkspaceMenu::GetMenuStructure().GetToolsCategory());
+#endif
 
 	// Register tab spawner for the Asset Loading Insights.
 	auto& LoadingProfilerTabSpawnerEntry = FGlobalTabmanager::Get()->RegisterNomadTabSpawner(FInsightsManagerTabs::LoadingProfilerTabId,
@@ -135,11 +148,11 @@ void FTraceInsightsModule::StartupModule()
 		.SetTooltipText(NSLOCTEXT("FTraceInsightsModule", "LoadingProfilerTooltipText", "Open the Asset Loading Insights tab."))
 		.SetIcon(FSlateIcon(FInsightsStyle::GetStyleSetName(), "LoadingProfiler.Icon.Small"));
 
-//#if WITH_EDITOR
-//	LoadingProfilerTabSpawnerEntry.SetGroup(WorkspaceMenu::GetMenuStructure().GetDeveloperToolsMiscCategory());
-//#else
-//	LoadingProfilerTabSpawnerEntry.SetGroup(WorkspaceMenu::GetMenuStructure().GetToolsCategory());
-//#endif
+#if WITH_EDITOR
+	LoadingProfilerTabSpawnerEntry.SetGroup(WorkspaceMenu::GetMenuStructure().GetDeveloperToolsProfilingCategory());
+#else
+	LoadingProfilerTabSpawnerEntry.SetGroup(WorkspaceMenu::GetMenuStructure().GetToolsCategory());
+#endif
 
 	// Register tab spawner for the Networking Insights.
 	auto& NetworkingProfilerTabSpawnerEntry = FGlobalTabmanager::Get()->RegisterNomadTabSpawner(FInsightsManagerTabs::NetworkingProfilerTabId,
@@ -149,18 +162,36 @@ void FTraceInsightsModule::StartupModule()
 		.SetTooltipText(NSLOCTEXT("FTraceInsightsModule", "NetworkingProfilerTooltipText", "Open the Networking Insights tab."))
 		.SetIcon(FSlateIcon(FInsightsStyle::GetStyleSetName(), "NetworkingProfiler.Icon.Small"));
 
-//#if WITH_EDITOR
-//	NetworkingProfilerTabSpawnerEntry.SetGroup(WorkspaceMenu::GetMenuStructure().GetDeveloperToolsMiscCategory());
-//#else
-//	NetworkingProfilerTabSpawnerEntry.SetGroup(WorkspaceMenu::GetMenuStructure().GetToolsCategory());
-//#endif
+#if WITH_EDITOR
+	NetworkingProfilerTabSpawnerEntry.SetGroup(WorkspaceMenu::GetMenuStructure().GetDeveloperToolsProfilingCategory());
+#else
+	NetworkingProfilerTabSpawnerEntry.SetGroup(WorkspaceMenu::GetMenuStructure().GetToolsCategory());
+#endif
+
+#if WITH_EDITOR
+	// In editor, load the layout and display it here. In the standalone tool this is done in UserInterfaceCommand
+	TSharedRef<FTabManager::FLayout> DefaultLayout = FTabManager::NewLayout("UnrealInsightsLayout_v1.0");
+	OnNewLayout(DefaultLayout);
+	PersistentLayout = FLayoutSaveRestore::LoadFromConfig(GEditorLayoutIni, DefaultLayout);
+	FGlobalTabmanager::Get()->RestoreFrom(PersistentLayout.ToSharedRef(), TSharedPtr<SWindow>());
+
+	FCoreDelegates::OnExit.AddRaw(this, &FTraceInsightsModule::HandleExit);
+#endif
 }
 
 ////////////////////////////////////////////////////////////////////////////////////////////////////
 
 void FTraceInsightsModule::ShutdownModule()
 {
+#if WITH_EDITOR
+	// Save application layout.
+	FLayoutSaveRestore::SaveToConfig(GEditorLayoutIni, PersistentLayout.ToSharedRef());
+	GConfig->Flush(false, GEditorLayoutIni);
+#endif
+
+#if !WITH_EDITOR
 	TraceSessionService->StopRecorderServer();
+#endif
 
 	FGlobalTabmanager::Get()->UnregisterNomadTabSpawner(FInsightsManagerTabs::NetworkingProfilerTabId);
 	FGlobalTabmanager::Get()->UnregisterNomadTabSpawner(FInsightsManagerTabs::LoadingProfilerTabId);
@@ -194,6 +225,20 @@ void FTraceInsightsModule::ShutdownModule()
 
 ////////////////////////////////////////////////////////////////////////////////////////////////////
 
+#if WITH_EDITOR
+void FTraceInsightsModule::HandleExit()
+{
+	// In editor, as module lifetimes are different, we need to shut down the recorder service before 
+	// threads get killed in the shutdown process otherwise we will hand forever waiting on the recorder server
+	if(TraceSessionService.IsValid())
+	{
+		TraceSessionService->StopRecorderServer();
+	}
+}
+#endif
+
+////////////////////////////////////////////////////////////////////////////////////////////////////
+
 void FTraceInsightsModule::OnNewLayout(TSharedRef<FTabManager::FLayout> NewLayout)
 {
 	const float DPIScaleFactor = FPlatformApplicationMisc::GetDPIScaleFactorAtPoint(10.0f, 10.0f);
@@ -203,7 +248,12 @@ void FTraceInsightsModule::OnNewLayout(TSharedRef<FTabManager::FLayout> NewLayou
 		->Split
 		(
 			FTabManager::NewStack()
+#if WITH_EDITOR
+			// In editor we default to all tabs closed
+			->AddTab(FInsightsManagerTabs::StartPageTabId, ETabState::ClosedTab)
+#else
 			->AddTab(FInsightsManagerTabs::StartPageTabId, ETabState::OpenedTab)
+#endif
 			->AddTab(FInsightsManagerTabs::TimingProfilerTabId, ETabState::ClosedTab)
 			->AddTab(FInsightsManagerTabs::LoadingProfilerTabId, ETabState::ClosedTab)
 			->AddTab(FInsightsManagerTabs::NetworkingProfilerTabId, ETabState::ClosedTab)
diff --git a/Engine/Source/Developer/TraceInsights/Private/Insights/Widgets/SStartPageWindow.cpp b/Engine/Source/Developer/TraceInsights/Private/Insights/Widgets/SStartPageWindow.cpp
--- a/Engine/Source/Developer/TraceInsights/Private/Insights/Widgets/SStartPageWindow.cpp
+++ b/Engine/Source/Developer/TraceInsights/Private/Insights/Widgets/SStartPageWindow.cpp
@@ -480,6 +480,14 @@ void SStartPageWindow::Construct(const FArguments& InArgs)
 				.Expose(OverlaySettingsSlot)
 		];
 
+#if WITH_EDITOR
+	// In the editor, attempt to connect to ourselves so we can shoose our own session
+	TSharedRef<Trace::ISessionService> SessionService = FInsightsManager::Get()->GetSessionService();
+	SessionService->ConnectSession(TEXT("127.0.0.1"));
+#endif
+
+	RefreshTraceSessionList();
+
 	FSlateApplication::Get().SetKeyboardFocus(TraceSessionsListView);
 	FSlateApplication::Get().SetUserFocus(0, TraceSessionsListView);
 }
@@ -763,8 +771,6 @@ TSharedRef<SWidget> SStartPageWindow::ConstructRecorderPanel()
 
 	;
 
-	RefreshTraceSessionList();
-
 	return Widget;
 }
 
diff --git a/Engine/Source/Developer/TraceInsights/TraceInsights.Build.cs b/Engine/Source/Developer/TraceInsights/TraceInsights.Build.cs
--- a/Engine/Source/Developer/TraceInsights/TraceInsights.Build.cs
+++ b/Engine/Source/Developer/TraceInsights/TraceInsights.Build.cs
@@ -26,7 +26,8 @@ public class TraceInsights : ModuleRules
 				"TraceLog",
 				"TraceServices",
 				"DesktopPlatform",
-			}
+				"WorkspaceMenuStructure",
+            }
 		);
 
 		if (Target.bBuildEditor)
diff --git a/Engine/Source/Editor/UnrealEd/Private/EditorEngine.cpp b/Engine/Source/Editor/UnrealEd/Private/EditorEngine.cpp
--- a/Engine/Source/Editor/UnrealEd/Private/EditorEngine.cpp
+++ b/Engine/Source/Editor/UnrealEd/Private/EditorEngine.cpp
@@ -1022,7 +1022,8 @@ void UEditorEngine::Init(IEngineLoop* InEngineLoop)
 			TEXT("InputBindingEditor"),
 			TEXT("AudioEditor"),
 			TEXT("TimeManagementEditor"),
-			TEXT("EditorInteractiveToolsFramework")
+			TEXT("EditorInteractiveToolsFramework"),
+			TEXT("TraceInsights")
 		};
 
 		FScopedSlowTask ModuleSlowTask(UE_ARRAY_COUNT(ModuleNames));
diff --git a/Engine/Source/Editor/UnrealEd/UnrealEd.Build.cs b/Engine/Source/Editor/UnrealEd/UnrealEd.Build.cs
--- a/Engine/Source/Editor/UnrealEd/UnrealEd.Build.cs
+++ b/Engine/Source/Editor/UnrealEd/UnrealEd.Build.cs
@@ -248,7 +248,8 @@ public class UnrealEd : ModuleRules
 				"Media",
 				"TimeManagementEditor",
 				"VirtualTexturingEditor",
-				"EditorInteractiveToolsFramework"
+				"EditorInteractiveToolsFramework",
+				"TraceInsights",
 			}
 		);
 
diff --git a/Engine/Source/Editor/WorkspaceMenuStructure/Private/WorkspaceMenuStructureModule.cpp b/Engine/Source/Editor/WorkspaceMenuStructure/Private/WorkspaceMenuStructureModule.cpp
--- a/Engine/Source/Editor/WorkspaceMenuStructure/Private/WorkspaceMenuStructureModule.cpp
+++ b/Engine/Source/Editor/WorkspaceMenuStructure/Private/WorkspaceMenuStructureModule.cpp
@@ -61,6 +61,11 @@ public:
 		return DeveloperToolsLogCategory.ToSharedRef();
 	}
 
+	virtual TSharedRef<FWorkspaceItem> GetDeveloperToolsProfilingCategory() const override
+	{
+		return DeveloperToolsProfilingCategory.ToSharedRef();
+	}
+
 	virtual TSharedRef<FWorkspaceItem> GetDeveloperToolsMiscCategory() const override
 	{
 		return DeveloperToolsMiscCategory.ToSharedRef();
@@ -95,6 +100,7 @@ public:
 		// Developer tools sections
 		DeveloperToolsDebugCategory = DeveloperToolsCategory->AddGroup(LOCTEXT("WorkspaceMenu_DeveloperToolsDebugCategory", "Debug"), FSlateIcon(), true);
 		DeveloperToolsLogCategory = DeveloperToolsCategory->AddGroup(LOCTEXT("WorkspaceMenu_DeveloperToolsLogCategory", "Log"), FSlateIcon(), true);
+		DeveloperToolsProfilingCategory = DeveloperToolsCategory->AddGroup(LOCTEXT("WorkspaceMenu_DeveloperToolsProfilingCategory", "Profiling"), FSlateIcon(), true);
 		DeveloperToolsMiscCategory = DeveloperToolsCategory->AddGroup(LOCTEXT("WorkspaceMenu_DeveloperToolsMiscCategory", "Miscellaneous"), FSlateIcon(), true);
 		
 		// Automation tools sub menu
@@ -127,6 +133,7 @@ private:
 	TSharedPtr<FWorkspaceItem> DeveloperToolsCategory;
 	TSharedPtr<FWorkspaceItem> DeveloperToolsDebugCategory;
 	TSharedPtr<FWorkspaceItem> DeveloperToolsLogCategory;
+	TSharedPtr<FWorkspaceItem> DeveloperToolsProfilingCategory;
 	TSharedPtr<FWorkspaceItem> DeveloperToolsMiscCategory;
 	
 	TSharedPtr<FWorkspaceItem> AutomationToolsCategory;
diff --git a/Engine/Source/Editor/WorkspaceMenuStructure/Public/WorkspaceMenuStructure.h b/Engine/Source/Editor/WorkspaceMenuStructure/Public/WorkspaceMenuStructure.h
--- a/Engine/Source/Editor/WorkspaceMenuStructure/Public/WorkspaceMenuStructure.h
+++ b/Engine/Source/Editor/WorkspaceMenuStructure/Public/WorkspaceMenuStructure.h
@@ -55,6 +55,9 @@ public:
 	/** See diagram above */
 	virtual TSharedRef<FWorkspaceItem> GetDeveloperToolsLogCategory() const = 0;
 
+	/** See diagram above */
+	virtual TSharedRef<FWorkspaceItem> GetDeveloperToolsProfilingCategory() const = 0;	
+	
 	/** See diagram above */
 	virtual TSharedRef<FWorkspaceItem> GetDeveloperToolsMiscCategory() const = 0;
 
