UnrealInsights:
Added possiblity for providers to add custom counters
Added option for including or excluding the points just outside of the enumeration interval when enumerating counter values

diff --git a/Engine/Source/Developer/TraceInsights/Private/Insights/ViewModels/TimingGraphTrack.cpp b/Engine/Source/Developer/TraceInsights/Private/Insights/ViewModels/TimingGraphTrack.cpp
--- a/Engine/Source/Developer/TraceInsights/Private/Insights/ViewModels/TimingGraphTrack.cpp
+++ b/Engine/Source/Developer/TraceInsights/Private/Insights/ViewModels/TimingGraphTrack.cpp
@@ -274,7 +274,7 @@ void FTimingGraphTrack::UpdateStatsCounterSeries(FTimingGraphSeries& Series, con
 
 				if (Counter.IsFloatingPoint())
 				{
-					Counter.EnumerateFloatValues(Viewport.GetStartTime(), Viewport.GetEndTime(), [&Builder, &MinValue, &MaxValue](double Time, double Value)
+					Counter.EnumerateFloatValues(Viewport.GetStartTime(), Viewport.GetEndTime(), true, [&Builder, &MinValue, &MaxValue](double Time, double Value)
 					{
 						if (Value < MinValue)
 						{
@@ -288,7 +288,7 @@ void FTimingGraphTrack::UpdateStatsCounterSeries(FTimingGraphSeries& Series, con
 				}
 				else
 				{
-					Counter.EnumerateValues(Viewport.GetStartTime(), Viewport.GetEndTime(), [&Builder, &MinValue, &MaxValue](double Time, int64 IntValue)
+					Counter.EnumerateValues(Viewport.GetStartTime(), Viewport.GetEndTime(), true, [&Builder, &MinValue, &MaxValue](double Time, int64 IntValue)
 					{
 						const double Value = static_cast<double>(IntValue);
 						if (Value < MinValue)
@@ -352,7 +352,7 @@ void FTimingGraphTrack::UpdateStatsCounterSeries(FTimingGraphSeries& Series, con
 
 			if (Counter.IsFloatingPoint())
 			{
-				Counter.EnumerateFloatValues(Viewport.GetStartTime(), Viewport.GetEndTime(), [&Builder](double Time, double Value)
+				Counter.EnumerateFloatValues(Viewport.GetStartTime(), Viewport.GetEndTime(), true, [&Builder](double Time, double Value)
 				{
 					//TODO: add a "value unit converter"
 					Builder.AddEvent(Time, 0.0, Value);
@@ -360,7 +360,7 @@ void FTimingGraphTrack::UpdateStatsCounterSeries(FTimingGraphSeries& Series, con
 			}
 			else
 			{
-				Counter.EnumerateValues(Viewport.GetStartTime(), Viewport.GetEndTime(), [&Builder](double Time, int64 IntValue)
+				Counter.EnumerateValues(Viewport.GetStartTime(), Viewport.GetEndTime(), true, [&Builder](double Time, int64 IntValue)
 				{
 					//TODO: add a "value unit converter"
 					Builder.AddEvent(Time, 0.0, static_cast<double>(IntValue));
diff --git a/Engine/Source/Developer/TraceInsights/Private/Insights/Widgets/SStatsView.cpp b/Engine/Source/Developer/TraceInsights/Private/Insights/Widgets/SStatsView.cpp
--- a/Engine/Source/Developer/TraceInsights/Private/Insights/Widgets/SStatsView.cpp
+++ b/Engine/Source/Developer/TraceInsights/Private/Insights/Widgets/SStatsView.cpp
@@ -1547,17 +1547,17 @@ void SStatsView::RebuildTree(bool bResync)
 			StatsNodesIdMap.Empty(PreviousNodeCount);
 			bListHasChanged = true;
 
-			CountersProvider.EnumerateCounters([this](const Trace::ICounter& Counter)
+			CountersProvider.EnumerateCounters([this](uint32 CounterId, const Trace::ICounter& Counter)
 			{
 				FName Name(Counter.GetName());
 				FName Group(Counter.GetDisplayHint() == Trace::CounterDisplayHint_Memory ? TEXT("Memory") :
 							Counter.IsFloatingPoint() ? TEXT("float") : TEXT("int64"));
 				EStatsNodeType Type = Counter.IsFloatingPoint() ? EStatsNodeType::Float : EStatsNodeType::Int64;
-				FStatsNodePtr StatsNodePtr = MakeShareable(new FStatsNode(Counter.GetId(), Name, Group, Type));
+				FStatsNodePtr StatsNodePtr = MakeShareable(new FStatsNode(CounterId, Name, Group, Type));
 				UpdateStatsNode(StatsNodePtr);
 				StatsNodes.Add(StatsNodePtr);
 				//StatsNodesMap.Add(Name, StatsNodePtr);
-				StatsNodesIdMap.Add(Counter.GetId(), StatsNodePtr);
+				StatsNodesIdMap.Add(CounterId, StatsNodePtr);
 			});
 		}
 	}
@@ -1603,23 +1603,23 @@ public:
 	{
 	}
 
-	void Update(const Trace::ICounter& Counter)
+	void Update(uint32 CounterId, const Trace::ICounter& Counter)
 	{
-		EnumerateValues(Counter, UpdateMinMax);
+		EnumerateValues(CounterId, Counter, UpdateMinMax);
 	}
 
 	void PrecomputeHistograms();
 
-	void UpdateHistograms(const Trace::ICounter& Counter)
+	void UpdateHistograms(uint32 CounterId, const Trace::ICounter& Counter)
 	{
-		EnumerateValues(Counter, UpdateHistogram);
+		EnumerateValues(CounterId, Counter, UpdateHistogram);
 	}
 
 	void PostProcess(TMap<uint64, FStatsNodePtr>& StatsNodesIdMap, bool bComputeMedian);
 
 private:
 	template<typename CallbackType>
-	void EnumerateValues(const Trace::ICounter& Counter, CallbackType Callback);
+	void EnumerateValues(uint32 CounterId, const Trace::ICounter& Counter, CallbackType Callback);
 
 	static void UpdateMinMax(TAggregatedStatsEx<Type>& Stats, Type Value);
 	static void UpdateHistogram(TAggregatedStatsEx<Type>& StatsEx, Type Value);
@@ -1635,18 +1635,18 @@ private:
 
 template<>
 template<typename CallbackType>
-void TTimeCalculationHelper<double>::EnumerateValues(const Trace::ICounter& Counter, CallbackType Callback)
+void TTimeCalculationHelper<double>::EnumerateValues(uint32 CounterId, const Trace::ICounter& Counter, CallbackType Callback)
 {
-	TAggregatedStatsEx<double>* StatsExPtr = StatsMap.Find(Counter.GetId());
+	TAggregatedStatsEx<double>* StatsExPtr = StatsMap.Find(CounterId);
 	if (!StatsExPtr)
 	{
-		StatsExPtr = &StatsMap.Add(Counter.GetId());
+		StatsExPtr = &StatsMap.Add(CounterId);
 		StatsExPtr->BaseStats.Min = +MAX_dbl;
 		StatsExPtr->BaseStats.Max = -MAX_dbl;
 	}
 	TAggregatedStatsEx<double>& StatsEx = *StatsExPtr;
 
-	Counter.EnumerateFloatValues(IntervalStartTime, IntervalEndTime, [this, &StatsEx, Callback](double Time, double Value)
+	Counter.EnumerateFloatValues(IntervalStartTime, IntervalEndTime, false, [this, &StatsEx, Callback](double Time, double Value)
 	{
 		if (Time >= IntervalStartTime && Time < IntervalEndTime)
 		{
@@ -1660,18 +1660,18 @@ void TTimeCalculationHelper<double>::EnumerateValues(const Trace::ICounter& Coun
 
 template<>
 template<typename CallbackType>
-void TTimeCalculationHelper<int64>::EnumerateValues(const Trace::ICounter& Counter, CallbackType Callback)
+void TTimeCalculationHelper<int64>::EnumerateValues(uint32 CounterId, const Trace::ICounter& Counter, CallbackType Callback)
 {
-	TAggregatedStatsEx<int64>* StatsExPtr = StatsMap.Find(Counter.GetId());
+	TAggregatedStatsEx<int64>* StatsExPtr = StatsMap.Find(CounterId);
 	if (!StatsExPtr)
 	{
-		StatsExPtr = &StatsMap.Add(Counter.GetId());
+		StatsExPtr = &StatsMap.Add(CounterId);
 		StatsExPtr->BaseStats.Min = +MAX_int64;
 		StatsExPtr->BaseStats.Max = -MAX_int64;
 	}
 	TAggregatedStatsEx<int64>& StatsEx = *StatsExPtr;
 
-	Counter.EnumerateValues(IntervalStartTime, IntervalEndTime, [this, &StatsEx, Callback](double Time, int64 Value)
+	Counter.EnumerateValues(IntervalStartTime, IntervalEndTime, false, [this, &StatsEx, Callback](double Time, int64 Value)
 	{
 		if (Time >= IntervalStartTime && Time < IntervalEndTime)
 		{
@@ -1897,15 +1897,15 @@ void SStatsView::UpdateStats(double StartTime, double EndTime)
 
 			// Compute instance count and total/min/max inclusive/exclusive times for each counter.
 			// Iterate through all counters.
-			CountersProvider.EnumerateCounters([&CalculationHelperDbl, &CalculationHelperInt](const Trace::ICounter& Counter)
+			CountersProvider.EnumerateCounters([&CalculationHelperDbl, &CalculationHelperInt](uint32 CounterId, const Trace::ICounter& Counter)
 			{
 				if (Counter.IsFloatingPoint())
 				{
-					CalculationHelperDbl.Update(Counter);
+					CalculationHelperDbl.Update(CounterId, Counter);
 				}
 				else
 				{
-					CalculationHelperInt.Update(Counter);
+					CalculationHelperInt.Update(CounterId, Counter);
 				}
 			});
 
@@ -1918,15 +1918,15 @@ void SStatsView::UpdateStats(double StartTime, double EndTime)
 
 				// Compute histogram.
 				// Iterate again through all counters.
-				CountersProvider.EnumerateCounters([&CalculationHelperDbl, &CalculationHelperInt](const Trace::ICounter& Counter)
+				CountersProvider.EnumerateCounters([&CalculationHelperDbl, &CalculationHelperInt](uint32 CounterId, const Trace::ICounter& Counter)
 				{
 					if (Counter.IsFloatingPoint())
 					{
-						CalculationHelperDbl.UpdateHistograms(Counter);
+						CalculationHelperDbl.UpdateHistograms(CounterId, Counter);
 					}
 					else
 					{
-						CalculationHelperInt.UpdateHistograms(Counter);
+						CalculationHelperInt.UpdateHistograms(CounterId, Counter);
 					}
 				});
 			}
diff --git a/Engine/Source/Developer/TraceServices/Private/Analyzers/CountersTraceAnalysis.cpp b/Engine/Source/Developer/TraceServices/Private/Analyzers/CountersTraceAnalysis.cpp
--- a/Engine/Source/Developer/TraceServices/Private/Analyzers/CountersTraceAnalysis.cpp
+++ b/Engine/Source/Developer/TraceServices/Private/Analyzers/CountersTraceAnalysis.cpp
@@ -31,7 +31,7 @@ bool FCountersAnalyzer::OnEvent(uint16 RouteId, const FOnEventContext& Context)
 		uint16 CounterId = EventData.GetValue<uint16>("Id");
 		ETraceCounterType CounterType = static_cast<ETraceCounterType>(EventData.GetValue<uint8>("Type"));
 		ETraceCounterDisplayHint CounterDisplayHint = static_cast<ETraceCounterDisplayHint>(EventData.GetValue<uint8>("DisplayHint"));
-		Trace::ICounter* Counter = CounterProvider.CreateCounter();
+		Trace::IEditableCounter* Counter = CounterProvider.CreateCounter();
 		if (CounterType == TraceCounterType_Float)
 		{
 			Counter->SetIsFloatingPoint(true);
@@ -49,7 +49,7 @@ bool FCountersAnalyzer::OnEvent(uint16 RouteId, const FOnEventContext& Context)
 		uint16 CounterId = EventData.GetValue<uint16>("CounterId");
 		int64 Value = EventData.GetValue<int64>("Value");
 		double Timestamp = Context.SessionContext.TimestampFromCycle(EventData.GetValue<uint64>("Cycle"));
-		Trace::ICounter** FindCounter = CountersMap.Find(CounterId);
+		Trace::IEditableCounter** FindCounter = CountersMap.Find(CounterId);
 		if (FindCounter)
 		{
 			(*FindCounter)->SetValue(Timestamp, Value);
@@ -61,10 +61,10 @@ bool FCountersAnalyzer::OnEvent(uint16 RouteId, const FOnEventContext& Context)
 		uint16 CounterId = EventData.GetValue<uint16>("CounterId");
 		float Value = EventData.GetValue<float>("Value");
 		double Timestamp = Context.SessionContext.TimestampFromCycle(EventData.GetValue<uint64>("Cycle"));
-		Trace::ICounter** FindCounter = CountersMap.Find(CounterId);
+		Trace::IEditableCounter* FindCounter = CountersMap.FindRef(CounterId);
 		if (FindCounter)
 		{
-			(*FindCounter)->SetValue(Timestamp, Value);
+			FindCounter->SetValue(Timestamp, Value);
 		}
 		break;
 	}
diff --git a/Engine/Source/Developer/TraceServices/Private/Analyzers/CountersTraceAnalysis.h b/Engine/Source/Developer/TraceServices/Private/Analyzers/CountersTraceAnalysis.h
--- a/Engine/Source/Developer/TraceServices/Private/Analyzers/CountersTraceAnalysis.h
+++ b/Engine/Source/Developer/TraceServices/Private/Analyzers/CountersTraceAnalysis.h
@@ -9,7 +9,7 @@
 namespace Trace
 {
 	class IAnalysisSession;
-	class ICounter;
+	class IEditableCounter;
 	class ICounterProvider;
 }
 
@@ -31,5 +31,5 @@ private:
 
 	Trace::IAnalysisSession& Session;
 	Trace::ICounterProvider& CounterProvider;
-	TMap<uint16, Trace::ICounter*> CountersMap;
+	TMap<uint16, Trace::IEditableCounter*> CountersMap;
 };
diff --git a/Engine/Source/Developer/TraceServices/Private/Analyzers/CsvProfilerTraceAnalysis.h b/Engine/Source/Developer/TraceServices/Private/Analyzers/CsvProfilerTraceAnalysis.h
--- a/Engine/Source/Developer/TraceServices/Private/Analyzers/CsvProfilerTraceAnalysis.h
+++ b/Engine/Source/Developer/TraceServices/Private/Analyzers/CsvProfilerTraceAnalysis.h
@@ -14,7 +14,7 @@ namespace Trace
 	class IFrameProvider;
 	class IThreadProvider;
 	class ICounterProvider;
-	class ICounter;
+	class IEditableCounter;
 }
 
 class FCsvProfilerAnalyzer
@@ -87,7 +87,7 @@ private:
 	{
 		uint64 ProviderHandle = uint64(-1);
 		uint64 ProviderCountHandle = uint64(-1);
-		Trace::ICounter* Counter = nullptr;
+		Trace::IEditableCounter* Counter = nullptr;
 		int64 CurrentFrame = -1;
 		FStatSeriesValue CurrentValue;
 		int64 CurrentCount = 0;
diff --git a/Engine/Source/Developer/TraceServices/Private/Analyzers/StatsTraceAnalysis.cpp b/Engine/Source/Developer/TraceServices/Private/Analyzers/StatsTraceAnalysis.cpp
--- a/Engine/Source/Developer/TraceServices/Private/Analyzers/StatsTraceAnalysis.cpp
+++ b/Engine/Source/Developer/TraceServices/Private/Analyzers/StatsTraceAnalysis.cpp
@@ -29,17 +29,12 @@ bool FStatsAnalyzer::OnEvent(uint16 RouteId, const FOnEventContext& Context)
 	case RouteId_Spec:
 	{
 		uint32 StatId = EventData.GetValue<uint32>("Id");
-		Trace::ICounter* Counter;
-		Trace::ICounter** FindIt = CountersMap.Find(StatId);
-		if (!FindIt)
+		Trace::IEditableCounter* Counter = CountersMap.FindRef(StatId);
+		if (!Counter)
 		{
 			Counter = CounterProvider.CreateCounter();
 			CountersMap.Add(StatId, Counter);
 		}
-		else
-		{
-			Counter = *FindIt;
-		}
 		const ANSICHAR* Name = reinterpret_cast<const ANSICHAR*>(EventData.GetAttachment());
 		const TCHAR* Description = reinterpret_cast<const TCHAR*>(EventData.GetAttachment() + strlen(Name) + 1);
 		Trace::ECounterDisplayHint DisplayHint = Trace::CounterDisplayHint_None;
@@ -74,17 +69,12 @@ bool FStatsAnalyzer::OnEvent(uint16 RouteId, const FOnEventContext& Context)
 
 			uint64 DecodedIdAndOp = FTraceAnalyzerUtils::Decode7bit(BufferPtr);
 			uint32 StatId = DecodedIdAndOp >> 3;
-			Trace::ICounter* Counter;
-			Trace::ICounter** FindIt = CountersMap.Find(StatId);
-			if (!FindIt)
+			Trace::IEditableCounter* Counter = CountersMap.FindRef(StatId);
+			if (!Counter)
 			{
 				Counter = CounterProvider.CreateCounter();
 				CountersMap.Add(StatId, Counter);
 			}
-			else
-			{
-				Counter = *FindIt;
-			}
 			uint8 Op = DecodedIdAndOp & 0x7;
 			uint64 CycleDiff = FTraceAnalyzerUtils::Decode7bit(BufferPtr);
 			uint64 Cycle = ThreadState->LastCycle + CycleDiff;
diff --git a/Engine/Source/Developer/TraceServices/Private/Analyzers/StatsTraceAnalysis.h b/Engine/Source/Developer/TraceServices/Private/Analyzers/StatsTraceAnalysis.h
--- a/Engine/Source/Developer/TraceServices/Private/Analyzers/StatsTraceAnalysis.h
+++ b/Engine/Source/Developer/TraceServices/Private/Analyzers/StatsTraceAnalysis.h
@@ -9,7 +9,7 @@
 namespace Trace
 {
 	class IAnalysisSession;
-	class ICounter;
+	class IEditableCounter;
 	class ICounterProvider;
 }
 
@@ -37,6 +37,6 @@ private:
 
 	Trace::IAnalysisSession& Session;
 	Trace::ICounterProvider& CounterProvider;
-	TMap<uint32, Trace::ICounter*> CountersMap;
+	TMap<uint32, Trace::IEditableCounter*> CountersMap;
 	TMap<uint32, TSharedRef<FThreadState>> ThreadStatesMap;
 };
diff --git a/Engine/Source/Developer/TraceServices/Private/Model/Counters.cpp b/Engine/Source/Developer/TraceServices/Private/Model/Counters.cpp
--- a/Engine/Source/Developer/TraceServices/Private/Model/Counters.cpp
+++ b/Engine/Source/Developer/TraceServices/Private/Model/Counters.cpp
@@ -9,17 +9,16 @@ namespace Trace
 
 const FName FCounterProvider::ProviderName("CounterProvider");
 
-FCounter::FCounter(ILinearAllocator& Allocator, const TArray<double>& InFrameStartTimes, uint32 InId)
+FCounter::FCounter(ILinearAllocator& Allocator, const TArray<double>& InFrameStartTimes)
 	: FrameStartTimes(InFrameStartTimes)
 	, IntCounterData(Allocator)
 	, DoubleCounterData(Allocator)
-	, Id(InId)
 {
 
 }
 
 template<typename CounterType, typename EnumerationType>
-static void EnumerateCounterValuesInternal(const TCounterData<CounterType>& CounterData, const TArray<double>& FrameStartTimes, bool bResetEveryFrame, double IntervalStart, double IntervalEnd, TFunctionRef<void(double, EnumerationType)> Callback)
+static void EnumerateCounterValuesInternal(const TCounterData<CounterType>& CounterData, const TArray<double>& FrameStartTimes, bool bResetEveryFrame, double IntervalStart, double IntervalEnd, bool bIncludeExternalBounds, TFunctionRef<void(double, EnumerationType)> Callback)
 {
 	if (!CounterData.Num())
 	{
@@ -40,18 +39,21 @@ static void EnumerateCounterValuesInternal(const TCounterData<CounterType>& Coun
 		{
 			if (bFirstEnumeratedValue)
 			{
-				if (!bFirstValue)
+				if (!bFirstValue && bIncludeExternalBounds)
 				{
 					Callback(LastTime, LastValue);
 				}
 				bFirstEnumeratedValue = false;
 			}
-			Callback(Time, CurrentValue);
-			bFirstEnumeratedValue = false;
 			if (Time > IntervalEnd)
 			{
+				if (bIncludeExternalBounds)
+				{
+					Callback(Time, CurrentValue);
+				}
 				break;
 			}
+			Callback(Time, CurrentValue);
 		}
 		LastTime = Time;
 		LastValue = CurrentValue;
@@ -66,27 +68,27 @@ void FCounter::SetIsFloatingPoint(bool bInIsFloatingPoint)
 	bIsFloatingPoint = bInIsFloatingPoint;
 }
 
-void FCounter::EnumerateValues(double IntervalStart, double IntervalEnd, TFunctionRef<void(double, int64)> Callback) const
+void FCounter::EnumerateValues(double IntervalStart, double IntervalEnd, bool bIncludeExternalBounds, TFunctionRef<void(double, int64)> Callback) const
 {
 	if (bIsFloatingPoint)
 	{
-		EnumerateCounterValuesInternal<double, int64>(DoubleCounterData, FrameStartTimes, bIsResetEveryFrame, IntervalStart, IntervalEnd, Callback);
+		EnumerateCounterValuesInternal<double, int64>(DoubleCounterData, FrameStartTimes, bIsResetEveryFrame, IntervalStart, IntervalEnd, bIncludeExternalBounds, Callback);
 	}
 	else
 	{
-		EnumerateCounterValuesInternal<int64, int64>(IntCounterData, FrameStartTimes, bIsResetEveryFrame, IntervalStart, IntervalEnd, Callback);
+		EnumerateCounterValuesInternal<int64, int64>(IntCounterData, FrameStartTimes, bIsResetEveryFrame, IntervalStart, IntervalEnd, bIncludeExternalBounds, Callback);
 	}
 }
 
-void FCounter::EnumerateFloatValues(double IntervalStart, double IntervalEnd, TFunctionRef<void(double, double)> Callback) const
+void FCounter::EnumerateFloatValues(double IntervalStart, double IntervalEnd, bool bIncludeExternalBounds, TFunctionRef<void(double, double)> Callback) const
 {
 	if (bIsFloatingPoint)
 	{
-		EnumerateCounterValuesInternal<double, double>(DoubleCounterData, FrameStartTimes, bIsResetEveryFrame, IntervalStart, IntervalEnd, Callback);
+		EnumerateCounterValuesInternal<double, double>(DoubleCounterData, FrameStartTimes, bIsResetEveryFrame, IntervalStart, IntervalEnd, bIncludeExternalBounds, Callback);
 	}
 	else
 	{
-		EnumerateCounterValuesInternal<int64, double>(IntCounterData, FrameStartTimes, bIsResetEveryFrame, IntervalStart, IntervalEnd, Callback);
+		EnumerateCounterValuesInternal<int64, double>(IntCounterData, FrameStartTimes, bIsResetEveryFrame, IntervalStart, IntervalEnd, bIncludeExternalBounds, Callback);
 	}
 }
 
@@ -150,18 +152,18 @@ FCounterProvider::FCounterProvider(IAnalysisSession& InSession, IFrameProvider&
 
 FCounterProvider::~FCounterProvider()
 {
-	for (FCounter* Counter : Counters)
+	for (const ICounter* Counter : Counters)
 	{
 		delete Counter;
 	}
 }
 
-void FCounterProvider::EnumerateCounters(TFunctionRef<void(const ICounter &)> Callback) const
+void FCounterProvider::EnumerateCounters(TFunctionRef<void(uint32, const ICounter&)> Callback) const
 {
 	uint32 Id = 0;
-	for (FCounter* Counter : Counters)
+	for (const ICounter* Counter : Counters)
 	{
-		Callback(*Counter);
+		Callback(Id++, *Counter);
 	}
 }
 
@@ -175,13 +177,18 @@ bool FCounterProvider::ReadCounter(uint32 CounterId, TFunctionRef<void(const ICo
 	return true;
 }
 
-ICounter* FCounterProvider::CreateCounter()
+IEditableCounter* FCounterProvider::CreateCounter()
 {
-	FCounter* Counter = new FCounter(Session.GetLinearAllocator(), FrameProvider.GetFrameStartTimes(TraceFrameType_Game), Counters.Num());
+	FCounter* Counter = new FCounter(Session.GetLinearAllocator(), FrameProvider.GetFrameStartTimes(TraceFrameType_Game));
 	Counters.Add(Counter);
 	return Counter;
 }
 
+void FCounterProvider::AddCounter(const ICounter* Counter)
+{
+	Counters.Add(Counter);
+}
+
 const ICounterProvider& ReadCounterProvider(const IAnalysisSession& Session)
 {
 	return *Session.ReadProvider<ICounterProvider>(FCounterProvider::ProviderName);
diff --git a/Engine/Source/Developer/TraceServices/Private/Model/CountersPrivate.h b/Engine/Source/Developer/TraceServices/Private/Model/CountersPrivate.h
--- a/Engine/Source/Developer/TraceServices/Private/Model/CountersPrivate.h
+++ b/Engine/Source/Developer/TraceServices/Private/Model/CountersPrivate.h
@@ -160,11 +160,10 @@ private:
 };
 
 class FCounter
-	: public ICounter
+	: public IEditableCounter
 {
 public:
-	FCounter(ILinearAllocator& Allocator, const TArray<double>& FrameStartTimes, uint32 Id);
-	virtual uint32 GetId() const override { return Id; }
+	FCounter(ILinearAllocator& Allocator, const TArray<double>& FrameStartTimes);
 	virtual const TCHAR* GetName() const override { return Name; }
 	virtual void SetName(const TCHAR* InName) override { Name = InName; }
 	virtual const TCHAR* GetDescription() const override { return Description; }
@@ -174,8 +173,8 @@ public:
 	virtual ECounterDisplayHint GetDisplayHint() const { return DisplayHint; }
 	virtual void SetDisplayHint(ECounterDisplayHint InDisplayHint) override { DisplayHint = InDisplayHint; }
 	virtual void SetIsResetEveryFrame(bool bInIsResetEveryFrame) override { bIsResetEveryFrame = bInIsResetEveryFrame; }
-	virtual void EnumerateValues(double IntervalStart, double IntervalEnd, TFunctionRef<void(double, int64)> Callback) const override;
-	virtual void EnumerateFloatValues(double IntervalStart, double IntervalEnd, TFunctionRef<void(double, double)> Callback) const override;
+	virtual void EnumerateValues(double IntervalStart, double IntervalEnd, bool bIncludeExternalBounds, TFunctionRef<void(double, int64)> Callback) const override;
+	virtual void EnumerateFloatValues(double IntervalStart, double IntervalEnd, bool bIncludeExternalBounds, TFunctionRef<void(double, double)> Callback) const override;
 	virtual void AddValue(double Time, int64 Value) override;
 	virtual void AddValue(double Time, double Value) override;
 	virtual void SetValue(double Time, int64 Value) override;
@@ -188,7 +187,6 @@ private:
 	const TCHAR* Name = nullptr;
 	const TCHAR* Description = nullptr;
 	uint64 ModCount = 0;
-	uint32 Id;
 	ECounterDisplayHint DisplayHint = CounterDisplayHint_None;
 	bool bIsFloatingPoint = false;
 	bool bIsResetEveryFrame = false;
@@ -203,14 +201,15 @@ public:
 	FCounterProvider(IAnalysisSession& Session, IFrameProvider& FrameProvider);
 	virtual ~FCounterProvider();
 	virtual uint64 GetCounterCount() const override { return Counters.Num(); }
-	virtual void EnumerateCounters(TFunctionRef<void(const ICounter&)> Callback) const override;
+	virtual void EnumerateCounters(TFunctionRef<void(uint32, const ICounter&)> Callback) const override;
 	virtual bool ReadCounter(uint32 CounterId, TFunctionRef<void(const ICounter&)> Callback) const override;
-	virtual ICounter* CreateCounter() override;
+	virtual IEditableCounter* CreateCounter() override;
+	virtual void AddCounter(const ICounter* Counter) override;
 
 private:
 	IAnalysisSession& Session;
 	IFrameProvider& FrameProvider;
-	TArray<FCounter*> Counters;
+	TArray<const ICounter*> Counters;
 };
 
 }
diff --git a/Engine/Source/Developer/TraceServices/Public/TraceServices/Model/Counters.h b/Engine/Source/Developer/TraceServices/Public/TraceServices/Model/Counters.h
--- a/Engine/Source/Developer/TraceServices/Public/TraceServices/Model/Counters.h
+++ b/Engine/Source/Developer/TraceServices/Public/TraceServices/Model/Counters.h
@@ -19,18 +19,25 @@ class ICounter
 public:
 	virtual ~ICounter() = default;
 
-	virtual uint32 GetId() const = 0;
 	virtual const TCHAR* GetName() const = 0;
-	virtual void SetName(const TCHAR* Name) = 0;
 	virtual const TCHAR* GetDescription() const = 0;
-	virtual void SetDescription(const TCHAR* Description) = 0;
 	virtual bool IsFloatingPoint() const = 0;
-	virtual void SetIsFloatingPoint(bool bIsFloatingPoint) = 0;
 	virtual ECounterDisplayHint GetDisplayHint() const = 0;
+	virtual void EnumerateValues(double IntervalStart, double IntervalEnd, bool bIncludeExternalBounds, TFunctionRef<void(double, int64)> Callback) const = 0;
+	virtual void EnumerateFloatValues(double IntervalStart, double IntervalEnd, bool bIncludeExternalBounds, TFunctionRef<void(double, double)> Callback) const = 0;
+};
+
+class IEditableCounter
+	: public ICounter
+{
+public:
+	virtual ~IEditableCounter() = default;
+
+	virtual void SetName(const TCHAR* Name) = 0;
+	virtual void SetDescription(const TCHAR* Description) = 0;
+	virtual void SetIsFloatingPoint(bool bIsFloatingPoint) = 0;
 	virtual void SetDisplayHint(ECounterDisplayHint DisplayHint) = 0;
 	virtual void SetIsResetEveryFrame(bool bInIsResetEveryFrame) = 0;
-	virtual void EnumerateValues(double IntervalStart, double IntervalEnd, TFunctionRef<void(double, int64)> Callback) const = 0;
-	virtual void EnumerateFloatValues(double IntervalStart, double IntervalEnd, TFunctionRef<void(double, double)> Callback) const = 0;
 	virtual void AddValue(double Time, int64 Value) = 0;
 	virtual void AddValue(double Time, double Value) = 0;
 	virtual void SetValue(double Time, int64 Value) = 0;
@@ -44,8 +51,9 @@ public:
 	virtual ~ICounterProvider() = default;
 	virtual uint64 GetCounterCount() const = 0;
 	virtual bool ReadCounter(uint32 CounterId, TFunctionRef<void(const ICounter&)> Callback) const = 0;
-	virtual void EnumerateCounters(TFunctionRef<void(const ICounter&)> Callback) const = 0;
-	virtual ICounter* CreateCounter() = 0;
+	virtual void EnumerateCounters(TFunctionRef<void(uint32, const ICounter&)> Callback) const = 0;
+	virtual IEditableCounter* CreateCounter() = 0;
+	virtual void AddCounter(const ICounter* Counter) = 0;
 };
 
 TRACESERVICES_API const ICounterProvider& ReadCounterProvider(const IAnalysisSession& Session);
