[Insights] Start Page: Fixed session list to correctly update when a trace file is renamed.

diff --git a/Engine/Source/Developer/TraceInsights/Private/Insights/Widgets/SStartPageWindow.cpp b/Engine/Source/Developer/TraceInsights/Private/Insights/Widgets/SStartPageWindow.cpp
--- a/Engine/Source/Developer/TraceInsights/Private/Insights/Widgets/SStartPageWindow.cpp
+++ b/Engine/Source/Developer/TraceInsights/Private/Insights/Widgets/SStartPageWindow.cpp
@@ -481,7 +481,7 @@ void SStartPageWindow::Construct(const FArguments& InArgs)
 		];
 
 #if WITH_EDITOR
-	// In the editor, attempt to connect to ourselves so we can shoose our own session
+	// In the editor, attempt to connect to ourselves so we can choose our own session.
 	TSharedRef<Trace::ISessionService> SessionService = FInsightsManager::Get()->GetSessionService();
 	SessionService->ConnectSession(TEXT("127.0.0.1"));
 #endif
@@ -899,6 +899,8 @@ void SStartPageWindow::RefreshTraceSessionList()
 
 	AvailableSessionCount = AvailableSessions.Num();
 
+	bool bSessionChanged = false;
+
 	// Count number of live sessions and update file sizes.
 	int32 OldLiveSessionCount = LiveSessionCount;
 	LiveSessionCount = 0;
@@ -916,11 +918,16 @@ void SStartPageWindow::RefreshTraceSessionList()
 		{
 			(*TraceSessionPtrPtr)->Size = SessionInfo.Size;
 		}
+		else
+		{
+			bSessionChanged = true;
+		}
 	}
 
 	// If session list has changed on analysis side, recreate the TraceSessions list view widget.
 	//TODO: if (AvailableSessionsChangeNumber != SessionService->GetAvailableSessionsChangeNumber())
-	if (AvailableSessionCount != TraceSessions.Num() ||
+	if (bSessionChanged ||
+		AvailableSessionCount != TraceSessions.Num() ||
 		LiveSessionCount != OldLiveSessionCount)
 	{
 		TSharedPtr<FTraceSession> NewSelectedTraceSession;
