UnrealInsights:
Separated FTimingProfilerAggregatedStats and FAggregatedTimingStats

diff --git a/Engine/Source/Developer/TraceInsights/Private/Insights/ViewModels/TimerNode.cpp b/Engine/Source/Developer/TraceInsights/Private/Insights/ViewModels/TimerNode.cpp
--- a/Engine/Source/Developer/TraceInsights/Private/Insights/ViewModels/TimerNode.cpp
+++ b/Engine/Source/Developer/TraceInsights/Private/Insights/ViewModels/TimerNode.cpp
@@ -10,12 +10,12 @@ const FName FTimerNode::TypeName(TEXT("FTimerNode"));
 
 void FTimerNode::ResetAggregatedStats()
 {
-	AggregatedStats = Trace::FAggregatedTimingStats();
+	AggregatedStats = Trace::FTimingProfilerAggregatedStats();
 }
 
 ////////////////////////////////////////////////////////////////////////////////////////////////////
 
-void FTimerNode::SetAggregatedStats(const Trace::FAggregatedTimingStats& InAggregatedStats)
+void FTimerNode::SetAggregatedStats(const Trace::FTimingProfilerAggregatedStats& InAggregatedStats)
 {
 	AggregatedStats = InAggregatedStats;
 }
diff --git a/Engine/Source/Developer/TraceInsights/Private/Insights/ViewModels/TimerNode.h b/Engine/Source/Developer/TraceInsights/Private/Insights/ViewModels/TimerNode.h
--- a/Engine/Source/Developer/TraceInsights/Private/Insights/ViewModels/TimerNode.h
+++ b/Engine/Source/Developer/TraceInsights/Private/Insights/ViewModels/TimerNode.h
@@ -4,6 +4,7 @@
 
 #include "CoreMinimal.h"
 #include "TraceServices/Containers/Timelines.h"
+#include "TraceServices/Model/TimingProfiler.h"
 
 // Insights
 #include "Insights/Table/ViewModels/BaseTreeNode.h"
@@ -89,10 +90,10 @@ public:
 	/**
 	 * @return the aggregated stats for this timer.
 	 */
-	const Trace::FAggregatedTimingStats& GetAggregatedStats() const { return AggregatedStats; }
+	const Trace::FTimingProfilerAggregatedStats& GetAggregatedStats() const { return AggregatedStats; }
 
 	void ResetAggregatedStats();
-	void SetAggregatedStats(const Trace::FAggregatedTimingStats& AggregatedStats);
+	void SetAggregatedStats(const Trace::FTimingProfilerAggregatedStats& AggregatedStats);
 
 	bool IsHotPath() const { return bIsHotPath; }
 	void SetIsHotPath(bool bOnOff) { bIsHotPath = bOnOff; }
@@ -105,7 +106,7 @@ private:
 	const ETimerNodeType Type;
 
 	/** Aggregated stats. */
-	Trace::FAggregatedTimingStats AggregatedStats;
+	Trace::FTimingProfilerAggregatedStats AggregatedStats;
 
 	/** True if this tree node is on the hot path. */
 	bool bIsHotPath;
diff --git a/Engine/Source/Developer/TraceInsights/Private/Insights/Widgets/STimerTreeView.cpp b/Engine/Source/Developer/TraceInsights/Private/Insights/Widgets/STimerTreeView.cpp
--- a/Engine/Source/Developer/TraceInsights/Private/Insights/Widgets/STimerTreeView.cpp
+++ b/Engine/Source/Developer/TraceInsights/Private/Insights/Widgets/STimerTreeView.cpp
@@ -957,7 +957,7 @@ FTimerNodePtr STimerTreeView::CreateTimerNodeRec(const Trace::FTimingProfilerBut
 
 	FTimerNodePtr TimerNodePtr = MakeShareable(new FTimerNode(Id, Name, Group, Type));
 
-	Trace::FAggregatedTimingStats AggregatedStats;
+	Trace::FTimingProfilerAggregatedStats AggregatedStats;
 	AggregatedStats.InstanceCount = Node.Count;
 	AggregatedStats.TotalInclusiveTime = Node.InclusiveTime;
 	AggregatedStats.TotalExclusiveTime = Node.ExclusiveTime;
diff --git a/Engine/Source/Developer/TraceInsights/Private/Insights/Widgets/STimersViewTooltip.cpp b/Engine/Source/Developer/TraceInsights/Private/Insights/Widgets/STimersViewTooltip.cpp
--- a/Engine/Source/Developer/TraceInsights/Private/Insights/Widgets/STimersViewTooltip.cpp
+++ b/Engine/Source/Developer/TraceInsights/Private/Insights/Widgets/STimersViewTooltip.cpp
@@ -90,7 +90,7 @@ TSharedPtr<SToolTip> STimersViewTooltip::GetColumnTooltip(const Insights::FTable
 
 TSharedPtr<SToolTip> STimersViewTooltip::GetRowTooltip(const TSharedPtr<FTimerNode> TimerNodePtr)
 {
-	const Trace::FAggregatedTimingStats& Stats = TimerNodePtr->GetAggregatedStats();
+	const Trace::FTimingProfilerAggregatedStats& Stats = TimerNodePtr->GetAggregatedStats();
 
 	const FText InstanceCountText = FText::AsNumber(Stats.InstanceCount);
 
diff --git a/Engine/Source/Developer/TraceServices/Private/Common/TimelineStatistics.h b/Engine/Source/Developer/TraceServices/Private/Common/TimelineStatistics.h
--- a/Engine/Source/Developer/TraceServices/Private/Common/TimelineStatistics.h
+++ b/Engine/Source/Developer/TraceServices/Private/Common/TimelineStatistics.h
@@ -5,6 +5,21 @@
 namespace Trace
 {
 
+struct FAggregatedTimingStats
+{
+	uint64 InstanceCount = 0;
+	double TotalInclusiveTime = 0.0;
+	double MinInclusiveTime = DBL_MAX;
+	double MaxInclusiveTime = -DBL_MAX;
+	double AverageInclusiveTime = 0.0;
+	double MedianInclusiveTime = 0.0;
+	double TotalExclusiveTime = 0.0;
+	double MinExclusiveTime = DBL_MAX;
+	double MaxExclusiveTime = -DBL_MAX;
+	double AverageExclusiveTime = 0.0;
+	double MedianExclusiveTime = 0.0;
+};
+
 class FTimelineStatistics
 {
 public:
diff --git a/Engine/Source/Developer/TraceServices/Private/Model/TimingProfiler.cpp b/Engine/Source/Developer/TraceServices/Private/Model/TimingProfiler.cpp
--- a/Engine/Source/Developer/TraceServices/Private/Model/TimingProfiler.cpp
+++ b/Engine/Source/Developer/TraceServices/Private/Model/TimingProfiler.cpp
@@ -209,8 +209,19 @@ ITable<FTimingProfilerAggregatedStats>* FTimingProfilerProvider::CreateAggregati
 	for (const auto& KV : Aggregation)
 	{
 		FTimingProfilerAggregatedStats& Row = Table->AddRow();
-		*static_cast<FAggregatedTimingStats*>(&Row) = KV.Value;
 		Row.Timer = KV.Key;
+		const FAggregatedTimingStats& Stats = KV.Value;
+		Row.InstanceCount = Stats.InstanceCount;
+		Row.TotalInclusiveTime = Stats.TotalExclusiveTime;
+		Row.MinInclusiveTime = Stats.MinExclusiveTime;
+		Row.MaxInclusiveTime = Stats.MaxExclusiveTime;
+		Row.AverageInclusiveTime = Stats.AverageExclusiveTime;
+		Row.MedianInclusiveTime = Stats.MedianInclusiveTime;
+		Row.TotalExclusiveTime = Stats.TotalExclusiveTime;
+		Row.MinExclusiveTime = Stats.MinExclusiveTime;
+		Row.MaxExclusiveTime = Stats.MaxExclusiveTime;
+		Row.AverageExclusiveTime = Stats.AverageExclusiveTime;
+		Row.MedianExclusiveTime = Stats.MedianExclusiveTime;
 	}
 	return Table;
 }
diff --git a/Engine/Source/Developer/TraceServices/Public/TraceServices/Containers/Timelines.h b/Engine/Source/Developer/TraceServices/Public/TraceServices/Containers/Timelines.h
--- a/Engine/Source/Developer/TraceServices/Public/TraceServices/Containers/Timelines.h
+++ b/Engine/Source/Developer/TraceServices/Public/TraceServices/Containers/Timelines.h
@@ -22,19 +22,4 @@ public:
 	virtual void EnumerateEvents(double IntervalStart, double IntervalEnd, TFunctionRef<void(double, double, uint32, const EventType&)> Callback) const = 0;
 };
 
-struct FAggregatedTimingStats
-{
-	uint64 InstanceCount = 0;
-	double TotalInclusiveTime = 0.0;
-	double MinInclusiveTime = DBL_MAX;
-	double MaxInclusiveTime = -DBL_MAX;
-	double AverageInclusiveTime = 0.0;
-	double MedianInclusiveTime = 0.0;
-	double TotalExclusiveTime = 0.0;
-	double MinExclusiveTime = DBL_MAX;
-	double MaxExclusiveTime = -DBL_MAX;
-	double AverageExclusiveTime = 0.0;
-	double MedianExclusiveTime = 0.0;
-};
-
 }
\ No newline at end of file
diff --git a/Engine/Source/Developer/TraceServices/Public/TraceServices/Model/TimingProfiler.h b/Engine/Source/Developer/TraceServices/Public/TraceServices/Model/TimingProfiler.h
--- a/Engine/Source/Developer/TraceServices/Public/TraceServices/Model/TimingProfiler.h
+++ b/Engine/Source/Developer/TraceServices/Public/TraceServices/Model/TimingProfiler.h
@@ -24,9 +24,19 @@ struct FTimingProfilerEvent
 };
 
 struct FTimingProfilerAggregatedStats
-	: public FAggregatedTimingStats
 {
 	const FTimingProfilerTimer* Timer = nullptr;
+	uint64 InstanceCount = 0;
+	double TotalInclusiveTime = 0.0;
+	double MinInclusiveTime = DBL_MAX;
+	double MaxInclusiveTime = -DBL_MAX;
+	double AverageInclusiveTime = 0.0;
+	double MedianInclusiveTime = 0.0;
+	double TotalExclusiveTime = 0.0;
+	double MinExclusiveTime = DBL_MAX;
+	double MaxExclusiveTime = -DBL_MAX;
+	double AverageExclusiveTime = 0.0;
+	double MedianExclusiveTime = 0.0;
 };
 
 struct FTimingProfilerButterflyNode
