[Insights]
  - Networking / Packets view: Fixed infinite loop if zooming in (vertically) too much.
  - Networking / Packets view: Fixed rare crash (trying to use an invalid ConnectionIndex).

diff --git a/Engine/Source/Developer/TraceInsights/Private/Insights/NetworkingProfiler/Widgets/SPacketView.cpp b/Engine/Source/Developer/TraceInsights/Private/Insights/NetworkingProfiler/Widgets/SPacketView.cpp
--- a/Engine/Source/Developer/TraceInsights/Private/Insights/NetworkingProfiler/Widgets/SPacketView.cpp
+++ b/Engine/Source/Developer/TraceInsights/Private/Insights/NetworkingProfiler/Widgets/SPacketView.cpp
@@ -235,11 +235,24 @@ void SPacketView::UpdateState()
 		Trace::FAnalysisSessionReadScope SessionReadScope(*Session.Get());
 		const Trace::INetProfilerProvider& NetProfilerProvider = Trace::ReadNetProfilerProvider(*Session.Get());
 
+		// TODO: if (NetProfilerProvider.IsConnectionValid(GameInstanceIndex, ConnectionIndex, ConnectionMode))
 		const uint32 GameInstanceCount = NetProfilerProvider.GetGameInstanceCount();
 		if (GameInstanceIndex < GameInstanceCount)
 		{
-			//const uint32 ConnectionCount = NetProfilerProvider.GetConnectionCount(GameInstanceIndex);
-			//if (ConnectionIndex < ConnectionCount) TODO: (if NetProfilerProvider.IsConnectionValid(GameInstanceIndex))
+			bool bIsValidConnection = false;
+			NetProfilerProvider.ReadConnections(GameInstanceIndex, [this, &bIsValidConnection](const Trace::FNetProfilerConnection& Connection)
+			{
+				if (ConnectionIndex == Connection.ConnectionIndex)
+				{
+					if ((ConnectionMode == Trace::ENetProfilerConnectionMode::Outgoing && Connection.bHasOutgoingData) ||
+						(ConnectionMode == Trace::ENetProfilerConnectionMode::Incoming && Connection.bHasIncomingData))
+					{
+						bIsValidConnection = true;
+					}
+				}
+			});
+
+			if (bIsValidConnection)
 			{
 				const uint32 NumPackets = NetProfilerProvider.GetPacketCount(ConnectionIndex, ConnectionMode);
 
@@ -731,7 +744,11 @@ void SPacketView::DrawVerticalAxisGrid(FDrawContext& DrawContext, const FSlateBr
 
 	if (Delta > 0.0)
 	{
-		const int64 DeltaBits = static_cast<int64>(Delta);
+		int64 DeltaBits = static_cast<int64>(Delta);
+		if (DeltaBits <= 0)
+		{
+			DeltaBits = 1;
+		}
 
 		// Compute rounding based on magnitude of visible range of values (Delta).
 		int64 Power10 = 1;
