Plugin API for STimingView & supporting changes

Plugin API overview:

Exposed a small subset of STimingView via a public interface: ITimingViewSession.
ITimingViewExtender. The modular feature that plugins need to implement. Allows hooks into the start and end of timing view sessions, as well as Tick, Menu extension and track changed events.

Other supporting changes:

Added extra accessors to ITimeline and implemented them for TMonotonicTimeline and TIntervalTimeline. Added a simple TPointTimeline too, for point-events with no duration.
Typedef-ed the various callback TFunctionRefs in ITimeline to save on typing...
Moved Type, SubType, Name and Order from FTimingEventTrack to FBaseTimingTrack. This will eventually allow for interleaved event and non-event tracks. Removed bCollapsed as it was not used.
Added FBaseTimingTrack SetSession/GetSession/IsSessionValid accessors to avoid exposing the whole of FInsightsManager to access the session.
Tweaked the height of events in compact view. It was extremly difficult to hover over them before as they were so thin. Overall track height is unchanged.
Added FTimingEventsTrack::OnEventSelected so tracks can handle selection themselves. Moved thread timing selection into this instead of checking type in the timing view.
Exposed a subset of FTimingViewDrawHelper via ITimingViewDrawHelper. FTimingEventTrack::Draw() now takes ITimingViewDrawHelper.
BuildContextMenu is now used to build per-track context menus. Moved thread timing track context menu building into its BuildContextMenu() override.
Reworked the vestigial selected event/hovered event/selection delegates from parameters on the view itself to accessor multicast events. This allows non-owners to hook into these events.
Added the ability to scrub the TimeMarker by holding ctrl or shift and scrubbing the ruler track.
Time marker now also displays the time, similar to the current time readout. FTimeRulerTrack now draws the marker instrad of STimingView to share/centralize code. Smoothed widths for time readouts are no longer static to allow for multiple instances to work correctly.

diff --git a/Engine/Source/Developer/TraceInsights/Private/Insights/Common/PaintUtils.h b/Engine/Source/Developer/TraceInsights/Private/Insights/Common/PaintUtils.h
deleted file mode 100644
diff --git a/Engine/Source/Developer/TraceInsights/Private/Insights/ViewModels/BaseTimingTrack.cpp b/Engine/Source/Developer/TraceInsights/Private/Insights/ViewModels/BaseTimingTrack.cpp
new file mode 100644
--- /dev/null
+++ b/Engine/Source/Developer/TraceInsights/Private/Insights/ViewModels/BaseTimingTrack.cpp
@@ -0,0 +1,24 @@
+// Copyright 1998-2019 Epic Games, Inc. All Rights Reserved.
+
+#include "Insights/ViewModels/BaseTimingTrack.h"
+#include "Framework/MultiBox/MultiBoxBuilder.h"
+
+#define LOCTEXT_NAMESPACE "BaseTimingTrack"
+
+void FBaseTimingTrack::BuildContextMenu(FMenuBuilder& MenuBuilder)
+{
+	MenuBuilder.BeginSection(TEXT("Empty"), FText::FromString(GetName()));
+	{
+		MenuBuilder.AddMenuEntry(
+			LOCTEXT("ContextMenu_NA", "N/A"),
+			LOCTEXT("ContextMenu_NA_Desc", "No actions available."),
+			FSlateIcon(),
+			FUIAction(FExecuteAction(), FCanExecuteAction::CreateLambda([](){ return false; })),
+			NAME_None,
+			EUserInterfaceActionType::Button
+		);
+	}
+	MenuBuilder.EndSection();
+}
+
+#undef LOCTEXT_NAMESPACE
\ No newline at end of file
diff --git a/Engine/Source/Developer/TraceInsights/Private/Insights/ViewModels/BaseTimingTrack.h b/Engine/Source/Developer/TraceInsights/Private/Insights/ViewModels/BaseTimingTrack.h
deleted file mode 100644
diff --git a/Engine/Source/Developer/TraceInsights/Private/Insights/ViewModels/FileActivityTimingTrack.cpp b/Engine/Source/Developer/TraceInsights/Private/Insights/ViewModels/FileActivityTimingTrack.cpp
--- a/Engine/Source/Developer/TraceInsights/Private/Insights/ViewModels/FileActivityTimingTrack.cpp
+++ b/Engine/Source/Developer/TraceInsights/Private/Insights/ViewModels/FileActivityTimingTrack.cpp
@@ -495,27 +495,23 @@ bool FFileActivityTimingTrack::SearchEvent(const double InStartTime,
 
 	FSearchTimingEventContext Ctx(InStartTime, InEndTime, InPredicate, InOutTimingEvent, bInStopAtFirstMatch, bInSearchForLargestEvent);
 
-	TSharedPtr<const Trace::IAnalysisSession> Session = FInsightsManager::Get()->GetSession();
-	if (Session.IsValid())
-	{
-		Trace::FAnalysisSessionReadScope SessionReadScope(*Session.Get());
+	Trace::FAnalysisSessionReadScope SessionReadScope(GetSession());
 
-		const FTimingEventsTrack* Track = Ctx.TimingEvent.Track;
+	const FTimingEventsTrack* Track = Ctx.TimingEvent.Track;
 
-		for (const FFileActivitySharedState::FIoTimingEvent& Event : State->GetAllEvents())
+	for (const FFileActivitySharedState::FIoTimingEvent& Event : State->GetAllEvents())
+	{
+		if (Event.EndTime <= Ctx.StartTime)
 		{
-			if (Event.EndTime <= Ctx.StartTime)
-			{
-				continue;
-			}
-
-			if (!Ctx.bContinueSearching || Event.StartTime >= Ctx.EndTime)
-			{
-				break;
-			}
+			continue;
+		}
 
-			Ctx.CheckEvent(Event, bIgnoreEventDepth ? 0 : Event.Depth);
+		if (!Ctx.bContinueSearching || Event.StartTime >= Ctx.EndTime)
+		{
+			break;
 		}
+
+		Ctx.CheckEvent(Event, bIgnoreEventDepth ? 0 : Event.Depth);
 	}
 
 	return Ctx.bFound;
@@ -525,7 +521,7 @@ bool FFileActivityTimingTrack::SearchEvent(const double InStartTime,
 // FOverviewFileActivityTimingTrack
 ////////////////////////////////////////////////////////////////////////////////////////////////////
 
-void FOverviewFileActivityTimingTrack::Draw(FTimingViewDrawHelper& Helper) const
+void FOverviewFileActivityTimingTrack::Draw(ITimingViewDrawHelper& Helper) const
 {
 	FTimingEventsTrack& Track = *const_cast<FOverviewFileActivityTimingTrack*>(this);
 
@@ -579,7 +575,7 @@ void FOverviewFileActivityTimingTrack::Draw(FTimingViewDrawHelper& Helper) const
 // FDetailedFileActivityTimingTrack
 ////////////////////////////////////////////////////////////////////////////////////////////////////
 
-void FDetailedFileActivityTimingTrack::Draw(FTimingViewDrawHelper& Helper) const
+void FDetailedFileActivityTimingTrack::Draw(ITimingViewDrawHelper& Helper) const
 {
 	FTimingEventsTrack& Track = *const_cast<FDetailedFileActivityTimingTrack*>(this);
 
diff --git a/Engine/Source/Developer/TraceInsights/Private/Insights/ViewModels/FileActivityTimingTrack.h b/Engine/Source/Developer/TraceInsights/Private/Insights/ViewModels/FileActivityTimingTrack.h
--- a/Engine/Source/Developer/TraceInsights/Private/Insights/ViewModels/FileActivityTimingTrack.h
+++ b/Engine/Source/Developer/TraceInsights/Private/Insights/ViewModels/FileActivityTimingTrack.h
@@ -94,7 +94,7 @@ public:
 	{
 	}
 
-	virtual void Draw(FTimingViewDrawHelper& Helper) const override;
+	virtual void Draw(ITimingViewDrawHelper& Helper) const override;
 	virtual bool SearchTimingEvent(const double InStartTime, const double InEndTime, TFunctionRef<bool(double, double, uint32)> InPredicate, FTimingEvent& InOutTimingEvent, bool bInStopAtFirstMatch, bool bInSearchForLargestEvent) const override
 	{
 		constexpr bool bIgnoreEventDepth = true;
@@ -112,7 +112,7 @@ public:
 	{
 	}
 
-	virtual void Draw(FTimingViewDrawHelper& Helper) const override;
+	virtual void Draw(ITimingViewDrawHelper& Helper) const override;
 	virtual bool SearchTimingEvent(const double InStartTime, const double InEndTime, TFunctionRef<bool(double, double, uint32)> InPredicate, FTimingEvent& InOutTimingEvent, bool bInStopAtFirstMatch, bool bInSearchForLargestEvent) const override
 	{
 		constexpr bool bIgnoreEventDepth = false;
diff --git a/Engine/Source/Developer/TraceInsights/Private/Insights/ViewModels/GraphTrack.cpp b/Engine/Source/Developer/TraceInsights/Private/Insights/ViewModels/GraphTrack.cpp
--- a/Engine/Source/Developer/TraceInsights/Private/Insights/ViewModels/GraphTrack.cpp
+++ b/Engine/Source/Developer/TraceInsights/Private/Insights/ViewModels/GraphTrack.cpp
@@ -1,6 +1,6 @@
 // Copyright 1998-2019 Epic Games, Inc. All Rights Reserved.
 
-#include "GraphTrack.h"
+#include "Insights/ViewModels/GraphTrack.h"
 
 #include "EditorStyleSet.h"
 #include "Framework/MultiBox/MultiBoxBuilder.h"
diff --git a/Engine/Source/Developer/TraceInsights/Private/Insights/ViewModels/GraphTrack.h b/Engine/Source/Developer/TraceInsights/Private/Insights/ViewModels/GraphTrack.h
deleted file mode 100644
diff --git a/Engine/Source/Developer/TraceInsights/Private/Insights/ViewModels/LoadingTimingTrack.cpp b/Engine/Source/Developer/TraceInsights/Private/Insights/ViewModels/LoadingTimingTrack.cpp
--- a/Engine/Source/Developer/TraceInsights/Private/Insights/ViewModels/LoadingTimingTrack.cpp
+++ b/Engine/Source/Developer/TraceInsights/Private/Insights/ViewModels/LoadingTimingTrack.cpp
@@ -224,18 +224,17 @@ struct FSearchLoadTimeProfilerCpuEvent
 // FLoadingMainThreadTimingTrack
 ////////////////////////////////////////////////////////////////////////////////////////////////////
 
-void FLoadingMainThreadTimingTrack::Draw(FTimingViewDrawHelper& Helper) const
+void FLoadingMainThreadTimingTrack::Draw(ITimingViewDrawHelper& Helper) const
 {
 	FTimingEventsTrack& Track = *const_cast<FLoadingMainThreadTimingTrack*>(this);
 
 	if (Helper.BeginTimeline(Track))
 	{
-		TSharedPtr<const Trace::IAnalysisSession> Session = FInsightsManager::Get()->GetSession();
-		if (Session.IsValid() && Trace::ReadLoadTimeProfilerProvider(*Session.Get()))
+		if (Trace::ReadLoadTimeProfilerProvider(GetSession()))
 		{
-			Trace::FAnalysisSessionReadScope SessionReadScope(*Session.Get());
+			Trace::FAnalysisSessionReadScope SessionReadScope(GetSession());
 
-			const Trace::ILoadTimeProfilerProvider& LoadTimeProfilerProvider = *Trace::ReadLoadTimeProfilerProvider(*Session.Get());
+			const Trace::ILoadTimeProfilerProvider& LoadTimeProfilerProvider = *Trace::ReadLoadTimeProfilerProvider(GetSession());
 
 			LoadTimeProfilerProvider.ReadMainThreadCpuTimeline([this, &Helper](const Trace::ILoadTimeProfilerProvider::CpuTimeline& Timeline)
 			{
@@ -266,33 +265,29 @@ void FLoadingMainThreadTimingTrack::Draw(FTimingViewDrawHelper& Helper) const
 ////////////////////////////////////////////////////////////////////////////////////////////////////
 
 bool FLoadingMainThreadTimingTrack::SearchTimingEvent(const double InStartTime,
-	const double InEndTime,
-	TFunctionRef<bool(double, double, uint32)> InPredicate,
-	FTimingEvent& InOutTimingEvent,
-	bool bInStopAtFirstMatch,
-	bool bInSearchForLargestEvent) const
+													  const double InEndTime,
+													  TFunctionRef<bool(double, double, uint32)> InPredicate,
+													  FTimingEvent& InOutTimingEvent,
+													  bool bInStopAtFirstMatch,
+													  bool bInSearchForLargestEvent) const
 {
 	FSearchLoadTimeProfilerCpuEvent Ctx(InStartTime, InEndTime, InPredicate, InOutTimingEvent, bInStopAtFirstMatch, bInSearchForLargestEvent);
 
-	TSharedPtr<const Trace::IAnalysisSession> Session = FInsightsManager::Get()->GetSession();
-	if (Session.IsValid())
-	{
-		Trace::FAnalysisSessionReadScope SessionReadScope(*Session.Get());
+	Trace::FAnalysisSessionReadScope SessionReadScope(GetSession());
 
-		const FTimingEventsTrack* Track = Ctx.TimingEvent.Track;
+	const FTimingEventsTrack* Track = Ctx.TimingEvent.Track;
 
-		if (Trace::ReadLoadTimeProfilerProvider(*Session.Get()))
-		{
-			const Trace::ILoadTimeProfilerProvider& LoadTimeProfilerProvider = *Trace::ReadLoadTimeProfilerProvider(*Session.Get());
+	if (Trace::ReadLoadTimeProfilerProvider(GetSession()))
+	{
+		const Trace::ILoadTimeProfilerProvider& LoadTimeProfilerProvider = *Trace::ReadLoadTimeProfilerProvider(GetSession());
 
-			LoadTimeProfilerProvider.ReadMainThreadCpuTimeline([&Ctx](const Trace::ILoadTimeProfilerProvider::CpuTimeline& Timeline)
+		LoadTimeProfilerProvider.ReadMainThreadCpuTimeline([&Ctx](const Trace::ILoadTimeProfilerProvider::CpuTimeline& Timeline)
+		{
+			Timeline.EnumerateEvents(Ctx.StartTime, Ctx.EndTime, [&Ctx](double EventStartTime, double EventEndTime, uint32 EventDepth, const Trace::FLoadTimeProfilerCpuEvent& Event)
 			{
-				Timeline.EnumerateEvents(Ctx.StartTime, Ctx.EndTime, [&Ctx](double EventStartTime, double EventEndTime, uint32 EventDepth, const Trace::FLoadTimeProfilerCpuEvent& Event)
-				{
-					Ctx.CheckEvent(EventStartTime, EventEndTime, EventDepth, Event);
-				});
+				Ctx.CheckEvent(EventStartTime, EventEndTime, EventDepth, Event);
 			});
-		}
+		});
 	}
 
 	return Ctx.bFound;
@@ -302,18 +297,17 @@ bool FLoadingMainThreadTimingTrack::SearchTimingEvent(const double InStartTime,
 // FLoadingAsyncThreadTimingTrack
 ////////////////////////////////////////////////////////////////////////////////////////////////////
 
-void FLoadingAsyncThreadTimingTrack::Draw(FTimingViewDrawHelper& Helper) const
+void FLoadingAsyncThreadTimingTrack::Draw(ITimingViewDrawHelper& Helper) const
 {
 	FTimingEventsTrack& Track = *const_cast<FLoadingAsyncThreadTimingTrack*>(this);
 
 	if (Helper.BeginTimeline(Track))
 	{
-		TSharedPtr<const Trace::IAnalysisSession> Session = FInsightsManager::Get()->GetSession();
-		if (Session.IsValid() && Trace::ReadLoadTimeProfilerProvider(*Session.Get()))
+		if (Trace::ReadLoadTimeProfilerProvider(GetSession()))
 		{
-			Trace::FAnalysisSessionReadScope SessionReadScope(*Session.Get());
+			Trace::FAnalysisSessionReadScope SessionReadScope(GetSession());
 
-			const Trace::ILoadTimeProfilerProvider& LoadTimeProfilerProvider = *Trace::ReadLoadTimeProfilerProvider(*Session.Get());
+			const Trace::ILoadTimeProfilerProvider& LoadTimeProfilerProvider = *Trace::ReadLoadTimeProfilerProvider(GetSession());
 
 			LoadTimeProfilerProvider.ReadAsyncLoadingThreadCpuTimeline([this, &Helper](const Trace::ILoadTimeProfilerProvider::CpuTimeline& Timeline)
 			{
@@ -352,25 +346,21 @@ bool FLoadingAsyncThreadTimingTrack::SearchTimingEvent(const double InStartTime,
 {
 	FSearchLoadTimeProfilerCpuEvent Ctx(InStartTime, InEndTime, InPredicate, InOutTimingEvent, bInStopAtFirstMatch, bInSearchForLargestEvent);
 
-	TSharedPtr<const Trace::IAnalysisSession> Session = FInsightsManager::Get()->GetSession();
-	if (Session.IsValid())
-	{
-		Trace::FAnalysisSessionReadScope SessionReadScope(*Session.Get());
+	Trace::FAnalysisSessionReadScope SessionReadScope(GetSession());
 
-		const FTimingEventsTrack* Track = Ctx.TimingEvent.Track;
+	const FTimingEventsTrack* Track = Ctx.TimingEvent.Track;
 
-		if (Trace::ReadLoadTimeProfilerProvider(*Session.Get()))
-		{
-			const Trace::ILoadTimeProfilerProvider& LoadTimeProfilerProvider = *Trace::ReadLoadTimeProfilerProvider(*Session.Get());
+	if (Trace::ReadLoadTimeProfilerProvider(GetSession()))
+	{
+		const Trace::ILoadTimeProfilerProvider& LoadTimeProfilerProvider = *Trace::ReadLoadTimeProfilerProvider(GetSession());
 
-			LoadTimeProfilerProvider.ReadAsyncLoadingThreadCpuTimeline([&Ctx](const Trace::ILoadTimeProfilerProvider::CpuTimeline& Timeline)
+		LoadTimeProfilerProvider.ReadAsyncLoadingThreadCpuTimeline([&Ctx](const Trace::ILoadTimeProfilerProvider::CpuTimeline& Timeline)
+		{
+			Timeline.EnumerateEvents(Ctx.StartTime, Ctx.EndTime, [&Ctx](double EventStartTime, double EventEndTime, uint32 EventDepth, const Trace::FLoadTimeProfilerCpuEvent& Event)
 			{
-				Timeline.EnumerateEvents(Ctx.StartTime, Ctx.EndTime, [&Ctx](double EventStartTime, double EventEndTime, uint32 EventDepth, const Trace::FLoadTimeProfilerCpuEvent& Event)
-				{
-					Ctx.CheckEvent(EventStartTime, EventEndTime, EventDepth, Event);
-				});
+				Ctx.CheckEvent(EventStartTime, EventEndTime, EventDepth, Event);
 			});
-		}
+		});
 	}
 
 	return Ctx.bFound;
diff --git a/Engine/Source/Developer/TraceInsights/Private/Insights/ViewModels/LoadingTimingTrack.h b/Engine/Source/Developer/TraceInsights/Private/Insights/ViewModels/LoadingTimingTrack.h
--- a/Engine/Source/Developer/TraceInsights/Private/Insights/ViewModels/LoadingTimingTrack.h
+++ b/Engine/Source/Developer/TraceInsights/Private/Insights/ViewModels/LoadingTimingTrack.h
@@ -60,7 +60,7 @@ public:
 	{
 	}
 
-	virtual void Draw(FTimingViewDrawHelper& Helper) const override;
+	virtual void Draw(ITimingViewDrawHelper& Helper) const override;
 	virtual bool SearchTimingEvent(const double InStartTime, const double InEndTime, TFunctionRef<bool(double, double, uint32)> InPredicate, FTimingEvent& InOutTimingEvent, bool bInStopAtFirstMatch, bool bInSearchForLargestEvent) const override;
 };
 
@@ -74,7 +74,7 @@ public:
 	{
 	}
 
-	virtual void Draw(FTimingViewDrawHelper& Helper) const override;
+	virtual void Draw(ITimingViewDrawHelper& Helper) const override;
 	virtual bool SearchTimingEvent(const double InStartTime, const double InEndTime, TFunctionRef<bool(double, double, uint32)> InPredicate, FTimingEvent& InOutTimingEvent, bool bInStopAtFirstMatch, bool bInSearchForLargestEvent) const override;
 };
 
diff --git a/Engine/Source/Developer/TraceInsights/Private/Insights/ViewModels/MarkersTimingTrack.cpp b/Engine/Source/Developer/TraceInsights/Private/Insights/ViewModels/MarkersTimingTrack.cpp
--- a/Engine/Source/Developer/TraceInsights/Private/Insights/ViewModels/MarkersTimingTrack.cpp
+++ b/Engine/Source/Developer/TraceInsights/Private/Insights/ViewModels/MarkersTimingTrack.cpp
@@ -136,12 +136,11 @@ void FMarkersTimingTrack::Update(const FTimingTrackViewport& InViewport)
 {
 	FTimeMarkerTrackBuilder Builder(*this, InViewport);
 
-	TSharedPtr<const Trace::IAnalysisSession> Session = FInsightsManager::Get()->GetSession();
-	if (Session.IsValid())
+	if (IsSessionValid())
 	{
-		Trace::FAnalysisSessionReadScope SessionReadScope(*Session.Get());
+		Trace::FAnalysisSessionReadScope SessionReadScope(GetSession());
 
-		const Trace::ILogProvider& LogProvider = Trace::ReadLogProvider(*Session.Get());
+		const Trace::ILogProvider& LogProvider = Trace::ReadLogProvider(GetSession());
 		Builder.BeginLog(LogProvider);
 
 		LogProvider.EnumerateMessages(
@@ -232,7 +231,7 @@ void FMarkersTimingTrack::Draw(FDrawContext& DrawContext, const FTimingTrackView
 
 void FMarkersTimingTrack::DrawHeader(FDrawContext& DrawContext, bool bFirstDraw) const
 {
-	const FString Name = IsBookmarksTrack() ? TEXT("Bookmarks") : TEXT("Logs");
+	const FString NameString = IsBookmarksTrack() ? TEXT("Bookmarks") : TEXT("Logs");
 	//const float ArrowSize = 6.0f;
 	const float ArrowSizeX = 4.0f;
 	const float ArrowSizeY = 8.0f;
@@ -262,7 +261,7 @@ void FMarkersTimingTrack::DrawHeader(FDrawContext& DrawContext, bool bFirstDraw)
 	}
 
 	// Draw "Bookmarks" or "Logs" text.
-	DrawContext.DrawText(2.0f, GetPosY() + 1.0f, Name, Font, Color);
+	DrawContext.DrawText(2.0f, GetPosY() + 1.0f, NameString, Font, Color);
 
 	if (IsCollapsed())
 	{
diff --git a/Engine/Source/Developer/TraceInsights/Private/Insights/ViewModels/ThreadTimingTrack.cpp b/Engine/Source/Developer/TraceInsights/Private/Insights/ViewModels/ThreadTimingTrack.cpp
--- a/Engine/Source/Developer/TraceInsights/Private/Insights/ViewModels/ThreadTimingTrack.cpp
+++ b/Engine/Source/Developer/TraceInsights/Private/Insights/ViewModels/ThreadTimingTrack.cpp
@@ -20,18 +20,17 @@
 
 ////////////////////////////////////////////////////////////////////////////////////////////////////
 
-void FThreadTimingTrack::Draw(FTimingViewDrawHelper& Helper) const
+void FThreadTimingTrack::Draw(ITimingViewDrawHelper& Helper) const
 {
 	FTimingEventsTrack& Track = *const_cast<FThreadTimingTrack*>(this);
 
 	if (Helper.BeginTimeline(Track))
 	{
-		TSharedPtr<const Trace::IAnalysisSession> Session = FInsightsManager::Get()->GetSession();
-		if (Session.IsValid() && Trace::ReadTimingProfilerProvider(*Session.Get()))
+		if (Trace::ReadTimingProfilerProvider(GetSession()))
 		{
-			Trace::FAnalysisSessionReadScope SessionReadScope(*Session.Get());
+			Trace::FAnalysisSessionReadScope SessionReadScope(GetSession());
 
-			const Trace::ITimingProfilerProvider& TimingProfilerProvider = *Trace::ReadTimingProfilerProvider(*Session.Get());
+			const Trace::ITimingProfilerProvider& TimingProfilerProvider = *Trace::ReadTimingProfilerProvider(GetSession());
 
 			TimingProfilerProvider.ReadTimers([this, &Helper, &Track, &TimingProfilerProvider](const Trace::FTimingProfilerTimer* Timers, uint64 TimersCount)
 			{
@@ -149,14 +148,13 @@ void FThreadTimingTrack::GetParentAndRoot(const FTimingEvent& TimingEvent, FTimi
 
 	if (TimingEvent.Depth > 0)
 	{
-		TSharedPtr<const Trace::IAnalysisSession> Session = FInsightsManager::Get()->GetSession();
-		if (Session.IsValid())
+		if (IsSessionValid())
 		{
-			Trace::FAnalysisSessionReadScope SessionReadScope(*Session.Get());
+			Trace::FAnalysisSessionReadScope SessionReadScope(GetSession());
 
-			if (Trace::ReadTimingProfilerProvider(*Session.Get()))
+			if (Trace::ReadTimingProfilerProvider(GetSession()))
 			{
-				const Trace::ITimingProfilerProvider& TimingProfilerProvider = *Trace::ReadTimingProfilerProvider(*Session.Get());
+				const Trace::ITimingProfilerProvider& TimingProfilerProvider = *Trace::ReadTimingProfilerProvider(GetSession());
 
 				TimingProfilerProvider.ReadTimeline(TimingEvent.Track->GetId(), [&](const Trace::ITimingProfilerProvider::Timeline& Timeline)
 				{
@@ -239,23 +237,19 @@ bool FThreadTimingTrack::SearchTimingEvent(const double InStartTime,
 
 	FSearchTimingEventContext Ctx(InStartTime, InEndTime, InPredicate, InOutTimingEvent, bInStopAtFirstMatch, bInSearchForLargestEvent);
 
-	TSharedPtr<const Trace::IAnalysisSession> Session = FInsightsManager::Get()->GetSession();
-	if (Session.IsValid())
+	Trace::FAnalysisSessionReadScope SessionReadScope(GetSession());
+
+	if (Trace::ReadTimingProfilerProvider(GetSession()))
 	{
-		Trace::FAnalysisSessionReadScope SessionReadScope(*Session.Get());
+		const Trace::ITimingProfilerProvider& TimingProfilerProvider = *Trace::ReadTimingProfilerProvider(GetSession());
 
-		if (Trace::ReadTimingProfilerProvider(*Session.Get()))
+		TimingProfilerProvider.ReadTimeline(Ctx.TimingEvent.Track->GetId(), [&Ctx](const Trace::ITimingProfilerProvider::Timeline& Timeline)
 		{
-			const Trace::ITimingProfilerProvider& TimingProfilerProvider = *Trace::ReadTimingProfilerProvider(*Session.Get());
-
-			TimingProfilerProvider.ReadTimeline(Ctx.TimingEvent.Track->GetId(), [&Ctx](const Trace::ITimingProfilerProvider::Timeline& Timeline)
+			Timeline.EnumerateEvents(Ctx.StartTime, Ctx.EndTime, [&Ctx](double EventStartTime, double EventEndTime, uint32 EventDepth, const Trace::FTimingProfilerEvent& Event)
 			{
-				Timeline.EnumerateEvents(Ctx.StartTime, Ctx.EndTime, [&Ctx](double EventStartTime, double EventEndTime, uint32 EventDepth, const Trace::FTimingProfilerEvent& Event)
-				{
-					Ctx.CheckEvent(EventStartTime, EventEndTime, EventDepth, Event);
-				});
+				Ctx.CheckEvent(EventStartTime, EventEndTime, EventDepth, Event);
 			});
-		}
+		});
 	}
 
 	return Ctx.bFound;
@@ -267,73 +261,105 @@ void FThreadTimingTrack::ComputeTimingEventStats(FTimingEvent& InOutTimingEvent)
 {
 	if (InOutTimingEvent.IsValid())
 	{
-		TSharedPtr<const Trace::IAnalysisSession> Session = FInsightsManager::Get()->GetSession();
-		if (Session.IsValid())
+		Trace::FAnalysisSessionReadScope SessionReadScope(GetSession());
+
+		if (Trace::ReadTimingProfilerProvider(GetSession()))
 		{
-			Trace::FAnalysisSessionReadScope SessionReadScope(*Session.Get());
+			const Trace::ITimingProfilerProvider& TimingProfilerProvider = *Trace::ReadTimingProfilerProvider(GetSession());
 
-			if (Trace::ReadTimingProfilerProvider(*Session.Get()))
+			// Compute Exclusive Time.
+			TimingProfilerProvider.ReadTimeline(InOutTimingEvent.Track->GetId(), [&InOutTimingEvent](const Trace::ITimingProfilerProvider::Timeline& Timeline)
 			{
-				const Trace::ITimingProfilerProvider& TimingProfilerProvider = *Trace::ReadTimingProfilerProvider(*Session.Get());
-
-				// Compute Exclusive Time.
-				TimingProfilerProvider.ReadTimeline(InOutTimingEvent.Track->GetId(), [&InOutTimingEvent](const Trace::ITimingProfilerProvider::Timeline& Timeline)
+				struct FEnumerationState
 				{
-					struct FEnumerationState
-					{
-						double EventStartTime;
-						double EventEndTime;
-						uint64 EventDepth;
-
-						uint64 CurrentDepth;
-						double LastTime;
-						double ExclusiveTime;
-						bool IsInEventScope;
-					};
-					FEnumerationState State;
-
-					State.EventStartTime = InOutTimingEvent.StartTime;
-					State.EventEndTime = InOutTimingEvent.EndTime;
-					State.EventDepth = InOutTimingEvent.Depth;
-
-					State.CurrentDepth = 0;
-					State.LastTime = 0.0;
-					State.ExclusiveTime = 0.0;
-					State.IsInEventScope = false;
-
-					Timeline.EnumerateEvents(InOutTimingEvent.StartTime, InOutTimingEvent.EndTime, [&State](bool IsEnter, double Time, const Trace::FTimingProfilerEvent& Event)
+					double EventStartTime;
+					double EventEndTime;
+					uint64 EventDepth;
+
+					uint64 CurrentDepth;
+					double LastTime;
+					double ExclusiveTime;
+					bool IsInEventScope;
+				};
+				FEnumerationState State;
+
+				State.EventStartTime = InOutTimingEvent.StartTime;
+				State.EventEndTime = InOutTimingEvent.EndTime;
+				State.EventDepth = InOutTimingEvent.Depth;
+
+				State.CurrentDepth = 0;
+				State.LastTime = 0.0;
+				State.ExclusiveTime = 0.0;
+				State.IsInEventScope = false;
+
+				Timeline.EnumerateEvents(InOutTimingEvent.StartTime, InOutTimingEvent.EndTime, [&State](bool IsEnter, double Time, const Trace::FTimingProfilerEvent& Event)
+				{
+					if (IsEnter)
 					{
-						if (IsEnter)
+						if (State.IsInEventScope && State.CurrentDepth == State.EventDepth + 1)
 						{
-							if (State.IsInEventScope && State.CurrentDepth == State.EventDepth + 1)
-							{
-								State.ExclusiveTime += Time - State.LastTime;
-							}
-							if (State.CurrentDepth == State.EventDepth && Time == State.EventStartTime)
-							{
-								State.IsInEventScope = true;
-							}
-							++State.CurrentDepth;
+							State.ExclusiveTime += Time - State.LastTime;
 						}
-						else
+						if (State.CurrentDepth == State.EventDepth && Time == State.EventStartTime)
 						{
-							--State.CurrentDepth;
-							if (State.CurrentDepth == State.EventDepth && Time == State.EventEndTime)
-							{
-								State.IsInEventScope = false;
-								State.ExclusiveTime += Time - State.LastTime;
-							}
+							State.IsInEventScope = true;
 						}
-						State.LastTime = Time;
-					});
-
-					InOutTimingEvent.ExclusiveTime = State.ExclusiveTime;
+						++State.CurrentDepth;
+					}
+					else
+					{
+						--State.CurrentDepth;
+						if (State.CurrentDepth == State.EventDepth && Time == State.EventEndTime)
+						{
+							State.IsInEventScope = false;
+							State.ExclusiveTime += Time - State.LastTime;
+						}
+					}
+					State.LastTime = Time;
 				});
-			}
+
+				InOutTimingEvent.ExclusiveTime = State.ExclusiveTime;
+			});
 		}
 	}
 }
 
 ////////////////////////////////////////////////////////////////////////////////////////////////////
 
+void FThreadTimingTrack::OnEventSelected(const FTimingEvent& SelectedTimingEvent) const
+{
+	// Select the timer node corresponding to timing event type of selected timing event.
+	FTimingProfilerManager::Get()->SetSelectedTimer(SelectedTimingEvent.TypeId);
+}
+
+////////////////////////////////////////////////////////////////////////////////////////////////////
+
+void FThreadTimingTrack::BuildContextMenu(FMenuBuilder& MenuBuilder)
+{
+	FText SectionTitle;
+	if(GetGroupName() != nullptr)
+	{
+		SectionTitle = FText::Format(LOCTEXT("TrackTitleGroupFmt", "{0} (Group: {1})"), FText::FromString(GetName()), FText::FromString(GetGroupName()));
+	}
+	else
+	{
+		SectionTitle = FText::FromString(GetName());
+	}
+
+	MenuBuilder.BeginSection(TEXT("Empty"), SectionTitle);
+	{
+		MenuBuilder.AddMenuEntry(
+			LOCTEXT("ContextMenu_NA", "N/A"),
+			LOCTEXT("ContextMenu_NA_Desc", "No actions available."),
+			FSlateIcon(),
+			FUIAction(FExecuteAction(), FCanExecuteAction::CreateLambda([](){ return false; })),
+			NAME_None,
+			EUserInterfaceActionType::Button
+		);
+	}
+	MenuBuilder.EndSection();
+}
+
+////////////////////////////////////////////////////////////////////////////////////////////////////
+
 #undef LOCTEXT_NAMESPACE
diff --git a/Engine/Source/Developer/TraceInsights/Private/Insights/ViewModels/ThreadTimingTrack.h b/Engine/Source/Developer/TraceInsights/Private/Insights/ViewModels/ThreadTimingTrack.h
--- a/Engine/Source/Developer/TraceInsights/Private/Insights/ViewModels/ThreadTimingTrack.h
+++ b/Engine/Source/Developer/TraceInsights/Private/Insights/ViewModels/ThreadTimingTrack.h
@@ -26,12 +26,14 @@ public:
 	void SetThreadId(uint32 InThreadId) { ThreadId = InThreadId; }
 	uint32 GetThreadId() const { return ThreadId; }
 
-	virtual void Draw(FTimingViewDrawHelper& Helper) const override;
+	virtual void Draw(ITimingViewDrawHelper& Helper) const override;
 	virtual void DrawSelectedEventInfo(const FTimingEvent& SelectedTimingEvent, const FTimingTrackViewport& Viewport, const FDrawContext& DrawContext, const FSlateBrush* WhiteBrush, const FSlateFontInfo& Font) const override;
 	virtual void InitTooltip(FTooltipDrawState& Tooltip, const FTimingEvent& HoveredTimingEvent) const override;
 
 	virtual bool SearchTimingEvent(const double InStartTime, const double InEndTime, TFunctionRef<bool(double, double, uint32)> InPredicate, FTimingEvent& InOutTimingEvent, bool bInStopAtFirstMatch, bool bInSearchForLargestEvent) const override;
 	virtual void ComputeTimingEventStats(FTimingEvent& InOutTimingEvent) const override;
+	virtual void OnEventSelected(const FTimingEvent& SelectedTimingEvent) const override;
+	virtual void BuildContextMenu(FMenuBuilder& MenuBuilder) override;
 
 private:
 	void GetParentAndRoot(const FTimingEvent& TimingEvent, FTimingEvent& OutParentTimingEvent, FTimingEvent& OutRootTimingEvent) const;
diff --git a/Engine/Source/Developer/TraceInsights/Private/Insights/ViewModels/TimeRulerTrack.cpp b/Engine/Source/Developer/TraceInsights/Private/Insights/ViewModels/TimeRulerTrack.cpp
--- a/Engine/Source/Developer/TraceInsights/Private/Insights/ViewModels/TimeRulerTrack.cpp
+++ b/Engine/Source/Developer/TraceInsights/Private/Insights/ViewModels/TimeRulerTrack.cpp
@@ -23,6 +23,8 @@ FTimeRulerTrack::FTimeRulerTrack(uint64 InTrackId)
 	: FBaseTimingTrack(InTrackId)
 	, WhiteBrush(FInsightsStyle::Get().GetBrush("WhiteBrush"))
 	, Font(FCoreStyle::GetDefaultFontStyle("Regular", 8))
+	, CrtMousePosTextWidth(0.0f)
+	, CrtTimeMarkerTextWidth(0.0f)
 {
 }
 
@@ -38,7 +40,7 @@ void FTimeRulerTrack::Reset()
 {
 	FBaseTimingTrack::Reset();
 
-	constexpr float TimeRulerHeight = 22.0f;
+	constexpr float TimeRulerHeight = 24.0f;
 	SetHeight(TimeRulerHeight);
 }
 
@@ -69,13 +71,13 @@ void FTimeRulerTrack::DrawBackground(FDrawContext& DrawContext, const FTimingTra
 
 ////////////////////////////////////////////////////////////////////////////////////////////////////
 
-void FTimeRulerTrack::Draw(FDrawContext& DrawContext, const FTimingTrackViewport& Viewport, const FVector2D& MousePosition, const bool bIsSelecting, const double SelectionStartTime, const double SelectionEndTime) const
+void FTimeRulerTrack::Draw(FDrawContext& DrawContext, const FTimingTrackViewport& Viewport, const FVector2D& MousePosition, const bool bIsSelecting, const double SelectionStartTime, const double SelectionEndTime, const double TimeMarker) const
 {
 	const float MinorTickMark = 5.0f;
 	const float MajorTickMark = 20 * MinorTickMark;
 
 	const float MinorTickMarkHeight = 5.0f;
-	const float MajorTickMarkHeight = 9.0f;
+	const float MajorTickMarkHeight = 11.0f;
 
 	const float TextY = GetPosY() + MajorTickMarkHeight;
 
@@ -146,7 +148,6 @@ void FTimeRulerTrack::Draw(FDrawContext& DrawContext, const FTimingTrackViewport
 		}
 
 		const float MousePosTextWidth = FMath::RoundToFloat(FontMeasureService->Measure(MousePosText, Font).X);
-		static float CrtMousePosTextWidth = 0.0f;
 
 		if (!FMath::IsNearlyEqual(CrtMousePosTextWidth, MousePosTextWidth))
 		{
@@ -191,6 +192,52 @@ void FTimeRulerTrack::Draw(FDrawContext& DrawContext, const FTimingTrackViewport
 		DrawContext.DrawText(X - MousePosTextWidth / 2, TextY, MousePosText, Font, MousePosTextForegroundColor);
 		DrawContext.LayerId++;
 	}
+
+	// Draw the time marker (orange vertical line).
+	float TimeMarkerX = Viewport.TimeToSlateUnitsRounded(TimeMarker);
+	if (TimeMarkerX >= 0.0f && TimeMarkerX < Viewport.GetWidth())
+	{
+		const FLinearColor TimeMarkerColor(0.85f, 0.5f, 0.03f, 0.5f);
+		const FLinearColor TimeMarkerTextBackgroundColor(TimeMarkerColor.CopyWithNewOpacity(1.0f));
+		FLinearColor TimeMarkerTextForegroundColor(0.1f, 0.1f, 0.1f, 1.0f);
+
+		DrawContext.DrawBox(TimeMarkerX, 0.0f, 1.0f, Viewport.GetHeight(), WhiteBrush, TimeMarkerColor);
+		DrawContext.LayerId++;
+
+		// Time at current marker
+		FString TimeMarkerText;
+
+		const double MousePosPrecision = FMath::Max(DT / 100.0, TimeUtils::Nanosecond);
+		if (!MousePosition.IsZero() && MousePosition.Y >= GetPosY() && MousePosition.Y < GetPosY() + GetHeight())
+		{
+			// If mouse is hovering the time ruler, format time with a better precision (split seconds in ms, us, ns and ps).
+			TimeMarkerText = TimeUtils::FormatTimeSplit(TimeMarker, MousePosPrecision);
+		}
+		else
+		{
+			// Format current time with one more digit than the time at major tick marks.
+			TimeMarkerText = TimeUtils::FormatTime(TimeMarker, MousePosPrecision);
+		}
+
+		const float TimeMarkerTextWidth = FMath::RoundToFloat(FontMeasureService->Measure(TimeMarkerText, Font).X);
+
+		if (!FMath::IsNearlyEqual(CrtTimeMarkerTextWidth, TimeMarkerTextWidth))
+		{
+			// Animate the box's width (to avoid flickering).
+			CrtTimeMarkerTextWidth = CrtTimeMarkerTextWidth * 0.6f + TimeMarkerTextWidth * 0.4f;
+		}
+
+		float X = TimeMarkerX;
+		float W = CrtTimeMarkerTextWidth + 4.0f;
+
+		// Fill the time marker box.
+		DrawContext.DrawBox(X - W / 2, 0.0f, W, 12.0f, WhiteBrush, TimeMarkerTextBackgroundColor);
+		DrawContext.LayerId++;
+
+		// Draw time marker text.
+		DrawContext.DrawText(X - TimeMarkerTextWidth / 2, 0.0f, TimeMarkerText, Font, TimeMarkerTextForegroundColor);
+		DrawContext.LayerId++;
+	}
 }
 
 ////////////////////////////////////////////////////////////////////////////////////////////////////
diff --git a/Engine/Source/Developer/TraceInsights/Private/Insights/ViewModels/TimeRulerTrack.h b/Engine/Source/Developer/TraceInsights/Private/Insights/ViewModels/TimeRulerTrack.h
--- a/Engine/Source/Developer/TraceInsights/Private/Insights/ViewModels/TimeRulerTrack.h
+++ b/Engine/Source/Developer/TraceInsights/Private/Insights/ViewModels/TimeRulerTrack.h
@@ -27,7 +27,8 @@ public:
 				const FVector2D& MousePosition = FVector2D(0.0f, 0.0f),
 				const bool bIsSelecting = false,
 				const double SelectionStartTime = 0.0,
-				const double SelectionEndTime = 0.0) const;
+				const double SelectionEndTime = 0.0,
+				const double TimeMarker = 0.0) const;
 
 private:
 	void DrawBackground(FDrawContext& DrawContext, const FTimingTrackViewport& Viewport) const;
@@ -36,6 +37,12 @@ public:
 	// Slate resources
 	const FSlateBrush* WhiteBrush;
 	const FSlateFontInfo Font;
+
+	// Smoothed mouse pos text width to avoid flickering
+	mutable float CrtMousePosTextWidth;
+
+	// Smoothed time marker text width to avoid flickering
+	mutable float CrtTimeMarkerTextWidth;
 };
 
 ////////////////////////////////////////////////////////////////////////////////////////////////////
diff --git a/Engine/Source/Developer/TraceInsights/Private/Insights/ViewModels/TimingEvent.h b/Engine/Source/Developer/TraceInsights/Private/Insights/ViewModels/TimingEvent.h
deleted file mode 100644
diff --git a/Engine/Source/Developer/TraceInsights/Private/Insights/ViewModels/TimingEventsTrack.cpp b/Engine/Source/Developer/TraceInsights/Private/Insights/ViewModels/TimingEventsTrack.cpp
--- a/Engine/Source/Developer/TraceInsights/Private/Insights/ViewModels/TimingEventsTrack.cpp
+++ b/Engine/Source/Developer/TraceInsights/Private/Insights/ViewModels/TimingEventsTrack.cpp
@@ -1,6 +1,6 @@
 // Copyright 1998-2019 Epic Games, Inc. All Rights Reserved.
 
-#include "TimingEventsTrack.h"
+#include "Insights/ViewModels/TimingEventsTrack.h"
 
 #define LOCTEXT_NAMESPACE "TimingEventsTrack"
 
@@ -86,13 +86,8 @@ bool FTimingEventsTrack::bUseDownSampling = true;
 ////////////////////////////////////////////////////////////////////////////////////////////////////
 
 FTimingEventsTrack::FTimingEventsTrack(uint64 InTrackId, const FName& InType, const FName& InSubType, const FString& InName)
-	: FBaseTimingTrack(InTrackId)
-	, Type(InType)
-	, SubType(InSubType)
-	, Name(InName)
-	, Order(0)
+	: FBaseTimingTrack(InTrackId, InType, InSubType, InName)
 	, NumLanes(0)
-	, bIsCollapsed(false)
 {
 }
 
@@ -109,7 +104,6 @@ void FTimingEventsTrack::Reset()
 	FBaseTimingTrack::Reset();
 
 	NumLanes = 0;
-	bIsCollapsed = false;
 }
 
 ////////////////////////////////////////////////////////////////////////////////////////////////////
diff --git a/Engine/Source/Developer/TraceInsights/Private/Insights/ViewModels/TimingEventsTrack.h b/Engine/Source/Developer/TraceInsights/Private/Insights/ViewModels/TimingEventsTrack.h
deleted file mode 100644
diff --git a/Engine/Source/Developer/TraceInsights/Private/Insights/ViewModels/TimingGraphTrack.cpp b/Engine/Source/Developer/TraceInsights/Private/Insights/ViewModels/TimingGraphTrack.cpp
--- a/Engine/Source/Developer/TraceInsights/Private/Insights/ViewModels/TimingGraphTrack.cpp
+++ b/Engine/Source/Developer/TraceInsights/Private/Insights/ViewModels/TimingGraphTrack.cpp
@@ -115,11 +115,10 @@ TSharedPtr<FTimingGraphSeries> FTimingGraphTrack::AddStatsCounterSeries(uint32 C
 	bool bIsFloatingPoint = false;
 	bool bIsMemory = false;
 
-	TSharedPtr<const Trace::IAnalysisSession> Session = FInsightsManager::Get()->GetSession();
-	if (Session.IsValid())
+	if (IsSessionValid())
 	{
-		Trace::FAnalysisSessionReadScope SessionReadScope(*Session.Get());
-		const Trace::ICounterProvider& CountersProvider = Trace::ReadCounterProvider(*Session.Get());
+		Trace::FAnalysisSessionReadScope SessionReadScope(GetSession());
+		const Trace::ICounterProvider& CountersProvider = Trace::ReadCounterProvider(GetSession());
 		if (CounterId < CountersProvider.GetCounterCount())
 		{
 			CountersProvider.ReadCounter(CounterId, [&](const Trace::ICounter& Counter)
@@ -232,11 +231,10 @@ void FTimingGraphTrack::Update(const FTimingTrackViewport& Viewport)
 void FTimingGraphTrack::UpdateFrameSeries(FTimingGraphSeries& Series, const FTimingTrackViewport& Viewport)
 {
 	FGraphTrackBuilder Builder(*this, Series, Viewport);
-	TSharedPtr<const Trace::IAnalysisSession> Session = FInsightsManager::Get()->GetSession();
-	if (Session.IsValid())
+	if (IsSessionValid())
 	{
-		Trace::FAnalysisSessionReadScope SessionReadScope(*Session.Get());
-		const Trace::IFrameProvider& FramesProvider = ReadFrameProvider(*Session.Get());
+		Trace::FAnalysisSessionReadScope SessionReadScope(GetSession());
+		const Trace::IFrameProvider& FramesProvider = ReadFrameProvider(GetSession());
 		ETraceFrameType FrameType = static_cast<ETraceFrameType>(Series.Id);
 		uint64 FrameCount = FramesProvider.GetFrameCount(FrameType);
 		FramesProvider.EnumerateFrames(FrameType, 0, FrameCount - 1, [&Builder](const Trace::FFrame& Frame)
@@ -260,11 +258,10 @@ void FTimingGraphTrack::UpdateTimerSeries(FTimingGraphSeries& Series, const FTim
 void FTimingGraphTrack::UpdateStatsCounterSeries(FTimingGraphSeries& Series, const FTimingTrackViewport& Viewport)
 {
 	FGraphTrackBuilder Builder(*this, Series, Viewport);
-	TSharedPtr<const Trace::IAnalysisSession> Session = FInsightsManager::Get()->GetSession();
-	if (Session.IsValid())
+	if (IsSessionValid())
 	{
-		Trace::FAnalysisSessionReadScope SessionReadScope(*Session.Get());
-		const Trace::ICounterProvider& CountersProvider = Trace::ReadCounterProvider(*Session.Get());
+		Trace::FAnalysisSessionReadScope SessionReadScope(GetSession());
+		const Trace::ICounterProvider& CountersProvider = Trace::ReadCounterProvider(GetSession());
 		CountersProvider.ReadCounter(Series.Id, [this, &Viewport, &Builder, &Series](const Trace::ICounter& Counter)
 		{
 			if (Series.IsAutoZoomEnabled())
diff --git a/Engine/Source/Developer/TraceInsights/Private/Insights/ViewModels/TimingTrackViewport.cpp b/Engine/Source/Developer/TraceInsights/Private/Insights/ViewModels/TimingTrackViewport.cpp
--- a/Engine/Source/Developer/TraceInsights/Private/Insights/ViewModels/TimingTrackViewport.cpp
+++ b/Engine/Source/Developer/TraceInsights/Private/Insights/ViewModels/TimingTrackViewport.cpp
@@ -1,6 +1,6 @@
 // Copyright 1998-2019 Epic Games, Inc. All Rights Reserved.
 
-#include "TimingTrackViewport.h"
+#include "Insights/ViewModels/TimingTrackViewport.h"
 
 #include "Widgets/Layout/SScrollBar.h"
 
diff --git a/Engine/Source/Developer/TraceInsights/Private/Insights/ViewModels/TimingTrackViewport.h b/Engine/Source/Developer/TraceInsights/Private/Insights/ViewModels/TimingTrackViewport.h
deleted file mode 100644
diff --git a/Engine/Source/Developer/TraceInsights/Private/Insights/ViewModels/TimingViewDrawHelper.h b/Engine/Source/Developer/TraceInsights/Private/Insights/ViewModels/TimingViewDrawHelper.h
--- a/Engine/Source/Developer/TraceInsights/Private/Insights/ViewModels/TimingViewDrawHelper.h
+++ b/Engine/Source/Developer/TraceInsights/Private/Insights/ViewModels/TimingViewDrawHelper.h
@@ -5,6 +5,7 @@
 #include "CoreMinimal.h"
 #include "Fonts/SlateFontInfo.h"
 #include "Math/Color.h"
+#include "Insights/ITimingViewDrawHelper.h"
 
 ////////////////////////////////////////////////////////////////////////////////////////////////////
 
@@ -17,7 +18,7 @@ class FTimingTrackViewport;
 
 ////////////////////////////////////////////////////////////////////////////////////////////////////
 
-class FTimingViewDrawHelper
+class FTimingViewDrawHelper final : public ITimingViewDrawHelper
 {
 public:
 	enum class EHighlightMode : uint32
@@ -94,9 +95,13 @@ public:
 	FTimingViewDrawHelper(const FTimingViewDrawHelper&) = delete;
 	FTimingViewDrawHelper& operator=(const FTimingViewDrawHelper&) = delete;
 
-	const FDrawContext& GetDrawContext() const { return DrawContext; }
-	const FTimingTrackViewport& GetViewport() const { return Viewport; }
-	const FTimingEventsTrackLayout& GetLayout() const { return Layout; }
+	// ITimingViewDrawHelper interface
+	virtual const FDrawContext& GetDrawContext() const override { return DrawContext; }
+	virtual const FTimingTrackViewport& GetViewport() const override { return Viewport; }
+	virtual const FTimingEventsTrackLayout& GetLayout() const override { return Layout; }
+	virtual bool BeginTimeline(FTimingEventsTrack& Track) override;
+	virtual void AddEvent(double StartTime, double EndTime, uint32 Depth, const TCHAR* EventName, uint32 Color = 0) override;
+	virtual void EndTimeline(FTimingEventsTrack& Track) override;
 
 	const FSlateBrush* GetWhiteBrush() const { return WhiteBrush; }
 	const FSlateFontInfo& GetEventFont() const { return EventFont; }
@@ -112,9 +117,6 @@ public:
 
 	//TODO: move the following in a Builder class
 	void BeginTimelines();
-	bool BeginTimeline(FTimingEventsTrack& Track);
-	void AddEvent(double StartTime, double EndTime, uint32 Depth, const TCHAR* EventName, uint32 Color = 0);
-	void EndTimeline(FTimingEventsTrack& Track);
 	void EndTimelines();
 
 private:
diff --git a/Engine/Source/Developer/TraceInsights/Private/Insights/ViewModels/TooltipDrawState.cpp b/Engine/Source/Developer/TraceInsights/Private/Insights/ViewModels/TooltipDrawState.cpp
--- a/Engine/Source/Developer/TraceInsights/Private/Insights/ViewModels/TooltipDrawState.cpp
+++ b/Engine/Source/Developer/TraceInsights/Private/Insights/ViewModels/TooltipDrawState.cpp
@@ -1,6 +1,6 @@
 // Copyright 1998-2019 Epic Games, Inc. All Rights Reserved.
 
-#include "TooltipDrawState.h"
+#include "Insights/ViewModels/TooltipDrawState.h"
 
 #include "Fonts/FontMeasure.h"
 #include "Fonts/SlateFontInfo.h"
diff --git a/Engine/Source/Developer/TraceInsights/Private/Insights/ViewModels/TooltipDrawState.h b/Engine/Source/Developer/TraceInsights/Private/Insights/ViewModels/TooltipDrawState.h
deleted file mode 100644
diff --git a/Engine/Source/Developer/TraceInsights/Private/Insights/Widgets/STimingView.cpp b/Engine/Source/Developer/TraceInsights/Private/Insights/Widgets/STimingView.cpp
--- a/Engine/Source/Developer/TraceInsights/Private/Insights/Widgets/STimingView.cpp
+++ b/Engine/Source/Developer/TraceInsights/Private/Insights/Widgets/STimingView.cpp
@@ -26,6 +26,7 @@
 #include "Widgets/Layout/SSpacer.h"
 #include "Widgets/SOverlay.h"
 #include "Widgets/Text/STextBlock.h"
+#include "Features/IModularFeatures.h"
 
 // Insights
 #include "Insights/Common/PaintUtils.h"
@@ -48,6 +49,7 @@
 #include "Insights/Widgets/SStatsView.h"
 #include "Insights/Widgets/STimersView.h"
 #include "Insights/Widgets/STimingProfilerWindow.h"
+#include "Insights/ITimingViewExtender.h"
 
 #include <limits>
 
@@ -65,6 +67,8 @@ const TCHAR* GetName(ELoadTimeProfilerObjectEventType Type);
 const TCHAR* GetFileActivityTypeName(Trace::EFileActivityType Type);
 uint32 GetFileActivityTypeColor(Trace::EFileActivityType Type);
 
+namespace Insights { const FName TimingViewExtenderFeatureName(TEXT("TimingViewExtender")); }
+
 ////////////////////////////////////////////////////////////////////////////////////////////////////
 
 STimingView::STimingView()
@@ -89,16 +93,17 @@ STimingView::~STimingView()
 		delete KV.Value;
 	}
 	CachedTimelines.Reset();
+
+	for (Insights::ITimingViewExtender* Extender : GetExtenders())
+	{
+		Extender->OnEndSession(*this);
+	}
 }
 
 ////////////////////////////////////////////////////////////////////////////////////////////////////
 
 void STimingView::Construct(const FArguments& InArgs)
 {
-	OnSelectionChanged = InArgs._OnSelectionChanged;
-	OnHoveredEventChanged = InArgs._OnHoveredEventChanged;
-	OnSelectedEventChanged = InArgs._OnSelectedEventChanged;
-
 	ChildSlot
 	[
 		SNew(SOverlay)
@@ -199,6 +204,14 @@ void STimingView::Reset()
 	}
 	CachedTimelines.Reset();
 
+	//////////////////////////////////////////////////
+
+	for (Insights::ITimingViewExtender* Extender : GetExtenders())
+	{
+		Extender->OnEndSession(*this);
+		Extender->OnBeginSession(*this);
+	}
+
 	//TODO: TopTracks.Reset();
 	//TODO: ScrollableTracks.Reset();
 	//TODO: BottomTracks.Reset();
@@ -208,7 +221,8 @@ void STimingView::Reset()
 
 	TimingEventsTracks.Reset();
 
-	bAreTimingEventsTracksDirty = true;
+	bTimingEventsNeedUpdate = true;
+	bTimingEventsTracksChanged = true;
 
 	FTimingEventsTrack::bUseDownSampling = true;
 
@@ -275,6 +289,7 @@ void STimingView::Reset()
 	LastSelectionType = ESelectionType::None;
 
 	TimeMarker = std::numeric_limits<double>::infinity();
+	bIsScrubbing = false;
 
 	//ThisGeometry
 
@@ -439,12 +454,13 @@ void STimingView::Tick(const FGeometry& AllottedGeometry, const double InCurrent
 
 			if (SessionTime >= Viewport.GetStartTime() && SessionTime <= Viewport.GetEndTime())
 			{
-				bAreTimingEventsTracksDirty = true;
+				bTimingEventsNeedUpdate = true;
 				MarkersTrack.SetDirtyFlag();
 			}
 		}
 
-		bool bIsTimingEventsTrackDirty = false;
+		// After this point we can add new tracks, so reset the dirty flag
+		bTimingEventsTracksChanged = false;
 
 		if (Trace::ReadLoadTimeProfilerProvider(*Session.Get()))
 		{
@@ -460,7 +476,6 @@ void STimingView::Tick(const FGeometry& AllottedGeometry, const double InCurrent
 				LoadingMainThreadTrack->SetOrder(-3);
 				LoadingMainThreadTrack->SetVisibilityFlag(bShowHideAllLoadingTracks);
 				AddTimingEventsTrack(LoadingMainThreadTrack);
-				bIsTimingEventsTrackDirty = true;
 			}
 
 			if (LoadingAsyncThreadTrack == nullptr)
@@ -470,7 +485,6 @@ void STimingView::Tick(const FGeometry& AllottedGeometry, const double InCurrent
 				LoadingMainThreadTrack->SetOrder(-2);
 				LoadingAsyncThreadTrack->SetVisibilityFlag(bShowHideAllLoadingTracks);
 				AddTimingEventsTrack(LoadingAsyncThreadTrack);
-				bIsTimingEventsTrackDirty = true;
 			}
 		}
 
@@ -484,10 +498,11 @@ void STimingView::Tick(const FGeometry& AllottedGeometry, const double InCurrent
 				IoOverviewTrack->SetOrder(-1);
 				IoOverviewTrack->SetVisibilityFlag(bShowHideAllIoTracks);
 				AddTimingEventsTrack(IoOverviewTrack);
-				bIsTimingEventsTrackDirty = true;
 			}
 		}
 
+		int32 Order = 1;
+
 		if (Trace::ReadTimingProfilerProvider(*Session.Get()))
 		{
 			const Trace::ITimingProfilerProvider& TimingProfilerProvider = *Trace::ReadTimingProfilerProvider(*Session.Get());
@@ -503,14 +518,12 @@ void STimingView::Tick(const FGeometry& AllottedGeometry, const double InCurrent
 					GpuTrack->SetOrder(0);
 					GpuTrack->SetVisibilityFlag(bShowHideAllGpuTracks);
 					AddTimingEventsTrack(GpuTrack);
-					bIsTimingEventsTrackDirty = true;
 				}
 			}
 
 			// Check available CPU tracks.
-			int32 Order = 1;
 			const Trace::IThreadProvider& ThreadProvider = Trace::ReadThreadProvider(*Session.Get());
-			ThreadProvider.EnumerateThreads([this, &Order, &bIsTimingEventsTrackDirty, &TimingProfilerProvider](const Trace::FThreadInfo& ThreadInfo)
+			ThreadProvider.EnumerateThreads([this, &Order, &TimingProfilerProvider](const Trace::FThreadInfo& ThreadInfo)
 			{
 				const TCHAR* GroupName = ThreadInfo.GroupName;
 				if (GroupName == nullptr)
@@ -564,7 +577,6 @@ void STimingView::Tick(const FGeometry& AllottedGeometry, const double InCurrent
 						}
 
 						AddTimingEventsTrack(Track);
-						bIsTimingEventsTrackDirty = true;
 					}
 					else
 					{
@@ -573,7 +585,7 @@ void STimingView::Tick(const FGeometry& AllottedGeometry, const double InCurrent
 						if (Track->GetOrder() != Order)
 						{
 							Track->SetOrder(Order);
-							bIsTimingEventsTrackDirty = true;
+							bTimingEventsTracksChanged = true;
 						}
 					}
 
@@ -592,14 +604,29 @@ void STimingView::Tick(const FGeometry& AllottedGeometry, const double InCurrent
 				IoActivityTrack->SetOrder(999999);
 				IoActivityTrack->SetVisibilityFlag(bShowHideAllIoTracks);
 				AddTimingEventsTrack(IoActivityTrack);
-				bIsTimingEventsTrackDirty = true;
 			}
 		}
 
-		if (bIsTimingEventsTrackDirty)
+		// Tick plugin extenders
+		for (Insights::ITimingViewExtender* Extender : GetExtenders())
 		{
+			Extender->Tick(*this, *Session.Get());
+		}
+
+		if (bTimingEventsTracksChanged)
+		{
+			// Let extenders re-order their track data
+			for (Insights::ITimingViewExtender* Extender : GetExtenders())
+			{
+				Extender->OnTracksChanged(*this, Order);
+			}
+
 			// The list has changed. Sort the list again.
 			Algo::SortBy(TimingEventsTracks, &FTimingEventsTrack::GetOrder);
+
+			bTimingEventsTracksChanged = false;
+
+			bTimingEventsNeedUpdate = true;
 		}
 	}
 
@@ -655,14 +682,14 @@ void STimingView::Tick(const FGeometry& AllottedGeometry, const double InCurrent
 
 		//TimeRuler.SetDirtyFlag();
 		MarkersTrack.SetDirtyFlag();
-		bAreTimingEventsTracksDirty = true;
+		bTimingEventsNeedUpdate = true;
 	}
 
 	if (bIsVerticalViewportDirty)
 	{
 		bIsVerticalViewportDirty = false;
 
-		bAreTimingEventsTracksDirty = true;
+		bTimingEventsNeedUpdate = true;
 	}
 
 	//if (TimeRuler.IsVisible() && TimeRuler.IsDirty())
@@ -685,9 +712,9 @@ void STimingView::Tick(const FGeometry& AllottedGeometry, const double InCurrent
 		TimeMarkerCacheUpdateDurationHistory.AddValue(Stopwatch.AccumulatedTime);
 	}
 
-	if (bAreTimingEventsTracksDirty)
+	if (bTimingEventsNeedUpdate)
 	{
-		bAreTimingEventsTracksDirty = false;
+		bTimingEventsNeedUpdate = false;
 
 		GraphTrack->SetDirtyFlag();
 
@@ -696,6 +723,7 @@ void STimingView::Tick(const FGeometry& AllottedGeometry, const double InCurrent
 
 		for (int32 TrackIndex = 0; TrackIndex < TimingEventsTracks.Num(); ++TrackIndex)
 		{
+			TimingEventsTracks[TrackIndex]->SetSession(Session.Get());
 			TimingEventsTracks[TrackIndex]->Update(Viewport);
 		}
 
@@ -768,6 +796,8 @@ int32 STimingView::OnPaint(const FPaintArgs& Args, const FGeometry& AllottedGeom
 	Helper.DrawBackground();
 
 	// Draw scrollable tracks.
+	TSharedPtr<const Trace::IAnalysisSession> Session = FInsightsManager::Get()->GetSession();
+	if(Session.IsValid())
 	{
 		Helper.BeginTimelines();
 
@@ -825,7 +855,7 @@ int32 STimingView::OnPaint(const FPaintArgs& Args, const FGeometry& AllottedGeom
 	// Draw the time ruler.
 	if (TimeRulerTrack.IsVisible())
 	{
-		TimeRulerTrack.Draw(DrawContext, Viewport, MousePosition, bIsSelecting, SelectionStartTime, SelectionEndTime);
+		TimeRulerTrack.Draw(DrawContext, Viewport, MousePosition, bIsSelecting, SelectionStartTime, SelectionEndTime, TimeMarker);
 	}
 
 	// Fill background for the Tracks filter combobox.
@@ -838,17 +868,6 @@ int32 STimingView::OnPaint(const FPaintArgs& Args, const FGeometry& AllottedGeom
 
 	//////////////////////////////////////////////////
 
-	// Draw the time marker (orange vertical line).
-	//DrawTimeMarker(OnPaintState);
-	float TimeMarkerX = Viewport.TimeToSlateUnitsRounded(TimeMarker);
-	if (TimeMarkerX >= 0.0f && TimeMarkerX < Viewport.GetWidth())
-	{
-		DrawContext.DrawBox(TimeMarkerX, 0.0f, 1.0f, Viewport.GetHeight(), WhiteBrush, FLinearColor(0.85f, 0.5f, 0.03f, 0.5f));
-		DrawContext.LayerId++;
-	}
-
-	//////////////////////////////////////////////////
-
 	Stopwatch.Stop();
 	DrawDurationHistory.AddValue(Stopwatch.AccumulatedTime);
 
@@ -1060,6 +1079,8 @@ void STimingView::AddTimingEventsTrack(FTimingEventsTrack* Track)
 
 	CachedTimelines.Add(Track->GetId(), Track);
 	TimingEventsTracks.Add(Track);
+
+	bTimingEventsTracksChanged = true;
 }
 
 ////////////////////////////////////////////////////////////////////////////////////////////////////
@@ -1087,6 +1108,7 @@ FReply STimingView::OnMouseButtonDown(const FGeometry& MyGeometry, const FPointe
 
 	bool bStartPanning = false;
 	bool bStartSelecting = false;
+	bool bStartScrubbing = false;
 
 	if (MouseEvent.GetEffectingButton() == EKeys::LeftMouseButton)
 	{
@@ -1099,7 +1121,11 @@ FReply STimingView::OnMouseButtonDown(const FGeometry& MyGeometry, const FPointe
 		{
 			bIsLMB_Pressed = true;
 
-			if (!bIsSpaceBarKeyPressed &&
+			if(MousePositionOnButtonDown.Y < TimeRulerTrack.GetHeight() && (MouseEvent.GetModifierKeys().IsControlDown() || MouseEvent.GetModifierKeys().IsShiftDown()))
+			{
+				bStartScrubbing = true;
+			}
+			else if (!bIsSpaceBarKeyPressed &&
 				(MousePositionOnButtonDown.Y < TimeRulerTrack.GetHeight() ||
 				 (MouseEvent.GetModifierKeys().IsControlDown() && MouseEvent.GetModifierKeys().IsShiftDown())))
 			{
@@ -1136,10 +1162,17 @@ FReply STimingView::OnMouseButtonDown(const FGeometry& MyGeometry, const FPointe
 		}
 	}
 
-	if (bStartPanning)
+	if(bStartScrubbing)
+	{
+		bIsPanning = false;
+		bIsDragging = false;
+		bIsScrubbing = true;
+	}
+	else if (bStartPanning)
 	{
 		bIsPanning = true;
 		bIsDragging = false;
+		bIsScrubbing = false;
 
 		ViewportStartTimeOnButtonDown = Viewport.GetStartTime();
 		ViewportScrollPosYOnButtonDown = Viewport.GetScrollPosY();
@@ -1164,6 +1197,7 @@ FReply STimingView::OnMouseButtonDown(const FGeometry& MyGeometry, const FPointe
 	{
 		bIsSelecting = true;
 		bIsDragging = false;
+		bIsScrubbing = false;
 
 		SelectionStartTime = Viewport.SlateUnitsToTime(MousePositionOnButtonDown.X);
 		SelectionEndTime = SelectionStartTime;
@@ -1200,6 +1234,10 @@ FReply STimingView::OnMouseButtonUp(const FGeometry& MyGeometry, const FPointerE
 				RaiseSelectionChanged();
 				bIsSelecting = false;
 			}
+			else if(bIsScrubbing)
+			{
+				bIsScrubbing = false;
+			}
 
 			if (bIsValidForMouseClick)
 			{
@@ -1240,6 +1278,10 @@ FReply STimingView::OnMouseButtonUp(const FGeometry& MyGeometry, const FPointerE
 				RaiseSelectionChanged();
 				bIsSelecting = false;
 			}
+			else if(bIsScrubbing)
+			{
+				bIsScrubbing = false;
+			}
 
 			if (bIsValidForMouseClick)
 			{
@@ -1329,6 +1371,16 @@ FReply STimingView::OnMouseMove(const FGeometry& MyGeometry, const FPointerEvent
 				RaiseSelectionChanging();
 			}
 		}
+		else if (bIsScrubbing)
+		{
+			if (HasMouseCapture())
+			{
+				bIsDragging = true;
+
+				TimeMarker = Viewport.SlateUnitsToTime(MousePosition.X);
+				RaiseTimeMarkerChanged();
+			}
+		}
 		else
 		{
 			if (MarkersTrack.IsVisible())
@@ -1592,7 +1644,7 @@ FReply STimingView::OnKeyDown(const FGeometry& MyGeometry, const FKeyEvent& InKe
 		else
 		{
 			Layout.bIsCompactMode = !Layout.bIsCompactMode;
-			bAreTimingEventsTracksDirty = true;
+			bTimingEventsNeedUpdate = true;
 			return FReply::Handled();
 		}
 	}
@@ -1711,7 +1763,7 @@ FReply STimingView::OnKeyDown(const FGeometry& MyGeometry, const FKeyEvent& InKe
 	else if (InKeyEvent.GetKey() == EKeys::D) // debug: toggles down-sampling on/off
 	{
 		FTimingEventsTrack::bUseDownSampling = !FTimingEventsTrack::bUseDownSampling;
-		bAreTimingEventsTracksDirty = true;
+		bTimingEventsNeedUpdate = true;
 		return FReply::Handled();
 	}
 	else if (InKeyEvent.GetKey() == EKeys::G) // debug: toggles Graph track on/off
@@ -1815,45 +1867,19 @@ void STimingView::ShowContextMenu(const FPointerEvent& MouseEvent)
 			FSlateIcon()
 		);
 	}
+	else if(HoveredTimingEvent.Track != nullptr)
+	{
+		const_cast<FTimingEventsTrack*>(HoveredTimingEvent.Track)->BuildContextMenu(MenuBuilder);
+	}
 	else
 	{
 		MenuBuilder.BeginSection(TEXT("Empty"));
 		{
-			struct FLocal
-			{
-				static bool ReturnFalse()
-				{
-					return false;
-				}
-			};
-
-			FUIAction DummyUIAction;
-			DummyUIAction.CanExecuteAction = FCanExecuteAction::CreateStatic(&FLocal::ReturnFalse);
-
-			FText Title;
-			if (HoveredTimingEvent.Track != nullptr)
-			{
-				if (HoveredTimingEvent.Track->GetType() == FName(TEXT("Thread")) &&
-					((FThreadTimingTrack*)HoveredTimingEvent.Track)->GetGroupName() != nullptr)
-				{
-					Title = FText::Format(LOCTEXT("TrackTitleGroupFmt", "{0} (Group: {1})"), FText::FromString(HoveredTimingEvent.Track->GetName()), FText::FromString(((FThreadTimingTrack*)HoveredTimingEvent.Track)->GetGroupName()));
-				}
-				else
-				{
-					Title = FText::FromString(HoveredTimingEvent.Track->GetName());
-				}
-			}
-			else
-			{
-				Title = LOCTEXT("ContextMenu_NA", "N/A");
-			}
-
-			MenuBuilder.AddMenuEntry
-			(
-				Title,
+			MenuBuilder.AddMenuEntry(
+				LOCTEXT("ContextMenu_NA", "N/A"),
 				LOCTEXT("ContextMenu_NA_Desc", "No actions available."),
 				FSlateIcon(),
-				DummyUIAction,
+				FUIAction(FExecuteAction(), FCanExecuteAction::CreateLambda([](){ return false; })),
 				NAME_None,
 				EUserInterfaceActionType::Button
 			);
@@ -2010,14 +2036,14 @@ void STimingView::SelectTimeInterval(double IntervalStartTime, double IntervalDu
 
 void STimingView::RaiseSelectionChanging()
 {
-	//TODO: SelectionChangingEvent.Broadcast(SelectionStartTime, SelectionEndTime);
+	OnSelectionChangedDelegate.Broadcast(LastSelectionType != ESelectionType::None, SelectionStartTime, SelectionEndTime);
 }
 
 ////////////////////////////////////////////////////////////////////////////////////////////////////
 
 void STimingView::RaiseSelectionChanged()
 {
-	//TODO: SelectionChangedEvent.Broadcast(SelectionStartTime, SelectionEndTime);
+	OnSelectionChangedDelegate.Broadcast(LastSelectionType != ESelectionType::None, SelectionStartTime, SelectionEndTime);
 
 	FTimingProfilerManager::Get()->SetSelectedTimeRange(SelectionStartTime, SelectionEndTime);
 
@@ -2029,6 +2055,13 @@ void STimingView::RaiseSelectionChanged()
 
 ////////////////////////////////////////////////////////////////////////////////////////////////////
 
+void STimingView::RaiseTimeMarkerChanged()
+{
+	OnTimeMarkerChangedDelegate.Broadcast(TimeMarker != std::numeric_limits<double>::infinity(), TimeMarker);
+}
+
+////////////////////////////////////////////////////////////////////////////////////////////////////
+
 void STimingView::UpdateAggregatedStats()
 {
 	if (bAssetLoadingMode)
@@ -2214,12 +2247,21 @@ void STimingView::UpdateHoveredTimingEvent(float MX, float MY)
 			HoveredTimingEvent = TimingEvent;
 			HoveredTimingEvent.Track->ComputeTimingEventStats(HoveredTimingEvent);
 			HoveredTimingEvent.Track->InitTooltip(Tooltip, HoveredTimingEvent);
+
+			OnHoveredEventChangedDelegate.Broadcast(HoveredTimingEvent);
 		}
 		Tooltip.SetDesiredOpacity(1.0f);
 	}
 	else
 	{
 		HoveredTimingEvent.Reset();
+
+		if(TimingEvent.Track != nullptr)
+		{
+			HoveredTimingEvent.Track = TimingEvent.Track;
+		}
+
+		OnHoveredEventChangedDelegate.Broadcast(HoveredTimingEvent);
 		Tooltip.SetDesiredOpacity(0.0f);
 	}
 }
@@ -2230,13 +2272,13 @@ void STimingView::OnSelectedTimingEventChanged()
 {
 	if (!bAssetLoadingMode)
 	{
-		if (SelectedTimingEvent.IsValid() &&
-			SelectedTimingEvent.Track->GetType() == FName(TEXT("Thread")))
+		if (SelectedTimingEvent.IsValid())
 		{
-			// Select the timer node coresponding to timing event type of selected timing event.
-			FTimingProfilerManager::Get()->SetSelectedTimer(SelectedTimingEvent.TypeId);
+			SelectedTimingEvent.Track->OnEventSelected(SelectedTimingEvent);
 		}
 	}
+
+	OnSelectedEventChangedDelegate.Broadcast(SelectedTimingEvent);
 }
 
 ////////////////////////////////////////////////////////////////////////////////////////////////////
@@ -2512,6 +2554,12 @@ TSharedRef<SWidget> STimingView::MakeTracksFilterMenu()
 			NAME_None,
 			EUserInterfaceActionType::ToggleButton
 		);
+
+		// Let any plugins extend the filter menu
+		for (Insights::ITimingViewExtender* Extender : GetExtenders())
+		{
+			Extender->ExtendFilterMenu(*this, MenuBuilder);
+		}
 	}
 	MenuBuilder.EndSection();
 
@@ -2594,7 +2642,7 @@ bool STimingView::ShowHideGraphTrack_IsChecked() const
 void STimingView::ShowHideGraphTrack_Execute()
 {
 	GraphTrack->ToggleVisibility();
-	bAreTimingEventsTracksDirty = true;
+	bTimingEventsNeedUpdate = true;
 }
 
 ////////////////////////////////////////////////////////////////////////////////////////////////////
@@ -2814,9 +2862,22 @@ void STimingView::OnTimingEventsTrackVisibilityChanged()
 		UpdateAggregatedStats();
 	}
 
-	bAreTimingEventsTracksDirty = true;
+	bTimingEventsNeedUpdate = true;
+}
+
+////////////////////////////////////////////////////////////////////////////////////////////////////
+
+FTimingEventsTrack* STimingView::FindTimingEventsTrack(uint64 InTrackID)
+{
+	FTimingEventsTrack** TrackPtr = CachedTimelines.Find(InTrackID);
+	return TrackPtr != nullptr ? *TrackPtr : nullptr;
 }
 
 ////////////////////////////////////////////////////////////////////////////////////////////////////
 
+TArray<Insights::ITimingViewExtender*> STimingView::GetExtenders() const
+{
+	return IModularFeatures::Get().GetModularFeatureImplementations<Insights::ITimingViewExtender>(Insights::TimingViewExtenderFeatureName);
+}
+
 #undef LOCTEXT_NAMESPACE
diff --git a/Engine/Source/Developer/TraceInsights/Private/Insights/Widgets/STimingView.h b/Engine/Source/Developer/TraceInsights/Private/Insights/Widgets/STimingView.h
--- a/Engine/Source/Developer/TraceInsights/Private/Insights/Widgets/STimingView.h
+++ b/Engine/Source/Developer/TraceInsights/Private/Insights/Widgets/STimingView.h
@@ -22,6 +22,7 @@
 #include "Insights/ViewModels/TimingTrackViewport.h"
 #include "Insights/ViewModels/TimingViewDrawHelper.h"
 #include "Insights/ViewModels/TooltipDrawState.h"
+#include "Insights/ITimingViewSession.h"
 
 class FFileActivitySharedState;
 class FLoadingSharedState;
@@ -29,18 +30,10 @@ class FMenuBuilder;
 class FTimingGraphTrack;
 class FTimingViewDrawHelper;
 class SScrollBar;
-
-/** The delegate to be invoked when the selection have been changed. */
-DECLARE_DELEGATE_TwoParams(FSelectionChangedDelegate, double /*StartTime*/, double /*EndTime*/);
-
-/** The delegate to be invoked when the timing event being hovered by the mouse has changed. Returns id of timing event or -1 (if no event is hovered). */
-DECLARE_DELEGATE_RetVal(int32, FHoveredEventChangedDelegate);
-
-/** The delegate to be invoked when the selected timing event has changed. Returns id of timing event or -1 (if no event is hovered). */
-DECLARE_DELEGATE_RetVal(int32, FSelectedEventChangedDelegate);
+namespace Insights { class ITimingViewExtender; }
 
 /** A custom widget used to display timing events. */
-class STimingView : public SCompoundWidget
+class STimingView : public SCompoundWidget, public Insights::ITimingViewSession
 {
 protected:
 	struct FAssetLoadingEventAggregationRow
@@ -62,16 +55,10 @@ public:
 	virtual ~STimingView();
 
 	SLATE_BEGIN_ARGS(STimingView)
-		: _OnSelectionChanged()
-		, _OnHoveredEventChanged()
-		, _OnSelectedEventChanged()
 		{
 			_Clipping = EWidgetClipping::ClipToBounds;
 		}
 
-		SLATE_EVENT(FSelectionChangedDelegate, OnSelectionChanged)
-		SLATE_EVENT(FHoveredEventChangedDelegate, OnHoveredEventChanged)
-		SLATE_EVENT(FSelectedEventChangedDelegate, OnSelectedEventChanged)
 	SLATE_END_ARGS()
 
 	/**
@@ -248,6 +235,14 @@ public:
 
 	////////////////////////////////////////////////////////////////////////////////////////////////////
 
+	// ITimingViewSession interface
+	virtual FTimingEventsTrack* FindTimingEventsTrack(uint64 InTrackID) override;
+	virtual void AddTimingEventsTrack(FTimingEventsTrack* Track) override;
+	virtual Insights::FTimeMarkerChangedDelegate& OnTimeMarkerChanged() override { return OnTimeMarkerChangedDelegate; }
+	virtual Insights::FSelectionChangedDelegate& OnSelectionChanged() override { return OnSelectionChangedDelegate; }
+	virtual Insights::FHoveredEventChangedDelegate& OnHoveredEventChanged() override { return OnHoveredEventChangedDelegate; }
+	virtual Insights::FSelectedEventChangedDelegate& OnSelectedEventChanged() override { return OnSelectedEventChangedDelegate; }
+
 	const FTimingTrackViewport& GetViewport() { return Viewport; }
 
 	double GetSelectionStartTime() const { return SelectionStartTime; }
@@ -280,8 +275,6 @@ protected:
 		return FVector2D(16.0f, 16.0f);
 	}
 
-	void AddTimingEventsTrack(FTimingEventsTrack* Track);
-
 	void ShowContextMenu(const FPointerEvent& MouseEvent);
 
 	/** Binds our UI commands to delegates. */
@@ -309,6 +302,9 @@ protected:
 
 	void RaiseSelectionChanging();
 	void RaiseSelectionChanged();
+
+	void RaiseTimeMarkerChanged();
+
 	void UpdateAggregatedStats();
 
 	void UpdateHoveredTimingEvent(float MX, float MY);
@@ -324,6 +320,9 @@ protected:
 
 	void OnTimingEventsTrackVisibilityChanged();
 
+	// Get all the plugin extenders we care about
+	TArray<Insights::ITimingViewExtender*> GetExtenders() const;
+
 	////////////////////////////////////////////////////////////////////////////////////////////////////
 	// FrameSelectionChanged Event
 
@@ -377,7 +376,11 @@ protected:
 
 	TArray<FTimingEventsTrack*> TimingEventsTracks; /**< all timing events tracks, in the order to be displayed */
 
-	bool bAreTimingEventsTracksDirty;
+	// Dirty flag that indicates the timing events need updating
+	bool bTimingEventsNeedUpdate;
+
+	// Dirty flag that indicates that the TimingEventsTracks list has changed
+	bool bTimingEventsTracksChanged;
 
 	////////////////////////////////////////////////////////////
 	// Cpu/Gpu
@@ -498,6 +501,9 @@ protected:
 
 	double TimeMarker;
 
+	/** True of the user is currently dragging the time marker */
+	bool bIsScrubbing;
+
 	////////////////////////////////////////////////////////////
 	// Misc
 
@@ -519,7 +525,8 @@ protected:
 	////////////////////////////////////////////////////////////
 	// Delegates
 
-	FSelectionChangedDelegate OnSelectionChanged;
-	FHoveredEventChangedDelegate OnHoveredEventChanged;
-	FSelectedEventChangedDelegate OnSelectedEventChanged;
+	Insights::FSelectionChangedDelegate OnSelectionChangedDelegate;
+	Insights::FHoveredEventChangedDelegate OnHoveredEventChangedDelegate;
+	Insights::FSelectedEventChangedDelegate OnSelectedEventChangedDelegate;
+	Insights::FTimeMarkerChangedDelegate OnTimeMarkerChangedDelegate;
 };
diff --git a/Engine/Source/Developer/TraceInsights/Public/Insights/Common/PaintUtils.h b/Engine/Source/Developer/TraceInsights/Public/Insights/Common/PaintUtils.h
new file mode 100644
--- /dev/null
+++ b/Engine/Source/Developer/TraceInsights/Public/Insights/Common/PaintUtils.h
@@ -0,0 +1,105 @@
+// Copyright 1998-2019 Epic Games, Inc. All Rights Reserved.
+
+#pragma once
+
+#include "CoreMinimal.h"
+#include "Fonts/FontMeasure.h"
+#include "Framework/Application/SlateApplication.h"
+#include "Rendering/DrawElements.h"
+
+#define MAKE_PAINT_GEOMETRY_PT(Geometry, X, Y)       Geometry.ToPaintGeometry(FSlateLayoutTransform(1.0f, FVector2D(X, Y)))
+#define MAKE_PAINT_GEOMETRY_RC(Geometry, X, Y, W, H) Geometry.ToPaintGeometry(FVector2D(W, H), FSlateLayoutTransform(1.0f, FVector2D(X, Y)))
+
+struct FGeometry;
+class FSlateRect;
+class WidgetStyle;
+enum class ESlateDrawEffect : uint8;
+class FSlateWindowElementList;
+
+/**
+ * Holds current state provided by OnPaint function, used to simplify drawing.
+ */
+struct TRACEINSIGHTS_API FDrawContext
+{
+	FDrawContext(const FGeometry& InGeometry,
+				 const FSlateRect& InCullingRect,
+				 const FWidgetStyle& InWidgetStyle,
+				 const ESlateDrawEffect InDrawEffects,
+				 FSlateWindowElementList& InOutElementList,
+				 int32& InOutLayerId)
+		: Geometry(InGeometry)
+		, CullingRect(InCullingRect)
+		, WidgetStyle(InWidgetStyle)
+		, DrawEffects(InDrawEffects)
+		, ElementList(InOutElementList)
+		, LayerId(InOutLayerId)
+	{}
+
+	/**
+	 * Non-copyable
+	 */
+	FDrawContext(const FDrawContext&) = delete;
+	FDrawContext& operator=(const FDrawContext&) = delete;
+
+	inline void DrawBox(const float X, const float Y, const float W, const float H, const FSlateBrush* Brush, const FLinearColor& Color) const
+	{
+		FSlateDrawElement::MakeBox(ElementList, LayerId, MAKE_PAINT_GEOMETRY_RC(Geometry, X, Y, W, H), Brush, DrawEffects, Color);
+	}
+
+	inline void DrawBox(const int32 InLayer, const float X, const float Y, const float W, const float H, const FSlateBrush* Brush, const FLinearColor& Color) const
+	{
+		FSlateDrawElement::MakeBox(ElementList, InLayer, MAKE_PAINT_GEOMETRY_RC(Geometry, X, Y, W, H), Brush, DrawEffects, Color);
+	}
+
+	inline void DrawRotatedBox(const float X, const float Y, const float W, const float H, const FSlateBrush* Brush, const FLinearColor& Color, float Angle, TOptional<FVector2D> RotationPoint) const
+	{
+		FSlateDrawElement::MakeRotatedBox(ElementList, LayerId, MAKE_PAINT_GEOMETRY_RC(Geometry, X, Y, W, H), Brush, DrawEffects, Angle, RotationPoint, FSlateDrawElement::RelativeToElement, Color);
+	}
+
+	inline void DrawText(const float X, const float Y, const FString& Text, const FSlateFontInfo& Font, const FLinearColor& Color) const
+	{
+		FSlateDrawElement::MakeText(ElementList, LayerId, MAKE_PAINT_GEOMETRY_PT(Geometry, X, Y), Text, Font, DrawEffects, Color);
+	}
+
+	inline void DrawText(const int32 InLayer, const float X, const float Y, const FString& Text, const FSlateFontInfo& Font, const FLinearColor& Color) const
+	{
+		FSlateDrawElement::MakeText(ElementList, InLayer, MAKE_PAINT_GEOMETRY_PT(Geometry, X, Y), Text, Font, DrawEffects, Color);
+	}
+
+	inline void DrawText(const float X, const float Y, const FString& Text, const int32 StartIndex, const int32 EndIndex, const FSlateFontInfo& Font, const FLinearColor& Color) const
+	{
+		FSlateDrawElement::MakeText(ElementList, LayerId, MAKE_PAINT_GEOMETRY_PT(Geometry, X, Y), Text, StartIndex, EndIndex, Font, DrawEffects, Color);
+	}
+
+	inline void DrawText(const int32 InLayer, const float X, const float Y, const FString& Text, const int32 StartIndex, const int32 EndIndex, const FSlateFontInfo& Font, const FLinearColor& Color) const
+	{
+		FSlateDrawElement::MakeText(ElementList, InLayer, MAKE_PAINT_GEOMETRY_PT(Geometry, X, Y), Text, StartIndex, EndIndex, Font, DrawEffects, Color);
+	}
+
+	inline void DrawTextAligned(EHorizontalAlignment HAlign, const float X, const float Y, const FString& Text, const FSlateFontInfo& Font, const FLinearColor& Color) const
+	{
+		const TSharedRef<FSlateFontMeasure> FontMeasureService = FSlateApplication::Get().GetRenderer()->GetFontMeasureService();
+		const float TextWidth = FontMeasureService->Measure(Text, Font).X;
+		float TextX = X;
+		if (HAlign == HAlign_Right)
+		{
+			TextX -= TextWidth;
+		}
+		else if (HAlign == HAlign_Center)
+		{
+			TextX -= TextWidth / 2;
+		}
+		FSlateDrawElement::MakeText(ElementList, LayerId, MAKE_PAINT_GEOMETRY_PT(Geometry, TextX, Y), Text, Font, DrawEffects, Color);
+	}
+
+	//inline void IncrementLayer() const { ++LayerId; }
+
+	// Accessors
+	const FGeometry& Geometry;
+	const FSlateRect& CullingRect;
+	const FWidgetStyle& WidgetStyle;
+	const ESlateDrawEffect DrawEffects;
+
+	FSlateWindowElementList& ElementList;
+	int32& LayerId;
+};
diff --git a/Engine/Source/Developer/TraceInsights/Public/Insights/ITimingViewDrawHelper.h b/Engine/Source/Developer/TraceInsights/Public/Insights/ITimingViewDrawHelper.h
new file mode 100644
--- /dev/null
+++ b/Engine/Source/Developer/TraceInsights/Public/Insights/ITimingViewDrawHelper.h
@@ -0,0 +1,22 @@
+// Copyright 1998-2019 Epic Games, Inc. All Rights Reserved.
+
+#pragma once
+
+#include "CoreMinimal.h"
+
+struct FDrawContext;
+class FTimingTrackViewport;
+struct FTimingEventsTrackLayout;
+class FTimingEventsTrack;
+
+/** Interface used to draw events into the timing view */
+class ITimingViewDrawHelper
+{
+public:
+	virtual const FDrawContext& GetDrawContext() const = 0;
+	virtual const FTimingTrackViewport& GetViewport() const = 0;
+	virtual const FTimingEventsTrackLayout& GetLayout() const = 0;
+	virtual bool BeginTimeline(FTimingEventsTrack& Track) = 0;
+	virtual void AddEvent(double StartTime, double EndTime, uint32 Depth, const TCHAR* EventName, uint32 Color = 0) = 0;
+	virtual void EndTimeline(FTimingEventsTrack& Track) = 0;
+};
\ No newline at end of file
diff --git a/Engine/Source/Developer/TraceInsights/Public/Insights/ITimingViewExtender.h b/Engine/Source/Developer/TraceInsights/Public/Insights/ITimingViewExtender.h
new file mode 100644
--- /dev/null
+++ b/Engine/Source/Developer/TraceInsights/Public/Insights/ITimingViewExtender.h
@@ -0,0 +1,35 @@
+// Copyright 1998-2019 Epic Games, Inc. All Rights Reserved.
+
+#pragma once
+
+#include "CoreMinimal.h"
+#include "Features/IModularFeature.h"
+
+namespace Insights
+{
+
+class ITimingViewSession;
+extern TRACEINSIGHTS_API const FName TimingViewExtenderFeatureName;
+
+class ITimingViewExtender : public IModularFeature
+{
+public:
+	virtual ~ITimingViewExtender() = default;
+
+	/** Called to set up any data at the end of the timing view session */
+	virtual void OnBeginSession(ITimingViewSession& InSession) = 0;
+
+	/** Called to clear out any data at the end of the timing view session */
+	virtual void OnEndSession(ITimingViewSession& InSession) = 0;
+
+	/** Called each frame. If any new tracks are created they can be added via ITimingViewSession::AddTimingEventsTrack() */
+	virtual void Tick(Insights::ITimingViewSession& InSession, const Trace::IAnalysisSession& InAnalysisSession) = 0;
+
+	/** Extension hook for the 'quick filter' menu */
+	virtual void ExtendFilterMenu(Insights::ITimingViewSession& InSession, FMenuBuilder& InMenuBuilder) = 0;
+
+	/** Called when tracks have changed and need to be re-ordered */
+	virtual void OnTracksChanged(Insights::ITimingViewSession& InSession, int32& InOutOrder) = 0;
+};
+
+}
\ No newline at end of file
diff --git a/Engine/Source/Developer/TraceInsights/Public/Insights/ITimingViewSession.h b/Engine/Source/Developer/TraceInsights/Public/Insights/ITimingViewSession.h
new file mode 100644
--- /dev/null
+++ b/Engine/Source/Developer/TraceInsights/Public/Insights/ITimingViewSession.h
@@ -0,0 +1,50 @@
+// Copyright 1998-2019 Epic Games, Inc. All Rights Reserved.
+
+#pragma once
+
+#include "CoreMinimal.h"
+
+class FTimingEventsTrack;
+namespace Insights { class ITimingViewVisualizer; }
+
+namespace Insights
+{
+
+/** The delegate to be invoked when the time marker has changed */
+DECLARE_MULTICAST_DELEGATE_TwoParams(FTimeMarkerChangedDelegate, bool /*bValid*/, double /*TimeMarker*/);
+
+/** The delegate to be invoked when the selection have been changed. */
+DECLARE_MULTICAST_DELEGATE_ThreeParams(FSelectionChangedDelegate, bool /*bValid*/, double /*StartTime*/, double /*EndTime*/);
+
+/** The delegate to be invoked when the timing event being hovered by the mouse has changed. */
+DECLARE_MULTICAST_DELEGATE_OneParam(FHoveredEventChangedDelegate, const FTimingEvent& /*InEvent*/);
+
+/** The delegate to be invoked when the selected timing event has changed. */
+DECLARE_MULTICAST_DELEGATE_OneParam(FSelectedEventChangedDelegate, const FTimingEvent& /*InEvent*/);
+
+/** Hosts a number of timing view visualizers, represents a session of the timing view */
+class ITimingViewSession
+{
+public:
+	virtual ~ITimingViewSession() = default;
+
+	/** Adds a new timing events track. Transfers ownership of the track pointer to the session. */
+	virtual void AddTimingEventsTrack(FTimingEventsTrack* InTrack) = 0;
+
+	/** Find a timing events track has been added via AddTimingEventsTrack() */
+	virtual FTimingEventsTrack* FindTimingEventsTrack(uint64 InTrackID) = 0;
+
+	/** Get the delegate to be invoked when the time marker has changed */
+	virtual FTimeMarkerChangedDelegate& OnTimeMarkerChanged() = 0;
+
+	/** Get the delegate to be invoked when the selection have been changed. */
+	virtual FSelectionChangedDelegate& OnSelectionChanged() = 0;
+
+	/** Get the delegate to be invoked when the timing event being hovered by the mouse has changed. */
+	virtual FHoveredEventChangedDelegate& OnHoveredEventChanged() = 0;
+
+	/** Get the delegate to be invoked when the selected timing event has changed. */
+	virtual FSelectedEventChangedDelegate& OnSelectedEventChanged() = 0;
+};
+
+}
\ No newline at end of file
diff --git a/Engine/Source/Developer/TraceInsights/Public/Insights/ViewModels/BaseTimingTrack.h b/Engine/Source/Developer/TraceInsights/Public/Insights/ViewModels/BaseTimingTrack.h
new file mode 100644
--- /dev/null
+++ b/Engine/Source/Developer/TraceInsights/Public/Insights/ViewModels/BaseTimingTrack.h
@@ -0,0 +1,131 @@
+// Copyright 1998-2019 Epic Games, Inc. All Rights Reserved.
+
+#pragma once
+
+#include "CoreMinimal.h"
+#include "Misc/EnumClassFlags.h"
+#include "Templates/SharedPointer.h"
+
+////////////////////////////////////////////////////////////////////////////////////////////////////
+
+class FTimingTrackViewport;
+class FMenuBuilder;
+namespace Trace { class IAnalysisSession; };
+
+////////////////////////////////////////////////////////////////////////////////////////////////////
+
+enum class ETimingTrackFlags : uint32
+{
+	None = 0,
+	IsVisible = (1 << 0),
+	IsDirty = (1 << 1),
+	IsSelected = (1 << 2),
+	IsHovered = (1 << 3),
+	IsHeaderHovered = (1 << 4),
+};
+ENUM_CLASS_FLAGS(ETimingTrackFlags);
+
+////////////////////////////////////////////////////////////////////////////////////////////////////
+
+class TRACEINSIGHTS_API FBaseTimingTrack : public TSharedFromThis<FBaseTimingTrack>
+{
+	friend class FTimingViewDrawHelper;
+
+protected:
+	explicit FBaseTimingTrack(uint64 InId, const FName& InType = NAME_None, const FName& InSubType = NAME_None, const FString& InName = FString())
+		: Id(InId)
+		, Type(InType)
+		, SubType(InSubType)
+		, Name(InName)
+		, Order(0)
+		, PosY(0.0f)
+		, Height(0.0f)
+		, Session(nullptr)
+		, Flags(ETimingTrackFlags::IsVisible | ETimingTrackFlags::IsDirty)
+	{}
+
+	virtual ~FBaseTimingTrack()
+	{}
+
+public:
+	virtual void Reset()
+	{
+		Id = 0;
+		PosY = 0.0f;
+		Height = 0.0f;
+		Session = nullptr;
+		Flags = ETimingTrackFlags::IsVisible | ETimingTrackFlags::IsDirty;
+	}
+
+	uint64 GetId() const { return Id; }
+
+	const FName& GetType() const { return Type; }
+	const FName& GetSubType() const { return SubType; }
+
+	const FString& GetName() const { return Name; }
+
+	void SetOrder(int32 InOrder) { Order = InOrder; }
+	int32 GetOrder() const { return Order; }
+
+	float GetPosY() const { return PosY; }
+	virtual void SetPosY(float InPosY) { PosY = InPosY; }
+
+	float GetHeight() const { return Height; }
+	virtual void SetHeight(float InHeight) { Height = InHeight; }
+
+	bool IsVisible() const { return EnumHasAnyFlags(Flags, ETimingTrackFlags::IsVisible); }
+	void Show() { Flags |= ETimingTrackFlags::IsVisible; OnVisibilityChanged(); }
+	void Hide() { Flags &= ~ETimingTrackFlags::IsVisible; OnVisibilityChanged(); }
+	void ToggleVisibility() { Flags ^= ETimingTrackFlags::IsVisible; OnVisibilityChanged(); }
+	void SetVisibilityFlag(bool bIsVisible) { bIsVisible ? Show() : Hide(); }
+	virtual void OnVisibilityChanged() {}
+
+	bool IsDirty() const { return EnumHasAnyFlags(Flags, ETimingTrackFlags::IsDirty); }
+	void SetDirtyFlag() { Flags |= ETimingTrackFlags::IsDirty; OnDirtyFlagChanged(); }
+	void ClearDirtyFlag() { Flags &= ~ETimingTrackFlags::IsDirty; OnDirtyFlagChanged(); }
+	virtual void OnDirtyFlagChanged() {}
+
+	bool IsSelected() const { return EnumHasAnyFlags(Flags, ETimingTrackFlags::IsSelected); }
+	void Select() { Flags |= ETimingTrackFlags::IsSelected; OnSelectedFlagChanged(); }
+	void Unselect() { Flags &= ~ETimingTrackFlags::IsSelected; OnSelectedFlagChanged(); }
+	void ToggleSelectedFlag() { Flags ^= ETimingTrackFlags::IsSelected; OnSelectedFlagChanged(); }
+	void SetSelectedFlag(bool bIsSelected) { bIsSelected ? Select() : Unselect(); }
+	virtual void OnSelectedFlagChanged() {}
+
+	bool IsHovered() const { return EnumHasAnyFlags(Flags, ETimingTrackFlags::IsHovered); }
+	void SetHoveredState(bool bIsHovered) { bIsHovered ? Flags |= ETimingTrackFlags::IsHovered : Flags &= ~ETimingTrackFlags::IsHovered; }
+	bool IsHeaderHovered() const { return EnumHasAllFlags(Flags, ETimingTrackFlags::IsHovered | ETimingTrackFlags::IsHeaderHovered); }
+	void SetHeaderHoveredState(bool bIsHeaderHovered) { bIsHeaderHovered ? Flags |= ETimingTrackFlags::IsHeaderHovered : Flags &= ~ETimingTrackFlags::IsHeaderHovered; }
+	virtual void UpdateHoveredState(float MouseX, float MouseY, const FTimingTrackViewport& Viewport) = 0;
+
+	virtual void Tick(const double CurrentTime, const float DeltaTime) {}
+	virtual void Update(const FTimingTrackViewport& Viewport) {}
+
+	static uint64 GenerateId() { return IdGenerator++; }
+
+	virtual void BuildContextMenu(FMenuBuilder& MenuBuilder);
+
+	// Check whether the session is currently valid
+	bool IsSessionValid() const { return Session != nullptr; }
+
+	// Access the session that this track is using
+	const Trace::IAnalysisSession& GetSession() const { check(Session); return *Session; }
+
+	// Set the session that this track is using. This is refreshed each tick we have a valid session
+	void SetSession(const Trace::IAnalysisSession* InSession) { Session = InSession; }
+
+private:
+	uint64 Id;
+	FName Type; // Thread, Loading, FileActivity
+	FName SubType; // Cpu, Gpu, MainThread, AsyncThread, Overview, Detailed
+	FString Name;
+	int32 Order;
+	float PosY; // y position, in Slate units
+	float Height; // height, in Slate units
+	const Trace::IAnalysisSession* Session;
+	ETimingTrackFlags Flags;
+
+	static uint64 IdGenerator;
+};
+
+////////////////////////////////////////////////////////////////////////////////////////////////////
diff --git a/Engine/Source/Developer/TraceInsights/Public/Insights/ViewModels/GraphTrack.h b/Engine/Source/Developer/TraceInsights/Public/Insights/ViewModels/GraphTrack.h
new file mode 100644
--- /dev/null
+++ b/Engine/Source/Developer/TraceInsights/Public/Insights/ViewModels/GraphTrack.h
@@ -0,0 +1,380 @@
+// Copyright 1998-2019 Epic Games, Inc. All Rights Reserved.
+
+#pragma once
+
+#include "CoreMinimal.h"
+#include "Fonts/SlateFontInfo.h"
+
+// Insights
+#include "Insights/ViewModels/BaseTimingTrack.h"
+#include "Insights/ViewModels/TooltipDrawState.h"
+
+struct FDrawContext;
+struct FSlateBrush;
+class FMenuBuilder;
+class FTimingTrackViewport;
+
+////////////////////////////////////////////////////////////////////////////////////////////////////
+
+class FGraphSeries
+{
+	friend class FGraphTrack;
+	friend class FGraphTrackBuilder;
+
+public:
+	struct FEvent
+	{
+		double Time;
+		double Duration;
+		double Value;
+	};
+
+	struct FBox
+	{
+		float X;
+		float W;
+		float Y;
+	};
+
+public:
+	FGraphSeries();
+	virtual ~FGraphSeries();
+
+	const FText& GetName() const { return Name; }
+	void SetName(const TCHAR* InName) { Name = FText::FromString(InName); }
+	void SetName(const FString& InName) { Name = FText::FromString(InName); }
+	void SetName(const FText& InName) { Name = InName; }
+
+	const FText& GetDescription() const { return Description; }
+	void SetDescription(const TCHAR* InDescription) { Description = FText::FromString(InDescription); }
+	void SetDescription(const FString& InDescription) { Description = FText::FromString(InDescription); }
+	void SetDescription(const FText& InDescription) { Description = InDescription; }
+
+	bool IsVisible() const { return bIsVisible; }
+	void SetVisibility(bool bOnOff) { bIsVisible = bOnOff; }
+
+	bool IsDirty() const { return bIsDirty; }
+	void SetDirtyFlag() { bIsDirty = true; }
+	void ClearDirtyFlag() { bIsDirty = false; }
+
+	const FLinearColor& GetColor() const { return Color; }
+	const FLinearColor& GetBorderColor() const { return BorderColor; }
+
+	void SetColor(FLinearColor InColor, FLinearColor InBorderColor)
+	{
+		Color = InColor;
+		BorderColor = InBorderColor;
+	}
+
+	/**
+	 * @return Y position (in viewport local space) of the baseline (with Value == 0); in pixels (Slate units).
+	 * Y == 0 at the top of the graph track, positive values are downward.
+	 */
+	double GetBaselineY() const { return BaselineY; }
+	void SetBaselineY(const double InBaselineY) { BaselineY = InBaselineY; }
+
+	/**
+	 * @return The scale between Value units and viewport units; in pixels (Slate units) / Value unit.
+	 */
+	double GetScaleY() const { return ScaleY; }
+	void SetScaleY(const double InScaleY) { ScaleY = InScaleY; }
+
+	/**
+	 * @param Value a value; in Value units
+	 * @return Y position (in viewport local space) for a Value; in pixels (Slate units).
+	 * Y == 0 at the top of the graph track, positive values are downward.
+	 */
+	float GetYForValue(double Value) const
+	{
+		return static_cast<float>(BaselineY - Value * ScaleY);
+	}
+
+	float GetRoundedYForValue(double Value) const
+	{
+		return FMath::RoundToFloat(FMath::Clamp<float>(GetYForValue(Value), -FLT_MAX, FLT_MAX));
+	}
+
+	/**
+	 * @param Y a Y position (in viewport local space); in pixels (Slate units).
+	 * @return Value for specified Y position.
+	 */
+	double GetValueForY(float Y) const
+	{
+		return (BaselineY - static_cast<double>(Y)) / ScaleY;
+	}
+
+	bool IsAutoZoomEnabled() const { return bAutoZoom; }
+	void EnableAutoZoom() { bAutoZoom = true; }
+
+	/** Target low value of auto zoom interval (corresponding to bottom of the track) */
+	double GetTargetAutoZoomLowValue() const { return TargetAutoZoomLowValue; }
+	/** Target high value of auto zoom interval (corresponding to top of the track) */
+	double GetTargetAutoZoomHighValue() const { return TargetAutoZoomHighValue; }
+	void SetTargetAutoZoomRange(double LowValue, double HighValue) { TargetAutoZoomLowValue = LowValue; TargetAutoZoomHighValue = HighValue; }
+
+	/** Current auto zoom low value */
+	double GetAutoZoomLowValue() const { return AutoZoomLowValue; }
+	/** Current auto zoom high value */
+	double GetAutoZoomHighValue() const { return AutoZoomHighValue; }
+	void SetAutoZoomRange(double LowValue, double HighValue) { AutoZoomLowValue = LowValue; AutoZoomHighValue = HighValue; }
+
+	/**
+	 * Compute BaselineY and ScaleY so the [Low, High] Value range will correspond to [Top, Bottom] Y position range.
+	 * GetYForValue(InHighValue) == InTopY
+	 * GetYForValue(InLowValue) == InBottomY
+	 */
+	void ComputeBaselineAndScale(const double InLowValue, const double InHighValue, const float InTopY, const float InBottomY, double& OutBaselineY, double& OutScaleY) const
+	{
+		ensure(InLowValue < InHighValue);
+		ensure(InTopY < InBottomY);
+		const double InvRange = 1.0 / (InHighValue - InLowValue);
+		OutScaleY = static_cast<double>(InBottomY - InTopY) * InvRange;
+		//OutBaselineY = (InHighValue * static_cast<double>(InBottomY) - InLowValue * static_cast<double>(InTopY)) * InvRange;
+		OutBaselineY = static_cast<double>(InTopY) + InHighValue * OutScaleY;
+	}
+
+	/**
+	 * @param X The horizontal coordinate of the point tested; in Slate pixels (local graph coordinates)
+	 * @param Y The vertical coordinate of the point tested; in Slate pixels (local graph coordinates)
+	 * @param Viewport The timing viewport used to transform time in local graph coordinates
+	 * @param bCheckLine If needs to check the bounding box of the horizontal line (determined by duration of event and value) or only the bounding box of the visual point
+	 * @param bCheckBox If needs to check the bounding box of the entire visual box (determined by duration of event, value and baseline)
+	 * @return A pointer to an Event located at (X, Y) coordinates, if any; nullptr if no event is located at respective coordinates
+	 */
+	const FGraphSeries::FEvent* GetEvent(const float PosX, const float PosY, const FTimingTrackViewport& Viewport, bool bCheckLine, bool bCheckBox) const;
+
+	virtual FString FormatValue(double Value) const;
+
+private:
+	FText Name;
+	FText Description;
+
+	bool bIsVisible;
+	bool bIsDirty;
+
+	bool bAutoZoom;
+	double TargetAutoZoomLowValue; // target low value of auto zoom interval (corresponding to bottom of the track)
+	double TargetAutoZoomHighValue; // target high value of auto zoom interval (corresponding to top of the track)
+	double AutoZoomLowValue; // current auto zoom low value
+	double AutoZoomHighValue; // current auto zoom high value
+
+	double BaselineY; // Y position (in viewport local space) of the baseline (with Value == 0); in pixels (Slate units)
+	double ScaleY; // scale between Value units and viewport units; in pixels (Slate units) / Value unit
+
+	FLinearColor Color;
+	FLinearColor BorderColor;
+
+protected:
+	TArray<FEvent> Events; // reduced list of events; used to identify an event at a certain screen position (ex.: the event hovered by mouse)
+	TArray<FVector2D> Points; // reduced list of points; for drawing points
+	TArray<FVector2D> LinePoints; // reduced list of points; for drawing the connected line and filled polygon
+	TArray<FBox> Boxes; // reduced list of boxes; for drawing boxes
+};
+
+////////////////////////////////////////////////////////////////////////////////////////////////////
+
+class TRACEINSIGHTS_API FGraphTrack : public FBaseTimingTrack
+{
+	friend class FGraphTrackBuilder;
+
+public:
+	struct FEvent
+	{
+		TSharedPtr<FGraphSeries> Series;
+		FGraphSeries::FEvent SeriesEvent;
+
+		void Reset() { Series.Reset(); }
+		bool IsValid() const { return Series.IsValid(); }
+	};
+
+private:
+	// Visual size of points (in pixels).
+	static constexpr float PointVisualSize = 5.0f;
+
+	// Size of points (in pixels) used in reduction algorithm.
+	static constexpr double PointSizeX = 3.0f;
+	static constexpr float PointSizeY = 3.0f;
+
+public:
+	explicit FGraphTrack(uint64 InTrackId);
+	virtual ~FGraphTrack();
+
+	virtual void Reset() override;
+
+	virtual void UpdateHoveredState(float MouseX, float MouseY, const FTimingTrackViewport& Viewport) override;
+
+	virtual void Update(const FTimingTrackViewport& Viewport) override;
+
+	void Draw(FDrawContext& DrawContext, const FTimingTrackViewport& Viewport, const FVector2D& MousePosition) const;
+	void PostDraw(FDrawContext& DrawContext, const FTimingTrackViewport& Viewport, const FVector2D& MousePosition) const;
+
+	/**
+	 * @param X The horizontal coordinate of the point tested; in Slate pixels (local graph coordinates)
+	 * @param Y The vertical coordinate of the point tested; in Slate pixels (local graph coordinates)
+	 * @param Viewport The timing viewport used to transform time in local graph coordinates
+	 * @param OutEvent The returned event located at (PosX, PosY) coordinates, if any.
+	 * @return True if an event is found at (PosX, PosY) coordinates.
+	 */
+	const bool GetEvent(const float MouseX, const float MouseY, const FTimingTrackViewport& Viewport, FGraphTrack::FEvent& OutEvent) const;
+
+	virtual void BuildContextMenu(FMenuBuilder& MenuBuilder) override;
+
+	TArray<TSharedPtr<FGraphSeries>>& GetSeries() { return AllSeries; }
+
+	//virtual int GetDebugLineCount() override;
+	//virtual void BuildDebugLines(FString& OutStr) override;
+
+	int32 GetNumAddedEvents() const { return NumAddedEvents; }
+	int32 GetNumDrawPoints() const { return NumDrawPoints; }
+	int32 GetNumDrawLines() const { return NumDrawLines; }
+	int32 GetNumDrawBoxes() const { return NumDrawBoxes; }
+
+protected:
+	void UpdateStats();
+
+	void DrawBackground(FDrawContext& DrawContext, const FTimingTrackViewport& Viewport) const;
+	void DrawSeries(const FGraphSeries& Series, FDrawContext& DrawContext, const FTimingTrackViewport& Viewport) const;
+	void DrawHighlightedEvent(FDrawContext& DrawContext, const FTimingTrackViewport& Viewport, const FGraphTrack::FEvent& GraphEvent) const;
+
+	void InitTooltip();
+
+private:
+	void ContextMenu_ShowPoints_Execute();
+	bool ContextMenu_ShowPoints_CanExecute();
+	bool ContextMenu_ShowPoints_IsChecked();
+
+	void ContextMenu_ShowPointsWithBorder_Execute();
+	bool ContextMenu_ShowPointsWithBorder_CanExecute();
+	bool ContextMenu_ShowPointsWithBorder_IsChecked();
+
+	void ContextMenu_ShowLines_Execute();
+	bool ContextMenu_ShowLines_CanExecute();
+	bool ContextMenu_ShowLines_IsChecked();
+
+	void ContextMenu_ShowPolygon_Execute();
+	bool ContextMenu_ShowPolygon_CanExecute();
+	bool ContextMenu_ShowPolygon_IsChecked();
+
+	void ContextMenu_UseEventDuration_Execute();
+	bool ContextMenu_UseEventDuration_CanExecute();
+	bool ContextMenu_UseEventDuration_IsChecked();
+
+	void ContextMenu_ShowBars_Execute();
+	bool ContextMenu_ShowBars_CanExecute();
+	bool ContextMenu_ShowBars_IsChecked();
+
+	void ContextMenu_ShowSeries_Execute(FGraphSeries* Series);
+	bool ContextMenu_ShowSeries_CanExecute(FGraphSeries* Series);
+	bool ContextMenu_ShowSeries_IsChecked(FGraphSeries* Series);
+
+protected:
+	TArray<TSharedPtr<FGraphSeries>> AllSeries;
+
+	// Slate resources
+	const FSlateBrush* WhiteBrush;
+	const FSlateBrush* PointBrush;
+	const FSlateBrush* BorderBrush;
+	const FSlateFontInfo Font;
+
+	bool bDrawPoints;
+	bool bDrawPointsWithBorder;
+	bool bDrawLines;
+	bool bDrawPolygon;
+	bool bUseEventDuration;
+	bool bDrawBoxes;
+
+	mutable FGraphTrack::FEvent HoveredGraphEvent;
+	FTooltipDrawState Tooltip;
+	FVector2D TooltipMousePosition;
+
+	// Stats
+	int32 NumAddedEvents; // total event count
+	int32 NumDrawPoints;
+	int32 NumDrawLines;
+	int32 NumDrawBoxes;
+};
+
+////////////////////////////////////////////////////////////////////////////////////////////////////
+
+class FRandomGraphTrack : public FGraphTrack
+{
+public:
+	explicit FRandomGraphTrack(uint64 InTrackId);
+	virtual ~FRandomGraphTrack();
+
+	virtual void Update(const FTimingTrackViewport& Viewport) override;
+
+	void AddDefaultSeries();
+
+protected:
+	void GenerateSeries(FGraphSeries& Series, const FTimingTrackViewport& Viewport, const int32 EventCount, int32 Seed);
+};
+
+////////////////////////////////////////////////////////////////////////////////////////////////////
+
+class FGraphTrackBuilder
+{
+private:
+	struct FPointInfo
+	{
+		bool bValid;
+		float X;
+		float Y;
+
+		FPointInfo() : bValid(false) {}
+	};
+
+public:
+	explicit FGraphTrackBuilder(FGraphTrack& InTrack, FGraphSeries& InSeries, const FTimingTrackViewport& InViewport);
+	~FGraphTrackBuilder();
+
+	/**
+	 * Non-copyable
+	 */
+	FGraphTrackBuilder(const FGraphTrackBuilder&) = delete;
+	FGraphTrackBuilder& operator=(const FGraphTrackBuilder&) = delete;
+
+	FGraphTrack& GetTrack() const { return Track; }
+	FGraphSeries& GetSeries() const { return Series; }
+	const FTimingTrackViewport& GetViewport() const { return Viewport; }
+
+	void AddEvent(double Time, double Duration, double Value);
+
+private:
+	void BeginPoints();
+	bool AddPoint(double Time, double Value);
+	void FlushPoints();
+	void EndPoints();
+
+	void BeginConnectedLines();
+	void AddConnectedLine(double Time, double Value);
+	void FlushConnectedLine();
+	void EndConnectedLines();
+
+	void BeginBoxes();
+	void AddBox(double Time, double Duration, double Value);
+	void FlushBox();
+	void EndBoxes();
+
+private:
+	FGraphTrack& Track;
+	FGraphSeries& Series;
+	const FTimingTrackViewport& Viewport;
+
+	// Used by the point reduction algorithm.
+	double PointsCurrentX;
+	TArray<FPointInfo> PointsAtCurrentX;
+
+	// Used by the line reduction algorithm.
+	float LinesCurrentX;
+	float LinesMinY;
+	float LinesMaxY;
+	float LinesFirstY;
+	float LinesLastY;
+	bool bIsLastLineAdded;
+
+	// Used by the box reduction algorithm.
+	//...
+};
+
+////////////////////////////////////////////////////////////////////////////////////////////////////
diff --git a/Engine/Source/Developer/TraceInsights/Public/Insights/ViewModels/TimingEvent.h b/Engine/Source/Developer/TraceInsights/Public/Insights/ViewModels/TimingEvent.h
new file mode 100644
--- /dev/null
+++ b/Engine/Source/Developer/TraceInsights/Public/Insights/ViewModels/TimingEvent.h
@@ -0,0 +1,87 @@
+// Copyright 1998-2019 Epic Games, Inc. All Rights Reserved.
+
+#pragma once
+
+#include "CoreMinimal.h"
+#include "TraceServices/Model/LoadTimeProfiler.h"
+
+class FTimingEventsTrack;
+
+struct FTimingEvent
+{
+	static const uint64 InvalidTypeId = uint64(-1);
+
+	const FTimingEventsTrack* Track;
+	uint64 TypeId;
+	uint32 Depth;
+	double StartTime;
+	double EndTime;
+	double ExclusiveTime;
+
+	// For Asset Loading events...
+	Trace::FLoadTimeProfilerCpuEvent LoadingInfo;
+
+	// For I/O events...
+	uint64 Offset;
+	uint64 Size;
+	const TCHAR* Path;
+
+	FTimingEvent()
+		: Track(nullptr)
+		, TypeId(InvalidTypeId)
+		, Depth(0)
+		, StartTime(0.0)
+		, EndTime(0.0)
+		, ExclusiveTime(0.0)
+		, LoadingInfo()
+		, Offset(0)
+		, Size(0)
+		, Path(nullptr)
+	{
+	}
+
+	FTimingEvent(const FTimingEvent&) = default;
+	FTimingEvent& operator=(const FTimingEvent&) = default;
+
+	FTimingEvent(FTimingEvent&&) = default;
+	FTimingEvent& operator=(FTimingEvent&&) = default;
+
+	void Reset()
+	{
+		Track = nullptr;
+		TypeId = InvalidTypeId;
+		Depth = 0;
+		StartTime = 0.0;
+		EndTime = 0.0;
+		ExclusiveTime = 0.0;
+
+		LoadingInfo.Package = nullptr;
+		LoadingInfo.Export = nullptr;
+		LoadingInfo.PackageEventType = LoadTimeProfilerPackageEventType_None;
+		LoadingInfo.ExportEventType = LoadTimeProfilerObjectEventType_None;
+
+		Offset = 0;
+		Size = 0;
+		Path = nullptr;
+	}
+
+	bool IsValidTrack() const { return Track != nullptr; }
+
+	bool IsValid() const { return Track != nullptr && TypeId != InvalidTypeId; }
+
+	double Duration() const { return EndTime - StartTime; }
+
+	bool Equals(const FTimingEvent& Other) const
+	{
+		return Track == Other.Track
+			&& TypeId == Other.TypeId
+			&& Depth == Other.Depth
+			&& StartTime == Other.StartTime
+			&& EndTime == Other.EndTime;
+	}
+
+	static bool AreEquals(const FTimingEvent& A, const FTimingEvent& B)
+	{
+		return A.Equals(B);
+	}
+};
diff --git a/Engine/Source/Developer/TraceInsights/Public/Insights/ViewModels/TimingEventsTrack.h b/Engine/Source/Developer/TraceInsights/Public/Insights/ViewModels/TimingEventsTrack.h
new file mode 100644
--- /dev/null
+++ b/Engine/Source/Developer/TraceInsights/Public/Insights/ViewModels/TimingEventsTrack.h
@@ -0,0 +1,93 @@
+// Copyright 1998-2019 Epic Games, Inc. All Rights Reserved.
+
+#pragma once
+
+#include "CoreMinimal.h"
+
+// Insights
+#include "Insights/ViewModels/BaseTimingTrack.h"
+
+struct FDrawContext;
+struct FSlateBrush;
+struct FSlateFontInfo;
+struct FTimingEvent;
+class FTimingTrackViewport;
+struct FTimingViewTooltip;
+class ITimingViewDrawHelper;
+class FTooltipDrawState;
+namespace Trace { class IAnalysisSession; }
+
+////////////////////////////////////////////////////////////////////////////////////////////////////
+
+struct TRACEINSIGHTS_API FTimingEventsTrackLayout
+{
+	static constexpr float RealMinTimelineH = 13.0f;
+
+	static constexpr float NormalLayoutEventH = 14.0f;
+	static constexpr float NormalLayoutEventDY = 2.0f;
+	static constexpr float NormalLayoutTimelineDY = 14.0f;
+	static constexpr float NormalLayoutMinTimelineH = 0.0f;
+
+	static constexpr float CompactLayoutEventH = 8.0f;
+	static constexpr float CompactLayoutEventDY = 1.0f;
+	static constexpr float CompactLayoutTimelineDY = 2.0f;
+	static constexpr float CompactLayoutMinTimelineH = 0.0f;
+
+	//////////////////////////////////////////////////
+
+	bool bIsCompactMode;
+
+	float EventH; // height of a timing event, in Slate units
+	float EventDY; // vertical distance between two timing event sub-tracks, in Slate units
+	float TimelineDY; // space at top and bottom of each timeline, in Slate units
+	float MinTimelineH;
+	float TargetMinTimelineH;
+
+	//////////////////////////////////////////////////
+
+	float GetLaneY(uint32 Depth) const { return 1.0f + TimelineDY + Depth * (EventDY + EventH); }
+
+	void ForceNormalMode();
+	void ForceCompactMode();
+	void Update();
+};
+
+////////////////////////////////////////////////////////////////////////////////////////////////////
+
+class TRACEINSIGHTS_API FTimingEventsTrack : public FBaseTimingTrack
+{
+public:
+	explicit FTimingEventsTrack(uint64 InTrackId, const FName& InType, const FName& InSubType, const FString& InName);
+	virtual ~FTimingEventsTrack();
+
+	int32 GetDepth() const { return NumLanes - 1; }
+	int32 GetNumLanes() const { return NumLanes; }
+	void SetDepth(int32 InDepth) { NumLanes = InDepth + 1; }
+
+	virtual void Reset() override;
+	virtual void UpdateHoveredState(float MX, float MY, const FTimingTrackViewport& Viewport);
+
+	virtual void Draw(ITimingViewDrawHelper& Helper) const = 0;
+	virtual void DrawSelectedEventInfo(const FTimingEvent& SelectedTimingEvent, const FTimingTrackViewport& Viewport, const FDrawContext& DrawContext, const FSlateBrush* WhiteBrush, const FSlateFontInfo& Font) const {}
+	virtual void InitTooltip(FTooltipDrawState& Tooltip, const FTimingEvent& HoveredTimingEvent) const {}
+
+	virtual bool SearchTimingEvent(const double InStartTime, const double InEndTime, TFunctionRef<bool(double, double, uint32)> InPredicate, FTimingEvent& InOutTimingEvent, bool bInStopAtFirstMatch, bool bInSearchForLargestEvent) const { return false; }
+	virtual void ComputeTimingEventStats(FTimingEvent& InOutTimingEvent) const {}
+
+	// Called back from the timing view when an event is selected
+	virtual void OnEventSelected(const FTimingEvent& SelectedTimingEvent) const {}
+
+private:
+	int32 NumLanes; // number of lanes (sub-tracks)
+
+	// TODO: Cached OnPaint state.
+	//TArray<FEventBoxInfo> Boxes;
+	//TArray<FEventBoxInfo> MergedBorderBoxes;
+	//TArray<FEventBoxInfo> Borders;
+	//TArray<FTextInfo> Texts;
+
+public:
+	static bool bUseDownSampling;
+};
+
+////////////////////////////////////////////////////////////////////////////////////////////////////
diff --git a/Engine/Source/Developer/TraceInsights/Public/Insights/ViewModels/TimingTrackViewport.h b/Engine/Source/Developer/TraceInsights/Public/Insights/ViewModels/TimingTrackViewport.h
new file mode 100644
--- /dev/null
+++ b/Engine/Source/Developer/TraceInsights/Public/Insights/ViewModels/TimingTrackViewport.h
@@ -0,0 +1,174 @@
+// Copyright 1998-2019 Epic Games, Inc. All Rights Reserved.
+
+#pragma once
+
+#include "CoreMinimal.h"
+
+class SScrollBar;
+
+////////////////////////////////////////////////////////////////////////////////////////////////////
+
+class TRACEINSIGHTS_API FTimingTrackViewport
+{
+public:
+	FTimingTrackViewport()
+	{
+		Reset();
+	}
+
+	void Reset()
+	{
+		Width = 0.0f;
+		Height = 0.0f;
+
+		MinValidTime = 0.0;
+		MaxValidTime = 0.0;
+
+		StartTime = 0.0;
+		EndTime = 0.0;
+
+		MinScaleX = (5 * 20) / 3600.0; // 1h between major tick marks
+		MaxScaleX = 1.0E10; // 10ns between major tick marks
+		ScaleX = (5 * 20) / 5.0; // 5s between major tick marks
+
+		TopOffset = 0.0f;
+		ScrollHeight = 1.0f;
+		ScrollPosY = 0.0f;
+	}
+
+	/** Width of viewport, in pixels. [px] */
+	float GetWidth() const { return Width; }
+
+	/** Height of viewport, in pixels. [px] */
+	float GetHeight() const { return Height; }
+
+	bool UpdateSize(const float InWidth, const float InHeight);
+
+	/** Minimum session time, in seconds. [s] */
+	double GetMinValidTime() const { return MinValidTime; }
+	/** Maximum session time, in seconds. [s] */
+	double GetMaxValidTime() const { return MaxValidTime; }
+	void SetMaxValidTime(const double InMaxValidTime) { MaxValidTime = InMaxValidTime; }
+
+	/**
+	  * Time of viewport's left side, in seconds. [s]
+	  */
+	double GetStartTime() const { return StartTime; }
+	/**
+	  * Time of viewport's right side, in seconds. [s]
+	  * Computed when StartTime, ScaleX or Width changes.
+	  */
+	double GetEndTime() const { return EndTime; }
+	double GetDuration() const { return EndTime - StartTime; }
+
+	double GetMinScaleX() const { return MinScaleX; }
+	double GetMaxScaleX() const { return MaxScaleX; }
+
+	/** Current scale factor between seconds and pixels (Slate units). [px/s] */
+	double GetScaleX() const { return ScaleX; }
+
+	//////////////////////////////////////////////////
+	// Conversions
+
+	float TimeToSlateUnits(const double Time) const
+	{
+		return static_cast<float>((Time - StartTime) * ScaleX);
+	}
+
+	float TimeToSlateUnitsRounded(const double Time) const
+	{
+		return static_cast<float>(FMath::RoundToDouble((Time - StartTime) * ScaleX));
+	}
+
+	float GetViewportDXForDuration(const double DT) const
+	{
+		return static_cast<float>(DT * ScaleX);
+	}
+
+	double GetDurationForViewportDX(const float DX) const
+	{
+		return static_cast<double>(DX) / ScaleX;
+	}
+
+	double SlateUnitsToTime(const float X) const
+	{
+		return StartTime + static_cast<double>(X) / ScaleX;
+	}
+
+	double AlignTimeToPixel(const double InTime, const double InScaleX) const
+	{
+		return FMath::RoundToDouble(InTime * InScaleX) / InScaleX;
+	}
+
+	double AlignTimeToPixel(const double Time) const
+	{
+		return AlignTimeToPixel(Time, ScaleX);
+	}
+
+	double RestrictEndTime(const double InEndTime) const;
+	double RestrictDuration(const double InStartTime, const double InEndTime) const;
+
+	//////////////////////////////////////////////////
+	// Scrolling
+
+	bool ScrollAtTime(const double Time);
+	bool CenterOnTimeInterval(const double Time, const double Duration);
+
+	void GetHorizontalScrollLimits(double& OutMinT, double& OutMaxT);
+	bool EnforceHorizontalScrollLimits(const double U);
+
+	//////////////////////////////////////////////////
+	// Zooming
+
+	bool ZoomOnTimeInterval(const double Time, const double Duration);
+	bool RelativeZoomWithFixedX(const float Delta, const float X);
+	bool ZoomWithFixedX(const double NewScaleX, const float X);
+
+	//////////////////////////////////////////////////
+	// Vertical axis
+
+	/** Top offset (to allow the time ruller to be visible), in pixels. [px] */
+	float GetTopOffset() const { return TopOffset; }
+
+	void SetTopOffset(float InTopOffset) { TopOffset = InTopOffset; }
+
+	/** Height of the vertical scrollable area, in pixels. [px] */
+	float GetScrollHeight() const { return ScrollHeight; }
+	void SetScrollHeight(const float InScrollHeight) { ScrollHeight = InScrollHeight; }
+
+	/** Current vertical scroll position, in pixels. [px] */
+	float GetScrollPosY() const { return ScrollPosY; }
+	void SetScrollPosY(const float InScrollPosY) { ScrollPosY = InScrollPosY; }
+
+	float GetViewportY(const float Y) const { return TopOffset + Y - ScrollPosY; }
+
+	//////////////////////////////////////////////////
+
+	bool OnUserScrolled(TSharedPtr<SScrollBar> ScrollBar, float ScrollOffset);
+	void UpdateScrollBar(TSharedPtr<SScrollBar> ScrollBar) const;
+
+	bool OnUserScrolledY(TSharedPtr<SScrollBar> ScrollBar, float ScrollOffset);
+	void UpdateScrollBarY(TSharedPtr<SScrollBar> ScrollBar) const;
+
+	//////////////////////////////////////////////////
+
+private:
+	float Width; // width of viewport, in pixels; [px]
+	float Height; // height of viewport, in pixels; [px]
+
+	double MinValidTime; // min session time, in seconds; [s]
+	double MaxValidTime; // max session time, in seconds; [s]
+
+	double StartTime; // time of viewport's left side, in seconds; [s]
+	double EndTime; // time of viewport's right side, in seconds; [s]; computed when StartTime, Scale or Width changes
+
+	double MinScaleX; // min scale factor; [px/s]
+	double MaxScaleX; // max scale factor; [px/s]
+	double ScaleX; // scale factor between seconds and pixels; [px/s]
+
+	float TopOffset; // top offset (to allow the time ruller to be visible), in pixels; [px]
+	float ScrollHeight; // height of the vertical scrollable area, in pixels; [px]
+	float ScrollPosY; // current vertical scroll position, in pixels; [px]
+};
+
+////////////////////////////////////////////////////////////////////////////////////////////////////
diff --git a/Engine/Source/Developer/TraceInsights/Public/Insights/ViewModels/TooltipDrawState.h b/Engine/Source/Developer/TraceInsights/Public/Insights/ViewModels/TooltipDrawState.h
new file mode 100644
--- /dev/null
+++ b/Engine/Source/Developer/TraceInsights/Public/Insights/ViewModels/TooltipDrawState.h
@@ -0,0 +1,105 @@
+// Copyright 1998-2019 Epic Games, Inc. All Rights Reserved.
+
+#pragma once
+
+#include "CoreMinimal.h"
+#include "Fonts/SlateFontInfo.h"
+#include "Math/Color.h"
+
+////////////////////////////////////////////////////////////////////////////////////////////////////
+
+struct FDrawContext;
+struct FSlateBrush;
+
+////////////////////////////////////////////////////////////////////////////////////////////////////
+
+class TRACEINSIGHTS_API FTooltipDrawState
+{
+public:
+	static constexpr float BorderX = 6.0f;
+	static constexpr float BorderY = 3.0f;
+	static constexpr float MinWidth = 128.0f;
+	static constexpr float MinHeight = 0.0f;
+
+	static constexpr float DefaultTitleHeight = 14.0f;
+	static constexpr float DefaultLineHeight = 14.0f;
+	static constexpr float NameValueDX = 2.0f;
+
+	static FLinearColor DefaultTitleColor;
+	static FLinearColor DefaultNameColor;
+	static FLinearColor DefaultValueColor;
+
+private:
+	enum class FDrawTextType
+	{
+		Misc,
+		Title,
+		Name,  // X is computed at Draw (right aligned at ValueOffsetX)
+		Value, // X is computed at Draw (left aligned at ValueOffsetX + NameValueDX)
+	};
+
+	struct FDrawTextInfo
+	{
+		float X; // relative X position of text inside tooltip
+		float Y; // relative Y position of text inside tooltip
+		FVector2D TextSize;
+		FString Text;
+		FLinearColor Color;
+		FDrawTextType Type;
+	};
+
+public:
+	FTooltipDrawState();
+	~FTooltipDrawState();
+
+	void Reset();
+
+	void ResetContent();
+	void AddTitle(const FString& Title);
+	void AddTitle(const FString& Title, const FLinearColor& Color);
+	void AddNameValueTextLine(const FString& Name, const FString& Value);
+	void AddTextLine(const FString& Text, const FLinearColor& Color);
+	void AddTextLine(const float X, const float Y, const FString& Text, const FLinearColor& Color);
+	void UpdateLayout(); // updates ValueOffsetX and DesiredSize
+
+	const FLinearColor& GetBackgroundColor() const { return BackgroundColor; }
+	void SetBackgroundColor(const FLinearColor& InBackgroundColor) { BackgroundColor = InBackgroundColor; }
+
+	float GetOpacity() const { return Opacity; }
+	void SetOpacity(const float InOpacity) { Opacity = InOpacity; }
+
+	float GetDesiredOpacity() const { return DesiredOpacity; }
+	void SetDesiredOpacity(const float InDesiredOpacity) { DesiredOpacity = InDesiredOpacity; }
+
+	const FVector2D& GetSize() const { return Size; }
+	const FVector2D& GetDesiredSize() const { return DesiredSize; }
+
+	void Update(); // updates current Opacity and Size
+
+	const FVector2D& GetPosition() const { return Position; }
+	void SetPosition(const FVector2D& MousePosition, const float MinX, const float MaxX, const float MinY, const float MaxY);
+	void SetPosition(const float PosX, const float PosY) { Position.X = PosX; Position.Y = PosY; }
+
+	void Draw(const FDrawContext& DrawContext) const;
+
+private:
+	const FSlateBrush* WhiteBrush;
+	const FSlateFontInfo Font;
+
+	FLinearColor BackgroundColor;
+
+	FVector2D Size;
+	FVector2D DesiredSize;
+
+	FVector2D Position;
+
+	float ValueOffsetX;
+	float NewLineY;
+
+	float Opacity;
+	float DesiredOpacity;
+
+	TArray<FDrawTextInfo> Texts;
+};
+
+////////////////////////////////////////////////////////////////////////////////////////////////////
diff --git a/Engine/Source/Developer/TraceServices/Private/Common/PagedArray.h b/Engine/Source/Developer/TraceServices/Private/Common/PagedArray.h
deleted file mode 100644
diff --git a/Engine/Source/Developer/TraceServices/Private/Model/IntervalTimeline.h b/Engine/Source/Developer/TraceServices/Private/Model/IntervalTimeline.h
deleted file mode 100644
diff --git a/Engine/Source/Developer/TraceServices/Private/Model/MonotonicTimeline.h b/Engine/Source/Developer/TraceServices/Private/Model/MonotonicTimeline.h
deleted file mode 100644
diff --git a/Engine/Source/Developer/TraceServices/Public/Common/PagedArray.h b/Engine/Source/Developer/TraceServices/Public/Common/PagedArray.h
new file mode 100644
--- /dev/null
+++ b/Engine/Source/Developer/TraceServices/Public/Common/PagedArray.h
@@ -0,0 +1,382 @@
+// Copyright 1998-2019 Epic Games, Inc. All Rights Reserved.
+
+#pragma once
+
+#include "TraceServices/Containers/Allocators.h"
+#include "Containers/Array.h"
+
+template<typename ItemType>
+struct TPagedArrayPage
+{
+	ItemType* Items = nullptr;
+	uint64 Count = 0;
+};
+
+template<typename ItemType>
+inline const ItemType* GetData(const TPagedArrayPage<ItemType>& Page)
+{
+	return Page.Items;
+}
+
+template<typename ItemType>
+inline SIZE_T GetNum(const TPagedArrayPage<ItemType>& Page)
+{
+	return Page.Count;
+}
+
+template<typename ItemType>
+inline const ItemType* GetFirstItem(const TPagedArrayPage<ItemType>& Page)
+{
+	return Page.Items;
+}
+
+template<typename ItemType>
+inline const ItemType* GetLastItem(const TPagedArrayPage<ItemType>& Page)
+{
+	if (Page.Count)
+	{
+		return Page.Items + Page.Count - 1;
+	}
+	else
+	{
+		return nullptr;
+	}
+}
+
+template<typename ItemType, typename PageType>
+class TPagedArray;
+
+template<typename ItemType, typename PageType>
+class TPagedArrayIterator
+{
+public:
+	TPagedArrayIterator()
+	{
+
+	}
+
+	TPagedArrayIterator(const TPagedArray<ItemType, PageType>& InOuter, uint64 InitialPageIndex, uint64 InitialItemIndex)
+		: Outer(&InOuter)
+	{
+		if (Outer->PagesArray.Num())
+		{
+			check(InitialPageIndex < Outer->PagesArray.Num());
+			CurrentPageIndex = InitialPageIndex;
+			OnCurrentPageChanged();
+			PageType* CurrentPage = Outer->FirstPage + CurrentPageIndex;
+			check(InitialItemIndex < CurrentPage->Count);
+			CurrentItem = CurrentPage->Items + InitialItemIndex;
+		}
+	}
+
+	uint64 GetCurrentPageIndex()
+	{
+		return CurrentPageIndex;
+	}
+
+	const PageType* GetCurrentPage()
+	{
+		return Outer->FirstPage + CurrentPageIndex;
+	}
+
+	const ItemType* GetCurrentItem()
+	{
+		return CurrentItem;
+	}
+
+	const ItemType& operator*() const
+	{
+		return *CurrentItem;
+	}
+
+	const ItemType* operator->() const
+	{
+		return CurrentItem;
+	}
+
+	explicit operator bool() const
+	{
+		return CurrentItem != nullptr;
+	}
+
+	const PageType* PrevPage()
+	{
+		if (CurrentPageIndex == 0)
+		{
+			CurrentItem = nullptr;
+			CurrentPageFirstItem = nullptr;
+			CurrentPageLastItem = nullptr;
+			return nullptr;
+		}
+		--CurrentPageIndex;
+		OnCurrentPageChanged();
+		CurrentItem = CurrentPageLastItem;
+		return GetCurrentPage();
+	}
+
+	const PageType* NextPage()
+	{
+		if (CurrentPageIndex == Outer->PagesArray.Num() - 1)
+		{
+			CurrentItem = nullptr;
+			CurrentPageFirstItem = nullptr;
+			CurrentPageLastItem = nullptr;
+			return nullptr;
+		}
+		++CurrentPageIndex;
+		OnCurrentPageChanged();
+		CurrentItem = CurrentPageFirstItem;
+		return GetCurrentPage();
+	}
+
+	const ItemType* NextItem()
+	{
+		if (CurrentItem == CurrentPageLastItem)
+		{
+			if (!NextPage())
+			{
+				return nullptr;
+			}
+			else
+			{
+				return CurrentItem;
+			}
+		}
+		++CurrentItem;
+		return CurrentItem;
+	}
+
+	TPagedArrayIterator& operator++()
+	{
+		NextItem();
+		return *this;
+	}
+
+	TPagedArrayIterator operator++(int)
+	{
+		TPagedArrayIterator Tmp(*this);
+		NextItem();
+		return Tmp;
+	}
+
+	const ItemType* PrevItem()
+	{
+		if (CurrentItem == CurrentPageFirstItem)
+		{
+			if (!PrevPage())
+			{
+				return nullptr;
+			}
+			else
+			{
+				return CurrentItem;
+			}
+		}
+		--CurrentItem;
+		return CurrentItem;
+	}
+
+	TPagedArrayIterator& operator--()
+	{
+		PrevItem();
+		return *this;
+	}
+
+	TPagedArrayIterator operator--(int)
+	{
+		TPagedArrayIterator Tmp(*this);
+		PrevItem();
+		return Tmp;
+	}
+
+	const ItemType* SetPosition(uint64 Index)
+	{
+		uint64 PageIndex = Index / Outer->PageSize;
+		uint64 ItemIndexInPage = Index % Outer->PageSize;
+		if (PageIndex != CurrentPageIndex)
+		{
+			check(PageIndex < Outer->PagesArray.Num());
+			CurrentPageIndex = PageIndex;
+			OnCurrentPageChanged();
+		}
+		PageType* CurrentPage = Outer->FirstPage + CurrentPageIndex;
+		check(ItemIndexInPage < CurrentPage->Count);
+		CurrentItem = CurrentPage->Items + ItemIndexInPage;
+		return CurrentItem;
+	}
+
+private:
+	void OnCurrentPageChanged()
+	{
+		PageType* CurrentPage = Outer->FirstPage + CurrentPageIndex;
+		CurrentPageFirstItem = CurrentPage->Items;
+		if (CurrentPage->Items)
+		{
+			CurrentPageLastItem = CurrentPage->Items + CurrentPage->Count - 1;
+		}
+		else
+		{
+			CurrentPageLastItem = nullptr;
+		}
+	}
+
+	const TPagedArray<ItemType, PageType>* Outer = nullptr;
+	const ItemType* CurrentItem = nullptr;
+	const ItemType* CurrentPageFirstItem = nullptr;
+	const ItemType* CurrentPageLastItem = nullptr;
+	uint64 CurrentPageIndex = 0;
+};
+
+template<typename InItemType, typename InPageType = TPagedArrayPage<InItemType>>
+class TPagedArray
+{
+public:
+	typedef InItemType ItemType;
+	typedef InPageType PageType;
+	typedef TPagedArrayIterator<InItemType, InPageType> TIterator;
+	
+	TPagedArray(Trace::ILinearAllocator& InAllocator, uint64 InPageSize)
+		: Allocator(InAllocator)
+		, PageSize(InPageSize)
+	{
+
+	}
+
+	~TPagedArray()
+	{
+		for (PageType& Page : PagesArray)
+		{
+			ItemType* PageEnd = Page.Items + Page.Count;
+			for (ItemType* Item = Page.Items; Item != PageEnd; ++Item)
+			{
+				Item->~ItemType();
+			}
+		}
+	}
+
+	uint64 Num() const
+	{
+		return TotalItemCount;
+	}
+
+	uint64 GetPageSize() const
+	{
+		return PageSize;
+	}
+
+	uint64 NumPages() const
+	{
+		return PagesArray.Num();
+	}
+
+	ItemType& PushBack()
+	{
+		if (!LastPage || LastPage->Count == PageSize)
+		{
+			LastPage = &PagesArray.AddDefaulted_GetRef();
+			FirstPage = PagesArray.GetData();
+			LastPage->Items = reinterpret_cast<ItemType*>(Allocator.Allocate(PageSize * sizeof(ItemType)));
+		}
+		++TotalItemCount;
+		ItemType* ItemPtr = LastPage->Items + LastPage->Count;
+		new (ItemPtr) ItemType();
+		++LastPage->Count;
+		return *ItemPtr;
+	}
+
+	ItemType& Insert(uint64 Index)
+	{
+		if (Index >= TotalItemCount)
+		{
+			return PushBack();
+		}
+		PushBack();
+		uint64 PageIndex = Index / PageSize;
+		uint64 PageItemIndex = Index % PageSize;
+		for (uint64 CurrentPageIndex = PagesArray.Num() - 1; CurrentPageIndex > PageIndex; --CurrentPageIndex)
+		{
+			PageType* CurrentPage = FirstPage + CurrentPageIndex;
+			memmove(CurrentPage->Items + 1, CurrentPage->Items, sizeof(ItemType) * (CurrentPage->Count - 1));
+			PageType* PrevPage = CurrentPage - 1;
+			memcpy(CurrentPage->Items, PrevPage->Items + PrevPage->Count - 1, sizeof(ItemType));
+		}
+		PageType* Page = FirstPage + PageIndex;
+		memmove(Page->Items + PageItemIndex + 1, Page->Items + PageItemIndex, sizeof(ItemType) * (Page->Count - PageItemIndex - 1));
+		return Page->Items[PageItemIndex];
+	}
+
+	PageType* GetLastPage()
+	{
+		return LastPage;
+	}
+
+	PageType* GetPage(uint64 PageIndex)
+	{
+		return FirstPage + PageIndex;
+	}
+
+	PageType* GetItemPage(uint64 ItemIndex)
+	{
+		uint64 PageIndex = ItemIndex / PageSize;
+		return FirstPage + PageIndex;
+	}
+
+	TIterator GetIterator() const
+	{
+		return TIterator(*this, 0, 0);
+	}
+
+	TIterator GetIteratorFromPage(uint64 PageIndex) const
+	{
+		return TIterator(*this, PageIndex, 0);
+	}
+
+	TIterator GetIteratorFromItem(uint64 ItemIndex) const
+	{
+		uint64 PageIndex = ItemIndex / PageSize;
+		uint64 IndexInPage = ItemIndex % PageSize;
+		return TIterator(*this, PageIndex, IndexInPage);
+	}
+
+	const PageType* GetPages() const
+	{
+		return FirstPage;
+	}
+
+	ItemType& operator[](uint64 Index)
+	{
+		uint64 PageIndex = Index / PageSize;
+		uint64 IndexInPage = Index % PageSize;
+		PageType* Page = FirstPage + PageIndex;
+		ItemType* Item = Page->Items + IndexInPage;
+		return *Item;
+	}
+
+	const ItemType& operator[](uint64 Index) const
+	{
+		return const_cast<TPagedArray&>(*this)[Index];
+	}
+
+private:
+	template<typename ItemType, typename PageType>
+	friend class TPagedArrayIterator;
+
+	Trace::ILinearAllocator& Allocator;
+	TArray<PageType> PagesArray;
+	PageType* FirstPage = nullptr;
+	PageType* LastPage = nullptr;
+	uint64 PageSize;
+	uint64 TotalItemCount = 0;
+};
+
+template<typename ItemType, typename PageType>
+inline const PageType* GetData(const TPagedArray<ItemType, PageType>& PagedArray)
+{
+	return PagedArray.GetPages();
+}
+
+template<typename ItemType, typename PageType>
+inline SIZE_T GetNum(const TPagedArray<ItemType, PageType>& PagedArray)
+{
+	return PagedArray.NumPages();
+}
\ No newline at end of file
diff --git a/Engine/Source/Developer/TraceServices/Public/Model/IntervalTimeline.h b/Engine/Source/Developer/TraceServices/Public/Model/IntervalTimeline.h
new file mode 100644
--- /dev/null
+++ b/Engine/Source/Developer/TraceServices/Public/Model/IntervalTimeline.h
@@ -0,0 +1,145 @@
+// Copyright 1998-2019 Epic Games, Inc. All Rights Reserved.
+
+#pragma once
+
+#include "TraceServices/Containers/Timelines.h"
+#include "Common/PagedArray.h"
+#include "TraceServices/Containers/Allocators.h"
+#include <limits>
+
+namespace Trace
+{
+
+struct FIntervalTimelineDefaultSettings
+{
+	enum
+	{
+		EventsPerPage = 65536
+	};
+};
+
+template<typename EventType, typename SettingsType = FIntervalTimelineDefaultSettings>
+class TIntervalTimeline
+	: public ITimeline<EventType>
+{
+public:
+	TIntervalTimeline(ILinearAllocator& Allocator)
+		: Events(Allocator, SettingsType::EventsPerPage)
+	{
+	}
+
+	virtual ~TIntervalTimeline() = default;
+	virtual uint64 GetModCount() const override { return ModCount; }
+	virtual uint64 GetEventCount() const override { return Events.Num(); }
+	virtual const EventType& GetEvent(uint64 InIndex) const override { return Events[InIndex].Event; }
+	virtual double GetStartTime() const override { return Events.Num() > 0 ? Events[0].StartTime : 0.0; }
+	virtual double GetEndTime() const override { return Events.Num() > 0 ? Events[Events.Num() - 1].EndTime : 0.0; }
+	virtual void EnumerateEventsDownSampled(double IntervalStart, double IntervalEnd, double Resolution, ITimeline::EventCallback Callback) const override { check(false); }
+	virtual void EnumerateEventsDownSampled(double IntervalStart, double IntervalEnd, double Resolution, ITimeline::EventRangeCallback Callback) const override { check(false); }
+	virtual void EnumerateEvents(double IntervalStart, double IntervalEnd, ITimeline::EventCallback Callback) const override
+	{
+		EnumerateEvents(IntervalStart, IntervalEnd, [&Callback](double InStartTime, double InEndTime, uint32 InDepth, const EventType& InEvent)
+		{
+			Callback(true, InStartTime, InEvent);
+			Callback(false, InEndTime, InEvent);
+		});
+	}
+	
+	virtual void EnumerateEvents(double IntervalStart, double IntervalEnd, ITimeline::EventRangeCallback Callback) const override
+	{
+		auto EventIterator = Events.GetIteratorFromPage(0);
+		const FEventInternal* Event = EventIterator.GetCurrentItem();
+		while (Event)
+		{
+			const FEventPage* Page = EventIterator.GetCurrentPage();
+			if (Page->EndTime < IntervalStart || IntervalEnd < Page->BeginTime)
+			{
+				if (!EventIterator.NextPage())
+				{
+					return;
+				}
+				Event = EventIterator.GetCurrentItem();
+			}
+			else
+			{
+				if (!(Event->EndTime < IntervalStart || IntervalEnd < Event->StartTime))
+				{
+					Callback(Event->StartTime, Event->EndTime, 0, Event->Event);
+				}
+				Event = EventIterator.NextItem();
+			}
+		}
+	}
+	
+	uint64 AppendBeginEvent(double StartTime, const EventType& Event)
+	{
+		//check(StartTime >= LastTime);
+		//LastTime = StartTime;
+
+		uint64 Index = Events.Num();
+		FEventInternal& EventInternal = Events.PushBack();
+		EventInternal.StartTime = StartTime;
+		EventInternal.EndTime = std::numeric_limits<double>::infinity();
+		EventInternal.Event = Event;
+
+		FEventPage* LastPage = Events.GetLastPage();
+		if (LastPage != CurrentPage)
+		{
+			CurrentPage = LastPage;
+			LastPage->BeginTime = StartTime;
+			LastPage->EndTime = StartTime;
+			LastPage->BeginEventIndex = Index;
+		}
+		else
+		{
+			LastPage->BeginTime = FMath::Min(LastPage->BeginTime, StartTime);
+			LastPage->EndTime = FMath::Max(LastPage->EndTime, StartTime);
+		}
+
+		LastPage->EndEventIndex = Index + 1;
+
+		++ModCount;
+
+		return Index;
+	}
+
+	EventType& EndEvent(uint64 EventIndex, double EndTime)
+	{
+		//check(EndTime >= LastTime);
+		//LastTime = EndTime;
+		FEventInternal& EventInternal = Events[EventIndex];
+		check(EventIndex < Events.Num());
+		EventInternal.EndTime = EndTime;
+		FEventPage* Page = Events.GetItemPage(EventIndex);
+		Page->EndTime = FMath::Max(Page->EndTime, EndTime);
+
+		++ModCount;
+
+		return EventInternal.Event;
+	}
+
+private:
+	struct FEventInternal
+	{
+		double StartTime;
+		double EndTime;
+		EventType Event;
+	};
+
+	struct FEventPage
+	{
+		double BeginTime = 0.0;
+		double EndTime = 0.0;
+		uint64 BeginEventIndex = 0;
+		uint64 EndEventIndex = 0;
+		FEventInternal* Items = nullptr;
+		uint64 Count = 0;
+	};
+
+	TPagedArray<FEventInternal, FEventPage> Events;
+	FEventPage* CurrentPage = nullptr;
+	double LastTime = 0.0;
+	uint64 ModCount = 0;
+};
+
+}
\ No newline at end of file
diff --git a/Engine/Source/Developer/TraceServices/Public/Model/MonotonicTimeline.h b/Engine/Source/Developer/TraceServices/Public/Model/MonotonicTimeline.h
new file mode 100644
--- /dev/null
+++ b/Engine/Source/Developer/TraceServices/Public/Model/MonotonicTimeline.h
@@ -0,0 +1,488 @@
+// Copyright 1998-2019 Epic Games, Inc. All Rights Reserved.
+
+#pragma once
+
+#include "TraceServices/AnalysisService.h"
+#include "Common/PagedArray.h"
+
+namespace Trace
+{
+
+struct FMonotonicTimelineDefaultSettings
+{
+	enum
+	{
+		MaxDepth = 1024,
+		ScopeEntriesPageSize = 65536,
+		EventsPageSize = 65536,
+		DetailLevelsCount = 6,
+	};
+
+	constexpr static double DetailLevelResolution(int32 Index)
+	{
+		const double DetailLevels[DetailLevelsCount] = { 0.0, 0.0001, 0.001, 0.008, 0.04, 0.2 };
+		return DetailLevels[Index];
+	}
+};
+
+template<typename InEventType, typename SettingsType = FMonotonicTimelineDefaultSettings>
+class TMonotonicTimeline
+	: public ITimeline<InEventType>
+{
+public:
+	using EventType = InEventType;
+
+	TMonotonicTimeline(ILinearAllocator& InAllocator)
+		: Allocator(InAllocator)
+	{
+		
+		for (int32 DetailLevelIndex = 0; DetailLevelIndex < SettingsType::DetailLevelsCount; ++DetailLevelIndex)
+		{
+			double Resolution = SettingsType::DetailLevelResolution(DetailLevelIndex);
+			DetailLevels.Emplace(Allocator, Resolution);
+		}
+	}
+
+	virtual ~TMonotonicTimeline() = default;
+	
+	void AppendBeginEvent(double StartTime, const EventType& Event)
+	{
+		int32 CurrentDepth = DetailLevels[0].InsertionState.CurrentDepth;
+
+		AddScopeEntry(DetailLevels[0], StartTime, true);
+		AddEvent(DetailLevels[0], Event);
+		FDetailLevelDepthState& Lod0DepthState = DetailLevels[0].InsertionState.DepthStates[CurrentDepth];
+		Lod0DepthState.EnterTime = StartTime;
+		Lod0DepthState.DominatingEvent = Event;
+		//Lod0DepthState.DebugDominatingEventType = Owner.EventTypes[TypeId];
+
+		for (int32 DetailLevelIndex = 1; DetailLevelIndex < SettingsType::DetailLevelsCount; ++DetailLevelIndex)
+		{
+			FDetailLevel& DetailLevel = DetailLevels[DetailLevelIndex];
+			FDetailLevelDepthState& CurrentDepthState = DetailLevel.InsertionState.DepthStates[CurrentDepth];
+
+			if (CurrentDepthState.PendingScopeEnterIndex < 0 || StartTime >= CurrentDepthState.EnterTime + DetailLevel.Resolution)
+			{
+				for (int32 Depth = DetailLevel.InsertionState.PendingDepth; Depth >= CurrentDepth; --Depth)
+				{
+					FDetailLevelDepthState& DepthState = DetailLevel.InsertionState.DepthStates[Depth];
+					check(DepthState.PendingScopeEnterIndex >= 0);
+					AddScopeEntry(DetailLevel, DepthState.ExitTime, false);
+
+					DepthState.PendingScopeEnterIndex = -1;
+					DepthState.PendingEventIndex = -1;
+				}
+				DetailLevel.InsertionState.PendingDepth = CurrentDepth;
+
+				uint64 EnterScopeIndex = DetailLevel.ScopeEntries.Num();
+				uint64 EventIndex = DetailLevel.Events.Num();
+
+				AddScopeEntry(DetailLevel, StartTime, true);
+				AddEvent(DetailLevel, Event);
+
+				CurrentDepthState.DominatingEventStartTime = StartTime;
+				CurrentDepthState.DominatingEventEndTime = StartTime;
+				CurrentDepthState.DominatingEventDuration = 0.0;
+				CurrentDepthState.PendingScopeEnterIndex = EnterScopeIndex;
+				CurrentDepthState.PendingEventIndex = EventIndex;
+				CurrentDepthState.EnterTime = StartTime;
+				CurrentDepthState.DominatingEvent = Event;
+				//CurrentDepthState.DebugDominatingEventType = Owner.EventTypes[TypeId];
+			}
+			else if (CurrentDepth > DetailLevel.InsertionState.PendingDepth)
+			{
+				DetailLevel.InsertionState.PendingDepth = CurrentDepth;
+			}
+			SetEvent(DetailLevel, CurrentDepthState.PendingEventIndex, Event);
+		}
+		++ModCount;
+	}
+
+	void AppendEndEvent(double EndTime)
+	{
+		AddScopeEntry(DetailLevels[0], EndTime, false);
+
+		int32 CurrentDepth = DetailLevels[0].InsertionState.CurrentDepth;
+		for (int32 DetailLevelIndex = 1; DetailLevelIndex < SettingsType::DetailLevelsCount; ++DetailLevelIndex)
+		{
+			FDetailLevel& DetailLevel = DetailLevels[DetailLevelIndex];
+
+			DetailLevel.InsertionState.DepthStates[CurrentDepth].ExitTime = EndTime;
+
+			UpdateDominatingEvent(DetailLevel, CurrentDepth, EndTime);
+		}
+		++ModCount;
+	}
+
+	virtual uint64 GetModCount() const override
+	{
+		return ModCount;
+	}
+
+	virtual uint64 GetEventCount() const override
+	{
+		return DetailLevels[0].Events.Num();
+	}
+
+	virtual const EventType& GetEvent(uint64 InIndex) const override 
+	{ 
+		return DetailLevels[0].Events[InIndex]; 
+	}
+
+	virtual double GetStartTime() const override
+	{
+		return DetailLevels[0].ScopeEntries.Num() > 0 ? FMath::Abs(DetailLevels[0].ScopeEntries[0].Time) : 0.0;
+	}
+
+	virtual double GetEndTime() const override
+	{
+		uint64 NumScopeEntries = DetailLevels[0].ScopeEntries.Num();
+		return NumScopeEntries > 0 ? FMath::Abs(DetailLevels[0].ScopeEntries[NumScopeEntries - 1].Time) : 0.0;
+	}
+
+	virtual void EnumerateEventsDownSampled(double IntervalStart, double IntervalEnd, double Resolution, ITimeline::EventCallback Callback) const override
+	{
+		int32 DetailLevelIndex = SettingsType::DetailLevelsCount - 1;
+		for (; DetailLevelIndex > 0; --DetailLevelIndex)
+		{
+			if (DetailLevels[DetailLevelIndex].Resolution <= Resolution)
+			{
+				break;
+			}
+		}
+
+		const FDetailLevel& DetailLevel = DetailLevels[DetailLevelIndex];
+		if (DetailLevel.ScopeEntries.Num() == 0)
+		{
+			return;
+		}
+
+		uint64 FirstScopePageIndex = Algo::UpperBoundBy(DetailLevel.ScopeEntries, IntervalStart, [](const FEventScopeEntryPage& Page)
+		{
+			return Page.BeginTime;
+		});
+		if (FirstScopePageIndex > 0)
+		{
+			--FirstScopePageIndex;
+		}
+		auto ScopeEntryIterator = DetailLevel.ScopeEntries.GetIteratorFromPage(FirstScopePageIndex);
+		const FEventScopeEntryPage* ScopePage = ScopeEntryIterator.GetCurrentPage();
+		if (ScopePage->BeginTime > IntervalEnd)
+		{
+			return;
+		}
+		if (ScopePage->EndTime < IntervalStart)
+		{
+			return;
+		}
+		auto EventsIterator = DetailLevel.Events.GetIteratorFromItem(ScopePage->BeginEventIndex);
+		struct FEnumerationStackEntry
+		{
+			double StartTime;
+			EventType Event;
+		};
+		FEnumerationStackEntry EventStack[SettingsType::MaxDepth];
+		int32 CurrentStackDepth = ScopePage->InitialStackCount;
+		for (int32 InitialStackIndex = 0; InitialStackIndex < CurrentStackDepth; ++InitialStackIndex)
+		{
+			FEnumerationStackEntry& EnumerationStackEntry = EventStack[InitialStackIndex];
+			const FEventStackEntry& EventStackEntry = ScopePage->InitialStack[InitialStackIndex];
+			EnumerationStackEntry.StartTime = GetScopeEntryTime(DetailLevel, EventStackEntry.EnterScopeIndex);
+			EnumerationStackEntry.Event = GetEvent(DetailLevel, EventStackEntry.EventIndex);
+		}
+
+		const FEventScopeEntry* ScopeEntry = ScopeEntryIterator.GetCurrentItem();
+		const EventType* Event = EventsIterator.GetCurrentItem();
+		while (ScopeEntry && FMath::Abs(ScopeEntry->Time) < IntervalStart)
+		{
+			if (ScopeEntry->Time < 0.0)
+			{
+				check(CurrentStackDepth < SettingsType::MaxDepth);
+				FEnumerationStackEntry& StackEntry = EventStack[CurrentStackDepth++];
+				StackEntry.Event = *Event;
+				StackEntry.StartTime = -ScopeEntry->Time;
+				Event = EventsIterator.NextItem();
+			}
+			else
+			{
+				check(CurrentStackDepth > 0);
+				--CurrentStackDepth;
+			}
+			ScopeEntry = ScopeEntryIterator.NextItem();
+		}
+		if (CurrentStackDepth == 1 && EventStack[0].StartTime > IntervalEnd)
+		{
+			return;
+		}
+		for (int32 StackIndex = 0; StackIndex < CurrentStackDepth; ++StackIndex)
+		{
+			FEnumerationStackEntry& StackEntry = EventStack[StackIndex];
+			Callback(true, StackEntry.StartTime, StackEntry.Event);
+		}
+		while (ScopeEntry && FMath::Abs(ScopeEntry->Time) <= IntervalEnd)
+		{
+			if (ScopeEntry->Time < 0.0)
+			{
+				check(CurrentStackDepth < SettingsType::MaxDepth);
+				FEnumerationStackEntry& StackEntry = EventStack[CurrentStackDepth++];
+				StackEntry.Event = *Event;
+				Callback(true, -ScopeEntry->Time, StackEntry.Event);
+				Event = EventsIterator.NextItem();
+			}
+			else
+			{
+				check(CurrentStackDepth > 0);
+				FEnumerationStackEntry& StackEntry = EventStack[--CurrentStackDepth];
+				Callback(false, ScopeEntry->Time, StackEntry.Event);
+			}
+			ScopeEntry = ScopeEntryIterator.NextItem();
+		}
+		uint32 ExitDepth = 0;
+		while (CurrentStackDepth > 0 && ScopeEntry)
+		{
+			if (ScopeEntry->Time < 0.0)
+			{
+				++ExitDepth;
+			}
+			else
+			{
+				if (ExitDepth == 0)
+				{
+					FEnumerationStackEntry& StackEntry = EventStack[--CurrentStackDepth];
+					Callback(false, ScopeEntry->Time, StackEntry.Event);
+				}
+				else
+				{
+					--ExitDepth;
+				}
+			}
+			ScopeEntry = ScopeEntryIterator.NextItem();
+		}
+		while (CurrentStackDepth > 0)
+		{
+			FEnumerationStackEntry& StackEntry = EventStack[--CurrentStackDepth];
+			Callback(false, DetailLevel.InsertionState.LastTime, StackEntry.Event);
+		}
+	}
+
+	virtual void EnumerateEventsDownSampled(double IntervalStart, double IntervalEnd, double Resolution, ITimeline::EventRangeCallback Callback) const override
+	{
+		struct FStackEntry
+		{
+			uint64 LocalEventIndex;
+		};
+		FStackEntry EventStack[SettingsType::MaxDepth];
+		uint32 CurrentDepth = 0;
+
+		struct FOutputEvent
+		{
+			double StartTime;
+			double EndTime;
+			uint32 Depth;
+			EventType Event;
+		};
+		TArray<FOutputEvent> OutputEvents;
+
+		EnumerateEventsDownSampled(IntervalStart, IntervalEnd, Resolution, [&EventStack, &OutputEvents, &CurrentDepth, Callback](bool IsEnter, double Time, const EventType& Event)
+		{
+			if (IsEnter)
+			{
+				FStackEntry& StackEntry = EventStack[CurrentDepth];
+				StackEntry.LocalEventIndex = OutputEvents.Num();
+				FOutputEvent& OutputEvent = OutputEvents.AddDefaulted_GetRef();
+				OutputEvent.StartTime = Time;
+				OutputEvent.EndTime = Time;
+				OutputEvent.Depth = CurrentDepth;
+				OutputEvent.Event = Event;
+				++CurrentDepth;
+			}
+			else
+			{
+				{
+					FStackEntry& StackEntry = EventStack[--CurrentDepth];
+					FOutputEvent* OutputEvent = OutputEvents.GetData() + StackEntry.LocalEventIndex;
+					OutputEvent->EndTime = Time;
+				}
+				if (CurrentDepth == 0)
+				{
+					for (FOutputEvent& OutputEvent : OutputEvents)
+					{
+						Callback(OutputEvent.StartTime, OutputEvent.EndTime, OutputEvent.Depth, OutputEvent.Event);
+					}
+					OutputEvents.Empty(OutputEvents.Num());
+				}
+			}
+		});
+	}
+
+	virtual void EnumerateEvents(double IntervalStart, double IntervalEnd, ITimeline::EventCallback Callback) const override
+	{
+		EnumerateEventsDownSampled(IntervalStart, IntervalEnd, 0.0, Callback);
+	}
+
+	virtual void EnumerateEvents(double IntervalStart, double IntervalEnd, ITimeline::EventRangeCallback Callback) const override
+	{
+		EnumerateEventsDownSampled(IntervalStart, IntervalEnd, 0.0, Callback);
+	}
+
+private:
+	struct FEventScopeEntry
+	{
+		double Time;
+		//uint32 DebugDepth;
+	};
+
+	struct FEventStackEntry
+	{
+		uint64 EnterScopeIndex;
+		uint64 EventIndex;
+	};
+
+	struct FEventScopeEntryPage
+	{
+		FEventScopeEntry* Items = nullptr;
+		uint64 Count = 0;
+		double BeginTime = 0.0;
+		double EndTime = 0.0;
+		uint64 BeginEventIndex = 0;
+		uint64 EndEventIndex = 0;
+		FEventStackEntry* InitialStack = nullptr;
+		uint16 InitialStackCount = 0;
+	};
+
+	struct FDetailLevelDepthState
+	{
+		int64 PendingScopeEnterIndex = -1;
+		int64 PendingEventIndex = -1;
+
+		EventType DominatingEvent;
+		double DominatingEventStartTime = 0.0;
+		double DominatingEventEndTime = 0.0;
+		double DominatingEventDuration = 0.0;
+
+		double EnterTime = 0.0;
+		double ExitTime = 0.0;
+
+		//const FTimelineEventType* DebugDominatingEventType;
+	};
+
+	struct FDetailLevelInsertionState
+	{
+		double LastTime = -1.0;
+		uint16 CurrentDepth = 0;
+		int32 PendingDepth = -1;
+		FDetailLevelDepthState DepthStates[SettingsType::MaxDepth];
+		FEventStackEntry EventStack[SettingsType::MaxDepth];
+		FEventScopeEntryPage* CurrentScopeEntryPage = nullptr;
+	};
+
+	struct FDetailLevel
+	{
+		FDetailLevel(ILinearAllocator& Allocator, double InResolution)
+			: Resolution(InResolution)
+			, ScopeEntries(Allocator, SettingsType::ScopeEntriesPageSize)
+			, Events(Allocator, SettingsType::EventsPageSize)
+		{
+
+		}
+
+		double Resolution;
+		TPagedArray<FEventScopeEntry, FEventScopeEntryPage> ScopeEntries;
+		TPagedArray<EventType> Events;
+
+		FDetailLevelInsertionState InsertionState;
+	};
+
+	void UpdateDominatingEvent(FDetailLevel& DetailLevel, int32 Depth, double CurrentTime)
+	{
+		FDetailLevelDepthState& Lod0DepthState = DetailLevels[0].InsertionState.DepthStates[Depth];
+		double Lod0EventDuration = CurrentTime - Lod0DepthState.EnterTime;
+		FDetailLevelDepthState& CurrentDepthState = DetailLevel.InsertionState.DepthStates[Depth];
+		if (Lod0EventDuration > CurrentDepthState.DominatingEventDuration)
+		{
+			check(CurrentDepthState.PendingScopeEnterIndex >= 0);
+			check(CurrentDepthState.PendingEventIndex >= 0);
+
+			CurrentDepthState.DominatingEvent = Lod0DepthState.DominatingEvent;
+			CurrentDepthState.DominatingEventStartTime = Lod0DepthState.EnterTime;
+			CurrentDepthState.DominatingEventEndTime = CurrentTime;
+			CurrentDepthState.DominatingEventDuration = Lod0EventDuration;
+
+			SetEvent(DetailLevel, CurrentDepthState.PendingEventIndex, CurrentDepthState.DominatingEvent);
+
+			//CurrentDepthState.DebugDominatingEventType = Owner.EventTypes[CurrentDepthState.DominatingEventType];
+		}
+	}
+
+	void AddScopeEntry(FDetailLevel& DetailLevel, double Time, bool IsEnter)
+	{
+		check(Time >= DetailLevel.InsertionState.LastTime);
+		DetailLevel.InsertionState.LastTime = Time;
+
+		uint64 EventIndex = DetailLevel.Events.Num();
+		uint64 ScopeIndex = DetailLevel.ScopeEntries.Num();
+
+		FEventScopeEntry& ScopeEntry = DetailLevel.ScopeEntries.PushBack();
+		ScopeEntry.Time = IsEnter ? -Time : Time;
+		FEventScopeEntryPage* LastPage = DetailLevel.ScopeEntries.GetLastPage();
+		if (LastPage != DetailLevel.InsertionState.CurrentScopeEntryPage)
+		{
+			DetailLevel.InsertionState.CurrentScopeEntryPage = LastPage;
+			LastPage->BeginTime = Time;
+			LastPage->BeginEventIndex = DetailLevel.Events.Num();
+			LastPage->EndEventIndex = LastPage->BeginEventIndex;
+			LastPage->InitialStackCount = DetailLevel.InsertionState.CurrentDepth;
+			if (LastPage->InitialStackCount)
+			{
+				LastPage->InitialStack = reinterpret_cast<FEventStackEntry*>(Allocator.Allocate(LastPage->InitialStackCount * sizeof(FEventStackEntry)));
+				memcpy(LastPage->InitialStack, DetailLevel.InsertionState.EventStack, LastPage->InitialStackCount * sizeof(FEventStackEntry));
+			}
+		}
+		LastPage->EndTime = Time;
+
+		if (IsEnter)
+		{
+			FEventStackEntry& StackEntry = DetailLevel.InsertionState.EventStack[DetailLevel.InsertionState.CurrentDepth++];
+			check(DetailLevel.InsertionState.CurrentDepth < SettingsType::MaxDepth);
+			StackEntry.EventIndex = EventIndex;
+			StackEntry.EnterScopeIndex = ScopeIndex;
+		}
+		else
+		{
+			check(DetailLevel.InsertionState.CurrentDepth > 0);
+			--DetailLevel.InsertionState.CurrentDepth;
+		}
+
+		//ScopeEntry.DebugDepth = DetailLevel.InsertionState.CurrentDepth;
+	}
+
+	void AddEvent(FDetailLevel& DetailLevel, const EventType& Event)
+	{
+		++DetailLevel.InsertionState.CurrentScopeEntryPage->EndEventIndex;
+		DetailLevel.Events.PushBack() = Event;
+
+		//Event.DebugType = Owner.EventTypes[TypeIndex];
+	}
+
+	double GetScopeEntryTime(const FDetailLevel& DetailLevel, uint64 Index) const
+	{
+		const FEventScopeEntry& ScopeEntry = DetailLevel.ScopeEntries[Index];
+		return ScopeEntry.Time < 0 ? -ScopeEntry.Time : ScopeEntry.Time;
+	}
+
+	void SetEvent(FDetailLevel& DetailLevel, uint64 Index, const EventType& Event)
+	{
+		DetailLevel.Events[Index] = Event;
+	}
+
+	EventType GetEvent(const FDetailLevel& DetailLevel, uint64 Index) const
+	{
+		return DetailLevel.Events[Index];
+	}
+
+	ILinearAllocator& Allocator;
+	TArray<FDetailLevel> DetailLevels;
+	uint64 ModCount = 0;
+};
+
+}
\ No newline at end of file
diff --git a/Engine/Source/Developer/TraceServices/Public/Model/PointTimeline.h b/Engine/Source/Developer/TraceServices/Public/Model/PointTimeline.h
new file mode 100644
--- /dev/null
+++ b/Engine/Source/Developer/TraceServices/Public/Model/PointTimeline.h
@@ -0,0 +1,125 @@
+// Copyright 1998-2019 Epic Games, Inc. All Rights Reserved.
+
+#pragma once
+
+#include "TraceServices/Containers/Timelines.h"
+#include "Common/PagedArray.h"
+#include "TraceServices/Containers/Allocators.h"
+#include <limits>
+
+namespace Trace
+{
+
+struct FPointTimelineDefaultSettings
+{
+	enum
+	{
+		EventsPerPage = 65536
+	};
+};
+
+template<typename EventType, typename SettingsType = FPointTimelineDefaultSettings>
+class TPointTimeline
+	: public ITimeline<EventType>
+{
+public:
+	TPointTimeline(ILinearAllocator& Allocator)
+		: Events(Allocator, SettingsType::EventsPerPage)
+	{
+	}
+
+	virtual ~TPointTimeline() = default;
+	virtual uint64 GetModCount() const override { return ModCount; }
+	virtual uint64 GetEventCount() const override { return Events.Num(); }
+	virtual const EventType& GetEvent(uint64 InIndex) const override { return Events[InIndex].Event; }
+	virtual double GetStartTime() const override { return Events.Num() > 0 ? Events[0].Time : 0.0; }
+	virtual double GetEndTime() const override { return Events.Num() > 0 ? Events[Events.Num() - 1].Time : 0.0; }
+	virtual void EnumerateEventsDownSampled(double IntervalStart, double IntervalEnd, double Resolution, ITimeline::EventCallback Callback) const override { check(false); }
+	virtual void EnumerateEventsDownSampled(double IntervalStart, double IntervalEnd, double Resolution, ITimeline::EventRangeCallback Callback) const override { check(false); }
+	virtual void EnumerateEvents(double IntervalStart, double IntervalEnd, ITimeline::EventCallback Callback) const override
+	{
+		EnumerateEvents(IntervalStart, IntervalEnd, [&Callback](double InStartTime, double InEndTime, uint32 InDepth, const EventType& InEvent)
+		{
+			Callback(true, InStartTime, InEvent);
+			Callback(false, InEndTime, InEvent);
+		});
+	}
+	
+	virtual void EnumerateEvents(double IntervalStart, double IntervalEnd, ITimeline::EventRangeCallback Callback) const override
+	{
+		auto EventIterator = Events.GetIteratorFromPage(0);
+		const FEventInternal* Event = EventIterator.GetCurrentItem();
+		while (Event)
+		{
+			const FEventPage* Page = EventIterator.GetCurrentPage();
+			if (Page->EndTime < IntervalStart || IntervalEnd < Page->BeginTime)
+			{
+				if (!EventIterator.NextPage())
+				{
+					return;
+				}
+				Event = EventIterator.GetCurrentItem();
+			}
+			else
+			{
+				if (!(Event->Time < IntervalStart || IntervalEnd < Event->Time))
+				{
+					Callback(Event->Time, Event->Time, 0, Event->Event);
+				}
+				Event = EventIterator.NextItem();
+			}
+		}
+	}
+	
+	uint64 AppendEvent(double Time, const EventType& Event)
+	{
+		uint64 Index = Events.Num();
+		FEventInternal& EventInternal = Events.PushBack();
+		EventInternal.Time = Time;
+		EventInternal.Event = Event;
+
+		FEventPage* LastPage = Events.GetLastPage();
+		if (LastPage != CurrentPage)
+		{
+			CurrentPage = LastPage;
+			LastPage->BeginTime = Time;
+			LastPage->EndTime = Time;
+			LastPage->BeginEventIndex = Index;
+		}
+		else
+		{
+			LastPage->BeginTime = FMath::Min(LastPage->BeginTime, Time);
+			LastPage->EndTime = FMath::Max(LastPage->EndTime, Time);
+		}
+
+		LastPage->EndEventIndex = Index + 1;
+
+		++ModCount;
+
+		return Index;
+	}
+
+private:
+	struct FEventInternal
+	{
+		double Time;
+		EventType Event;
+	};
+
+	struct FEventPage
+	{
+		double BeginTime = 0.0;
+		double EndTime = 0.0;
+		uint64 BeginEventIndex = 0;
+		uint64 EndEventIndex = 0;
+		FEventInternal* Items = nullptr;
+		uint64 Count = 0;
+	};
+
+	TPagedArray<FEventInternal, FEventPage> Events;
+	FEventPage* CurrentPage = nullptr;
+	double LastTime = 0.0;
+	uint64 ModCount = 0;
+};
+
+}
\ No newline at end of file
diff --git a/Engine/Source/Developer/TraceServices/Public/TraceServices/Containers/Timelines.h b/Engine/Source/Developer/TraceServices/Public/TraceServices/Containers/Timelines.h
--- a/Engine/Source/Developer/TraceServices/Public/TraceServices/Containers/Timelines.h
+++ b/Engine/Source/Developer/TraceServices/Public/TraceServices/Containers/Timelines.h
@@ -12,14 +12,19 @@ class ITimeline
 {
 public:
 	typedef InEventType EventType;
+	typedef TFunctionRef<void(bool /*bStart*/, double /*Time*/, const EventType& /*Event*/)> EventCallback;
+	typedef TFunctionRef<void(double /*StartTime*/, double /*EndTime*/, uint32 /*Depth*/, const EventType&/*Event*/)> EventRangeCallback;
 
 	virtual ~ITimeline() = default;
 	virtual uint64 GetModCount() const = 0;
 	virtual uint64 GetEventCount() const = 0;
-	virtual void EnumerateEventsDownSampled(double IntervalStart, double IntervalEnd, double Resolution, TFunctionRef<void(bool, double, const EventType&)> Callback) const = 0;
-	virtual void EnumerateEventsDownSampled(double IntervalStart, double IntervalEnd, double Resolution, TFunctionRef<void(double, double, uint32, const EventType&)> Callback) const = 0;
-	virtual void EnumerateEvents(double IntervalStart, double IntervalEnd, TFunctionRef<void(bool, double, const EventType&)> Callback) const = 0;
-	virtual void EnumerateEvents(double IntervalStart, double IntervalEnd, TFunctionRef<void(double, double, uint32, const EventType&)> Callback) const = 0;
+	virtual const InEventType& GetEvent(uint64 InIndex) const = 0;
+	virtual double GetStartTime() const = 0;
+	virtual double GetEndTime() const = 0;
+	virtual void EnumerateEventsDownSampled(double IntervalStart, double IntervalEnd, double Resolution, EventCallback Callback) const = 0;
+	virtual void EnumerateEventsDownSampled(double IntervalStart, double IntervalEnd, double Resolution, EventRangeCallback Callback) const = 0;
+	virtual void EnumerateEvents(double IntervalStart, double IntervalEnd, EventCallback Callback) const = 0;
+	virtual void EnumerateEvents(double IntervalStart, double IntervalEnd, EventRangeCallback Callback) const = 0;
 };
 
 }
\ No newline at end of file
diff --git a/Engine/Source/Developer/TraceServices/Public/TraceServices/ModuleService.h b/Engine/Source/Developer/TraceServices/Public/TraceServices/ModuleService.h
--- a/Engine/Source/Developer/TraceServices/Public/TraceServices/ModuleService.h
+++ b/Engine/Source/Developer/TraceServices/Public/TraceServices/ModuleService.h
@@ -12,7 +12,7 @@ namespace Trace
 
 class IAnalyzer;
 class IAnalysisSession;
-extern const FName ModuleFeatureName;
+extern TRACESERVICES_API const FName ModuleFeatureName;
 
 struct FModuleInfo
 {
