Reverted compact layout changes

diff --git a/Engine/Source/Developer/TraceInsights/Public/Insights/ViewModels/TimingEventsTrack.h b/Engine/Source/Developer/TraceInsights/Public/Insights/ViewModels/TimingEventsTrack.h
--- a/Engine/Source/Developer/TraceInsights/Public/Insights/ViewModels/TimingEventsTrack.h
+++ b/Engine/Source/Developer/TraceInsights/Public/Insights/ViewModels/TimingEventsTrack.h
@@ -28,9 +28,9 @@ struct TRACEINSIGHTS_API FTimingEventsTrackLayout
 	static constexpr float NormalLayoutTimelineDY = 14.0f;
 	static constexpr float NormalLayoutMinTimelineH = 0.0f;
 
-	static constexpr float CompactLayoutEventH = 8.0f;
+	static constexpr float CompactLayoutEventH = 2.0f;
 	static constexpr float CompactLayoutEventDY = 1.0f;
-	static constexpr float CompactLayoutTimelineDY = 2.0f;
+	static constexpr float CompactLayoutTimelineDY = 3.0f;
 	static constexpr float CompactLayoutMinTimelineH = 0.0f;
 
 	//////////////////////////////////////////////////
