Fix graph track not drawing

Temporary fix pending refactor of GetSession/SetSession etc.

diff --git a/Engine/Source/Developer/TraceInsights/Private/Insights/Widgets/STimingView.cpp b/Engine/Source/Developer/TraceInsights/Private/Insights/Widgets/STimingView.cpp
--- a/Engine/Source/Developer/TraceInsights/Private/Insights/Widgets/STimingView.cpp
+++ b/Engine/Source/Developer/TraceInsights/Private/Insights/Widgets/STimingView.cpp
@@ -733,6 +733,7 @@ void STimingView::Tick(const FGeometry& AllottedGeometry, const double InCurrent
 
 	if (GraphTrack->IsVisible())
 	{
+		GraphTrack->SetSession(Session.Get());
 		GraphTrack->Update(Viewport);
 	}
 
