Insights public API tweaks

Reverted changes to cache session on tracks. Shared data/extender can cache the session if required. Previous access via Insights Manager re-instated until those tracks can be refactored out using the extender plugin API.
Events for selection and timer marker changing/changed now indicate interactivity.

diff --git a/Engine/Source/Developer/TraceInsights/Private/Insights/ViewModels/FileActivityTimingTrack.cpp b/Engine/Source/Developer/TraceInsights/Private/Insights/ViewModels/FileActivityTimingTrack.cpp
--- a/Engine/Source/Developer/TraceInsights/Private/Insights/ViewModels/FileActivityTimingTrack.cpp
+++ b/Engine/Source/Developer/TraceInsights/Private/Insights/ViewModels/FileActivityTimingTrack.cpp
@@ -495,23 +495,27 @@ bool FFileActivityTimingTrack::SearchEvent(const double InStartTime,
 
 	FSearchTimingEventContext Ctx(InStartTime, InEndTime, InPredicate, InOutTimingEvent, bInStopAtFirstMatch, bInSearchForLargestEvent);
 
-	Trace::FAnalysisSessionReadScope SessionReadScope(GetSession());
+	TSharedPtr<const Trace::IAnalysisSession> Session = FInsightsManager::Get()->GetSession();
+	if (Session.IsValid())
+	{
+		Trace::FAnalysisSessionReadScope SessionReadScope(*Session.Get());
 
-	const FTimingEventsTrack* Track = Ctx.TimingEvent.Track;
+		const FTimingEventsTrack* Track = Ctx.TimingEvent.Track;
 
-	for (const FFileActivitySharedState::FIoTimingEvent& Event : State->GetAllEvents())
-	{
-		if (Event.EndTime <= Ctx.StartTime)
+		for (const FFileActivitySharedState::FIoTimingEvent& Event : State->GetAllEvents())
 		{
-			continue;
-		}
+			if (Event.EndTime <= Ctx.StartTime)
+			{
+				continue;
+			}
 
-		if (!Ctx.bContinueSearching || Event.StartTime >= Ctx.EndTime)
-		{
-			break;
-		}
+			if (!Ctx.bContinueSearching || Event.StartTime >= Ctx.EndTime)
+			{
+				break;
+			}
 
-		Ctx.CheckEvent(Event, bIgnoreEventDepth ? 0 : Event.Depth);
+			Ctx.CheckEvent(Event, bIgnoreEventDepth ? 0 : Event.Depth);
+		}
 	}
 
 	return Ctx.bFound;
diff --git a/Engine/Source/Developer/TraceInsights/Private/Insights/ViewModels/LoadingTimingTrack.cpp b/Engine/Source/Developer/TraceInsights/Private/Insights/ViewModels/LoadingTimingTrack.cpp
--- a/Engine/Source/Developer/TraceInsights/Private/Insights/ViewModels/LoadingTimingTrack.cpp
+++ b/Engine/Source/Developer/TraceInsights/Private/Insights/ViewModels/LoadingTimingTrack.cpp
@@ -230,11 +230,12 @@ void FLoadingMainThreadTimingTrack::Draw(ITimingViewDrawHelper& Helper) const
 
 	if (Helper.BeginTimeline(Track))
 	{
-		if (Trace::ReadLoadTimeProfilerProvider(GetSession()))
+		TSharedPtr<const Trace::IAnalysisSession> Session = FInsightsManager::Get()->GetSession();
+		if (Session.IsValid() && Trace::ReadLoadTimeProfilerProvider(*Session.Get()))
 		{
-			Trace::FAnalysisSessionReadScope SessionReadScope(GetSession());
+			Trace::FAnalysisSessionReadScope SessionReadScope(*Session.Get());
 
-			const Trace::ILoadTimeProfilerProvider& LoadTimeProfilerProvider = *Trace::ReadLoadTimeProfilerProvider(GetSession());
+			const Trace::ILoadTimeProfilerProvider& LoadTimeProfilerProvider = *Trace::ReadLoadTimeProfilerProvider(*Session.Get());
 
 			LoadTimeProfilerProvider.ReadMainThreadCpuTimeline([this, &Helper](const Trace::ILoadTimeProfilerProvider::CpuTimeline& Timeline)
 			{
@@ -265,29 +266,33 @@ void FLoadingMainThreadTimingTrack::Draw(ITimingViewDrawHelper& Helper) const
 ////////////////////////////////////////////////////////////////////////////////////////////////////
 
 bool FLoadingMainThreadTimingTrack::SearchTimingEvent(const double InStartTime,
-													  const double InEndTime,
-													  TFunctionRef<bool(double, double, uint32)> InPredicate,
-													  FTimingEvent& InOutTimingEvent,
-													  bool bInStopAtFirstMatch,
-													  bool bInSearchForLargestEvent) const
+	const double InEndTime,
+	TFunctionRef<bool(double, double, uint32)> InPredicate,
+	FTimingEvent& InOutTimingEvent,
+	bool bInStopAtFirstMatch,
+	bool bInSearchForLargestEvent) const
 {
 	FSearchLoadTimeProfilerCpuEvent Ctx(InStartTime, InEndTime, InPredicate, InOutTimingEvent, bInStopAtFirstMatch, bInSearchForLargestEvent);
 
-	Trace::FAnalysisSessionReadScope SessionReadScope(GetSession());
-
-	const FTimingEventsTrack* Track = Ctx.TimingEvent.Track;
-
-	if (Trace::ReadLoadTimeProfilerProvider(GetSession()))
+	TSharedPtr<const Trace::IAnalysisSession> Session = FInsightsManager::Get()->GetSession();
+	if (Session.IsValid())
 	{
-		const Trace::ILoadTimeProfilerProvider& LoadTimeProfilerProvider = *Trace::ReadLoadTimeProfilerProvider(GetSession());
+		Trace::FAnalysisSessionReadScope SessionReadScope(*Session.Get());
+
+		const FTimingEventsTrack* Track = Ctx.TimingEvent.Track;
 
-		LoadTimeProfilerProvider.ReadMainThreadCpuTimeline([&Ctx](const Trace::ILoadTimeProfilerProvider::CpuTimeline& Timeline)
+		if (Trace::ReadLoadTimeProfilerProvider(*Session.Get()))
 		{
-			Timeline.EnumerateEvents(Ctx.StartTime, Ctx.EndTime, [&Ctx](double EventStartTime, double EventEndTime, uint32 EventDepth, const Trace::FLoadTimeProfilerCpuEvent& Event)
+			const Trace::ILoadTimeProfilerProvider& LoadTimeProfilerProvider = *Trace::ReadLoadTimeProfilerProvider(*Session.Get());
+
+			LoadTimeProfilerProvider.ReadMainThreadCpuTimeline([&Ctx](const Trace::ILoadTimeProfilerProvider::CpuTimeline& Timeline)
 			{
-				Ctx.CheckEvent(EventStartTime, EventEndTime, EventDepth, Event);
+				Timeline.EnumerateEvents(Ctx.StartTime, Ctx.EndTime, [&Ctx](double EventStartTime, double EventEndTime, uint32 EventDepth, const Trace::FLoadTimeProfilerCpuEvent& Event)
+				{
+					Ctx.CheckEvent(EventStartTime, EventEndTime, EventDepth, Event);
+				});
 			});
-		});
+		}
 	}
 
 	return Ctx.bFound;
@@ -303,11 +308,12 @@ void FLoadingAsyncThreadTimingTrack::Draw(ITimingViewDrawHelper& Helper) const
 
 	if (Helper.BeginTimeline(Track))
 	{
-		if (Trace::ReadLoadTimeProfilerProvider(GetSession()))
+		TSharedPtr<const Trace::IAnalysisSession> Session = FInsightsManager::Get()->GetSession();
+		if (Session.IsValid() && Trace::ReadLoadTimeProfilerProvider(*Session.Get()))
 		{
-			Trace::FAnalysisSessionReadScope SessionReadScope(GetSession());
+			Trace::FAnalysisSessionReadScope SessionReadScope(*Session.Get());
 
-			const Trace::ILoadTimeProfilerProvider& LoadTimeProfilerProvider = *Trace::ReadLoadTimeProfilerProvider(GetSession());
+			const Trace::ILoadTimeProfilerProvider& LoadTimeProfilerProvider = *Trace::ReadLoadTimeProfilerProvider(*Session.Get());
 
 			LoadTimeProfilerProvider.ReadAsyncLoadingThreadCpuTimeline([this, &Helper](const Trace::ILoadTimeProfilerProvider::CpuTimeline& Timeline)
 			{
@@ -346,21 +352,25 @@ bool FLoadingAsyncThreadTimingTrack::SearchTimingEvent(const double InStartTime,
 {
 	FSearchLoadTimeProfilerCpuEvent Ctx(InStartTime, InEndTime, InPredicate, InOutTimingEvent, bInStopAtFirstMatch, bInSearchForLargestEvent);
 
-	Trace::FAnalysisSessionReadScope SessionReadScope(GetSession());
-
-	const FTimingEventsTrack* Track = Ctx.TimingEvent.Track;
-
-	if (Trace::ReadLoadTimeProfilerProvider(GetSession()))
+	TSharedPtr<const Trace::IAnalysisSession> Session = FInsightsManager::Get()->GetSession();
+	if (Session.IsValid())
 	{
-		const Trace::ILoadTimeProfilerProvider& LoadTimeProfilerProvider = *Trace::ReadLoadTimeProfilerProvider(GetSession());
+		Trace::FAnalysisSessionReadScope SessionReadScope(*Session.Get());
+
+		const FTimingEventsTrack* Track = Ctx.TimingEvent.Track;
 
-		LoadTimeProfilerProvider.ReadAsyncLoadingThreadCpuTimeline([&Ctx](const Trace::ILoadTimeProfilerProvider::CpuTimeline& Timeline)
+		if (Trace::ReadLoadTimeProfilerProvider(*Session.Get()))
 		{
-			Timeline.EnumerateEvents(Ctx.StartTime, Ctx.EndTime, [&Ctx](double EventStartTime, double EventEndTime, uint32 EventDepth, const Trace::FLoadTimeProfilerCpuEvent& Event)
+			const Trace::ILoadTimeProfilerProvider& LoadTimeProfilerProvider = *Trace::ReadLoadTimeProfilerProvider(*Session.Get());
+
+			LoadTimeProfilerProvider.ReadAsyncLoadingThreadCpuTimeline([&Ctx](const Trace::ILoadTimeProfilerProvider::CpuTimeline& Timeline)
 			{
-				Ctx.CheckEvent(EventStartTime, EventEndTime, EventDepth, Event);
+				Timeline.EnumerateEvents(Ctx.StartTime, Ctx.EndTime, [&Ctx](double EventStartTime, double EventEndTime, uint32 EventDepth, const Trace::FLoadTimeProfilerCpuEvent& Event)
+				{
+					Ctx.CheckEvent(EventStartTime, EventEndTime, EventDepth, Event);
+				});
 			});
-		});
+		}
 	}
 
 	return Ctx.bFound;
diff --git a/Engine/Source/Developer/TraceInsights/Private/Insights/ViewModels/MarkersTimingTrack.cpp b/Engine/Source/Developer/TraceInsights/Private/Insights/ViewModels/MarkersTimingTrack.cpp
--- a/Engine/Source/Developer/TraceInsights/Private/Insights/ViewModels/MarkersTimingTrack.cpp
+++ b/Engine/Source/Developer/TraceInsights/Private/Insights/ViewModels/MarkersTimingTrack.cpp
@@ -136,11 +136,12 @@ void FMarkersTimingTrack::Update(const FTimingTrackViewport& InViewport)
 {
 	FTimeMarkerTrackBuilder Builder(*this, InViewport);
 
-	if (IsSessionValid())
+	TSharedPtr<const Trace::IAnalysisSession> Session = FInsightsManager::Get()->GetSession();
+	if (Session.IsValid())
 	{
-		Trace::FAnalysisSessionReadScope SessionReadScope(GetSession());
+		Trace::FAnalysisSessionReadScope SessionReadScope(*Session.Get());
 
-		const Trace::ILogProvider& LogProvider = Trace::ReadLogProvider(GetSession());
+		const Trace::ILogProvider& LogProvider = Trace::ReadLogProvider(*Session.Get());
 		Builder.BeginLog(LogProvider);
 
 		LogProvider.EnumerateMessages(
diff --git a/Engine/Source/Developer/TraceInsights/Private/Insights/ViewModels/ThreadTimingTrack.cpp b/Engine/Source/Developer/TraceInsights/Private/Insights/ViewModels/ThreadTimingTrack.cpp
--- a/Engine/Source/Developer/TraceInsights/Private/Insights/ViewModels/ThreadTimingTrack.cpp
+++ b/Engine/Source/Developer/TraceInsights/Private/Insights/ViewModels/ThreadTimingTrack.cpp
@@ -26,11 +26,12 @@ void FThreadTimingTrack::Draw(ITimingViewDrawHelper& Helper) const
 
 	if (Helper.BeginTimeline(Track))
 	{
-		if (Trace::ReadTimingProfilerProvider(GetSession()))
+		TSharedPtr<const Trace::IAnalysisSession> Session = FInsightsManager::Get()->GetSession();
+		if (Session.IsValid() && Trace::ReadTimingProfilerProvider(*Session.Get()))
 		{
-			Trace::FAnalysisSessionReadScope SessionReadScope(GetSession());
+			Trace::FAnalysisSessionReadScope SessionReadScope(*Session.Get());
 
-			const Trace::ITimingProfilerProvider& TimingProfilerProvider = *Trace::ReadTimingProfilerProvider(GetSession());
+			const Trace::ITimingProfilerProvider& TimingProfilerProvider = *Trace::ReadTimingProfilerProvider(*Session.Get());
 
 			TimingProfilerProvider.ReadTimers([this, &Helper, &Track, &TimingProfilerProvider](const Trace::FTimingProfilerTimer* Timers, uint64 TimersCount)
 			{
@@ -148,13 +149,14 @@ void FThreadTimingTrack::GetParentAndRoot(const FTimingEvent& TimingEvent, FTimi
 
 	if (TimingEvent.Depth > 0)
 	{
-		if (IsSessionValid())
+		TSharedPtr<const Trace::IAnalysisSession> Session = FInsightsManager::Get()->GetSession();
+		if (Session.IsValid())
 		{
-			Trace::FAnalysisSessionReadScope SessionReadScope(GetSession());
+			Trace::FAnalysisSessionReadScope SessionReadScope(*Session.Get());
 
-			if (Trace::ReadTimingProfilerProvider(GetSession()))
+			if (Trace::ReadTimingProfilerProvider(*Session.Get()))
 			{
-				const Trace::ITimingProfilerProvider& TimingProfilerProvider = *Trace::ReadTimingProfilerProvider(GetSession());
+				const Trace::ITimingProfilerProvider& TimingProfilerProvider = *Trace::ReadTimingProfilerProvider(*Session.Get());
 
 				TimingProfilerProvider.ReadTimeline(TimingEvent.Track->GetId(), [&](const Trace::ITimingProfilerProvider::Timeline& Timeline)
 				{
@@ -237,19 +239,23 @@ bool FThreadTimingTrack::SearchTimingEvent(const double InStartTime,
 
 	FSearchTimingEventContext Ctx(InStartTime, InEndTime, InPredicate, InOutTimingEvent, bInStopAtFirstMatch, bInSearchForLargestEvent);
 
-	Trace::FAnalysisSessionReadScope SessionReadScope(GetSession());
-
-	if (Trace::ReadTimingProfilerProvider(GetSession()))
+	TSharedPtr<const Trace::IAnalysisSession> Session = FInsightsManager::Get()->GetSession();
+	if (Session.IsValid())
 	{
-		const Trace::ITimingProfilerProvider& TimingProfilerProvider = *Trace::ReadTimingProfilerProvider(GetSession());
+		Trace::FAnalysisSessionReadScope SessionReadScope(*Session.Get());
 
-		TimingProfilerProvider.ReadTimeline(Ctx.TimingEvent.Track->GetId(), [&Ctx](const Trace::ITimingProfilerProvider::Timeline& Timeline)
+		if (Trace::ReadTimingProfilerProvider(*Session.Get()))
 		{
-			Timeline.EnumerateEvents(Ctx.StartTime, Ctx.EndTime, [&Ctx](double EventStartTime, double EventEndTime, uint32 EventDepth, const Trace::FTimingProfilerEvent& Event)
+			const Trace::ITimingProfilerProvider& TimingProfilerProvider = *Trace::ReadTimingProfilerProvider(*Session.Get());
+
+			TimingProfilerProvider.ReadTimeline(Ctx.TimingEvent.Track->GetId(), [&Ctx](const Trace::ITimingProfilerProvider::Timeline& Timeline)
 			{
-				Ctx.CheckEvent(EventStartTime, EventEndTime, EventDepth, Event);
+				Timeline.EnumerateEvents(Ctx.StartTime, Ctx.EndTime, [&Ctx](double EventStartTime, double EventEndTime, uint32 EventDepth, const Trace::FTimingProfilerEvent& Event)
+				{
+					Ctx.CheckEvent(EventStartTime, EventEndTime, EventDepth, Event);
+				});
 			});
-		});
+		}
 	}
 
 	return Ctx.bFound;
@@ -261,65 +267,69 @@ void FThreadTimingTrack::ComputeTimingEventStats(FTimingEvent& InOutTimingEvent)
 {
 	if (InOutTimingEvent.IsValid())
 	{
-		Trace::FAnalysisSessionReadScope SessionReadScope(GetSession());
-
-		if (Trace::ReadTimingProfilerProvider(GetSession()))
+		TSharedPtr<const Trace::IAnalysisSession> Session = FInsightsManager::Get()->GetSession();
+		if (Session.IsValid())
 		{
-			const Trace::ITimingProfilerProvider& TimingProfilerProvider = *Trace::ReadTimingProfilerProvider(GetSession());
+			Trace::FAnalysisSessionReadScope SessionReadScope(*Session.Get());
 
-			// Compute Exclusive Time.
-			TimingProfilerProvider.ReadTimeline(InOutTimingEvent.Track->GetId(), [&InOutTimingEvent](const Trace::ITimingProfilerProvider::Timeline& Timeline)
+			if (Trace::ReadTimingProfilerProvider(*Session.Get()))
 			{
-				struct FEnumerationState
-				{
-					double EventStartTime;
-					double EventEndTime;
-					uint64 EventDepth;
-
-					uint64 CurrentDepth;
-					double LastTime;
-					double ExclusiveTime;
-					bool IsInEventScope;
-				};
-				FEnumerationState State;
-
-				State.EventStartTime = InOutTimingEvent.StartTime;
-				State.EventEndTime = InOutTimingEvent.EndTime;
-				State.EventDepth = InOutTimingEvent.Depth;
-
-				State.CurrentDepth = 0;
-				State.LastTime = 0.0;
-				State.ExclusiveTime = 0.0;
-				State.IsInEventScope = false;
-
-				Timeline.EnumerateEvents(InOutTimingEvent.StartTime, InOutTimingEvent.EndTime, [&State](bool IsEnter, double Time, const Trace::FTimingProfilerEvent& Event)
+				const Trace::ITimingProfilerProvider& TimingProfilerProvider = *Trace::ReadTimingProfilerProvider(*Session.Get());
+
+				// Compute Exclusive Time.
+				TimingProfilerProvider.ReadTimeline(InOutTimingEvent.Track->GetId(), [&InOutTimingEvent](const Trace::ITimingProfilerProvider::Timeline& Timeline)
 				{
-					if (IsEnter)
+					struct FEnumerationState
 					{
-						if (State.IsInEventScope && State.CurrentDepth == State.EventDepth + 1)
-						{
-							State.ExclusiveTime += Time - State.LastTime;
-						}
-						if (State.CurrentDepth == State.EventDepth && Time == State.EventStartTime)
+						double EventStartTime;
+						double EventEndTime;
+						uint64 EventDepth;
+
+						uint64 CurrentDepth;
+						double LastTime;
+						double ExclusiveTime;
+						bool IsInEventScope;
+					};
+					FEnumerationState State;
+
+					State.EventStartTime = InOutTimingEvent.StartTime;
+					State.EventEndTime = InOutTimingEvent.EndTime;
+					State.EventDepth = InOutTimingEvent.Depth;
+
+					State.CurrentDepth = 0;
+					State.LastTime = 0.0;
+					State.ExclusiveTime = 0.0;
+					State.IsInEventScope = false;
+
+					Timeline.EnumerateEvents(InOutTimingEvent.StartTime, InOutTimingEvent.EndTime, [&State](bool IsEnter, double Time, const Trace::FTimingProfilerEvent& Event)
+					{
+						if (IsEnter)
 						{
-							State.IsInEventScope = true;
+							if (State.IsInEventScope && State.CurrentDepth == State.EventDepth + 1)
+							{
+								State.ExclusiveTime += Time - State.LastTime;
+							}
+							if (State.CurrentDepth == State.EventDepth && Time == State.EventStartTime)
+							{
+								State.IsInEventScope = true;
+							}
+							++State.CurrentDepth;
 						}
-						++State.CurrentDepth;
-					}
-					else
-					{
-						--State.CurrentDepth;
-						if (State.CurrentDepth == State.EventDepth && Time == State.EventEndTime)
+						else
 						{
-							State.IsInEventScope = false;
-							State.ExclusiveTime += Time - State.LastTime;
+							--State.CurrentDepth;
+							if (State.CurrentDepth == State.EventDepth && Time == State.EventEndTime)
+							{
+								State.IsInEventScope = false;
+								State.ExclusiveTime += Time - State.LastTime;
+							}
 						}
-					}
-					State.LastTime = Time;
-				});
+						State.LastTime = Time;
+					});
 
-				InOutTimingEvent.ExclusiveTime = State.ExclusiveTime;
-			});
+					InOutTimingEvent.ExclusiveTime = State.ExclusiveTime;
+				});
+			}
 		}
 	}
 }
diff --git a/Engine/Source/Developer/TraceInsights/Private/Insights/ViewModels/TimingGraphTrack.cpp b/Engine/Source/Developer/TraceInsights/Private/Insights/ViewModels/TimingGraphTrack.cpp
--- a/Engine/Source/Developer/TraceInsights/Private/Insights/ViewModels/TimingGraphTrack.cpp
+++ b/Engine/Source/Developer/TraceInsights/Private/Insights/ViewModels/TimingGraphTrack.cpp
@@ -115,10 +115,11 @@ TSharedPtr<FTimingGraphSeries> FTimingGraphTrack::AddStatsCounterSeries(uint32 C
 	bool bIsFloatingPoint = false;
 	bool bIsMemory = false;
 
-	if (IsSessionValid())
+	TSharedPtr<const Trace::IAnalysisSession> Session = FInsightsManager::Get()->GetSession();
+	if (Session.IsValid())
 	{
-		Trace::FAnalysisSessionReadScope SessionReadScope(GetSession());
-		const Trace::ICounterProvider& CountersProvider = Trace::ReadCounterProvider(GetSession());
+		Trace::FAnalysisSessionReadScope SessionReadScope(*Session.Get());
+		const Trace::ICounterProvider& CountersProvider = Trace::ReadCounterProvider(*Session.Get());
 		if (CounterId < CountersProvider.GetCounterCount())
 		{
 			CountersProvider.ReadCounter(CounterId, [&](const Trace::ICounter& Counter)
@@ -231,10 +232,11 @@ void FTimingGraphTrack::Update(const FTimingTrackViewport& Viewport)
 void FTimingGraphTrack::UpdateFrameSeries(FTimingGraphSeries& Series, const FTimingTrackViewport& Viewport)
 {
 	FGraphTrackBuilder Builder(*this, Series, Viewport);
-	if (IsSessionValid())
+	TSharedPtr<const Trace::IAnalysisSession> Session = FInsightsManager::Get()->GetSession();
+	if (Session.IsValid())
 	{
-		Trace::FAnalysisSessionReadScope SessionReadScope(GetSession());
-		const Trace::IFrameProvider& FramesProvider = ReadFrameProvider(GetSession());
+		Trace::FAnalysisSessionReadScope SessionReadScope(*Session.Get());
+		const Trace::IFrameProvider& FramesProvider = ReadFrameProvider(*Session.Get());
 		ETraceFrameType FrameType = static_cast<ETraceFrameType>(Series.Id);
 		uint64 FrameCount = FramesProvider.GetFrameCount(FrameType);
 		FramesProvider.EnumerateFrames(FrameType, 0, FrameCount - 1, [&Builder](const Trace::FFrame& Frame)
@@ -258,10 +260,11 @@ void FTimingGraphTrack::UpdateTimerSeries(FTimingGraphSeries& Series, const FTim
 void FTimingGraphTrack::UpdateStatsCounterSeries(FTimingGraphSeries& Series, const FTimingTrackViewport& Viewport)
 {
 	FGraphTrackBuilder Builder(*this, Series, Viewport);
-	if (IsSessionValid())
+	TSharedPtr<const Trace::IAnalysisSession> Session = FInsightsManager::Get()->GetSession();
+	if (Session.IsValid())
 	{
-		Trace::FAnalysisSessionReadScope SessionReadScope(GetSession());
-		const Trace::ICounterProvider& CountersProvider = Trace::ReadCounterProvider(GetSession());
+		Trace::FAnalysisSessionReadScope SessionReadScope(*Session.Get());
+		const Trace::ICounterProvider& CountersProvider = Trace::ReadCounterProvider(*Session.Get());
 		CountersProvider.ReadCounter(Series.Id, [this, &Viewport, &Builder, &Series](const Trace::ICounter& Counter)
 		{
 			if (Series.IsAutoZoomEnabled())
diff --git a/Engine/Source/Developer/TraceInsights/Private/Insights/Widgets/STimingView.cpp b/Engine/Source/Developer/TraceInsights/Private/Insights/Widgets/STimingView.cpp
--- a/Engine/Source/Developer/TraceInsights/Private/Insights/Widgets/STimingView.cpp
+++ b/Engine/Source/Developer/TraceInsights/Private/Insights/Widgets/STimingView.cpp
@@ -723,7 +723,6 @@ void STimingView::Tick(const FGeometry& AllottedGeometry, const double InCurrent
 
 		for (int32 TrackIndex = 0; TrackIndex < TimingEventsTracks.Num(); ++TrackIndex)
 		{
-			TimingEventsTracks[TrackIndex]->SetSession(Session.Get());
 			TimingEventsTracks[TrackIndex]->Update(Viewport);
 		}
 
@@ -733,7 +732,6 @@ void STimingView::Tick(const FGeometry& AllottedGeometry, const double InCurrent
 
 	if (GraphTrack->IsVisible())
 	{
-		GraphTrack->SetSession(Session.Get());
 		GraphTrack->Update(Viewport);
 	}
 
@@ -1237,6 +1235,7 @@ FReply STimingView::OnMouseButtonUp(const FGeometry& MyGeometry, const FPointerE
 			}
 			else if(bIsScrubbing)
 			{
+				RaiseTimeMarkerChanged();
 				bIsScrubbing = false;
 			}
 
@@ -1281,6 +1280,7 @@ FReply STimingView::OnMouseButtonUp(const FGeometry& MyGeometry, const FPointerE
 			}
 			else if(bIsScrubbing)
 			{
+				RaiseTimeMarkerChanged();
 				bIsScrubbing = false;
 			}
 
@@ -1379,7 +1379,7 @@ FReply STimingView::OnMouseMove(const FGeometry& MyGeometry, const FPointerEvent
 				bIsDragging = true;
 
 				TimeMarker = Viewport.SlateUnitsToTime(MousePosition.X);
-				RaiseTimeMarkerChanged();
+				RaiseTimeMarkerChanging();
 			}
 		}
 		else
@@ -2037,14 +2037,14 @@ void STimingView::SelectTimeInterval(double IntervalStartTime, double IntervalDu
 
 void STimingView::RaiseSelectionChanging()
 {
-	OnSelectionChangedDelegate.Broadcast(LastSelectionType != ESelectionType::None, SelectionStartTime, SelectionEndTime);
+	OnSelectionChangedDelegate.Broadcast(Insights::ETimeChangedFlags::Interactive, SelectionStartTime, SelectionEndTime);
 }
 
 ////////////////////////////////////////////////////////////////////////////////////////////////////
 
 void STimingView::RaiseSelectionChanged()
 {
-	OnSelectionChangedDelegate.Broadcast(LastSelectionType != ESelectionType::None, SelectionStartTime, SelectionEndTime);
+	OnSelectionChangedDelegate.Broadcast(Insights::ETimeChangedFlags::None, SelectionStartTime, SelectionEndTime);
 
 	FTimingProfilerManager::Get()->SetSelectedTimeRange(SelectionStartTime, SelectionEndTime);
 
@@ -2056,9 +2056,16 @@ void STimingView::RaiseSelectionChanged()
 
 ////////////////////////////////////////////////////////////////////////////////////////////////////
 
+void STimingView::RaiseTimeMarkerChanging()
+{
+	OnTimeMarkerChangedDelegate.Broadcast(Insights::ETimeChangedFlags::Interactive, TimeMarker);
+}
+
+////////////////////////////////////////////////////////////////////////////////////////////////////
+
 void STimingView::RaiseTimeMarkerChanged()
 {
-	OnTimeMarkerChangedDelegate.Broadcast(TimeMarker != std::numeric_limits<double>::infinity(), TimeMarker);
+	OnTimeMarkerChangedDelegate.Broadcast(Insights::ETimeChangedFlags::None, TimeMarker);
 }
 
 ////////////////////////////////////////////////////////////////////////////////////////////////////
diff --git a/Engine/Source/Developer/TraceInsights/Private/Insights/Widgets/STimingView.h b/Engine/Source/Developer/TraceInsights/Private/Insights/Widgets/STimingView.h
--- a/Engine/Source/Developer/TraceInsights/Private/Insights/Widgets/STimingView.h
+++ b/Engine/Source/Developer/TraceInsights/Private/Insights/Widgets/STimingView.h
@@ -303,6 +303,7 @@ protected:
 	void RaiseSelectionChanging();
 	void RaiseSelectionChanged();
 
+	void RaiseTimeMarkerChanging();
 	void RaiseTimeMarkerChanged();
 
 	void UpdateAggregatedStats();
diff --git a/Engine/Source/Developer/TraceInsights/Public/Insights/ITimingViewSession.h b/Engine/Source/Developer/TraceInsights/Public/Insights/ITimingViewSession.h
--- a/Engine/Source/Developer/TraceInsights/Public/Insights/ITimingViewSession.h
+++ b/Engine/Source/Developer/TraceInsights/Public/Insights/ITimingViewSession.h
@@ -10,11 +10,21 @@ namespace Insights { class ITimingViewVisualizer; }
 namespace Insights
 {
 
+enum class ETimeChangedFlags : int32
+{
+	None,
+
+	// The event fired in response to an interactive change from the user. Will be followed by a non-interactive change finished.
+	Interactive = (1 << 0)
+};
+
+ENUM_CLASS_FLAGS(ETimeChangedFlags);
+
 /** The delegate to be invoked when the time marker has changed */
-DECLARE_MULTICAST_DELEGATE_TwoParams(FTimeMarkerChangedDelegate, bool /*bValid*/, double /*TimeMarker*/);
+DECLARE_MULTICAST_DELEGATE_TwoParams(FTimeMarkerChangedDelegate, ETimeChangedFlags /*InFlags*/, double /*TimeMarker*/);
 
 /** The delegate to be invoked when the selection have been changed. */
-DECLARE_MULTICAST_DELEGATE_ThreeParams(FSelectionChangedDelegate, bool /*bValid*/, double /*StartTime*/, double /*EndTime*/);
+DECLARE_MULTICAST_DELEGATE_ThreeParams(FSelectionChangedDelegate, ETimeChangedFlags /*InFlags*/, double /*StartTime*/, double /*EndTime*/);
 
 /** The delegate to be invoked when the timing event being hovered by the mouse has changed. */
 DECLARE_MULTICAST_DELEGATE_OneParam(FHoveredEventChangedDelegate, const FTimingEvent& /*InEvent*/);
diff --git a/Engine/Source/Developer/TraceInsights/Public/Insights/ViewModels/BaseTimingTrack.h b/Engine/Source/Developer/TraceInsights/Public/Insights/ViewModels/BaseTimingTrack.h
--- a/Engine/Source/Developer/TraceInsights/Public/Insights/ViewModels/BaseTimingTrack.h
+++ b/Engine/Source/Developer/TraceInsights/Public/Insights/ViewModels/BaseTimingTrack.h
@@ -40,7 +40,6 @@ protected:
 		, Order(0)
 		, PosY(0.0f)
 		, Height(0.0f)
-		, Session(nullptr)
 		, Flags(ETimingTrackFlags::IsVisible | ETimingTrackFlags::IsDirty)
 	{}
 
@@ -53,7 +52,6 @@ public:
 		Id = 0;
 		PosY = 0.0f;
 		Height = 0.0f;
-		Session = nullptr;
 		Flags = ETimingTrackFlags::IsVisible | ETimingTrackFlags::IsDirty;
 	}
 
@@ -105,15 +103,6 @@ public:
 
 	virtual void BuildContextMenu(FMenuBuilder& MenuBuilder);
 
-	// Check whether the session is currently valid
-	bool IsSessionValid() const { return Session != nullptr; }
-
-	// Access the session that this track is using
-	const Trace::IAnalysisSession& GetSession() const { check(Session); return *Session; }
-
-	// Set the session that this track is using. This is refreshed each tick we have a valid session
-	void SetSession(const Trace::IAnalysisSession* InSession) { Session = InSession; }
-
 private:
 	uint64 Id;
 	FName Type; // Thread, Loading, FileActivity
@@ -122,7 +111,6 @@ private:
 	int32 Order;
 	float PosY; // y position, in Slate units
 	float Height; // height, in Slate units
-	const Trace::IAnalysisSession* Session;
 	ETimingTrackFlags Flags;
 
 	static uint64 IdGenerator;
