More non-unity fixes

diff --git a/Engine/Plugins/Animation/GameplayInsights/Source/GameplayInsights/Private/AnimationSharedData.cpp b/Engine/Plugins/Animation/GameplayInsights/Source/GameplayInsights/Private/AnimationSharedData.cpp
--- a/Engine/Plugins/Animation/GameplayInsights/Source/GameplayInsights/Private/AnimationSharedData.cpp
+++ b/Engine/Plugins/Animation/GameplayInsights/Source/GameplayInsights/Private/AnimationSharedData.cpp
@@ -8,6 +8,8 @@
 #include "AnimationTickRecordsTrack.h"
 #include "Insights/ITimingViewSession.h"
 #include "Insights/ViewModels/TimingEvent.h"
+#include "AnimationProvider.h"
+#include "GameplayProvider.h"
 
 #define LOCTEXT_NAMESPACE "AnimationSharedData"
 
diff --git a/Engine/Plugins/Animation/GameplayInsights/Source/GameplayInsights/Private/AnimationTickRecordsTrack.cpp b/Engine/Plugins/Animation/GameplayInsights/Source/GameplayInsights/Private/AnimationTickRecordsTrack.cpp
--- a/Engine/Plugins/Animation/GameplayInsights/Source/GameplayInsights/Private/AnimationTickRecordsTrack.cpp
+++ b/Engine/Plugins/Animation/GameplayInsights/Source/GameplayInsights/Private/AnimationTickRecordsTrack.cpp
@@ -7,6 +7,7 @@
 #include "Insights/ViewModels/TimingTrackViewport.h"
 #include "Insights/ViewModels/TimingEvent.h"
 #include "Insights/ViewModels/TooltipDrawState.h"
+#include "AnimationSharedData.h"
 
 #define LOCTEXT_NAMESPACE "AnimationTickRecordsTrack"
 
diff --git a/Engine/Plugins/Animation/GameplayInsights/Source/GameplayInsights/Private/AnimationTickRecordsTrack.h b/Engine/Plugins/Animation/GameplayInsights/Source/GameplayInsights/Private/AnimationTickRecordsTrack.h
--- a/Engine/Plugins/Animation/GameplayInsights/Source/GameplayInsights/Private/AnimationTickRecordsTrack.h
+++ b/Engine/Plugins/Animation/GameplayInsights/Source/GameplayInsights/Private/AnimationTickRecordsTrack.h
@@ -5,6 +5,8 @@
 #include "GameplayTrack.h"
 #include "Insights/ViewModels/TimingEventsTrack.h"
 
+class FAnimationSharedData;
+
 class FAnimationTickRecordsTrack : public TGameplayTrackMixin<FTimingEventsTrack>
 {
 public:
diff --git a/Engine/Plugins/Animation/GameplayInsights/Source/GameplayInsights/Private/GameplaySharedData.cpp b/Engine/Plugins/Animation/GameplayInsights/Source/GameplayInsights/Private/GameplaySharedData.cpp
--- a/Engine/Plugins/Animation/GameplayInsights/Source/GameplayInsights/Private/GameplaySharedData.cpp
+++ b/Engine/Plugins/Animation/GameplayInsights/Source/GameplayInsights/Private/GameplaySharedData.cpp
@@ -4,6 +4,8 @@
 #include "TraceServices/Model/AnalysisSession.h"
 #include "ObjectEventsTrack.h"
 #include "Algo/Sort.h"
+#include "GameplayProvider.h"
+#include "Insights/ITimingViewSession.h"
 
 #define LOCTEXT_NAMESPACE "GameplaySharedData"
 
diff --git a/Engine/Plugins/Animation/GameplayInsights/Source/GameplayInsights/Private/GameplaySharedData.h b/Engine/Plugins/Animation/GameplayInsights/Source/GameplayInsights/Private/GameplaySharedData.h
--- a/Engine/Plugins/Animation/GameplayInsights/Source/GameplayInsights/Private/GameplaySharedData.h
+++ b/Engine/Plugins/Animation/GameplayInsights/Source/GameplayInsights/Private/GameplaySharedData.h
@@ -9,6 +9,7 @@ namespace Insights { class ITimingViewSession; }
 class FObjectEventsTrack;
 class FSkeletalMeshPoseTrack;
 class FAnimationTickRecordsTrack;
+struct FObjectInfo;
 
 class FGameplaySharedData
 {
diff --git a/Engine/Plugins/Animation/GameplayInsights/Source/GameplayInsights/Private/GameplayTraceModule.cpp b/Engine/Plugins/Animation/GameplayInsights/Source/GameplayInsights/Private/GameplayTraceModule.cpp
--- a/Engine/Plugins/Animation/GameplayInsights/Source/GameplayInsights/Private/GameplayTraceModule.cpp
+++ b/Engine/Plugins/Animation/GameplayInsights/Source/GameplayInsights/Private/GameplayTraceModule.cpp
@@ -4,6 +4,7 @@
 #include "AnimationAnalyzer.h"
 #include "AnimationProvider.h"
 #include "GameplayProvider.h"
+#include "GameplayAnalyzer.h"
 
 FName FGameplayTraceModule::ModuleName("GameplayTrace");
 
diff --git a/Engine/Plugins/Animation/GameplayInsights/Source/GameplayInsights/Private/IGameplayProvider.h b/Engine/Plugins/Animation/GameplayInsights/Source/GameplayInsights/Private/IGameplayProvider.h
--- a/Engine/Plugins/Animation/GameplayInsights/Source/GameplayInsights/Private/IGameplayProvider.h
+++ b/Engine/Plugins/Animation/GameplayInsights/Source/GameplayInsights/Private/IGameplayProvider.h
@@ -3,6 +3,7 @@
 #pragma once
 
 #include "TraceServices/Model/AnalysisSession.h"
+#include "TraceServices/Containers/Timelines.h"
 
 struct FClassInfo
 {	
diff --git a/Engine/Plugins/Animation/GameplayInsights/Source/GameplayInsights/Private/ObjectEventsTrack.cpp b/Engine/Plugins/Animation/GameplayInsights/Source/GameplayInsights/Private/ObjectEventsTrack.cpp
--- a/Engine/Plugins/Animation/GameplayInsights/Source/GameplayInsights/Private/ObjectEventsTrack.cpp
+++ b/Engine/Plugins/Animation/GameplayInsights/Source/GameplayInsights/Private/ObjectEventsTrack.cpp
@@ -7,6 +7,7 @@
 #include "Insights/ViewModels/TimingEvent.h"
 #include "Algo/Sort.h"
 #include "Insights/ViewModels/TooltipDrawState.h"
+#include "GameplaySharedData.h"
 
 #define LOCTEXT_NAMESPACE "ObjectEventsTrack"
 
diff --git a/Engine/Plugins/Animation/GameplayInsights/Source/GameplayInsights/Private/ObjectEventsTrack.h b/Engine/Plugins/Animation/GameplayInsights/Source/GameplayInsights/Private/ObjectEventsTrack.h
--- a/Engine/Plugins/Animation/GameplayInsights/Source/GameplayInsights/Private/ObjectEventsTrack.h
+++ b/Engine/Plugins/Animation/GameplayInsights/Source/GameplayInsights/Private/ObjectEventsTrack.h
@@ -6,6 +6,8 @@
 #include "Insights/ViewModels/TimingEventsTrack.h"
 #include "Containers/ArrayView.h"
 
+class FGameplaySharedData;
+
 class FObjectEventsTrack : public TGameplayTrackMixin<FTimingEventsTrack>
 {
 public:
diff --git a/Engine/Plugins/Animation/GameplayInsights/Source/GameplayInsights/Private/SkeletalMeshPoseTrack.cpp b/Engine/Plugins/Animation/GameplayInsights/Source/GameplayInsights/Private/SkeletalMeshPoseTrack.cpp
--- a/Engine/Plugins/Animation/GameplayInsights/Source/GameplayInsights/Private/SkeletalMeshPoseTrack.cpp
+++ b/Engine/Plugins/Animation/GameplayInsights/Source/GameplayInsights/Private/SkeletalMeshPoseTrack.cpp
@@ -6,6 +6,7 @@
 #include "AnimationProvider.h"
 #include "Insights/ViewModels/TimingTrackViewport.h"
 #include "Insights/ViewModels/TimingEvent.h"
+#include "AnimationSharedData.h"
 
 #if WITH_ENGINE
 #include "Components/LineBatchComponent.h"
diff --git a/Engine/Plugins/Animation/GameplayInsights/Source/GameplayInsights/Private/SkeletalMeshPoseTrack.h b/Engine/Plugins/Animation/GameplayInsights/Source/GameplayInsights/Private/SkeletalMeshPoseTrack.h
--- a/Engine/Plugins/Animation/GameplayInsights/Source/GameplayInsights/Private/SkeletalMeshPoseTrack.h
+++ b/Engine/Plugins/Animation/GameplayInsights/Source/GameplayInsights/Private/SkeletalMeshPoseTrack.h
@@ -6,6 +6,7 @@
 #include "Insights/ViewModels/TimingEventsTrack.h"
 
 class ULineBatchComponent;
+class FAnimationSharedData;
 
 class FSkeletalMeshPoseTrack : public TGameplayTrackMixin<FTimingEventsTrack>
 {
diff --git a/Engine/Source/Developer/TraceInsights/Public/Insights/ITimingViewExtender.h b/Engine/Source/Developer/TraceInsights/Public/Insights/ITimingViewExtender.h
--- a/Engine/Source/Developer/TraceInsights/Public/Insights/ITimingViewExtender.h
+++ b/Engine/Source/Developer/TraceInsights/Public/Insights/ITimingViewExtender.h
@@ -5,6 +5,9 @@
 #include "CoreMinimal.h"
 #include "Features/IModularFeature.h"
 
+namespace Trace { class IAnalysisSession; }
+class FMenuBuilder;
+
 namespace Insights
 {
 
@@ -23,13 +26,13 @@ public:
 	virtual void OnEndSession(ITimingViewSession& InSession) = 0;
 
 	/** Called each frame. If any new tracks are created they can be added via ITimingViewSession::AddTimingEventsTrack() */
-	virtual void Tick(Insights::ITimingViewSession& InSession, const Trace::IAnalysisSession& InAnalysisSession) = 0;
+	virtual void Tick(ITimingViewSession& InSession, const Trace::IAnalysisSession& InAnalysisSession) = 0;
 
 	/** Extension hook for the 'quick filter' menu */
-	virtual void ExtendFilterMenu(Insights::ITimingViewSession& InSession, FMenuBuilder& InMenuBuilder) = 0;
+	virtual void ExtendFilterMenu(ITimingViewSession& InSession, FMenuBuilder& InMenuBuilder) = 0;
 
 	/** Called when tracks have changed and need to be re-ordered */
-	virtual void OnTracksChanged(Insights::ITimingViewSession& InSession, int32& InOutOrder) = 0;
+	virtual void OnTracksChanged(ITimingViewSession& InSession, int32& InOutOrder) = 0;
 };
 
 }
\ No newline at end of file
diff --git a/Engine/Source/Developer/TraceInsights/Public/Insights/ITimingViewSession.h b/Engine/Source/Developer/TraceInsights/Public/Insights/ITimingViewSession.h
--- a/Engine/Source/Developer/TraceInsights/Public/Insights/ITimingViewSession.h
+++ b/Engine/Source/Developer/TraceInsights/Public/Insights/ITimingViewSession.h
@@ -6,6 +6,7 @@
 
 class FTimingEventsTrack;
 namespace Insights { class ITimingViewVisualizer; }
+struct FTimingEvent;
 
 namespace Insights
 {
