Improved search for event tracks

Added an implementation of an improved search-based API for per-event data, using a new helper: TTimingEventSearch<>
This allows per-track type payloads to be pushed to implementations while keeping all track data encapsulated within the track itself.
Changed FTimingEventTrack::SearchTimingEvent to take FTimingEventSearchParameters instead of its multiple parameters.
Searching for the 'longest event' is now implemented in terms of the new search helper via custom Filter and Match callbacks.
Removed TypeId,  LoadingInfo, Offset, Size and Path members from FTimingEvent.
Reworked tracks to use the new helpers instead of storing their data in FTimingEvent.
Removed MessageId from gameplay and animation events as it isnt needed anymore.

Still to do:
- Add optimization to reduce per-frame search overhead for displaying selected event info

diff --git a/Engine/Plugins/Animation/GameplayInsights/Source/GameplayInsights/Private/AnimationProvider.cpp b/Engine/Plugins/Animation/GameplayInsights/Source/GameplayInsights/Private/AnimationProvider.cpp
--- a/Engine/Plugins/Animation/GameplayInsights/Source/GameplayInsights/Private/AnimationProvider.cpp
+++ b/Engine/Plugins/Animation/GameplayInsights/Source/GameplayInsights/Private/AnimationProvider.cpp
@@ -173,7 +173,6 @@ void FAnimationProvider::AppendTickRecord(uint64 InObjectId, uint64 InSubObjectI
 	}
 
 	FTickRecordMessage Message;
-	Message.MessageId = Timeline->GetEventCount();
 	Message.ComponentId = InObjectId;
 	Message.AnimInstanceId = InSubObjectId;
 	Message.AssetId = InAssetId;
@@ -244,7 +243,6 @@ void FAnimationProvider::AppendSkeletalMeshPose(uint64 InObjectId, uint64 InMesh
 	}
 
 	FSkeletalMeshPoseMessage Message;
-	Message.MessageId = Timeline->GetEventCount();
 	Message.TransformStartIndex = SkeletalMeshPoseTransforms.Num();
 	Message.ComponentId = InObjectId;
 	Message.MeshId = InMeshId;
diff --git a/Engine/Plugins/Animation/GameplayInsights/Source/GameplayInsights/Private/AnimationTickRecordsTrack.cpp b/Engine/Plugins/Animation/GameplayInsights/Source/GameplayInsights/Private/AnimationTickRecordsTrack.cpp
--- a/Engine/Plugins/Animation/GameplayInsights/Source/GameplayInsights/Private/AnimationTickRecordsTrack.cpp
+++ b/Engine/Plugins/Animation/GameplayInsights/Source/GameplayInsights/Private/AnimationTickRecordsTrack.cpp
@@ -8,6 +8,7 @@
 #include "Insights/ViewModels/TimingEvent.h"
 #include "Insights/ViewModels/TooltipDrawState.h"
 #include "AnimationSharedData.h"
+#include "Insights/ViewModels/TimingEventSearch.h"
 
 #define LOCTEXT_NAMESPACE "AnimationTickRecordsTrack"
 
@@ -59,94 +60,71 @@ void FAnimationTickRecordsTrack::Draw(ITimingViewDrawHelper& Helper) const
 
 void FAnimationTickRecordsTrack::InitTooltip(FTooltipDrawState& Tooltip, const FTimingEvent& HoveredTimingEvent) const
 {
-	Tooltip.ResetContent();
+	FTimingEventSearchParameters SearchParameters(HoveredTimingEvent.StartTime, HoveredTimingEvent.EndTime, ETimingEventSearchFlags::StopAtFirstMatch);
 
-	Tooltip.AddTitle(GetName());
-
-	const FAnimationProvider* AnimationProvider = SharedData.GetAnalysisSession().ReadProvider<FAnimationProvider>(FAnimationProvider::ProviderName);
-	if(AnimationProvider)
+	FindTickRecordMessage(SearchParameters, [this, &Tooltip](double InFoundStartTime, double InFoundEndTime, uint32 InFoundDepth, const FTickRecordMessage& InMessage)
 	{
-		Trace::FAnalysisSessionReadScope SessionReadScope(SharedData.GetAnalysisSession());
+		Tooltip.ResetContent();
 
-		AnimationProvider->ReadTickRecordMessage(GetGameplayTrack().GetObjectId(), GetAssetId(), HoveredTimingEvent.TypeId, [this, &Tooltip, &HoveredTimingEvent](const FTickRecordMessage& InMessage)
-		{
-			Tooltip.AddNameValueTextLine(LOCTEXT("EventTime", "Time").ToString(), FText::AsNumber(HoveredTimingEvent.StartTime).ToString());
-			Tooltip.AddNameValueTextLine(LOCTEXT("BlendWeight", "Blend Weight").ToString(), FText::AsNumber(InMessage.BlendWeight).ToString());
-			Tooltip.AddNameValueTextLine(LOCTEXT("PlaybackTime", "Playback Time").ToString(), FText::AsNumber(InMessage.PlaybackTime).ToString());
-			Tooltip.AddNameValueTextLine(LOCTEXT("RootMotionWeight", "Root Motion Weight").ToString(), FText::AsNumber(InMessage.RootMotionWeight).ToString());
-			Tooltip.AddNameValueTextLine(LOCTEXT("PlayRate", "Play Rate").ToString(), FText::AsNumber(InMessage.PlayRate).ToString());
-			Tooltip.AddNameValueTextLine(LOCTEXT("Looping", "Looping").ToString(), InMessage.Looping ? LOCTEXT("True", "True").ToString() : LOCTEXT("False", "False").ToString());
-		});
-	}
+		Tooltip.AddTitle(GetName());
 
-	Tooltip.UpdateLayout();
+		Tooltip.AddNameValueTextLine(LOCTEXT("EventTime", "Time").ToString(), FText::AsNumber(InFoundStartTime).ToString());
+		Tooltip.AddNameValueTextLine(LOCTEXT("BlendWeight", "Blend Weight").ToString(), FText::AsNumber(InMessage.BlendWeight).ToString());
+		Tooltip.AddNameValueTextLine(LOCTEXT("PlaybackTime", "Playback Time").ToString(), FText::AsNumber(InMessage.PlaybackTime).ToString());
+		Tooltip.AddNameValueTextLine(LOCTEXT("RootMotionWeight", "Root Motion Weight").ToString(), FText::AsNumber(InMessage.RootMotionWeight).ToString());
+		Tooltip.AddNameValueTextLine(LOCTEXT("PlayRate", "Play Rate").ToString(), FText::AsNumber(InMessage.PlayRate).ToString());
+		Tooltip.AddNameValueTextLine(LOCTEXT("Looping", "Looping").ToString(), InMessage.Looping ? LOCTEXT("True", "True").ToString() : LOCTEXT("False", "False").ToString());
+
+		Tooltip.UpdateLayout();
+	});
 }
 
-bool FAnimationTickRecordsTrack::SearchTimingEvent(const double InStartTime, const double InEndTime, TFunctionRef<bool(double, double, uint32)> InPredicate, FTimingEvent& InOutTimingEvent, bool bInStopAtFirstMatch, bool bInSearchForLargestEvent) const
+bool FAnimationTickRecordsTrack::SearchTimingEvent(const FTimingEventSearchParameters& InSearchParameters, FTimingEvent& InOutTimingEvent) const
 {
-	struct FSearchTimingEventContext
+	return FindTickRecordMessage(InSearchParameters, [this, &InOutTimingEvent](double InFoundStartTime, double InFoundEndTime, uint32 InFoundDepth, const FTickRecordMessage& InFoundMessage)
 	{
-		const double StartTime;
-		const double EndTime;
-		TFunctionRef<bool(double, double, uint32)> Predicate;
-		FTimingEvent& TimingEvent;
-		const bool bStopAtFirstMatch;
-		const bool bSearchForLargestEvent;
-		mutable bool bFound;
-		mutable bool bContinueSearching;
-		mutable double LargestDuration;
-
-		FSearchTimingEventContext(const double InStartTime, const double InEndTime, TFunctionRef<bool(double, double, uint32)> InPredicate, FTimingEvent& InOutTimingEvent, bool bInStopAtFirstMatch, bool bInSearchForLargestEvent)
-			: StartTime(InStartTime)
-			, EndTime(InEndTime)
-			, Predicate(InPredicate)
-			, TimingEvent(InOutTimingEvent)
-			, bStopAtFirstMatch(bInStopAtFirstMatch)
-			, bSearchForLargestEvent(bInSearchForLargestEvent)
-			, bFound(false)
-			, bContinueSearching(true)
-			, LargestDuration(-1.0)
-		{
-		}
+		InOutTimingEvent = FTimingEvent(this, InFoundStartTime, InFoundEndTime, InFoundDepth);
+	});
+}
 
-		void CheckMessage(double EventStartTime, double EventEndTime, uint32 EventDepth, uint64 InMessageId)
-		{
-			if (bContinueSearching && Predicate(EventStartTime, EventEndTime, EventDepth))
-			{
-				if (!bSearchForLargestEvent || EventEndTime - EventStartTime > LargestDuration)
-				{
-					LargestDuration = EventEndTime - EventStartTime;
+bool FAnimationTickRecordsTrack::FindTickRecordMessage(const FTimingEventSearchParameters& InParameters, TFunctionRef<void(double, double, uint32, const FTickRecordMessage&)> InFoundPredicate) const
+{
+	// Storage for the message we want to match (payload for an event)
+	FTickRecordMessage MatchedMessage;
 
-					TimingEvent.TypeId = InMessageId;
-					TimingEvent.Depth = EventDepth;
-					TimingEvent.StartTime = EventStartTime;
-					TimingEvent.EndTime = EventEndTime;
+	return TTimingEventSearch<FTickRecordMessage>::Search(
+		InParameters,
 
-					bFound = true;
-					bContinueSearching = !bStopAtFirstMatch || bSearchForLargestEvent;
-				}
-			}
-		}
-	};
+		// Search...
+		[this](TTimingEventSearch<FTickRecordMessage>::FContext& InContext)
+		{
+			const FAnimationProvider* AnimationProvider = SharedData.GetAnalysisSession().ReadProvider<FAnimationProvider>(FAnimationProvider::ProviderName);
 
-	const FAnimationProvider* AnimationProvider = SharedData.GetAnalysisSession().ReadProvider<FAnimationProvider>(FAnimationProvider::ProviderName);
+			if(AnimationProvider)
+			{
+				Trace::FAnalysisSessionReadScope SessionReadScope(SharedData.GetAnalysisSession());
 
-	if(AnimationProvider)
-	{
-		FSearchTimingEventContext Context(InStartTime, InEndTime, InPredicate, InOutTimingEvent, bInStopAtFirstMatch, bInSearchForLargestEvent);
+				AnimationProvider->ReadTickRecordTimeline(GetGameplayTrack().GetObjectId(), GetAssetId(), [this, &InContext](const FAnimationProvider::TickRecordTimeline& InTimeline)
+				{
+					InTimeline.EnumerateEvents(InContext.GetParameters().StartTime, InContext.GetParameters().EndTime, [&InContext](double InEventStartTime, double InEventEndTime, uint32 InDepth, const FTickRecordMessage& InMessage)
+					{
+						InContext.Check(InEventStartTime, InEventEndTime, 0, InMessage);
+					});
+				});
+			}
+		},
 
-		Trace::FAnalysisSessionReadScope SessionReadScope(SharedData.GetAnalysisSession());
+		// Matched...
+		[&MatchedMessage](double InStartTime, double InEndTime, uint32 InDepth, const FTickRecordMessage& InEvent)
+		{
+			MatchedMessage = InEvent;
+		},
 
-		AnimationProvider->ReadTickRecordTimeline(GetGameplayTrack().GetObjectId(), GetAssetId(), [this, &Context, &InStartTime, &InEndTime](const FAnimationProvider::TickRecordTimeline& InTimeline)
+		// Found!
+		[&InFoundPredicate, &MatchedMessage](double InFoundStartTime, double InFoundEndTime, uint32 InFoundDepth)
 		{
-			InTimeline.EnumerateEvents(InStartTime, InEndTime, [&Context](double InEventStartTime, double InEventEndTime, uint32 InDepth, const FTickRecordMessage& InMessage)
-			{
-				Context.CheckMessage(InEventStartTime, InEventEndTime, 0, InMessage.MessageId);
-			});
+			InFoundPredicate(InFoundStartTime, InFoundEndTime, InFoundDepth, MatchedMessage);
 		});
-	}
-
-	return false;
 }
 
 #undef LOCTEXT_NAMESPACE
\ No newline at end of file
diff --git a/Engine/Plugins/Animation/GameplayInsights/Source/GameplayInsights/Private/AnimationTickRecordsTrack.h b/Engine/Plugins/Animation/GameplayInsights/Source/GameplayInsights/Private/AnimationTickRecordsTrack.h
--- a/Engine/Plugins/Animation/GameplayInsights/Source/GameplayInsights/Private/AnimationTickRecordsTrack.h
+++ b/Engine/Plugins/Animation/GameplayInsights/Source/GameplayInsights/Private/AnimationTickRecordsTrack.h
@@ -6,6 +6,8 @@
 #include "Insights/ViewModels/TimingEventsTrack.h"
 
 class FAnimationSharedData;
+struct FTickRecordMessage;
+class FTimingEventSearchParameters;
 
 class FAnimationTickRecordsTrack : public TGameplayTrackMixin<FTimingEventsTrack>
 {
@@ -17,11 +19,15 @@ public:
 
 	virtual void Draw(ITimingViewDrawHelper& Helper) const override;
 	virtual void InitTooltip(FTooltipDrawState& Tooltip, const FTimingEvent& HoveredTimingEvent) const override;
-	virtual bool SearchTimingEvent(const double InStartTime, const double InEndTime, TFunctionRef<bool(double, double, uint32)> InPredicate, FTimingEvent& InOutTimingEvent, bool bInStopAtFirstMatch, bool bInSearchForLargestEvent) const override;
+	virtual bool SearchTimingEvent(const FTimingEventSearchParameters& InSearchParameters, FTimingEvent& InOutTimingEvent) const override;
 
 	/** Get the asset ID that this track uses */
 	uint64 GetAssetId() const { return AssetId; }
 
+private:
+	// Helper function used to find a tick record
+	bool FindTickRecordMessage(const FTimingEventSearchParameters& InParameters, TFunctionRef<void(double, double, uint32, const FTickRecordMessage&)> InFoundPredicate) const;
+
 private:
 	/** The shared data */
 	const FAnimationSharedData& SharedData;
diff --git a/Engine/Plugins/Animation/GameplayInsights/Source/GameplayInsights/Private/GameplayProvider.cpp b/Engine/Plugins/Animation/GameplayInsights/Source/GameplayInsights/Private/GameplayProvider.cpp
--- a/Engine/Plugins/Animation/GameplayInsights/Source/GameplayInsights/Private/GameplayProvider.cpp
+++ b/Engine/Plugins/Animation/GameplayInsights/Source/GameplayInsights/Private/GameplayProvider.cpp
@@ -128,7 +128,6 @@ void FGameplayProvider::AppendObjectEvent(uint64 InObjectId, double InTime, cons
 	}
 
 	FObjectEventMessage Message;
-	Message.MessageId = Timeline->GetEventCount();
 	Message.Id = InObjectId;
 	Message.Name = Session.StoreString(InEventName);
 
diff --git a/Engine/Plugins/Animation/GameplayInsights/Source/GameplayInsights/Private/IAnimationProvider.h b/Engine/Plugins/Animation/GameplayInsights/Source/GameplayInsights/Private/IAnimationProvider.h
--- a/Engine/Plugins/Animation/GameplayInsights/Source/GameplayInsights/Private/IAnimationProvider.h
+++ b/Engine/Plugins/Animation/GameplayInsights/Source/GameplayInsights/Private/IAnimationProvider.h
@@ -16,7 +16,6 @@ struct FSkeletalMeshInfo
 
 struct FSkeletalMeshPoseMessage
 {
-	uint64 MessageId = 0;
 	uint64 TransformStartIndex = 0;
 	uint64 ComponentId = 0;	
 	uint64 MeshId = 0;	
@@ -27,7 +26,6 @@ struct FSkeletalMeshPoseMessage
 
 struct FTickRecordMessage
 {
-	uint64 MessageId = 0;
 	uint64 ComponentId = 0;
 	uint64 AnimInstanceId = 0;
 	uint64 AssetId = 0;
diff --git a/Engine/Plugins/Animation/GameplayInsights/Source/GameplayInsights/Private/IGameplayProvider.h b/Engine/Plugins/Animation/GameplayInsights/Source/GameplayInsights/Private/IGameplayProvider.h
--- a/Engine/Plugins/Animation/GameplayInsights/Source/GameplayInsights/Private/IGameplayProvider.h
+++ b/Engine/Plugins/Animation/GameplayInsights/Source/GameplayInsights/Private/IGameplayProvider.h
@@ -24,7 +24,6 @@ struct FObjectInfo
 
 struct FObjectEventMessage
 {
-	uint64 MessageId = 0;
 	uint64 Id = 0;
 	const TCHAR* Name = nullptr;
 };
diff --git a/Engine/Plugins/Animation/GameplayInsights/Source/GameplayInsights/Private/ObjectEventsTrack.cpp b/Engine/Plugins/Animation/GameplayInsights/Source/GameplayInsights/Private/ObjectEventsTrack.cpp
--- a/Engine/Plugins/Animation/GameplayInsights/Source/GameplayInsights/Private/ObjectEventsTrack.cpp
+++ b/Engine/Plugins/Animation/GameplayInsights/Source/GameplayInsights/Private/ObjectEventsTrack.cpp
@@ -8,6 +8,7 @@
 #include "Algo/Sort.h"
 #include "Insights/ViewModels/TooltipDrawState.h"
 #include "GameplaySharedData.h"
+#include "Insights/ViewModels/TimingEventSearch.h"
 
 #define LOCTEXT_NAMESPACE "ObjectEventsTrack"
 
@@ -48,92 +49,65 @@ void FObjectEventsTrack::Draw(ITimingViewDrawHelper& Helper) const
 
 void FObjectEventsTrack::InitTooltip(FTooltipDrawState& Tooltip, const FTimingEvent& HoveredTimingEvent) const
 {
-	Tooltip.ResetContent();
+	FTimingEventSearchParameters SearchParameters(HoveredTimingEvent.StartTime, HoveredTimingEvent.EndTime, ETimingEventSearchFlags::StopAtFirstMatch);
 
-	const FGameplayProvider* GameplayProvider = SharedData.GetAnalysisSession().ReadProvider<FGameplayProvider>(FGameplayProvider::ProviderName);
-	if(GameplayProvider)
+	FindObjectEvent(SearchParameters, [this, &Tooltip, &HoveredTimingEvent](double InFoundStartTime, double InFoundEndTime, uint32 InFoundDepth, const FObjectEventMessage& InMessage)
 	{
-		Trace::FAnalysisSessionReadScope SessionReadScope(SharedData.GetAnalysisSession());
+		Tooltip.ResetContent();
 
-		GameplayProvider->ReadObjectEventsTimeline(GetGameplayTrack().GetObjectId(), [&Tooltip, &HoveredTimingEvent](const FGameplayProvider::ObjectEventsTimeline& InTimeline)
-		{
-			if(HoveredTimingEvent.TypeId < InTimeline.GetEventCount())
-			{
-				const FObjectEventMessage& Message = InTimeline.GetEvent(HoveredTimingEvent.TypeId);
-				Tooltip.AddTitle(FText::FromString(FString(Message.Name)).ToString());
-				Tooltip.AddNameValueTextLine(LOCTEXT("EventTime", "Time").ToString(), FText::AsNumber(HoveredTimingEvent.StartTime).ToString());
-			}
-		});
-	}
+		Tooltip.AddTitle(FText::FromString(FString(InMessage.Name)).ToString());
+		Tooltip.AddNameValueTextLine(LOCTEXT("EventTime", "Time").ToString(), FText::AsNumber(HoveredTimingEvent.StartTime).ToString());
 
-	Tooltip.UpdateLayout();
+		Tooltip.UpdateLayout();
+	});
 }
 
-bool FObjectEventsTrack::SearchTimingEvent(const double InStartTime, const double InEndTime, TFunctionRef<bool(double, double, uint32)> InPredicate, FTimingEvent& InOutTimingEvent, bool bInStopAtFirstMatch, bool bInSearchForLargestEvent) const
+bool FObjectEventsTrack::SearchTimingEvent(const FTimingEventSearchParameters& InSearchParameters, FTimingEvent& InOutTimingEvent) const
 {
-	struct FSearchTimingEventContext
+	return FindObjectEvent(InSearchParameters, [this, &InOutTimingEvent](double InFoundStartTime, double InFoundEndTime, uint32 InFoundDepth, const FObjectEventMessage& InFoundMessage)
 	{
-		const double StartTime;
-		const double EndTime;
-		TFunctionRef<bool(double, double, uint32)> Predicate;
-		FTimingEvent& TimingEvent;
-		const bool bStopAtFirstMatch;
-		const bool bSearchForLargestEvent;
-		mutable bool bFound;
-		mutable bool bContinueSearching;
-		mutable double LargestDuration;
-
-		FSearchTimingEventContext(const double InStartTime, const double InEndTime, TFunctionRef<bool(double, double, uint32)> InPredicate, FTimingEvent& InOutTimingEvent, bool bInStopAtFirstMatch, bool bInSearchForLargestEvent)
-			: StartTime(InStartTime)
-			, EndTime(InEndTime)
-			, Predicate(InPredicate)
-			, TimingEvent(InOutTimingEvent)
-			, bStopAtFirstMatch(bInStopAtFirstMatch)
-			, bSearchForLargestEvent(bInSearchForLargestEvent)
-			, bFound(false)
-			, bContinueSearching(true)
-			, LargestDuration(-1.0)
-		{
-		}
+		InOutTimingEvent = FTimingEvent(this, InFoundStartTime, InFoundEndTime, InFoundDepth);
+	});
+}
 
-		void CheckMessage(double EventStartTime, double EventEndTime, uint32 EventDepth, uint64 InMessageId)
-		{
-			if (bContinueSearching && Predicate(EventStartTime, EventEndTime, EventDepth))
-			{
-				if (!bSearchForLargestEvent || EventEndTime - EventStartTime > LargestDuration)
-				{
-					LargestDuration = EventEndTime - EventStartTime;
+bool FObjectEventsTrack::FindObjectEvent(const FTimingEventSearchParameters& InParameters, TFunctionRef<void(double, double, uint32, const FObjectEventMessage&)> InFoundPredicate) const
+{
+	// Storage for the message we want to match (payload for an event)
+	FObjectEventMessage MatchedMessage;
 
-					TimingEvent.TypeId = InMessageId;
-					TimingEvent.Depth = EventDepth;
-					TimingEvent.StartTime = EventStartTime;
-					TimingEvent.EndTime = EventEndTime;
+	return TTimingEventSearch<FObjectEventMessage>::Search(
+		InParameters,
 
-					bFound = true;
-					bContinueSearching = !bStopAtFirstMatch || bSearchForLargestEvent;
-				}
-			}
-		}
-	};
+		// Search...
+		[this](TTimingEventSearch<FObjectEventMessage>::FContext& InContext)
+		{
+			const FGameplayProvider* GameplayProvider = SharedData.GetAnalysisSession().ReadProvider<FGameplayProvider>(FGameplayProvider::ProviderName);
 
-	const FGameplayProvider* GameplayProvider = SharedData.GetAnalysisSession().ReadProvider<FGameplayProvider>(FGameplayProvider::ProviderName);
+			if(GameplayProvider)
+			{
+				Trace::FAnalysisSessionReadScope SessionReadScope(SharedData.GetAnalysisSession());
 
-	if(GameplayProvider)
-	{
-		FSearchTimingEventContext Context(InStartTime, InEndTime, InPredicate, InOutTimingEvent, bInStopAtFirstMatch, bInSearchForLargestEvent);
+				GameplayProvider->ReadObjectEventsTimeline(GetGameplayTrack().GetObjectId(), [&InContext](const FGameplayProvider::ObjectEventsTimeline& InTimeline)
+				{
+					InTimeline.EnumerateEvents(InContext.GetParameters().StartTime, InContext.GetParameters().EndTime, [&InContext](double InEventStartTime, double InEventEndTime, uint32 InDepth, const FObjectEventMessage& InMessage)
+					{
+						InContext.Check(InEventStartTime, InEventEndTime, 0, InMessage);
+					});
+				});
+			}
+		},
 
-		Trace::FAnalysisSessionReadScope SessionReadScope(SharedData.GetAnalysisSession());
+		// Matched...
+		[&MatchedMessage](double InStartTime, double InEndTime, uint32 InDepth, const FObjectEventMessage& InEvent)
+		{
+			MatchedMessage = InEvent;
+		},
 
-		GameplayProvider->ReadObjectEventsTimeline(GetGameplayTrack().GetObjectId(), [&Context, &InStartTime, &InEndTime](const FGameplayProvider::ObjectEventsTimeline& InTimeline)
+		// Found!
+		[&InFoundPredicate, &MatchedMessage](double InFoundStartTime, double InFoundEndTime, uint32 InFoundDepth)
 		{
-			InTimeline.EnumerateEvents(InStartTime, InEndTime, [&Context](double InEventStartTime, double InEventEndTime, uint32 InDepth, const FObjectEventMessage& InMessage)
-			{
-				Context.CheckMessage(InEventStartTime, InEventEndTime, 0, InMessage.MessageId);
-			});
+			InFoundPredicate(InFoundStartTime, InFoundEndTime, InFoundDepth, MatchedMessage);
 		});
-	}
-
-	return false;
 }
 
 #undef LOCTEXT_NAMESPACE
\ No newline at end of file
diff --git a/Engine/Plugins/Animation/GameplayInsights/Source/GameplayInsights/Private/ObjectEventsTrack.h b/Engine/Plugins/Animation/GameplayInsights/Source/GameplayInsights/Private/ObjectEventsTrack.h
--- a/Engine/Plugins/Animation/GameplayInsights/Source/GameplayInsights/Private/ObjectEventsTrack.h
+++ b/Engine/Plugins/Animation/GameplayInsights/Source/GameplayInsights/Private/ObjectEventsTrack.h
@@ -4,9 +4,10 @@
 
 #include "GameplayTrack.h"
 #include "Insights/ViewModels/TimingEventsTrack.h"
-#include "Containers/ArrayView.h"
 
 class FGameplaySharedData;
+struct FObjectEventMessage;
+class FTimingEventSearchParameters;
 
 class FObjectEventsTrack : public TGameplayTrackMixin<FTimingEventsTrack>
 {
@@ -18,7 +19,11 @@ public:
 
 	virtual void Draw(ITimingViewDrawHelper& Helper) const override;
 	virtual void InitTooltip(FTooltipDrawState& Tooltip, const FTimingEvent& HoveredTimingEvent) const override;
-	virtual bool SearchTimingEvent(const double InStartTime, const double InEndTime, TFunctionRef<bool(double, double, uint32)> InPredicate, FTimingEvent& InOutTimingEvent, bool bInStopAtFirstMatch, bool bInSearchForLargestEvent) const override;
+	virtual bool SearchTimingEvent(const FTimingEventSearchParameters& InSearchParameters, FTimingEvent& InOutTimingEvent) const override;
+
+private:
+	// Helper function used to find an object event
+	bool FindObjectEvent(const FTimingEventSearchParameters& InParameters, TFunctionRef<void(double, double, uint32, const FObjectEventMessage&)> InFoundPredicate) const;
 
 private:
 	const FGameplaySharedData& SharedData;
diff --git a/Engine/Plugins/Animation/GameplayInsights/Source/GameplayInsights/Private/SkeletalMeshPoseTrack.cpp b/Engine/Plugins/Animation/GameplayInsights/Source/GameplayInsights/Private/SkeletalMeshPoseTrack.cpp
--- a/Engine/Plugins/Animation/GameplayInsights/Source/GameplayInsights/Private/SkeletalMeshPoseTrack.cpp
+++ b/Engine/Plugins/Animation/GameplayInsights/Source/GameplayInsights/Private/SkeletalMeshPoseTrack.cpp
@@ -8,6 +8,7 @@
 #include "Insights/ViewModels/TimingEvent.h"
 #include "AnimationSharedData.h"
 #include "Framework/Multibox/MultiboxBuilder.h"
+#include "Insights/ViewModels/TimingEventSearch.h"
 
 #if WITH_ENGINE
 #include "Components/LineBatchComponent.h"
@@ -74,90 +75,67 @@ void FSkeletalMeshPoseTrack::Draw( ITimingViewDrawHelper& Helper) const
 
 void FSkeletalMeshPoseTrack::InitTooltip(FTooltipDrawState& Tooltip, const FTimingEvent& HoveredTimingEvent) const
 {
-	Tooltip.ResetContent();
+	FTimingEventSearchParameters SearchParameters(HoveredTimingEvent.StartTime, HoveredTimingEvent.EndTime, ETimingEventSearchFlags::StopAtFirstMatch);
 
-	Tooltip.AddTitle(LOCTEXT("SkeletalMeshPoseTooltipTitle", "Skeletal Mesh Pose").ToString());
-
-	const FAnimationProvider* AnimationProvider = SharedData.GetAnalysisSession().ReadProvider<FAnimationProvider>(FAnimationProvider::ProviderName);
-	if(AnimationProvider)
+	FindSkeletalMeshPoseMessage(SearchParameters, [this, &Tooltip](double InFoundStartTime, double InFoundEndTime, uint32 InFoundDepth, const FSkeletalMeshPoseMessage& InMessage)
 	{
-		Trace::FAnalysisSessionReadScope SessionReadScope(SharedData.GetAnalysisSession());
+		Tooltip.ResetContent();
 
-		AnimationProvider->ReadSkeletalMeshPoseMessage(GetGameplayTrack().GetObjectId(), HoveredTimingEvent.TypeId, [&Tooltip, &HoveredTimingEvent](const FSkeletalMeshPoseMessage& InMessage)
-		{
-			Tooltip.AddNameValueTextLine(LOCTEXT("EventTime", "Time").ToString(), FText::AsNumber(HoveredTimingEvent.StartTime).ToString());
-			Tooltip.AddNameValueTextLine(LOCTEXT("BoneCount", "Bone Count").ToString(), FText::AsNumber(InMessage.NumTransforms).ToString());
-		});
-	}
+		Tooltip.AddTitle(LOCTEXT("SkeletalMeshPoseTooltipTitle", "Skeletal Mesh Pose").ToString());
 
-	Tooltip.UpdateLayout();
+		Tooltip.AddNameValueTextLine(LOCTEXT("EventTime", "Time").ToString(), FText::AsNumber(InFoundStartTime).ToString());
+		Tooltip.AddNameValueTextLine(LOCTEXT("BoneCount", "Bone Count").ToString(), FText::AsNumber(InMessage.NumTransforms).ToString());
+
+		Tooltip.UpdateLayout();
+	});
 }
 
-bool FSkeletalMeshPoseTrack::SearchTimingEvent(const double InStartTime, const double InEndTime, TFunctionRef<bool(double, double, uint32)> InPredicate, FTimingEvent& InOutTimingEvent, bool bInStopAtFirstMatch, bool bInSearchForLargestEvent) const
+bool FSkeletalMeshPoseTrack::SearchTimingEvent(const FTimingEventSearchParameters& InSearchParameters, FTimingEvent& InOutTimingEvent) const
 {
-	struct FSearchTimingEventContext
+	return FindSkeletalMeshPoseMessage(InSearchParameters, [this, &InOutTimingEvent](double InFoundStartTime, double InFoundEndTime, uint32 InFoundDepth, const FSkeletalMeshPoseMessage& InFoundMessage)
 	{
-		const double StartTime;
-		const double EndTime;
-		TFunctionRef<bool(double, double, uint32)> Predicate;
-		FTimingEvent& TimingEvent;
-		const bool bStopAtFirstMatch;
-		const bool bSearchForLargestEvent;
-		mutable bool bFound;
-		mutable bool bContinueSearching;
-		mutable double LargestDuration;
-
-		FSearchTimingEventContext(const double InStartTime, const double InEndTime, TFunctionRef<bool(double, double, uint32)> InPredicate, FTimingEvent& InOutTimingEvent, bool bInStopAtFirstMatch, bool bInSearchForLargestEvent)
-			: StartTime(InStartTime)
-			, EndTime(InEndTime)
-			, Predicate(InPredicate)
-			, TimingEvent(InOutTimingEvent)
-			, bStopAtFirstMatch(bInStopAtFirstMatch)
-			, bSearchForLargestEvent(bInSearchForLargestEvent)
-			, bFound(false)
-			, bContinueSearching(true)
-			, LargestDuration(-1.0)
-		{
-		}
+		InOutTimingEvent = FTimingEvent(this, InFoundStartTime, InFoundEndTime, InFoundDepth);
+	});
+}
 
-		void CheckMessage(double EventStartTime, double EventEndTime, uint32 EventDepth, uint64 InMessageId)
-		{
-			if (bContinueSearching && Predicate(EventStartTime, EventEndTime, EventDepth))
-			{
-				if (!bSearchForLargestEvent || EventEndTime - EventStartTime > LargestDuration)
-				{
-					LargestDuration = EventEndTime - EventStartTime;
+bool FSkeletalMeshPoseTrack::FindSkeletalMeshPoseMessage(const FTimingEventSearchParameters& InParameters, TFunctionRef<void(double, double, uint32, const FSkeletalMeshPoseMessage&)> InFoundPredicate) const
+{
+	// Storage for the message we want to match (payload for an event)
+	FSkeletalMeshPoseMessage MatchedMessage;
 
-					TimingEvent.TypeId = InMessageId;
-					TimingEvent.Depth = EventDepth;
-					TimingEvent.StartTime = EventStartTime;
-					TimingEvent.EndTime = EventEndTime;
+	return TTimingEventSearch<FSkeletalMeshPoseMessage>::Search(
+		InParameters,
 
-					bFound = true;
-					bContinueSearching = !bStopAtFirstMatch || bSearchForLargestEvent;
-				}
-			}
-		}
-	};
+		// Search...
+		[this](TTimingEventSearch<FSkeletalMeshPoseMessage>::FContext& InContext)
+		{
+			const FAnimationProvider* AnimationProvider = SharedData.GetAnalysisSession().ReadProvider<FAnimationProvider>(FAnimationProvider::ProviderName);
 
-	const FAnimationProvider* AnimationProvider = SharedData.GetAnalysisSession().ReadProvider<FAnimationProvider>(FAnimationProvider::ProviderName);
+			if(AnimationProvider)
+			{
+				Trace::FAnalysisSessionReadScope SessionReadScope(SharedData.GetAnalysisSession());
 
-	if(AnimationProvider)
-	{
-		FSearchTimingEventContext Context(InStartTime, InEndTime, InPredicate, InOutTimingEvent, bInStopAtFirstMatch, bInSearchForLargestEvent);
+				AnimationProvider->ReadSkeletalMeshPoseTimeline(GetGameplayTrack().GetObjectId(), [&InContext](const FAnimationProvider::SkeletalMeshPoseTimeline& InTimeline)
+				{
+					InTimeline.EnumerateEvents(InContext.GetParameters().StartTime, InContext.GetParameters().EndTime, [&InContext](double InEventStartTime, double InEventEndTime, uint32 InDepth, const FSkeletalMeshPoseMessage& InMessage)
+					{
+						InContext.Check(InEventStartTime, InEventEndTime, 0, InMessage);
+					});
+				});
+			}
+		},
 
-		Trace::FAnalysisSessionReadScope SessionReadScope(SharedData.GetAnalysisSession());
+		// Matched...
+		[&MatchedMessage](double InStartTime, double InEndTime, uint32 InDepth, const FSkeletalMeshPoseMessage& InEvent)
+		{
+			MatchedMessage = InEvent;
+		},
 
-		AnimationProvider->ReadSkeletalMeshPoseTimeline(GetGameplayTrack().GetObjectId(), [&Context, &InStartTime, &InEndTime](const FAnimationProvider::SkeletalMeshPoseTimeline& InTimeline)
+		// Found!
+		[&InFoundPredicate, &MatchedMessage](double InFoundStartTime, double InFoundEndTime, uint32 InFoundDepth)
 		{
-			InTimeline.EnumerateEvents(InStartTime, InEndTime, [&Context](double InEventStartTime, double InEventEndTime, uint32 InDepth, const FSkeletalMeshPoseMessage& InMessage)
-			{
-				Context.CheckMessage(InEventStartTime, InEventEndTime, 0, InMessage.MessageId);
-			});
+			InFoundPredicate(InFoundStartTime, InFoundEndTime, InFoundDepth, MatchedMessage);
 		});
-	}
-
-	return false;
 }
 
 void FSkeletalMeshPoseTrack::BuildContextMenu(FMenuBuilder& MenuBuilder)
diff --git a/Engine/Plugins/Animation/GameplayInsights/Source/GameplayInsights/Private/SkeletalMeshPoseTrack.h b/Engine/Plugins/Animation/GameplayInsights/Source/GameplayInsights/Private/SkeletalMeshPoseTrack.h
--- a/Engine/Plugins/Animation/GameplayInsights/Source/GameplayInsights/Private/SkeletalMeshPoseTrack.h
+++ b/Engine/Plugins/Animation/GameplayInsights/Source/GameplayInsights/Private/SkeletalMeshPoseTrack.h
@@ -7,6 +7,8 @@
 
 class ULineBatchComponent;
 class FAnimationSharedData;
+class FTimingEventSearchParameters;
+struct FSkeletalMeshPoseMessage;
 
 class FSkeletalMeshPoseTrack : public TGameplayTrackMixin<FTimingEventsTrack>
 {
@@ -19,7 +21,7 @@ public:
 	virtual void Reset() override;
 	virtual void Draw(ITimingViewDrawHelper& Helper) const override;
 	virtual void InitTooltip(FTooltipDrawState& Tooltip, const FTimingEvent& HoveredTimingEvent) const override;
-	virtual bool SearchTimingEvent(const double InStartTime, const double InEndTime, TFunctionRef<bool(double, double, uint32)> InPredicate, FTimingEvent& InOutTimingEvent, bool bInStopAtFirstMatch, bool bInSearchForLargestEvent) const override;
+	virtual bool SearchTimingEvent(const FTimingEventSearchParameters& InSearchParameters, FTimingEvent& InOutTimingEvent) const override;
 	virtual void BuildContextMenu(FMenuBuilder& MenuBuilder) override;
 
 	// Access drawing flags
@@ -33,6 +35,10 @@ public:
 	void DrawPoses(ULineBatchComponent* InLineBatcher, double SelectionStartTime, double SelectionEndTime);
 #endif
 
+private:
+	// Helper function used to find a skeletal mesh pose
+	bool FindSkeletalMeshPoseMessage(const FTimingEventSearchParameters& InParameters, TFunctionRef<void(double, double, uint32, const FSkeletalMeshPoseMessage&)> InFoundPredicate) const;
+
 private:
 	/** The shared data */
 	const FAnimationSharedData& SharedData;
diff --git a/Engine/Source/Developer/TraceInsights/Private/Insights/ViewModels/FileActivityTimingTrack.cpp b/Engine/Source/Developer/TraceInsights/Private/Insights/ViewModels/FileActivityTimingTrack.cpp
--- a/Engine/Source/Developer/TraceInsights/Private/Insights/ViewModels/FileActivityTimingTrack.cpp
+++ b/Engine/Source/Developer/TraceInsights/Private/Insights/ViewModels/FileActivityTimingTrack.cpp
@@ -17,6 +17,7 @@
 #include "Insights/ViewModels/TimingTrackViewport.h"
 #include "Insights/ViewModels/TimingViewDrawHelper.h"
 #include "Insights/ViewModels/TooltipDrawState.h"
+#include "Insights/ViewModels/TimingEventSearch.h"
 
 #include <limits>
 
@@ -396,129 +397,99 @@ void FFileActivitySharedState::Update()
 
 void FFileActivityTimingTrack::InitTooltip(FTooltipDrawState& Tooltip, const FTimingEvent& HoveredTimingEvent) const
 {
-	Tooltip.ResetContent();
-
-	const Trace::EFileActivityType ActivityType = static_cast<Trace::EFileActivityType>(HoveredTimingEvent.TypeId & 0x0F);
-	const bool bHasFailed = ((HoveredTimingEvent.TypeId & 0xF0) != 0);
+	auto MatchDepth = [&HoveredTimingEvent](double, double, uint32 InDepth)
+	{ 
+		return InDepth == HoveredTimingEvent.Depth; 
+	};
 
-	FString TypeStr;
-	uint32 TypeColor;
-	if (bHasFailed)
-	{
-		TypeStr = TEXT("Failed ");
-		TypeStr += GetFileActivityTypeName(ActivityType);
-		TypeColor = 0xFFFF3333;
-	}
-	else
+	FTimingEventSearchParameters SearchParameters(HoveredTimingEvent.StartTime, HoveredTimingEvent.EndTime, ETimingEventSearchFlags::StopAtFirstMatch, MatchDepth);
+	FindIoTimingEvent(SearchParameters, false, [this, &Tooltip, &HoveredTimingEvent](double InFoundStartTime, double InFoundEndTime, uint32 InFoundDepth, const FFileActivitySharedState::FIoTimingEvent& InEvent)
 	{
-		TypeStr = GetFileActivityTypeName(ActivityType);
-		TypeColor = GetFileActivityTypeColor(ActivityType);
-	}
-	FLinearColor TypeLinearColor = FLinearColor(FColor(TypeColor));
-	TypeLinearColor.R *= 2.0f;
-	TypeLinearColor.G *= 2.0f;
-	TypeLinearColor.B *= 2.0f;
-	Tooltip.AddTitle(TypeStr, TypeLinearColor);
+		Tooltip.ResetContent();
 
-	Tooltip.AddTitle(HoveredTimingEvent.Path);
+		const Trace::EFileActivityType ActivityType = static_cast<Trace::EFileActivityType>(InEvent.Type & 0x0F);
+		const bool bHasFailed = ((InEvent.Type & 0xF0) != 0);
 
-	Tooltip.AddNameValueTextLine(TEXT("Duration:"), TimeUtils::FormatTimeAuto(HoveredTimingEvent.Duration()));
-	Tooltip.AddNameValueTextLine(TEXT("Depth:"), FString::Printf(TEXT("%d"), HoveredTimingEvent.Depth));
+		FString TypeStr;
+		uint32 TypeColor;
+		if (bHasFailed)
+		{
+			TypeStr = TEXT("Failed ");
+			TypeStr += GetFileActivityTypeName(ActivityType);
+			TypeColor = 0xFFFF3333;
+		}
+		else
+		{
+			TypeStr = GetFileActivityTypeName(ActivityType);
+			TypeColor = GetFileActivityTypeColor(ActivityType);
+		}
+		FLinearColor TypeLinearColor = FLinearColor(FColor(TypeColor));
+		TypeLinearColor.R *= 2.0f;
+		TypeLinearColor.G *= 2.0f;
+		TypeLinearColor.B *= 2.0f;
+		Tooltip.AddTitle(TypeStr, TypeLinearColor);
 
-	if (ActivityType == Trace::FileActivityType_Read || ActivityType == Trace::FileActivityType_Write)
-	{
-		Tooltip.AddNameValueTextLine(TEXT("Offset:"), FText::AsNumber(HoveredTimingEvent.Offset).ToString() + TEXT(" bytes"));
-		Tooltip.AddNameValueTextLine(TEXT("Size:"), FText::AsNumber(HoveredTimingEvent.Size).ToString() + TEXT(" bytes"));
-	}
+		Tooltip.AddTitle(InEvent.FileActivity->Path);
+
+		Tooltip.AddNameValueTextLine(TEXT("Duration:"), TimeUtils::FormatTimeAuto(HoveredTimingEvent.Duration()));
+		Tooltip.AddNameValueTextLine(TEXT("Depth:"), FString::Printf(TEXT("%d"), HoveredTimingEvent.Depth));
+
+		if (ActivityType == Trace::FileActivityType_Read || ActivityType == Trace::FileActivityType_Write)
+		{
+			Tooltip.AddNameValueTextLine(TEXT("Offset:"), FText::AsNumber(InEvent.Offset).ToString() + TEXT(" bytes"));
+			Tooltip.AddNameValueTextLine(TEXT("Size:"), FText::AsNumber(InEvent.Size).ToString() + TEXT(" bytes"));
+		}
 
-	Tooltip.UpdateLayout();
+		Tooltip.UpdateLayout();
+	});
 }
 
 ////////////////////////////////////////////////////////////////////////////////////////////////////
 
-bool FFileActivityTimingTrack::SearchEvent(const double InStartTime,
-										   const double InEndTime,
-										   TFunctionRef<bool(double, double, uint32)> InPredicate,
-										   FTimingEvent& InOutTimingEvent,
-										   bool bInStopAtFirstMatch,
-										   bool bInSearchForLargestEvent,
-										   bool bIgnoreEventDepth) const
+bool FFileActivityTimingTrack::FindIoTimingEvent(const FTimingEventSearchParameters& InParameters, bool bIgnoreEventDepth, TFunctionRef<void(double, double, uint32, const FFileActivitySharedState::FIoTimingEvent&)> InFoundPredicate) const
 {
-	struct FSearchTimingEventContext
-	{
-		const double StartTime;
-		const double EndTime;
-		TFunctionRef<bool(double, double, uint32)> Predicate;
-		FTimingEvent& TimingEvent;
-		const bool bStopAtFirstMatch;
-		const bool bSearchForLargestEvent;
-		mutable bool bFound;
-		mutable bool bContinueSearching;
-		mutable double LargestDuration;
-
-		FSearchTimingEventContext(const double InStartTime, const double InEndTime, TFunctionRef<bool(double, double, uint32)> InPredicate, FTimingEvent& InOutTimingEvent, bool bInStopAtFirstMatch, bool bInSearchForLargestEvent)
-			: StartTime(InStartTime)
-			, EndTime(InEndTime)
-			, Predicate(InPredicate)
-			, TimingEvent(InOutTimingEvent)
-			, bStopAtFirstMatch(bInStopAtFirstMatch)
-			, bSearchForLargestEvent(bInSearchForLargestEvent)
-			, bFound(false)
-			, bContinueSearching(true)
-			, LargestDuration(-1.0)
-		{
-		}
+	// Storage for the event we want to match
+	FFileActivitySharedState::FIoTimingEvent MatchedEvent;
+
+	return TTimingEventSearch<FFileActivitySharedState::FIoTimingEvent>::Search(
+		InParameters,
 
-		void CheckEvent(const FFileActivitySharedState::FIoTimingEvent& Event, uint32 EventDepth)
+		// Search...
+		[this, bIgnoreEventDepth](TTimingEventSearch<FFileActivitySharedState::FIoTimingEvent>::FContext& InContext)
 		{
-			if (bContinueSearching && Predicate(Event.StartTime, Event.EndTime, EventDepth))
+			TSharedPtr<const Trace::IAnalysisSession> Session = FInsightsManager::Get()->GetSession();
+			if (Session.IsValid())
 			{
-				if (!bSearchForLargestEvent || Event.EndTime - Event.StartTime > LargestDuration)
-				{
-					LargestDuration = Event.EndTime - Event.StartTime;
+				Trace::FAnalysisSessionReadScope SessionReadScope(*Session.Get());
 
-					TimingEvent.TypeId = Event.Type;
-					TimingEvent.Depth = EventDepth;
-					TimingEvent.StartTime = Event.StartTime;
-					TimingEvent.EndTime = Event.EndTime;
+				for (const FFileActivitySharedState::FIoTimingEvent& Event : State->GetAllEvents())
+				{
+					if (Event.EndTime <= InContext.GetParameters().StartTime)
+					{
+						continue;
+					}
 
-					TimingEvent.Offset = Event.Offset;
-					TimingEvent.Size = Event.Size;
-					TimingEvent.Path = Event.FileActivity->Path;
+					if (!InContext.ShouldContinueSearching() || InContext.GetParameters().StartTime >= InContext.GetParameters().EndTime)
+					{
+						break;
+					}
 
-					bFound = true;
-					bContinueSearching = !bStopAtFirstMatch || bSearchForLargestEvent;
+					InContext.Check(Event.StartTime, Event.EndTime, bIgnoreEventDepth ? 0 : Event.Depth, Event);
 				}
 			}
-		}
-	};
+		},
 
-	FSearchTimingEventContext Ctx(InStartTime, InEndTime, InPredicate, InOutTimingEvent, bInStopAtFirstMatch, bInSearchForLargestEvent);
-
-	TSharedPtr<const Trace::IAnalysisSession> Session = FInsightsManager::Get()->GetSession();
-	if (Session.IsValid())
-	{
-		Trace::FAnalysisSessionReadScope SessionReadScope(*Session.Get());
-
-		const FTimingEventsTrack* Track = Ctx.TimingEvent.Track;
-
-		for (const FFileActivitySharedState::FIoTimingEvent& Event : State->GetAllEvents())
+		// Matched...
+		[&MatchedEvent](double InStartTime, double InEndTime, uint32 InDepth, const FFileActivitySharedState::FIoTimingEvent& InEvent)
 		{
-			if (Event.EndTime <= Ctx.StartTime)
-			{
-				continue;
-			}
-
-			if (!Ctx.bContinueSearching || Event.StartTime >= Ctx.EndTime)
-			{
-				break;
-			}
+			MatchedEvent = InEvent;
+		},
 
-			Ctx.CheckEvent(Event, bIgnoreEventDepth ? 0 : Event.Depth);
-		}
-	}
-
-	return Ctx.bFound;
+		// Found!
+		[&InFoundPredicate, &MatchedEvent](double InFoundStartTime, double InFoundEndTime, uint32 InFoundDepth)
+		{
+			InFoundPredicate(InFoundStartTime, InFoundEndTime, InFoundDepth, MatchedEvent);
+		});
 }
 
 ////////////////////////////////////////////////////////////////////////////////////////////////////
@@ -575,6 +546,15 @@ void FOverviewFileActivityTimingTrack::Draw(ITimingViewDrawHelper& Helper) const
 	}
 }
 
+bool FOverviewFileActivityTimingTrack::SearchTimingEvent(const FTimingEventSearchParameters& InSearchParameters, FTimingEvent& InOutTimingEvent) const
+{
+	constexpr bool bIgnoreEventDepth = true;
+	return FindIoTimingEvent(InSearchParameters, bIgnoreEventDepth, [this, &InOutTimingEvent](double InFoundStartTime, double InFoundEndTime, uint32 InFoundDepth, const FFileActivitySharedState::FIoTimingEvent& InEvent)
+	{
+		InOutTimingEvent = FTimingEvent(this, InFoundStartTime, InFoundEndTime, InFoundDepth);
+	});
+}
+
 ////////////////////////////////////////////////////////////////////////////////////////////////////
 // FDetailedFileActivityTimingTrack
 ////////////////////////////////////////////////////////////////////////////////////////////////////
@@ -654,6 +634,15 @@ void FDetailedFileActivityTimingTrack::Draw(ITimingViewDrawHelper& Helper) const
 	}
 }
 
+bool FDetailedFileActivityTimingTrack::SearchTimingEvent(const FTimingEventSearchParameters& InSearchParameters, FTimingEvent& InOutTimingEvent) const
+{
+	constexpr bool bIgnoreEventDepth = false;
+	return FindIoTimingEvent(InSearchParameters, bIgnoreEventDepth, [this, &InOutTimingEvent](double InFoundStartTime, double InFoundEndTime, uint32 InFoundDepth, const FFileActivitySharedState::FIoTimingEvent& InEvent)
+	{
+		InOutTimingEvent = FTimingEvent(this, InFoundStartTime, InFoundEndTime, InFoundDepth);
+	});
+}
+
 ////////////////////////////////////////////////////////////////////////////////////////////////////
 
 #undef LOCTEXT_NAMESPACE
diff --git a/Engine/Source/Developer/TraceInsights/Private/Insights/ViewModels/FileActivityTimingTrack.h b/Engine/Source/Developer/TraceInsights/Private/Insights/ViewModels/FileActivityTimingTrack.h
--- a/Engine/Source/Developer/TraceInsights/Private/Insights/ViewModels/FileActivityTimingTrack.h
+++ b/Engine/Source/Developer/TraceInsights/Private/Insights/ViewModels/FileActivityTimingTrack.h
@@ -7,6 +7,8 @@
 // Insights
 #include "Insights/ViewModels/TimingEventsTrack.h"
 
+class FTimingEventSearchParameters;
+
 ////////////////////////////////////////////////////////////////////////////////////////////////////
 
 class FFileActivitySharedState
@@ -78,7 +80,7 @@ public:
 	virtual void InitTooltip(FTooltipDrawState& Tooltip, const FTimingEvent& HoveredTimingEvent) const override;
 
 protected:
-	bool SearchEvent(const double InStartTime, const double InEndTime, TFunctionRef<bool(double, double, uint32)> InPredicate, FTimingEvent& InOutTimingEvent, bool bInStopAtFirstMatch, bool bInSearchForLargestEvent, bool bIgnoreEventDepth) const;
+	bool FindIoTimingEvent(const FTimingEventSearchParameters& InParameters, bool bIgnoreEventDepth, TFunctionRef<void(double, double, uint32, const FFileActivitySharedState::FIoTimingEvent&)> InFoundPredicate) const;
 
 protected:
 	TSharedPtr<FFileActivitySharedState> State;
@@ -95,11 +97,7 @@ public:
 	}
 
 	virtual void Draw(ITimingViewDrawHelper& Helper) const override;
-	virtual bool SearchTimingEvent(const double InStartTime, const double InEndTime, TFunctionRef<bool(double, double, uint32)> InPredicate, FTimingEvent& InOutTimingEvent, bool bInStopAtFirstMatch, bool bInSearchForLargestEvent) const override
-	{
-		constexpr bool bIgnoreEventDepth = true;
-		return SearchEvent(InStartTime, InEndTime, InPredicate, InOutTimingEvent, bInStopAtFirstMatch, bInSearchForLargestEvent, bIgnoreEventDepth);
-	}
+	virtual bool SearchTimingEvent(const FTimingEventSearchParameters& InSearchParameters,  FTimingEvent& InOutTimingEvent) const override;
 };
 
 ////////////////////////////////////////////////////////////////////////////////////////////////////
@@ -113,11 +111,7 @@ public:
 	}
 
 	virtual void Draw(ITimingViewDrawHelper& Helper) const override;
-	virtual bool SearchTimingEvent(const double InStartTime, const double InEndTime, TFunctionRef<bool(double, double, uint32)> InPredicate, FTimingEvent& InOutTimingEvent, bool bInStopAtFirstMatch, bool bInSearchForLargestEvent) const override
-	{
-		constexpr bool bIgnoreEventDepth = false;
-		return SearchEvent(InStartTime, InEndTime, InPredicate, InOutTimingEvent, bInStopAtFirstMatch, bInSearchForLargestEvent, bIgnoreEventDepth);
-	}
+	virtual bool SearchTimingEvent(const FTimingEventSearchParameters& InSearchParameters, FTimingEvent& InOutTimingEvent) const override;
 };
 
 ////////////////////////////////////////////////////////////////////////////////////////////////////
diff --git a/Engine/Source/Developer/TraceInsights/Private/Insights/ViewModels/LoadingTimingTrack.cpp b/Engine/Source/Developer/TraceInsights/Private/Insights/ViewModels/LoadingTimingTrack.cpp
--- a/Engine/Source/Developer/TraceInsights/Private/Insights/ViewModels/LoadingTimingTrack.cpp
+++ b/Engine/Source/Developer/TraceInsights/Private/Insights/ViewModels/LoadingTimingTrack.cpp
@@ -15,6 +15,7 @@
 #include "Insights/ViewModels/TimingTrackViewport.h"
 #include "Insights/ViewModels/TimingViewDrawHelper.h"
 #include "Insights/ViewModels/TooltipDrawState.h"
+#include "Insights/ViewModels/TimingEventSearch.h"
 
 #define LOCTEXT_NAMESPACE "LoadingTimingTrack"
 
@@ -94,84 +95,42 @@ void FLoadingSharedState::SetColorSchema(int32 Schema)
 
 void FLoadingTimingTrack::InitTooltip(FTooltipDrawState& Tooltip, const FTimingEvent& HoveredTimingEvent) const
 {
-	Tooltip.ResetContent();
-
-	Tooltip.AddTitle(State->GetEventNameEx(HoveredTimingEvent.Depth, HoveredTimingEvent.LoadingInfo));
-
-	const Trace::FPackageInfo* Package = HoveredTimingEvent.LoadingInfo.Package;
-	const Trace::FPackageExportInfo* Export = HoveredTimingEvent.LoadingInfo.Export;
-
-	Tooltip.AddNameValueTextLine(TEXT("Duration:"), TimeUtils::FormatTimeAuto(HoveredTimingEvent.Duration()));
-	Tooltip.AddNameValueTextLine(TEXT("Depth:"), FString::Printf(TEXT("%d"), HoveredTimingEvent.Depth));
-
-	if (Package)
-	{
-		Tooltip.AddNameValueTextLine(TEXT("Package Name:"), Package->Name);
-		Tooltip.AddNameValueTextLine(TEXT("Header Size:"), FString::Printf(TEXT("%s bytes"), *FText::AsNumber(Package->Summary.TotalHeaderSize).ToString()));
-		Tooltip.AddNameValueTextLine(TEXT("Package Summary:"), FString::Printf(TEXT("%d imports, %d exports"), Package->Summary.ImportCount, Package->Summary.ExportCount));
-	}
-
-	Tooltip.AddNameValueTextLine(TEXT("Export Event:"), FString::Printf(TEXT("%s"), Trace::GetLoadTimeProfilerObjectEventTypeString(HoveredTimingEvent.LoadingInfo.EventType)));
+	auto MatchDepth = [&HoveredTimingEvent](double, double, uint32 InDepth)
+	{ 
+		return InDepth == HoveredTimingEvent.Depth; 
+	};
 
-	if (Export)
+	FTimingEventSearchParameters SearchParameters(HoveredTimingEvent.StartTime, HoveredTimingEvent.EndTime, ETimingEventSearchFlags::StopAtFirstMatch, MatchDepth);
+	FindLoadTimeProfilerCpuEvent(SearchParameters, [this, &Tooltip, &HoveredTimingEvent](double InFoundStartTime, double InFoundEndTime, uint32 InFoundDepth, const Trace::FLoadTimeProfilerCpuEvent& InFoundEvent)
 	{
-		Tooltip.AddNameValueTextLine(TEXT("Export Class:"), Export->Class ? Export->Class->Name : TEXT("N/A"));
-		Tooltip.AddNameValueTextLine(TEXT("Serial Size:"), FString::Printf(TEXT("%s bytes"), *FText::AsNumber(Export->SerialSize).ToString()));
-	}
+		Tooltip.ResetContent();
 
-	Tooltip.UpdateLayout();
-}
+		Tooltip.AddTitle(State->GetEventNameEx(HoveredTimingEvent.Depth, InFoundEvent));
 
-////////////////////////////////////////////////////////////////////////////////////////////////////
-// FSearchLoadTimeProfilerCpuEvent
-////////////////////////////////////////////////////////////////////////////////////////////////////
+		const Trace::FPackageInfo* Package = InFoundEvent.Package;
+		const Trace::FPackageExportInfo* Export = InFoundEvent.Export;
 
-struct FSearchLoadTimeProfilerCpuEvent
-{
-	const double StartTime;
-	const double EndTime;
-	TFunctionRef<bool(double, double, uint32)> Predicate;
-	FTimingEvent& TimingEvent;
-	const bool bStopAtFirstMatch;
-	const bool bSearchForLargestEvent;
-	mutable bool bFound;
-	mutable bool bContinueSearching;
-	mutable double LargestDuration;
-
-	FSearchLoadTimeProfilerCpuEvent(const double InStartTime, const double InEndTime, TFunctionRef<bool(double, double, uint32)> InPredicate, FTimingEvent& InOutTimingEvent, bool bInStopAtFirstMatch, bool bInSearchForLargestEvent)
-		: StartTime(InStartTime)
-		, EndTime(InEndTime)
-		, Predicate(InPredicate)
-		, TimingEvent(InOutTimingEvent)
-		, bStopAtFirstMatch(bInStopAtFirstMatch)
-		, bSearchForLargestEvent(bInSearchForLargestEvent)
-		, bFound(false)
-		, bContinueSearching(true)
-		, LargestDuration(-1.0)
-	{
-	}
+		Tooltip.AddNameValueTextLine(TEXT("Duration:"), TimeUtils::FormatTimeAuto(HoveredTimingEvent.Duration()));
+		Tooltip.AddNameValueTextLine(TEXT("Depth:"), FString::Printf(TEXT("%d"), HoveredTimingEvent.Depth));
 
-	void CheckEvent(double EventStartTime, double EventEndTime, uint32 EventDepth, const Trace::FLoadTimeProfilerCpuEvent& Event)
-	{
-		if (bContinueSearching && Predicate(EventStartTime, EventEndTime, EventDepth))
+		if (Package)
 		{
-			if (!bSearchForLargestEvent || EventEndTime - EventStartTime > LargestDuration)
-			{
-				LargestDuration = EventEndTime - EventStartTime;
-
-				TimingEvent.TypeId = 0;
-				TimingEvent.Depth = EventDepth;
-				TimingEvent.StartTime = EventStartTime;
-				TimingEvent.EndTime = EventEndTime;
+			Tooltip.AddNameValueTextLine(TEXT("Package Name:"), Package->Name);
+			Tooltip.AddNameValueTextLine(TEXT("Header Size:"), FString::Printf(TEXT("%s bytes"), *FText::AsNumber(Package->Summary.TotalHeaderSize).ToString()));
+			Tooltip.AddNameValueTextLine(TEXT("Package Summary:"), FString::Printf(TEXT("%d imports, %d exports"), Package->Summary.ImportCount, Package->Summary.ExportCount));
+		}
 
-				TimingEvent.LoadingInfo = Event;
+		Tooltip.AddNameValueTextLine(TEXT("Export Event:"), FString::Printf(TEXT("%s"), Trace::GetLoadTimeProfilerObjectEventTypeString(InFoundEvent.EventType)));
 
-				bFound = true;
-				bContinueSearching = !bStopAtFirstMatch || bSearchForLargestEvent;
-			}
+		if (Export)
+		{
+			Tooltip.AddNameValueTextLine(TEXT("Export Class:"), Export->Class ? Export->Class->Name : TEXT("N/A"));
+			Tooltip.AddNameValueTextLine(TEXT("Serial Size:"), FString::Printf(TEXT("%s bytes"), *FText::AsNumber(Export->SerialSize).ToString()));
 		}
-	}
-};
+
+		Tooltip.UpdateLayout();
+	});
+}
 
 ////////////////////////////////////////////////////////////////////////////////////////////////////
 // FLoadingMainThreadTimingTrack
@@ -218,34 +177,58 @@ void FLoadingTimingTrack::Draw(ITimingViewDrawHelper& Helper) const
 
 ////////////////////////////////////////////////////////////////////////////////////////////////////
 
-bool FLoadingTimingTrack::SearchTimingEvent(const double InStartTime,
-													  const double InEndTime,
-													  TFunctionRef<bool(double, double, uint32)> InPredicate,
-													  FTimingEvent& InOutTimingEvent,
-													  bool bInStopAtFirstMatch,
-													  bool bInSearchForLargestEvent) const
+bool FLoadingTimingTrack::SearchTimingEvent(const FTimingEventSearchParameters& InSearchParameters, FTimingEvent& InOutTimingEvent) const
 {
-	FSearchLoadTimeProfilerCpuEvent Ctx(InStartTime, InEndTime, InPredicate, InOutTimingEvent, bInStopAtFirstMatch, bInSearchForLargestEvent);
+	return FindLoadTimeProfilerCpuEvent(InSearchParameters, [this, &InOutTimingEvent](double InFoundStartTime, double InFoundEndTime, uint32 InFoundDepth, const Trace::FLoadTimeProfilerCpuEvent& InFoundEvent)
+	{
+		InOutTimingEvent = FTimingEvent(this, InFoundStartTime, InFoundEndTime, InFoundDepth);
+	});
+}
 
-	TSharedPtr<const Trace::IAnalysisSession> Session = FInsightsManager::Get()->GetSession();
-	Trace::FAnalysisSessionReadScope SessionReadScope(*Session.Get());
+////////////////////////////////////////////////////////////////////////////////////////////////////
 
-	const FTimingEventsTrack* Track = Ctx.TimingEvent.Track;
+bool FLoadingTimingTrack::FindLoadTimeProfilerCpuEvent(const FTimingEventSearchParameters& InParameters, TFunctionRef<void(double, double, uint32, const Trace::FLoadTimeProfilerCpuEvent&)> InFoundPredicate) const
+{
+	// Storage for the event we want to match
+	Trace::FLoadTimeProfilerCpuEvent MatchedEvent;
 
-	if (Trace::ReadLoadTimeProfilerProvider(*Session.Get()))
-	{
-		const Trace::ILoadTimeProfilerProvider& LoadTimeProfilerProvider = *Trace::ReadLoadTimeProfilerProvider(*Session.Get());
+	return TTimingEventSearch<Trace::FLoadTimeProfilerCpuEvent>::Search(
+		InParameters,
 
-		LoadTimeProfilerProvider.ReadTimeline(TimelineIndex, [&Ctx](const Trace::ILoadTimeProfilerProvider::CpuTimeline& Timeline)
+		// Search...
+		[this](TTimingEventSearch<Trace::FLoadTimeProfilerCpuEvent>::FContext& InContext)
 		{
-			Timeline.EnumerateEvents(Ctx.StartTime, Ctx.EndTime, [&Ctx](double EventStartTime, double EventEndTime, uint32 EventDepth, const Trace::FLoadTimeProfilerCpuEvent& Event)
+			TSharedPtr<const Trace::IAnalysisSession> Session = FInsightsManager::Get()->GetSession();
+			if (Session.IsValid())
 			{
-				Ctx.CheckEvent(EventStartTime, EventEndTime, EventDepth, Event);
-			});
-		});
-	}
+				Trace::FAnalysisSessionReadScope SessionReadScope(*Session.Get());
 
-	return Ctx.bFound;
+				if (Trace::ReadLoadTimeProfilerProvider(*Session.Get()))
+				{
+					const Trace::ILoadTimeProfilerProvider& LoadTimeProfilerProvider = *Trace::ReadLoadTimeProfilerProvider(*Session.Get());
+
+					LoadTimeProfilerProvider.ReadTimeline(TimelineIndex, [&InContext](const Trace::ILoadTimeProfilerProvider::CpuTimeline& Timeline)
+					{
+						Timeline.EnumerateEvents(InContext.GetParameters().StartTime, InContext.GetParameters().EndTime, [&InContext](double EventStartTime, double EventEndTime, uint32 EventDepth, const Trace::FLoadTimeProfilerCpuEvent& Event)
+						{
+							InContext.Check(EventStartTime, EventEndTime, EventDepth, Event);
+						});
+					});
+				}
+			}
+		},
+
+		// Matched...
+		[&MatchedEvent](double InStartTime, double InEndTime, uint32 InDepth, const Trace::FLoadTimeProfilerCpuEvent& InEvent)
+		{
+			MatchedEvent = InEvent;
+		},
+
+		// Found!
+		[&InFoundPredicate, &MatchedEvent](double InFoundStartTime, double InFoundEndTime, uint32 InFoundDepth)
+		{
+			InFoundPredicate(InFoundStartTime, InFoundEndTime, InFoundDepth, MatchedEvent);
+		});
 }
 
 ////////////////////////////////////////////////////////////////////////////////////////////////////
diff --git a/Engine/Source/Developer/TraceInsights/Private/Insights/ViewModels/LoadingTimingTrack.h b/Engine/Source/Developer/TraceInsights/Private/Insights/ViewModels/LoadingTimingTrack.h
--- a/Engine/Source/Developer/TraceInsights/Private/Insights/ViewModels/LoadingTimingTrack.h
+++ b/Engine/Source/Developer/TraceInsights/Private/Insights/ViewModels/LoadingTimingTrack.h
@@ -8,6 +8,8 @@
 // Insights
 #include "Insights/ViewModels/TimingEventsTrack.h"
 
+class FTimingEventSearchParameters;
+
 /** Defines FLoadingTrackGetEventNameDelegate delegate interface. Returns the name for a timing event in a Loading track. */
 DECLARE_DELEGATE_RetVal_TwoParams(const TCHAR*, FLoadingTrackGetEventNameDelegate, uint32 /*Depth*/, const Trace::FLoadTimeProfilerCpuEvent& /*Event*/);
 
@@ -45,9 +47,12 @@ public:
 	virtual ~FLoadingTimingTrack() {}
 
 	virtual void InitTooltip(FTooltipDrawState& Tooltip, const FTimingEvent& HoveredTimingEvent) const override;
-
 	virtual void Draw(ITimingViewDrawHelper& Helper) const override;
-	virtual bool SearchTimingEvent(const double InStartTime, const double InEndTime, TFunctionRef<bool(double, double, uint32)> InPredicate, FTimingEvent& InOutTimingEvent, bool bInStopAtFirstMatch, bool bInSearchForLargestEvent) const override;
+	virtual bool SearchTimingEvent(const FTimingEventSearchParameters& InSearchParameters, FTimingEvent& InOutTimingEvent) const override;
+
+protected:
+	// Helper function to find an event given search parameters
+	bool FindLoadTimeProfilerCpuEvent(const FTimingEventSearchParameters& InParameters, TFunctionRef<void(double, double, uint32, const Trace::FLoadTimeProfilerCpuEvent&)> InFoundPredicate) const;
 
 protected:
 	TSharedPtr<FLoadingSharedState> State;
diff --git a/Engine/Source/Developer/TraceInsights/Private/Insights/ViewModels/ThreadTimingTrack.cpp b/Engine/Source/Developer/TraceInsights/Private/Insights/ViewModels/ThreadTimingTrack.cpp
--- a/Engine/Source/Developer/TraceInsights/Private/Insights/ViewModels/ThreadTimingTrack.cpp
+++ b/Engine/Source/Developer/TraceInsights/Private/Insights/ViewModels/ThreadTimingTrack.cpp
@@ -4,6 +4,8 @@
 
 #include "Fonts/FontMeasure.h"
 #include "Styling/SlateBrush.h"
+#include "Framework/MultiBox/MultiBoxBuilder.h"
+#include "HAL/PlatformApplicationMisc.h"
 
 // Insights
 #include "Insights/Common/PaintUtils.h"
@@ -15,7 +17,7 @@
 #include "Insights/ViewModels/TimingTrackViewport.h"
 #include "Insights/ViewModels/TimingViewDrawHelper.h"
 #include "Insights/ViewModels/TooltipDrawState.h"
-#include "Framework/MultiBox/MultiBoxBuilder.h"
+#include "Insights/ViewModels/TimingEventSearch.h"
 
 #define LOCTEXT_NAMESPACE "ThreadTimingTrack"
 
@@ -65,89 +67,96 @@ void FThreadTimingTrack::Draw(ITimingViewDrawHelper& Helper) const
 
 void FThreadTimingTrack::DrawSelectedEventInfo(const FTimingEvent& SelectedTimingEvent, const FTimingTrackViewport& Viewport, const FDrawContext& DrawContext, const FSlateBrush* WhiteBrush, const FSlateFontInfo& Font) const
 {
-	const FTimerNodePtr TimerNodePtr = FTimingProfilerManager::Get()->GetTimerNode(SelectedTimingEvent.TypeId);
-
-	FString Str = FString::Printf(TEXT("%s (Incl.: %s, Excl.: %s)"),
-		TimerNodePtr ? *(TimerNodePtr->GetName().ToString()) : TEXT("N/A"),
-		*TimeUtils::FormatTimeAuto(SelectedTimingEvent.Duration()),
-		*TimeUtils::FormatTimeAuto(SelectedTimingEvent.ExclusiveTime));
+	FindTimingProfilerEvent(SelectedTimingEvent, [&SelectedTimingEvent, &Font, &Viewport, &DrawContext, &WhiteBrush](double InFoundStartTime, double InFoundEndTime, uint32 InFoundDepth, const Trace::FTimingProfilerEvent& InFoundEvent)
+	{	
+		const FTimerNodePtr TimerNodePtr = FTimingProfilerManager::Get()->GetTimerNode(InFoundEvent.TimerIndex);
+		if(TimerNodePtr.IsValid())
+		{
+			FString Str = FString::Printf(TEXT("%s (Incl.: %s, Excl.: %s)"),
+				TimerNodePtr ? *(TimerNodePtr->GetName().ToString()) : TEXT("N/A"),
+				*TimeUtils::FormatTimeAuto(SelectedTimingEvent.Duration()),
+				*TimeUtils::FormatTimeAuto(SelectedTimingEvent.ExclusiveTime));
 
-	const TSharedRef<FSlateFontMeasure> FontMeasureService = FSlateApplication::Get().GetRenderer()->GetFontMeasureService();
-	const FVector2D Size = FontMeasureService->Measure(Str, Font);
-	const float X = Viewport.GetWidth() - Size.X - 23.0f;
-	const float Y = Viewport.GetHeight() - Size.Y - 18.0f;
+			const TSharedRef<FSlateFontMeasure> FontMeasureService = FSlateApplication::Get().GetRenderer()->GetFontMeasureService();
+			const FVector2D Size = FontMeasureService->Measure(Str, Font);
+			const float X = Viewport.GetWidth() - Size.X - 23.0f;
+			const float Y = Viewport.GetHeight() - Size.Y - 18.0f;
 
-	const FLinearColor BackgroundColor(0.05f, 0.05f, 0.05f, 1.0f);
-	const FLinearColor TextColor(0.7f, 0.7f, 0.7f, 1.0f);
+			const FLinearColor BackgroundColor(0.05f, 0.05f, 0.05f, 1.0f);
+			const FLinearColor TextColor(0.7f, 0.7f, 0.7f, 1.0f);
 
-	DrawContext.DrawBox(X - 8.0f, Y - 2.0f, Size.X + 16.0f, Size.Y + 4.0f, WhiteBrush, BackgroundColor);
-	DrawContext.LayerId++;
+			DrawContext.DrawBox(X - 8.0f, Y - 2.0f, Size.X + 16.0f, Size.Y + 4.0f, WhiteBrush, BackgroundColor);
+			DrawContext.LayerId++;
 
-	DrawContext.DrawText(X, Y, Str, Font, TextColor);
-	DrawContext.LayerId++;
+			DrawContext.DrawText(X, Y, Str, Font, TextColor);
+			DrawContext.LayerId++;
+		}
+	});
 }
 
 ////////////////////////////////////////////////////////////////////////////////////////////////////
 
 void FThreadTimingTrack::InitTooltip(FTooltipDrawState& Tooltip, const FTimingEvent& HoveredTimingEvent) const
 {
-	FTimingEvent ParentTimingEvent;
-	FTimingEvent RootTimingEvent;
-	if (HoveredTimingEvent.IsValid())
-	{
-		GetParentAndRoot(HoveredTimingEvent, ParentTimingEvent, RootTimingEvent);
-	}
+	FindTimingProfilerEvent(HoveredTimingEvent, [this, &Tooltip, &HoveredTimingEvent](double InFoundStartTime, double InFoundEndTime, uint32 InFoundDepth, const Trace::FTimingProfilerEvent& InFoundEvent)
+	{	
+		FTimingEvent ParentTimingEvent;
+		Trace::FTimingProfilerEvent ParentEvent;
+		FTimingEvent RootTimingEvent;
+		Trace::FTimingProfilerEvent RootEvent;
+		if (HoveredTimingEvent.IsValid())
+		{
+			GetParentAndRoot(HoveredTimingEvent, ParentTimingEvent, ParentEvent, RootTimingEvent, RootEvent);
+		}
 
-	Tooltip.ResetContent();
+		Tooltip.ResetContent();
 
-	const FTimerNodePtr TimerNodePtr = FTimingProfilerManager::Get()->GetTimerNode(HoveredTimingEvent.TypeId);
-	FString TimerName = TimerNodePtr ? TimerNodePtr->GetName().ToString() : TEXT("N/A");
-	Tooltip.AddTitle(TimerName);
+		const FTimerNodePtr TimerNodePtr = FTimingProfilerManager::Get()->GetTimerNode(InFoundEvent.TimerIndex);
+		FString TimerName = TimerNodePtr ? TimerNodePtr->GetName().ToString() : TEXT("N/A");
+		Tooltip.AddTitle(TimerName);
 
-	if (ParentTimingEvent.IsValid() && HoveredTimingEvent.Depth > 0)
-	{
-		const FTimerNodePtr ParentTimerNodePtr = FTimingProfilerManager::Get()->GetTimerNode(ParentTimingEvent.TypeId);
-		const FString ParentTimerName = ParentTimerNodePtr ? ParentTimerNodePtr->GetName().ToString() : TEXT("N/A");
-		FNumberFormattingOptions FormattingOptions;
-		FormattingOptions.MaximumFractionalDigits = 2;
-		const FString ValueStr = FString::Printf(TEXT("%s %s"), *FText::AsPercent(HoveredTimingEvent.Duration() / ParentTimingEvent.Duration(), &FormattingOptions).ToString(), *ParentTimerName);
-		Tooltip.AddNameValueTextLine(TEXT("% of Parent:"), ValueStr);
-	}
+		if (ParentTimingEvent.IsValid() && HoveredTimingEvent.Depth > 0)
+		{
+			const FTimerNodePtr ParentTimerNodePtr = FTimingProfilerManager::Get()->GetTimerNode(ParentEvent.TimerIndex);
+			const FString ParentTimerName = ParentTimerNodePtr.IsValid() ? ParentTimerNodePtr->GetName().ToString() : TEXT("N/A");
+			FNumberFormattingOptions FormattingOptions;
+			FormattingOptions.MaximumFractionalDigits = 2;
+			const FString ValueStr = FString::Printf(TEXT("%s %s"), *FText::AsPercent(HoveredTimingEvent.Duration() / ParentTimingEvent.Duration(), &FormattingOptions).ToString(), *ParentTimerName);
+			Tooltip.AddNameValueTextLine(TEXT("% of Parent:"), ValueStr);
+		}
 
-	if (RootTimingEvent.IsValid() && HoveredTimingEvent.Depth > 1)
-	{
-		const FTimerNodePtr RootTimerNodePtr = FTimingProfilerManager::Get()->GetTimerNode(RootTimingEvent.TypeId);
-		const FString RootTimerName = RootTimerNodePtr ? RootTimerNodePtr->GetName().ToString() : TEXT("N/A");
-		FNumberFormattingOptions FormattingOptions;
-		FormattingOptions.MaximumFractionalDigits = 2;
-		const FString ValueStr = FString::Printf(TEXT("%s %s"), *FText::AsPercent(HoveredTimingEvent.Duration() / RootTimingEvent.Duration(), &FormattingOptions).ToString(), *RootTimerName);
-		Tooltip.AddNameValueTextLine(TEXT("% of Root:"), ValueStr);
-	}
+		if (RootTimingEvent.IsValid() && HoveredTimingEvent.Depth > 1)
+		{
+			const FTimerNodePtr RootTimerNodePtr = FTimingProfilerManager::Get()->GetTimerNode(RootEvent.TimerIndex);
+			const FString RootTimerName = RootTimerNodePtr.IsValid() ? RootTimerNodePtr->GetName().ToString() : TEXT("N/A");
+			FNumberFormattingOptions FormattingOptions;
+			FormattingOptions.MaximumFractionalDigits = 2;
+			const FString ValueStr = FString::Printf(TEXT("%s %s"), *FText::AsPercent(HoveredTimingEvent.Duration() / RootTimingEvent.Duration(), &FormattingOptions).ToString(), *RootTimerName);
+			Tooltip.AddNameValueTextLine(TEXT("% of Root:"), ValueStr);
+		}
 
-	Tooltip.AddNameValueTextLine(TEXT("Inclusive Time:"), TimeUtils::FormatTimeAuto(HoveredTimingEvent.Duration()));
+		Tooltip.AddNameValueTextLine(TEXT("Inclusive Time:"), TimeUtils::FormatTimeAuto(HoveredTimingEvent.Duration()));
 
-	//Tooltip.AddNameValueTextLine(TEXT("Exclusive Time:"), TimeUtils::FormatTimeAuto(HoveredTimingEvent.ExclusiveTime));
-	{
-		FNumberFormattingOptions FormattingOptions;
-		FormattingOptions.MaximumFractionalDigits = 2;
-		const FString ExclStr = FString::Printf(TEXT("%s (%s)"), *TimeUtils::FormatTimeAuto(HoveredTimingEvent.ExclusiveTime), *FText::AsPercent(HoveredTimingEvent.ExclusiveTime / HoveredTimingEvent.Duration(), &FormattingOptions).ToString());
-		Tooltip.AddNameValueTextLine(TEXT("Exclusive Time:"), ExclStr);
-	}
+		//Tooltip.AddNameValueTextLine(TEXT("Exclusive Time:"), TimeUtils::FormatTimeAuto(HoveredTimingEvent.ExclusiveTime));
+		{
+			FNumberFormattingOptions FormattingOptions;
+			FormattingOptions.MaximumFractionalDigits = 2;
+			const FString ExclStr = FString::Printf(TEXT("%s (%s)"), *TimeUtils::FormatTimeAuto(HoveredTimingEvent.ExclusiveTime), *FText::AsPercent(HoveredTimingEvent.ExclusiveTime / HoveredTimingEvent.Duration(), &FormattingOptions).ToString());
+			Tooltip.AddNameValueTextLine(TEXT("Exclusive Time:"), ExclStr);
+		}
 
-	Tooltip.AddNameValueTextLine(TEXT("Depth:"), FString::Printf(TEXT("%d"), HoveredTimingEvent.Depth));
+		Tooltip.AddNameValueTextLine(TEXT("Depth:"), FString::Printf(TEXT("%d"), HoveredTimingEvent.Depth));
 
-	Tooltip.UpdateLayout();
+		Tooltip.UpdateLayout();
+	});
 }
 
 ////////////////////////////////////////////////////////////////////////////////////////////////////
 
-void FThreadTimingTrack::GetParentAndRoot(const FTimingEvent& TimingEvent, FTimingEvent& OutParentTimingEvent, FTimingEvent& OutRootTimingEvent) const
+void FThreadTimingTrack::GetParentAndRoot(const FTimingEvent& TimingEvent, FTimingEvent& OutParentTimingEvent, Trace::FTimingProfilerEvent& OutParentEvent, FTimingEvent& OutRootTimingEvent, Trace::FTimingProfilerEvent& OutRootEvent) const
 {
 	// Note: This function does not compute Exclusive Time for parent and root events.
 
-	OutRootTimingEvent.Track = TimingEvent.Track;
-	OutParentTimingEvent.Track = TimingEvent.Track;
-
 	if (TimingEvent.Depth > 0)
 	{
 		TSharedPtr<const Trace::IAnalysisSession> Session = FInsightsManager::Get()->GetSession();
@@ -159,23 +168,19 @@ void FThreadTimingTrack::GetParentAndRoot(const FTimingEvent& TimingEvent, FTimi
 			{
 				const Trace::ITimingProfilerProvider& TimingProfilerProvider = *Trace::ReadTimingProfilerProvider(*Session.Get());
 
-				TimingProfilerProvider.ReadTimeline(TimingEvent.Track->GetId(), [&](const Trace::ITimingProfilerProvider::Timeline& Timeline)
+				TimingProfilerProvider.ReadTimeline(GetId(), [&TimingEvent, &OutParentTimingEvent, &OutParentEvent, &OutRootTimingEvent, &OutRootEvent](const Trace::ITimingProfilerProvider::Timeline& Timeline)
 				{
-					Timeline.EnumerateEvents(TimingEvent.StartTime, TimingEvent.EndTime, [&](double EventStartTime, double EventEndTime, uint32 EventDepth, const Trace::FTimingProfilerEvent& Event)
+					Timeline.EnumerateEvents(TimingEvent.StartTime, TimingEvent.EndTime, [&TimingEvent, &OutParentTimingEvent, &OutParentEvent, &OutRootTimingEvent, &OutRootEvent](double EventStartTime, double EventEndTime, uint32 EventDepth, const Trace::FTimingProfilerEvent& Event)
 					{
 						if (EventDepth == 0)
 						{
-							OutRootTimingEvent.TypeId = Event.TimerIndex;
-							OutRootTimingEvent.Depth = EventDepth;
-							OutRootTimingEvent.StartTime = EventStartTime;
-							OutRootTimingEvent.EndTime = EventEndTime;
+							OutRootEvent = Event;
+							OutRootTimingEvent = FTimingEvent(TimingEvent.Track, EventStartTime, EventEndTime, EventDepth);
 						}
 						if (EventDepth == TimingEvent.Depth - 1)
 						{
-							OutParentTimingEvent.TypeId = Event.TimerIndex;
-							OutParentTimingEvent.Depth = EventDepth;
-							OutParentTimingEvent.StartTime = EventStartTime;
-							OutParentTimingEvent.EndTime = EventEndTime;
+							OutParentEvent = Event;
+							OutParentTimingEvent = FTimingEvent(TimingEvent.Track, EventStartTime, EventEndTime, EventDepth);
 						}
 					});
 				});
@@ -186,80 +191,12 @@ void FThreadTimingTrack::GetParentAndRoot(const FTimingEvent& TimingEvent, FTimi
 
 ////////////////////////////////////////////////////////////////////////////////////////////////////
 
-bool FThreadTimingTrack::SearchTimingEvent(const double InStartTime,
-										   const double InEndTime,
-										   TFunctionRef<bool(double, double, uint32)> InPredicate,
-										   FTimingEvent& InOutTimingEvent,
-										   bool bInStopAtFirstMatch,
-										   bool bInSearchForLargestEvent) const
+bool FThreadTimingTrack::SearchTimingEvent(const FTimingEventSearchParameters& InSearchParameters, FTimingEvent& InOutTimingEvent) const
 {
-	struct FSearchTimingEventContext
+	return FindTimingProfilerEvent(InSearchParameters, [this, &InOutTimingEvent](double InFoundStartTime, double InFoundEndTime, uint32 InFoundDepth, const Trace::FTimingProfilerEvent& InFoundEvent)
 	{
-		const double StartTime;
-		const double EndTime;
-		TFunctionRef<bool(double, double, uint32)> Predicate;
-		FTimingEvent& TimingEvent;
-		const bool bStopAtFirstMatch;
-		const bool bSearchForLargestEvent;
-		mutable bool bFound;
-		mutable bool bContinueSearching;
-		mutable double LargestDuration;
-
-		FSearchTimingEventContext(const double InStartTime, const double InEndTime, TFunctionRef<bool(double, double, uint32)> InPredicate, FTimingEvent& InOutTimingEvent, bool bInStopAtFirstMatch, bool bInSearchForLargestEvent)
-			: StartTime(InStartTime)
-			, EndTime(InEndTime)
-			, Predicate(InPredicate)
-			, TimingEvent(InOutTimingEvent)
-			, bStopAtFirstMatch(bInStopAtFirstMatch)
-			, bSearchForLargestEvent(bInSearchForLargestEvent)
-			, bFound(false)
-			, bContinueSearching(true)
-			, LargestDuration(-1.0)
-		{
-		}
-
-		void CheckEvent(double EventStartTime, double EventEndTime, uint32 EventDepth, const Trace::FTimingProfilerEvent& Event)
-		{
-			if (bContinueSearching && Predicate(EventStartTime, EventEndTime, EventDepth))
-			{
-				if (!bSearchForLargestEvent || EventEndTime - EventStartTime > LargestDuration)
-				{
-					LargestDuration = EventEndTime - EventStartTime;
-
-					TimingEvent.TypeId = Event.TimerIndex;
-					TimingEvent.Depth = EventDepth;
-					TimingEvent.StartTime = EventStartTime;
-					TimingEvent.EndTime = EventEndTime;
-
-					bFound = true;
-					bContinueSearching = !bStopAtFirstMatch || bSearchForLargestEvent;
-				}
-			}
-		}
-	};
-
-	FSearchTimingEventContext Ctx(InStartTime, InEndTime, InPredicate, InOutTimingEvent, bInStopAtFirstMatch, bInSearchForLargestEvent);
-
-	TSharedPtr<const Trace::IAnalysisSession> Session = FInsightsManager::Get()->GetSession();
-	if (Session.IsValid())
-	{
-		Trace::FAnalysisSessionReadScope SessionReadScope(*Session.Get());
-
-		if (Trace::ReadTimingProfilerProvider(*Session.Get()))
-		{
-			const Trace::ITimingProfilerProvider& TimingProfilerProvider = *Trace::ReadTimingProfilerProvider(*Session.Get());
-
-			TimingProfilerProvider.ReadTimeline(Ctx.TimingEvent.Track->GetId(), [&Ctx](const Trace::ITimingProfilerProvider::Timeline& Timeline)
-			{
-				Timeline.EnumerateEvents(Ctx.StartTime, Ctx.EndTime, [&Ctx](double EventStartTime, double EventEndTime, uint32 EventDepth, const Trace::FTimingProfilerEvent& Event)
-				{
-					Ctx.CheckEvent(EventStartTime, EventEndTime, EventDepth, Event);
-				});
-			});
-		}
-	}
-
-	return Ctx.bFound;
+		InOutTimingEvent = FTimingEvent(this, InFoundStartTime, InFoundEndTime, InFoundDepth);
+	});
 }
 
 ////////////////////////////////////////////////////////////////////////////////////////////////////
@@ -339,8 +276,26 @@ void FThreadTimingTrack::ComputeTimingEventStats(FTimingEvent& InOutTimingEvent)
 
 void FThreadTimingTrack::OnEventSelected(const FTimingEvent& SelectedTimingEvent) const
 {
-	// Select the timer node corresponding to timing event type of selected timing event.
-	FTimingProfilerManager::Get()->SetSelectedTimer(SelectedTimingEvent.TypeId);
+	FindTimingProfilerEvent(SelectedTimingEvent, [](double InFoundStartTime, double InFoundEndTime, uint32 InFoundDepth, const Trace::FTimingProfilerEvent& InFoundEvent)
+	{
+		// Select the timer node corresponding to timing event type of selected timing event.
+		FTimingProfilerManager::Get()->SetSelectedTimer(InFoundEvent.TimerIndex);
+	});
+}
+
+////////////////////////////////////////////////////////////////////////////////////////////////////
+
+void FThreadTimingTrack::OnClipboardCopyEvent(const FTimingEvent& SelectedTimingEvent) const
+{
+	FindTimingProfilerEvent(SelectedTimingEvent, [](double InFoundStartTime, double InFoundEndTime, uint32 InFoundDepth, const Trace::FTimingProfilerEvent& InFoundEvent)
+	{
+		const FTimerNodePtr TimerNodePtr = FTimingProfilerManager::Get()->GetTimerNode(InFoundEvent.TimerIndex);
+		if (TimerNodePtr.IsValid())
+		{
+			// Copy name of selected timing event to clipboard.
+			FPlatformApplicationMisc::ClipboardCopy(*TimerNodePtr->GetName().ToString());
+		}
+	});
 }
 
 ////////////////////////////////////////////////////////////////////////////////////////////////////
@@ -371,6 +326,66 @@ void FThreadTimingTrack::BuildContextMenu(FMenuBuilder& MenuBuilder)
 	MenuBuilder.EndSection();
 }
 
+////////////////////////////////////////////////////////////////////////////////////////////////////
+
+bool FThreadTimingTrack::FindTimingProfilerEvent(const FTimingEvent& InTimingEvent, TFunctionRef<void(double, double, uint32, const Trace::FTimingProfilerEvent&)> InFoundPredicate) const
+{
+	auto MatchDepth = [&InTimingEvent](double, double, uint32 InDepth)
+	{ 
+		return InDepth == InTimingEvent.Depth; 
+	};
+
+	FTimingEventSearchParameters SearchParameters(InTimingEvent.StartTime, InTimingEvent.EndTime, ETimingEventSearchFlags::StopAtFirstMatch, MatchDepth);
+	return FindTimingProfilerEvent(SearchParameters, InFoundPredicate);
+}
+
+////////////////////////////////////////////////////////////////////////////////////////////////////
+
+bool FThreadTimingTrack::FindTimingProfilerEvent(const FTimingEventSearchParameters& InParameters, TFunctionRef<void(double, double, uint32, const Trace::FTimingProfilerEvent&)> InFoundPredicate) const
+{
+	// Storage for the event we want to match
+	Trace::FTimingProfilerEvent MatchedEvent;
+
+	return TTimingEventSearch<Trace::FTimingProfilerEvent>::Search(
+		InParameters,
+
+		// Search...
+		[this](TTimingEventSearch<Trace::FTimingProfilerEvent>::FContext& InContext)
+		{
+			TSharedPtr<const Trace::IAnalysisSession> Session = FInsightsManager::Get()->GetSession();
+			if (Session.IsValid())
+			{
+				Trace::FAnalysisSessionReadScope SessionReadScope(*Session.Get());
+
+				if (Trace::ReadTimingProfilerProvider(*Session.Get()))
+				{
+					const Trace::ITimingProfilerProvider& TimingProfilerProvider = *Trace::ReadTimingProfilerProvider(*Session.Get());
+
+					TimingProfilerProvider.ReadTimeline(GetId(), [&InContext](const Trace::ITimingProfilerProvider::Timeline& Timeline)
+					{
+						Timeline.EnumerateEvents(InContext.GetParameters().StartTime, InContext.GetParameters().EndTime, [&InContext](double EventStartTime, double EventEndTime, uint32 EventDepth, const Trace::FTimingProfilerEvent& Event)
+						{
+							InContext.Check(EventStartTime, EventEndTime, EventDepth, Event);
+						});
+					});
+				}
+			}
+		},
+
+		// Matched...
+		[&MatchedEvent](double InStartTime, double InEndTime, uint32 InDepth, const Trace::FTimingProfilerEvent& InEvent)
+		{
+			MatchedEvent = InEvent;
+		},
+
+		// Found!
+		[&InFoundPredicate, &MatchedEvent](double InFoundStartTime, double InFoundEndTime, uint32 InFoundDepth)
+		{
+			InFoundPredicate(InFoundStartTime, InFoundEndTime, InFoundDepth, MatchedEvent);
+		});
+}
+
+
 ////////////////////////////////////////////////////////////////////////////////////////////////////
 
 #undef LOCTEXT_NAMESPACE
diff --git a/Engine/Source/Developer/TraceInsights/Private/Insights/ViewModels/ThreadTimingTrack.h b/Engine/Source/Developer/TraceInsights/Private/Insights/ViewModels/ThreadTimingTrack.h
--- a/Engine/Source/Developer/TraceInsights/Private/Insights/ViewModels/ThreadTimingTrack.h
+++ b/Engine/Source/Developer/TraceInsights/Private/Insights/ViewModels/ThreadTimingTrack.h
@@ -7,6 +7,9 @@
 // Insights
 #include "Insights/ViewModels/TimingEventsTrack.h"
 
+namespace Trace { struct FTimingProfilerEvent; }
+class FTimingEventSearchParameters;
+
 ////////////////////////////////////////////////////////////////////////////////////////////////////
 
 class FThreadTimingTrack : public FTimingEventsTrack
@@ -30,13 +33,16 @@ public:
 	virtual void DrawSelectedEventInfo(const FTimingEvent& SelectedTimingEvent, const FTimingTrackViewport& Viewport, const FDrawContext& DrawContext, const FSlateBrush* WhiteBrush, const FSlateFontInfo& Font) const override;
 	virtual void InitTooltip(FTooltipDrawState& Tooltip, const FTimingEvent& HoveredTimingEvent) const override;
 
-	virtual bool SearchTimingEvent(const double InStartTime, const double InEndTime, TFunctionRef<bool(double, double, uint32)> InPredicate, FTimingEvent& InOutTimingEvent, bool bInStopAtFirstMatch, bool bInSearchForLargestEvent) const override;
+	virtual bool SearchTimingEvent(const FTimingEventSearchParameters& InSearchParameters, FTimingEvent& InOutTimingEvent) const override;
 	virtual void ComputeTimingEventStats(FTimingEvent& InOutTimingEvent) const override;
 	virtual void OnEventSelected(const FTimingEvent& SelectedTimingEvent) const override;
+	virtual void OnClipboardCopyEvent(const FTimingEvent& SelectedTimingEvent) const override;
 	virtual void BuildContextMenu(FMenuBuilder& MenuBuilder) override;
 
 private:
-	void GetParentAndRoot(const FTimingEvent& TimingEvent, FTimingEvent& OutParentTimingEvent, FTimingEvent& OutRootTimingEvent) const;
+	bool FindTimingProfilerEvent(const FTimingEvent& InTimingEvent, TFunctionRef<void(double, double, uint32, const Trace::FTimingProfilerEvent&)> InFoundPredicate) const;
+	bool FindTimingProfilerEvent(const FTimingEventSearchParameters& InParameters, TFunctionRef<void(double, double, uint32, const Trace::FTimingProfilerEvent&)> InFoundPredicate) const;
+	void GetParentAndRoot(const FTimingEvent& TimingEvent, FTimingEvent& OutParentTimingEvent, Trace::FTimingProfilerEvent& OutParentEvent, FTimingEvent& OutRootTimingEvent, Trace::FTimingProfilerEvent& OutRootEvent) const;
 
 private:
 	const TCHAR* GroupName;
diff --git a/Engine/Source/Developer/TraceInsights/Private/Insights/Widgets/STimingView.cpp b/Engine/Source/Developer/TraceInsights/Private/Insights/Widgets/STimingView.cpp
--- a/Engine/Source/Developer/TraceInsights/Private/Insights/Widgets/STimingView.cpp
+++ b/Engine/Source/Developer/TraceInsights/Private/Insights/Widgets/STimingView.cpp
@@ -10,7 +10,6 @@
 #include "Framework/Application/MenuStack.h"
 #include "Framework/Application/SlateApplication.h"
 #include "Framework/MultiBox/MultiBoxBuilder.h"
-#include "HAL/PlatformApplicationMisc.h"
 #include "HAL/PlatformTime.h"
 #include "Layout/WidgetPath.h"
 #include "Misc/Paths.h"
@@ -50,6 +49,7 @@
 #include "Insights/Widgets/STimersView.h"
 #include "Insights/Widgets/STimingProfilerWindow.h"
 #include "Insights/ITimingViewExtender.h"
+#include "Insights/ViewModels/TimingEventSearch.h"
 
 #include <limits>
 
@@ -1620,12 +1620,7 @@ FReply STimingView::OnKeyDown(const FGeometry& MyGeometry, const FKeyEvent& InKe
 		{
 			if (SelectedTimingEvent.IsValid())
 			{
-				const FTimerNodePtr TimerNodePtr = FTimingProfilerManager::Get()->GetTimerNode(SelectedTimingEvent.TypeId);
-				if (TimerNodePtr.IsValid())
-				{
-					// Copy name of selected timing event to clipboard.
-					FPlatformApplicationMisc::ClipboardCopy(*TimerNodePtr->GetName().ToString());
-				}
+				SelectedTimingEvent.Track->OnClipboardCopyEvent(SelectedTimingEvent);
 			}
 		}
 		else
@@ -2217,14 +2212,15 @@ void STimingView::UpdateHoveredTimingEvent(float MX, float MY)
 
 				const double StartTime = Viewport.SlateUnitsToTime(MX);
 				const double EndTime = StartTime + 2.0 / Viewport.GetScaleX(); // +2px
-				constexpr bool bStopAtFirstMatch = true; // get first one matching
-				constexpr bool bSearchForLargestEvent = false;
-				TimingEvent.Track->SearchTimingEvent(StartTime, EndTime,
+
+				FTimingEventSearchParameters SearchParameters(StartTime, EndTime, ETimingEventSearchFlags::StopAtFirstMatch,
+					// Filter
 					[Depth](double, double, uint32 EventDepth)
 					{
 						return EventDepth == Depth;
-					},
-					TimingEvent, bStopAtFirstMatch, bSearchForLargestEvent);
+					});
+				
+				TimingEvent.Track->SearchTimingEvent(SearchParameters, TimingEvent);
 			}
 		}
 	}
@@ -2292,15 +2288,16 @@ void STimingView::SelectLeftTimingEvent()
 		const uint32 Depth = SelectedTimingEvent.Depth;
 		const double StartTime = SelectedTimingEvent.StartTime;
 		const double EndTime = SelectedTimingEvent.EndTime;
-		const bool bStopAtFirstMatch = false; // get last one matching
-		const bool bSearchForLargestEvent = false;
-		if (SelectedTimingEvent.Track->SearchTimingEvent(0.0, StartTime,
+
+		FTimingEventSearchParameters SearchParameters(0.0, StartTime, ETimingEventSearchFlags::SearchAll,
+			// Filter
 			[Depth, StartTime, EndTime](double EventStartTime, double EventEndTime, uint32 EventDepth)
 			{
 				return EventDepth == Depth &&
 					(EventStartTime < StartTime || EventEndTime < EndTime);
-			},
-			SelectedTimingEvent, bStopAtFirstMatch, bSearchForLargestEvent))
+			});
+
+		if (SelectedTimingEvent.Track->SearchTimingEvent(SearchParameters, SelectedTimingEvent))
 		{
 			SelectedTimingEvent.Track->ComputeTimingEventStats(SelectedTimingEvent);
 			BringIntoView(SelectedTimingEvent.StartTime, SelectedTimingEvent.EndTime);
@@ -2318,15 +2315,15 @@ void STimingView::SelectRightTimingEvent()
 		const uint32 Depth = SelectedTimingEvent.Depth;
 		const double StartTime = SelectedTimingEvent.StartTime;
 		const double EndTime = SelectedTimingEvent.EndTime;
-		const bool bStopAtFirstMatch = true; // get first one matching
-		const bool bSearchForLargestEvent = false;
-		if (SelectedTimingEvent.Track->SearchTimingEvent(EndTime, Viewport.GetMaxValidTime(),
+		FTimingEventSearchParameters SearchParameters(EndTime, Viewport.GetMaxValidTime(), ETimingEventSearchFlags::StopAtFirstMatch,
+			// Filter
 			[Depth, StartTime, EndTime](double EventStartTime, double EventEndTime, uint32 EventDepth)
 			{
 				return EventDepth == Depth &&
 					(EventStartTime > StartTime || EventEndTime > EndTime);
-			},
-			SelectedTimingEvent, bStopAtFirstMatch, bSearchForLargestEvent))
+			});
+
+		if (SelectedTimingEvent.Track->SearchTimingEvent(SearchParameters,SelectedTimingEvent))
 		{
 			SelectedTimingEvent.Track->ComputeTimingEventStats(SelectedTimingEvent);
 			BringIntoView(SelectedTimingEvent.StartTime, SelectedTimingEvent.EndTime);
@@ -2343,18 +2340,18 @@ void STimingView::SelectUpTimingEvent()
 		SelectedTimingEvent.Depth > 0)
 	{
 		const uint32 Depth = SelectedTimingEvent.Depth - 1;
-		double StartTime = SelectedTimingEvent.StartTime;
-		double EndTime = SelectedTimingEvent.EndTime;
-		const bool bStopAtFirstMatch = true; // get first one matching
-		const bool bSearchForLargestEvent = false;
-		if (SelectedTimingEvent.Track->SearchTimingEvent(StartTime, EndTime,
+		const double StartTime = SelectedTimingEvent.StartTime;
+		const double EndTime = SelectedTimingEvent.EndTime;
+		FTimingEventSearchParameters SearchParameters(StartTime, EndTime, ETimingEventSearchFlags::StopAtFirstMatch,
+			// Filter
 			[Depth, StartTime, EndTime](double EventStartTime, double EventEndTime, uint32 EventDepth)
 			{
 				return EventDepth == Depth
 					&& EventStartTime <= EndTime
 					&& EventEndTime >= StartTime;
-			},
-			SelectedTimingEvent, bStopAtFirstMatch, bSearchForLargestEvent))
+			});
+
+		if (SelectedTimingEvent.Track->SearchTimingEvent(SearchParameters, SelectedTimingEvent))
 		{
 			SelectedTimingEvent.Track->ComputeTimingEventStats(SelectedTimingEvent);
 			BringIntoView(SelectedTimingEvent.StartTime, SelectedTimingEvent.EndTime);
@@ -2370,18 +2367,28 @@ void STimingView::SelectDownTimingEvent()
 	if (SelectedTimingEvent.IsValid())
 	{
 		const uint32 Depth = SelectedTimingEvent.Depth + 1;
-		double StartTime = SelectedTimingEvent.StartTime;
-		double EndTime = SelectedTimingEvent.EndTime;
-		const bool bStopAtFirstMatch = false; // check all timing events
-		const bool bSearchForLargestEvent = true; // get largest timing event
-		if (SelectedTimingEvent.Track->SearchTimingEvent(StartTime, EndTime,
-			[Depth, StartTime, EndTime](double EventStartTime, double EventEndTime, uint32 EventDepth)
-			{
-				return EventDepth == Depth
-					&& EventStartTime <= EndTime
-					&& EventEndTime >= StartTime;
-			},
-			SelectedTimingEvent, bStopAtFirstMatch, bSearchForLargestEvent))
+		const double StartTime = SelectedTimingEvent.StartTime;
+		const double EndTime = SelectedTimingEvent.EndTime;
+		double LargestDuration = 0.0;
+		FTimingEventSearchParameters SearchParameters(StartTime, EndTime, ETimingEventSearchFlags::SearchAll,
+		// Filter
+		[Depth, StartTime, EndTime, &LargestDuration](double EventStartTime, double EventEndTime, uint32 EventDepth)
+		{
+			const double Duration = EventEndTime - EventStartTime;
+
+			return Duration > LargestDuration
+				&& EventDepth == Depth
+				&& EventStartTime <= EndTime
+				&& EventEndTime >= StartTime;
+		},
+		// Matched
+		[&LargestDuration](double EventStartTime, double EventEndTime, uint32 EventDepth)
+		{
+			const double Duration = EventEndTime - EventStartTime;
+			LargestDuration = Duration;
+		});
+		
+		if (SelectedTimingEvent.Track->SearchTimingEvent(SearchParameters, SelectedTimingEvent))
 		{
 			SelectedTimingEvent.Track->ComputeTimingEventStats(SelectedTimingEvent);
 			BringIntoView(SelectedTimingEvent.StartTime, SelectedTimingEvent.EndTime);
diff --git a/Engine/Source/Developer/TraceInsights/Public/Insights/ViewModels/TimingEvent.h b/Engine/Source/Developer/TraceInsights/Public/Insights/ViewModels/TimingEvent.h
--- a/Engine/Source/Developer/TraceInsights/Public/Insights/ViewModels/TimingEvent.h
+++ b/Engine/Source/Developer/TraceInsights/Public/Insights/ViewModels/TimingEvent.h
@@ -9,37 +9,28 @@ class FTimingEventsTrack;
 
 struct FTimingEvent
 {
-	static const uint64 InvalidTypeId = uint64(-1);
-
 	const FTimingEventsTrack* Track;
-	uint64 TypeId;
-	uint32 Depth;
 	double StartTime;
 	double EndTime;
 	double ExclusiveTime;
-
-	// For Asset Loading events...
-	Trace::FLoadTimeProfilerCpuEvent LoadingInfo;
-
-	// For I/O events...
-	uint64 Offset;
-	uint64 Size;
-	const TCHAR* Path;
+	uint32 Depth;
 
 	FTimingEvent()
 		: Track(nullptr)
-		, TypeId(InvalidTypeId)
-		, Depth(0)
-		, StartTime(0.0)
-		, EndTime(0.0)
+		, StartTime(0.0f)
+		, EndTime(-1.0)
 		, ExclusiveTime(0.0)
-		, LoadingInfo()
-		, Offset(0)
-		, Size(0)
-		, Path(nullptr)
+		, Depth(0)
 	{
 	}
 
+	FTimingEvent(const FTimingEventsTrack* InTrack, double InStartTime, double InEndTime, uint32 InDepth)
+		: Track(InTrack)
+		, StartTime(InStartTime)
+		, EndTime(InEndTime)
+		, Depth(InDepth)
+	{}
+
 	FTimingEvent(const FTimingEvent&) = default;
 	FTimingEvent& operator=(const FTimingEvent&) = default;
 
@@ -49,31 +40,21 @@ struct FTimingEvent
 	void Reset()
 	{
 		Track = nullptr;
-		TypeId = InvalidTypeId;
-		Depth = 0;
 		StartTime = 0.0;
-		EndTime = 0.0;
+		EndTime = -1.0;
 		ExclusiveTime = 0.0;
-
-		LoadingInfo.Package = nullptr;
-		LoadingInfo.Export = nullptr;
-		LoadingInfo.EventType = Trace::LoadTimeProfilerObjectEventType_None;
-
-		Offset = 0;
-		Size = 0;
-		Path = nullptr;
+		Depth = 0;
 	}
 
 	bool IsValidTrack() const { return Track != nullptr; }
 
-	bool IsValid() const { return Track != nullptr && TypeId != InvalidTypeId; }
+	bool IsValid() const { return Track != nullptr && StartTime <= EndTime; }
 
 	double Duration() const { return EndTime - StartTime; }
 
 	bool Equals(const FTimingEvent& Other) const
 	{
 		return Track == Other.Track
-			&& TypeId == Other.TypeId
 			&& Depth == Other.Depth
 			&& StartTime == Other.StartTime
 			&& EndTime == Other.EndTime;
diff --git a/Engine/Source/Developer/TraceInsights/Public/Insights/ViewModels/TimingEventSearch.h b/Engine/Source/Developer/TraceInsights/Public/Insights/ViewModels/TimingEventSearch.h
new file mode 100644
--- /dev/null
+++ b/Engine/Source/Developer/TraceInsights/Public/Insights/ViewModels/TimingEventSearch.h
@@ -0,0 +1,197 @@
+// Copyright 1998-2019 Epic Games, Inc. All Rights Reserved.
+
+#pragma once
+
+#include "CoreMinimal.h"
+
+// Search behavior flags
+enum class ETimingEventSearchFlags : int32
+{
+	// Search all matches
+	None = 0,
+	SearchAll = None,
+
+	// Whether to stop at the first match
+	StopAtFirstMatch = (1 << 0),
+};
+
+ENUM_CLASS_FLAGS(ETimingEventSearchFlags);
+
+// Parameters for a timing event search
+class FTimingEventSearchParameters
+{
+public:
+	// Predicate called to filter event matches. 
+	// Note it is assumed that the SearchPredicate (below) will perform basic range checks on StartTime and EndTime,
+	// therefore this is for any other filtering that needs to take place (e.g. left of an event, below an event etc.)
+	// Returns true to pass the filter.
+	typedef TFunctionRef<bool(double /*StartTime*/, double /*EndTime*/, uint32 /*Depth*/)> EventFilterPredicate;
+
+	// Predicate called when we get a match
+	typedef TFunctionRef<void(double /*StartTime*/, double /*EndTime*/, uint32 /*Depth*/)> EventMatchedPredicate;
+
+	FTimingEventSearchParameters(double InStartTime, double InEndTime, ETimingEventSearchFlags Flags, EventFilterPredicate InEventFilter = NoFilter, EventMatchedPredicate InEventMatched = NoMatch)
+		: EventFilter(InEventFilter)
+		, EventMatched(InEventMatched)
+		, StartTime(InStartTime)
+		, EndTime(InEndTime)
+		, Flags(Flags)
+	{}
+
+private:
+	// Default predicates
+	static bool NoFilter(double, double, uint32) { return true; }
+	static void NoMatch(double, double, uint32) {}
+
+public:
+	// Predicate for search filtering
+	EventFilterPredicate EventFilter;
+
+	// Predicate called when we get a match
+	EventMatchedPredicate EventMatched;
+
+	// Start time of the search
+	double StartTime;
+
+	// End time of the search
+	double EndTime;
+
+	// Search behavior flags
+	ETimingEventSearchFlags Flags;
+};
+
+// Helper used to orchestrate a search of a timing event track's events
+// Example of usage:
+// 
+//		FMyStruct MatchedPayload;
+//
+//		TTimingEventSearch<FMyStruct>::Search(
+//		Parameters,
+// 		[](const TTimingEventSearch<FMyStruct>::FContext& InContext)
+// 		{
+// 			for(FMyStruct& Payload : Payloads)
+//			{
+// 				InContext.Check(InEventStartTime, InEventEndTime, 0, Payload);
+// 			}
+// 		},
+// 		[&MatchedPayload](double InStartTime, double InEndTime, uint32 InDepth, const FMyStruct& InEvent)
+// 		{
+// 			MatchedPayload = InEvent;
+// 		},
+// 		[&MatchedPayload](double InStartTime, double InEndTime, uint32 InDepth)
+// 		{
+// 			// Do something with MatchedPayload, e.g. call a captured lambda.
+// 		});
+// 
+template <typename PayloadType>
+struct TTimingEventSearch
+{
+public:
+	struct FContext;
+
+	// Predicate called when a match has been made to the search parameters
+	typedef TFunctionRef<void(double /*InStartTime*/, double /*InEndTime*/, uint32 /*InDepth*/, const PayloadType& /*InPayload*/)> PayloadMatchedPredicate;
+
+	// Predicate called to filter by payload contents
+	// Return true to pass the filer
+	typedef TFunctionRef<bool(double /*InStartTime*/, double /*InEndTime*/, uint32 /*InDepth*/, const PayloadType& /*InPayload*/)> PayloadFilterPredicate;
+
+	// Predicate called when a match has been found. Note does not include payload as it is expected this will be captured/stored externally.
+	typedef TFunctionRef<void(double /*InStartTime*/, double /*InEndTime*/, uint32 /*InDepth*/)> FoundPredicate;
+
+	// Predicate called to run the search, e.g. iterate over an array of events
+	// It is expected to call FContext::Check on each valid searched event.
+	typedef TFunctionRef<void(FContext& /*Context*/)> SearchPredicate;
+
+	// Context used to operate a search.
+	struct FContext
+	{
+	public:
+		FContext(const FTimingEventSearchParameters& InParameters, PayloadFilterPredicate InPayloadFilterPredicate, PayloadMatchedPredicate InPayloadMatchedPredicate)
+			: Parameters(InParameters)
+			, PayloadMatched(InPayloadMatchedPredicate)
+			, PayloadFilter(InPayloadFilterPredicate)
+			, FoundStartTime(-1.0)
+			, FoundEndTime(-1.0)
+			, FoundDepth(0)
+			, bFound(false)
+			, bContinueSearching(true)
+		{
+		}
+
+		// Function called to check and potentially match an event
+		void Check(double InEventStartTime, double InEventEndTime, uint32 InEventDepth, const PayloadType& InEvent)
+		{
+			if (bContinueSearching && Parameters.EventFilter(InEventStartTime, InEventEndTime, InEventDepth) && PayloadFilter(InEventStartTime, InEventEndTime, InEventDepth, InEvent))
+			{
+				Parameters.EventMatched(InEventStartTime, InEventEndTime, InEventDepth);
+				PayloadMatched(InEventStartTime, InEventEndTime, InEventDepth, InEvent);
+
+				FoundDepth = InEventDepth;
+				FoundStartTime = InEventStartTime;
+				FoundEndTime = InEventEndTime;
+
+				bFound = true;
+				bContinueSearching = (Parameters.Flags & ETimingEventSearchFlags::StopAtFirstMatch) == ETimingEventSearchFlags::None;
+			}
+		}
+
+		// Access the search parameters
+		const FTimingEventSearchParameters& GetParameters() const { return Parameters; }
+
+		// Accessors for read-only results
+		double GetStartTimeFound() const { return FoundStartTime; }
+		double GetEndTimeFound() const { return FoundEndTime; }
+		uint32 GetDepthFound() const { return FoundDepth; }
+		bool IsMatchFound() const { return bFound; }
+		bool ShouldContinueSearching() const { return bContinueSearching; }
+
+	private:
+		// Search parameters
+		FTimingEventSearchParameters Parameters;
+
+		// Predicate called when an event was matched
+		PayloadMatchedPredicate PayloadMatched;
+
+		// Filter applied to payloads
+		PayloadFilterPredicate PayloadFilter;
+
+		// The start time of the event that was found
+		double FoundStartTime;
+
+		// The end time of the event that was found
+		double FoundEndTime;
+
+		// The depth of the event that was found
+		uint32 FoundDepth;
+
+		// Whether a match was found
+		bool bFound;
+
+		// Internal flag to skip work if a match was found
+		bool bContinueSearching;
+	};
+
+public:
+	// Search using only the event filter
+	static bool Search(const FTimingEventSearchParameters& InParameters, SearchPredicate InSearchPredicate, PayloadMatchedPredicate InPayloadMatchedPredicate, FoundPredicate InFoundPredicate)
+	{
+		auto NoFilter = [](double, double, uint32, const PayloadType&){ return true; };
+		return Search(InParameters, InSearchPredicate, NoFilter, InPayloadMatchedPredicate, InFoundPredicate);
+	}
+
+	// Search using a specific payload filter
+	static bool Search(const FTimingEventSearchParameters& InParameters, SearchPredicate InSearchPredicate, PayloadFilterPredicate InFilterPredicate, PayloadMatchedPredicate InPayloadMatchedPredicate, FoundPredicate InFoundPredicate)
+	{
+		FContext Context(InParameters, InFilterPredicate, InPayloadMatchedPredicate);
+
+		InSearchPredicate(Context);
+
+		if(Context.IsMatchFound())
+		{
+			InFoundPredicate(Context.GetStartTimeFound(), Context.GetEndTimeFound(), Context.GetDepthFound());
+		}
+
+		return Context.IsMatchFound();
+	}
+};
\ No newline at end of file
diff --git a/Engine/Source/Developer/TraceInsights/Public/Insights/ViewModels/TimingEventsTrack.h b/Engine/Source/Developer/TraceInsights/Public/Insights/ViewModels/TimingEventsTrack.h
--- a/Engine/Source/Developer/TraceInsights/Public/Insights/ViewModels/TimingEventsTrack.h
+++ b/Engine/Source/Developer/TraceInsights/Public/Insights/ViewModels/TimingEventsTrack.h
@@ -16,6 +16,7 @@ struct FTimingViewTooltip;
 class ITimingViewDrawHelper;
 class FTooltipDrawState;
 namespace Trace { class IAnalysisSession; }
+class FTimingEventSearchParameters;
 
 ////////////////////////////////////////////////////////////////////////////////////////////////////
 
@@ -71,12 +72,15 @@ public:
 	virtual void DrawSelectedEventInfo(const FTimingEvent& SelectedTimingEvent, const FTimingTrackViewport& Viewport, const FDrawContext& DrawContext, const FSlateBrush* WhiteBrush, const FSlateFontInfo& Font) const {}
 	virtual void InitTooltip(FTooltipDrawState& Tooltip, const FTimingEvent& HoveredTimingEvent) const {}
 
-	virtual bool SearchTimingEvent(const double InStartTime, const double InEndTime, TFunctionRef<bool(double, double, uint32)> InPredicate, FTimingEvent& InOutTimingEvent, bool bInStopAtFirstMatch, bool bInSearchForLargestEvent) const { return false; }
+	virtual bool SearchTimingEvent(const FTimingEventSearchParameters& InSearchParameters, FTimingEvent& InOutTimingEvent) const { return false; }
 	virtual void ComputeTimingEventStats(FTimingEvent& InOutTimingEvent) const {}
 
 	// Called back from the timing view when an event is selected
 	virtual void OnEventSelected(const FTimingEvent& SelectedTimingEvent) const {}
 
+	// Called back from the timing view when a selected event is copied to the clipboard with Ctrl+C
+	virtual void OnClipboardCopyEvent(const FTimingEvent& SelectedTimingEvent) const {}
+
 private:
 	int32 NumLanes; // number of lanes (sub-tracks)
 
diff --git a/Engine/Source/Developer/TraceServices/Public/TraceServices/Model/TimingProfiler.h b/Engine/Source/Developer/TraceServices/Public/TraceServices/Model/TimingProfiler.h
--- a/Engine/Source/Developer/TraceServices/Public/TraceServices/Model/TimingProfiler.h
+++ b/Engine/Source/Developer/TraceServices/Public/TraceServices/Model/TimingProfiler.h
@@ -12,15 +12,15 @@ namespace Trace
 
 struct FTimingProfilerTimer
 {
-	const TCHAR* Name;
-	uint32 Id;
-	uint32 NameHash;
-	bool IsGpuTimer;
+	const TCHAR* Name = nullptr;
+	uint32 Id = 0;
+	uint32 NameHash = 0;
+	bool IsGpuTimer = false;
 };
 
 struct FTimingProfilerEvent
 {
-	uint32 TimerIndex;
+	uint32 TimerIndex = uint32(-1);
 };
 
 struct FTimingProfilerAggregatedStats
