Added optional caching to event searches

Tweaked API to not require external storage of matched events. Made the matched calback optional as it isnt needed any more.

diff --git a/Engine/Plugins/Animation/GameplayInsights/Source/GameplayInsights/Private/AnimationTickRecordsTrack.cpp b/Engine/Plugins/Animation/GameplayInsights/Source/GameplayInsights/Private/AnimationTickRecordsTrack.cpp
--- a/Engine/Plugins/Animation/GameplayInsights/Source/GameplayInsights/Private/AnimationTickRecordsTrack.cpp
+++ b/Engine/Plugins/Animation/GameplayInsights/Source/GameplayInsights/Private/AnimationTickRecordsTrack.cpp
@@ -89,13 +89,9 @@ bool FAnimationTickRecordsTrack::SearchTimingEvent(const FTimingEventSearchParam
 
 bool FAnimationTickRecordsTrack::FindTickRecordMessage(const FTimingEventSearchParameters& InParameters, TFunctionRef<void(double, double, uint32, const FTickRecordMessage&)> InFoundPredicate) const
 {
-	// Storage for the message we want to match (payload for an event)
-	FTickRecordMessage MatchedMessage;
-
 	return TTimingEventSearch<FTickRecordMessage>::Search(
 		InParameters,
 
-		// Search...
 		[this](TTimingEventSearch<FTickRecordMessage>::FContext& InContext)
 		{
 			const FAnimationProvider* AnimationProvider = SharedData.GetAnalysisSession().ReadProvider<FAnimationProvider>(FAnimationProvider::ProviderName);
@@ -114,16 +110,9 @@ bool FAnimationTickRecordsTrack::FindTickRecordMessage(const FTimingEventSearchP
 			}
 		},
 
-		// Matched...
-		[&MatchedMessage](double InStartTime, double InEndTime, uint32 InDepth, const FTickRecordMessage& InEvent)
-		{
-			MatchedMessage = InEvent;
-		},
-
-		// Found!
-		[&InFoundPredicate, &MatchedMessage](double InFoundStartTime, double InFoundEndTime, uint32 InFoundDepth)
+		[&InFoundPredicate](double InFoundStartTime, double InFoundEndTime, uint32 InFoundDepth, const FTickRecordMessage& InEvent)
 		{
-			InFoundPredicate(InFoundStartTime, InFoundEndTime, InFoundDepth, MatchedMessage);
+			InFoundPredicate(InFoundStartTime, InFoundEndTime, InFoundDepth, InEvent);
 		});
 }
 
diff --git a/Engine/Plugins/Animation/GameplayInsights/Source/GameplayInsights/Private/ObjectEventsTrack.cpp b/Engine/Plugins/Animation/GameplayInsights/Source/GameplayInsights/Private/ObjectEventsTrack.cpp
--- a/Engine/Plugins/Animation/GameplayInsights/Source/GameplayInsights/Private/ObjectEventsTrack.cpp
+++ b/Engine/Plugins/Animation/GameplayInsights/Source/GameplayInsights/Private/ObjectEventsTrack.cpp
@@ -72,9 +72,6 @@ bool FObjectEventsTrack::SearchTimingEvent(const FTimingEventSearchParameters& I
 
 bool FObjectEventsTrack::FindObjectEvent(const FTimingEventSearchParameters& InParameters, TFunctionRef<void(double, double, uint32, const FObjectEventMessage&)> InFoundPredicate) const
 {
-	// Storage for the message we want to match (payload for an event)
-	FObjectEventMessage MatchedMessage;
-
 	return TTimingEventSearch<FObjectEventMessage>::Search(
 		InParameters,
 
@@ -97,16 +94,10 @@ bool FObjectEventsTrack::FindObjectEvent(const FTimingEventSearchParameters& InP
 			}
 		},
 
-		// Matched...
-		[&MatchedMessage](double InStartTime, double InEndTime, uint32 InDepth, const FObjectEventMessage& InEvent)
-		{
-			MatchedMessage = InEvent;
-		},
-
 		// Found!
-		[&InFoundPredicate, &MatchedMessage](double InFoundStartTime, double InFoundEndTime, uint32 InFoundDepth)
+		[&InFoundPredicate](double InFoundStartTime, double InFoundEndTime, uint32 InFoundDepth, const FObjectEventMessage& InEvent)
 		{
-			InFoundPredicate(InFoundStartTime, InFoundEndTime, InFoundDepth, MatchedMessage);
+			InFoundPredicate(InFoundStartTime, InFoundEndTime, InFoundDepth, InEvent);
 		});
 }
 
diff --git a/Engine/Plugins/Animation/GameplayInsights/Source/GameplayInsights/Private/SkeletalMeshPoseTrack.cpp b/Engine/Plugins/Animation/GameplayInsights/Source/GameplayInsights/Private/SkeletalMeshPoseTrack.cpp
--- a/Engine/Plugins/Animation/GameplayInsights/Source/GameplayInsights/Private/SkeletalMeshPoseTrack.cpp
+++ b/Engine/Plugins/Animation/GameplayInsights/Source/GameplayInsights/Private/SkeletalMeshPoseTrack.cpp
@@ -100,13 +100,9 @@ bool FSkeletalMeshPoseTrack::SearchTimingEvent(const FTimingEventSearchParameter
 
 bool FSkeletalMeshPoseTrack::FindSkeletalMeshPoseMessage(const FTimingEventSearchParameters& InParameters, TFunctionRef<void(double, double, uint32, const FSkeletalMeshPoseMessage&)> InFoundPredicate) const
 {
-	// Storage for the message we want to match (payload for an event)
-	FSkeletalMeshPoseMessage MatchedMessage;
-
 	return TTimingEventSearch<FSkeletalMeshPoseMessage>::Search(
 		InParameters,
 
-		// Search...
 		[this](TTimingEventSearch<FSkeletalMeshPoseMessage>::FContext& InContext)
 		{
 			const FAnimationProvider* AnimationProvider = SharedData.GetAnalysisSession().ReadProvider<FAnimationProvider>(FAnimationProvider::ProviderName);
@@ -125,16 +121,9 @@ bool FSkeletalMeshPoseTrack::FindSkeletalMeshPoseMessage(const FTimingEventSearc
 			}
 		},
 
-		// Matched...
-		[&MatchedMessage](double InStartTime, double InEndTime, uint32 InDepth, const FSkeletalMeshPoseMessage& InEvent)
-		{
-			MatchedMessage = InEvent;
-		},
-
-		// Found!
-		[&InFoundPredicate, &MatchedMessage](double InFoundStartTime, double InFoundEndTime, uint32 InFoundDepth)
+		[&InFoundPredicate](double InFoundStartTime, double InFoundEndTime, uint32 InFoundDepth, const FSkeletalMeshPoseMessage& InEvent)
 		{
-			InFoundPredicate(InFoundStartTime, InFoundEndTime, InFoundDepth, MatchedMessage);
+			InFoundPredicate(InFoundStartTime, InFoundEndTime, InFoundDepth, InEvent);
 		});
 }
 
diff --git a/Engine/Source/Developer/TraceInsights/Private/Insights/ViewModels/FileActivityTimingTrack.cpp b/Engine/Source/Developer/TraceInsights/Private/Insights/ViewModels/FileActivityTimingTrack.cpp
--- a/Engine/Source/Developer/TraceInsights/Private/Insights/ViewModels/FileActivityTimingTrack.cpp
+++ b/Engine/Source/Developer/TraceInsights/Private/Insights/ViewModels/FileActivityTimingTrack.cpp
@@ -448,9 +448,6 @@ void FFileActivityTimingTrack::InitTooltip(FTooltipDrawState& Tooltip, const FTi
 
 bool FFileActivityTimingTrack::FindIoTimingEvent(const FTimingEventSearchParameters& InParameters, bool bIgnoreEventDepth, TFunctionRef<void(double, double, uint32, const FFileActivitySharedState::FIoTimingEvent&)> InFoundPredicate) const
 {
-	// Storage for the event we want to match
-	FFileActivitySharedState::FIoTimingEvent MatchedEvent;
-
 	return TTimingEventSearch<FFileActivitySharedState::FIoTimingEvent>::Search(
 		InParameters,
 
@@ -479,16 +476,10 @@ bool FFileActivityTimingTrack::FindIoTimingEvent(const FTimingEventSearchParamet
 			}
 		},
 
-		// Matched...
-		[&MatchedEvent](double InStartTime, double InEndTime, uint32 InDepth, const FFileActivitySharedState::FIoTimingEvent& InEvent)
-		{
-			MatchedEvent = InEvent;
-		},
-
 		// Found!
-		[&InFoundPredicate, &MatchedEvent](double InFoundStartTime, double InFoundEndTime, uint32 InFoundDepth)
+		[&InFoundPredicate](double InFoundStartTime, double InFoundEndTime, uint32 InFoundDepth, const FFileActivitySharedState::FIoTimingEvent& InEvent)
 		{
-			InFoundPredicate(InFoundStartTime, InFoundEndTime, InFoundDepth, MatchedEvent);
+			InFoundPredicate(InFoundStartTime, InFoundEndTime, InFoundDepth, InEvent);
 		});
 }
 
diff --git a/Engine/Source/Developer/TraceInsights/Private/Insights/ViewModels/LoadingTimingTrack.cpp b/Engine/Source/Developer/TraceInsights/Private/Insights/ViewModels/LoadingTimingTrack.cpp
--- a/Engine/Source/Developer/TraceInsights/Private/Insights/ViewModels/LoadingTimingTrack.cpp
+++ b/Engine/Source/Developer/TraceInsights/Private/Insights/ViewModels/LoadingTimingTrack.cpp
@@ -189,13 +189,9 @@ bool FLoadingTimingTrack::SearchTimingEvent(const FTimingEventSearchParameters&
 
 bool FLoadingTimingTrack::FindLoadTimeProfilerCpuEvent(const FTimingEventSearchParameters& InParameters, TFunctionRef<void(double, double, uint32, const Trace::FLoadTimeProfilerCpuEvent&)> InFoundPredicate) const
 {
-	// Storage for the event we want to match
-	Trace::FLoadTimeProfilerCpuEvent MatchedEvent;
-
 	return TTimingEventSearch<Trace::FLoadTimeProfilerCpuEvent>::Search(
 		InParameters,
 
-		// Search...
 		[this](TTimingEventSearch<Trace::FLoadTimeProfilerCpuEvent>::FContext& InContext)
 		{
 			TSharedPtr<const Trace::IAnalysisSession> Session = FInsightsManager::Get()->GetSession();
@@ -218,16 +214,9 @@ bool FLoadingTimingTrack::FindLoadTimeProfilerCpuEvent(const FTimingEventSearchP
 			}
 		},
 
-		// Matched...
-		[&MatchedEvent](double InStartTime, double InEndTime, uint32 InDepth, const Trace::FLoadTimeProfilerCpuEvent& InEvent)
-		{
-			MatchedEvent = InEvent;
-		},
-
-		// Found!
-		[&InFoundPredicate, &MatchedEvent](double InFoundStartTime, double InFoundEndTime, uint32 InFoundDepth)
+		[&InFoundPredicate](double InFoundStartTime, double InFoundEndTime, uint32 InFoundDepth, const Trace::FLoadTimeProfilerCpuEvent& InEvent)
 		{
-			InFoundPredicate(InFoundStartTime, InFoundEndTime, InFoundDepth, MatchedEvent);
+			InFoundPredicate(InFoundStartTime, InFoundEndTime, InFoundDepth, InEvent);
 		});
 }
 
diff --git a/Engine/Source/Developer/TraceInsights/Private/Insights/ViewModels/ThreadTimingTrack.cpp b/Engine/Source/Developer/TraceInsights/Private/Insights/ViewModels/ThreadTimingTrack.cpp
--- a/Engine/Source/Developer/TraceInsights/Private/Insights/ViewModels/ThreadTimingTrack.cpp
+++ b/Engine/Source/Developer/TraceInsights/Private/Insights/ViewModels/ThreadTimingTrack.cpp
@@ -336,6 +336,7 @@ bool FThreadTimingTrack::FindTimingProfilerEvent(const FTimingEvent& InTimingEve
 	};
 
 	FTimingEventSearchParameters SearchParameters(InTimingEvent.StartTime, InTimingEvent.EndTime, ETimingEventSearchFlags::StopAtFirstMatch, MatchDepth);
+	SearchParameters.SearchHandle = &InTimingEvent.SearchHandle;
 	return FindTimingProfilerEvent(SearchParameters, InFoundPredicate);
 }
 
@@ -343,13 +344,9 @@ bool FThreadTimingTrack::FindTimingProfilerEvent(const FTimingEvent& InTimingEve
 
 bool FThreadTimingTrack::FindTimingProfilerEvent(const FTimingEventSearchParameters& InParameters, TFunctionRef<void(double, double, uint32, const Trace::FTimingProfilerEvent&)> InFoundPredicate) const
 {
-	// Storage for the event we want to match
-	Trace::FTimingProfilerEvent MatchedEvent;
-
 	return TTimingEventSearch<Trace::FTimingProfilerEvent>::Search(
 		InParameters,
 
-		// Search...
 		[this](TTimingEventSearch<Trace::FTimingProfilerEvent>::FContext& InContext)
 		{
 			TSharedPtr<const Trace::IAnalysisSession> Session = FInsightsManager::Get()->GetSession();
@@ -372,17 +369,12 @@ bool FThreadTimingTrack::FindTimingProfilerEvent(const FTimingEventSearchParamet
 			}
 		},
 
-		// Matched...
-		[&MatchedEvent](double InStartTime, double InEndTime, uint32 InDepth, const Trace::FTimingProfilerEvent& InEvent)
+		[&InFoundPredicate](double InFoundStartTime, double InFoundEndTime, uint32 InFoundDepth, const Trace::FTimingProfilerEvent& InEvent)
 		{
-			MatchedEvent = InEvent;
+			InFoundPredicate(InFoundStartTime, InFoundEndTime, InFoundDepth, InEvent);
 		},
-
-		// Found!
-		[&InFoundPredicate, &MatchedEvent](double InFoundStartTime, double InFoundEndTime, uint32 InFoundDepth)
-		{
-			InFoundPredicate(InFoundStartTime, InFoundEndTime, InFoundDepth, MatchedEvent);
-		});
+		
+		SearchCache);
 }
 
 
diff --git a/Engine/Source/Developer/TraceInsights/Private/Insights/ViewModels/ThreadTimingTrack.h b/Engine/Source/Developer/TraceInsights/Private/Insights/ViewModels/ThreadTimingTrack.h
--- a/Engine/Source/Developer/TraceInsights/Private/Insights/ViewModels/ThreadTimingTrack.h
+++ b/Engine/Source/Developer/TraceInsights/Private/Insights/ViewModels/ThreadTimingTrack.h
@@ -6,8 +6,9 @@
 
 // Insights
 #include "Insights/ViewModels/TimingEventsTrack.h"
+#include "Insights/ViewModels/TimingEventSearch.h"
+#include "TraceServices/Model/TimingProfiler.h"
 
-namespace Trace { struct FTimingProfilerEvent; }
 class FTimingEventSearchParameters;
 
 ////////////////////////////////////////////////////////////////////////////////////////////////////
@@ -47,6 +48,9 @@ private:
 private:
 	const TCHAR* GroupName;
 	uint32 ThreadId;
+
+	// Search cache
+	mutable TTimingEventSearchCache<Trace::FTimingProfilerEvent> SearchCache;
 };
 
 ////////////////////////////////////////////////////////////////////////////////////////////////////
diff --git a/Engine/Source/Developer/TraceInsights/Public/Insights/ViewModels/TimingEvent.h b/Engine/Source/Developer/TraceInsights/Public/Insights/ViewModels/TimingEvent.h
--- a/Engine/Source/Developer/TraceInsights/Public/Insights/ViewModels/TimingEvent.h
+++ b/Engine/Source/Developer/TraceInsights/Public/Insights/ViewModels/TimingEvent.h
@@ -4,20 +4,33 @@
 
 #include "CoreMinimal.h"
 #include "TraceServices/Model/LoadTimeProfiler.h"
+#include "TimingEventSearch.h"
 
 class FTimingEventsTrack;
 
 struct FTimingEvent
 {
+	// The track this timing event is contained within
 	const FTimingEventsTrack* Track;
+
+	// Handle to a previous search, used to accelerate access to underlying event data
+	mutable FTimingEventSearchHandle SearchHandle;
+
+	// The start time of the event
 	double StartTime;
+
+	// The end time of the event
 	double EndTime;
+
+	// For hierarchical events, the cached exclusive time
 	double ExclusiveTime;
+
+	// The depth of the event
 	uint32 Depth;
 
 	FTimingEvent()
 		: Track(nullptr)
-		, StartTime(0.0f)
+		, StartTime(0.0)
 		, EndTime(-1.0)
 		, ExclusiveTime(0.0)
 		, Depth(0)
@@ -28,6 +41,7 @@ struct FTimingEvent
 		: Track(InTrack)
 		, StartTime(InStartTime)
 		, EndTime(InEndTime)
+		, ExclusiveTime(0.0)
 		, Depth(InDepth)
 	{}
 
diff --git a/Engine/Source/Developer/TraceInsights/Public/Insights/ViewModels/TimingEventSearch.h b/Engine/Source/Developer/TraceInsights/Public/Insights/ViewModels/TimingEventSearch.h
--- a/Engine/Source/Developer/TraceInsights/Public/Insights/ViewModels/TimingEventSearch.h
+++ b/Engine/Source/Developer/TraceInsights/Public/Insights/ViewModels/TimingEventSearch.h
@@ -3,6 +3,7 @@
 #pragma once
 
 #include "CoreMinimal.h"
+#include "Containers/Map.h"
 
 // Search behavior flags
 enum class ETimingEventSearchFlags : int32
@@ -17,6 +18,19 @@ enum class ETimingEventSearchFlags : int32
 
 ENUM_CLASS_FLAGS(ETimingEventSearchFlags);
 
+// A handle to a timing event that was previously searched
+struct FTimingEventSearchHandle
+{
+	template<typename T, int32 N>
+	friend struct TTimingEventSearchCache;
+
+	bool IsValid() const { return Id != uint64(-1); }
+	void Reset() { Id = uint64(-1); }
+
+private:
+	uint64 Id = uint64(-1);
+};
+
 // Parameters for a timing event search
 class FTimingEventSearchParameters
 {
@@ -33,6 +47,7 @@ public:
 	FTimingEventSearchParameters(double InStartTime, double InEndTime, ETimingEventSearchFlags Flags, EventFilterPredicate InEventFilter = NoFilter, EventMatchedPredicate InEventMatched = NoMatch)
 		: EventFilter(InEventFilter)
 		, EventMatched(InEventMatched)
+		, SearchHandle(nullptr)
 		, StartTime(InStartTime)
 		, EndTime(InEndTime)
 		, Flags(Flags)
@@ -50,6 +65,9 @@ public:
 	// Predicate called when we get a match
 	EventMatchedPredicate EventMatched;
 
+	// A handle to a previous search. This will be written during a search if a cache is utilized and a hit was not generated
+	FTimingEventSearchHandle* SearchHandle;
+
 	// Start time of the search
 	double StartTime;
 
@@ -60,6 +78,90 @@ public:
 	ETimingEventSearchFlags Flags;
 };
 
+// Simple acceleration structure used to return previously searched results
+template <typename PayloadType, int32 Size = 3>
+struct TTimingEventSearchCache
+{
+public:
+	TTimingEventSearchCache()
+		: CurrentHandle(0)
+		, WriteIndex(0)
+	{}
+
+	// Write to the cache
+	// @return a handle to the search
+	FTimingEventSearchHandle Write(double InStartTime, double InEndTime, uint32 InDepth, const PayloadType& InPayload)
+	{
+		TPair<FTimingEventSearchHandle, FResultData>& WritePair = CachedValues[WriteIndex];
+
+		WritePair.Key.Id = CurrentHandle++;
+		WritePair.Value.Payload = InPayload;
+		WritePair.Value.StartTime = InStartTime;
+		WritePair.Value.EndTime = InEndTime;
+		WritePair.Value.Depth = InDepth;
+
+		WriteIndex = (WriteIndex + 1) % Size;
+
+		return WritePair.Key;
+	}
+
+	// Attempt to read from the cache.
+	// @return false if the value was not found
+	bool Read(const FTimingEventSearchHandle& InHandle, double& OutStartTime, double& OutEndTime, uint32& OutDepth, PayloadType& OutPayload)
+	{
+		static_assert(Size > 1, "TTimingEventSearchCache only works with caches that have multiple values");
+
+		if(InHandle.IsValid())
+		{
+			// Start read at the previous write, to catch most recently written values
+			for(int32 ReadIndex = (WriteIndex + (Size - 1)) % Size, ReadCount = 0; ReadCount < Size; ReadCount++, ReadIndex = (ReadIndex + 1) % Size)
+			{
+				check(ReadIndex >= 0 && ReadIndex < Size);
+
+				const TPair<FTimingEventSearchHandle, FResultData>& CachedPair = CachedValues[ReadIndex];
+				if(InHandle.Id == CachedPair.Key.Id)
+				{
+					OutStartTime = CachedPair.Value.StartTime;
+					OutEndTime = CachedPair.Value.EndTime;
+					OutDepth = CachedPair.Value.Depth;
+					OutPayload = CachedPair.Value.Payload;
+					return true;
+				}
+			}
+		}
+
+		return false;
+	}
+
+	// Reset the cache. Any reads after this will return false
+	void Reset()
+	{
+		for(TPair<FTimingEventSearchHandle, FResultData>& CachedPair : CachedValues)
+		{
+			CachedPair.Key.Reset();
+		}
+	}
+
+private:
+	// Current handle
+	uint64 CurrentHandle;
+
+	// Current write index
+	int32 WriteIndex;
+
+	// A cached result
+	struct FResultData
+	{
+		PayloadType Payload;
+		double StartTime;
+		double EndTime;
+		uint32 Depth;
+	};
+
+	// Array of cached searches
+	TPair<FTimingEventSearchHandle, FResultData> CachedValues[Size];
+};
+
 // Helper used to orchestrate a search of a timing event track's events
 // Example of usage:
 // 
@@ -96,8 +198,8 @@ public:
 	// Return true to pass the filer
 	typedef TFunctionRef<bool(double /*InStartTime*/, double /*InEndTime*/, uint32 /*InDepth*/, const PayloadType& /*InPayload*/)> PayloadFilterPredicate;
 
-	// Predicate called when a match has been found. Note does not include payload as it is expected this will be captured/stored externally.
-	typedef TFunctionRef<void(double /*InStartTime*/, double /*InEndTime*/, uint32 /*InDepth*/)> FoundPredicate;
+	// Predicate called when a match has been found.
+	typedef TFunctionRef<void(double /*InStartTime*/, double /*InEndTime*/, uint32 /*InDepth*/, const PayloadType& /*InPayload*/)> FoundPredicate;
 
 	// Predicate called to run the search, e.g. iterate over an array of events
 	// It is expected to call FContext::Check on each valid searched event.
@@ -127,6 +229,7 @@ public:
 				Parameters.EventMatched(InEventStartTime, InEventEndTime, InEventDepth);
 				PayloadMatched(InEventStartTime, InEventEndTime, InEventDepth, InEvent);
 
+				FoundPayload = InEvent;
 				FoundDepth = InEventDepth;
 				FoundStartTime = InEventStartTime;
 				FoundEndTime = InEventEndTime;
@@ -140,12 +243,16 @@ public:
 		const FTimingEventSearchParameters& GetParameters() const { return Parameters; }
 
 		// Accessors for read-only results
+		const PayloadType& GetPayloadFound() const { return FoundPayload; }
 		double GetStartTimeFound() const { return FoundStartTime; }
 		double GetEndTimeFound() const { return FoundEndTime; }
 		uint32 GetDepthFound() const { return FoundDepth; }
 		bool IsMatchFound() const { return bFound; }
 		bool ShouldContinueSearching() const { return bContinueSearching; }
 
+		// Allows search to be aborted by predicates
+		void AbortSearch() { bContinueSearching = false; }
+
 	private:
 		// Search parameters
 		FTimingEventSearchParameters Parameters;
@@ -156,6 +263,9 @@ public:
 		// Filter applied to payloads
 		PayloadFilterPredicate PayloadFilter;
 
+		// The payload we have found
+		PayloadType FoundPayload;
+
 		// The start time of the event that was found
 		double FoundStartTime;
 
@@ -173,25 +283,55 @@ public:
 	};
 
 public:
-	// Search using only the event filter
-	static bool Search(const FTimingEventSearchParameters& InParameters, SearchPredicate InSearchPredicate, PayloadMatchedPredicate InPayloadMatchedPredicate, FoundPredicate InFoundPredicate)
+	// Default predicates
+	static bool NoFilter(double, double, uint32, const PayloadType&) { return true; }
+	static void NoMatch(double, double, uint32, const PayloadType&) {}
+
+public:
+	// Search using only the event filter and no match predicate
+	static bool Search(const FTimingEventSearchParameters& InParameters, SearchPredicate InSearchPredicate, FoundPredicate InFoundPredicate)
 	{
-		auto NoFilter = [](double, double, uint32, const PayloadType&){ return true; };
-		return Search(InParameters, InSearchPredicate, NoFilter, InPayloadMatchedPredicate, InFoundPredicate);
+		return Search(InParameters, InSearchPredicate, NoFilter, InFoundPredicate, NoMatch);
+	}
+
+	// Search using only the event filter, no match predicate and a cache
+	static bool Search(const FTimingEventSearchParameters& InParameters, SearchPredicate InSearchPredicate, FoundPredicate InFoundPredicate, TTimingEventSearchCache<PayloadType>& InCache)
+	{
+		return Search(InParameters, InSearchPredicate, NoFilter, InFoundPredicate, NoMatch, &InCache);
 	}
 
 	// Search using a specific payload filter
-	static bool Search(const FTimingEventSearchParameters& InParameters, SearchPredicate InSearchPredicate, PayloadFilterPredicate InFilterPredicate, PayloadMatchedPredicate InPayloadMatchedPredicate, FoundPredicate InFoundPredicate)
+	static bool Search(const FTimingEventSearchParameters& InParameters, SearchPredicate InSearchPredicate, PayloadFilterPredicate InFilterPredicate, FoundPredicate InFoundPredicate, PayloadMatchedPredicate InPayloadMatchedPredicate, TTimingEventSearchCache<PayloadType>* InCache = nullptr)
 	{
+		if(InCache != nullptr && InParameters.SearchHandle != nullptr)
+		{
+			double StartTime;
+			double EndTime;
+			uint32 Depth;
+			PayloadType Payload;
+			if(InCache->Read(*InParameters.SearchHandle, StartTime, EndTime, Depth, Payload))
+			{
+				InParameters.EventMatched(StartTime, EndTime, Depth);
+				InPayloadMatchedPredicate(StartTime, EndTime, Depth, Payload);
+				InFoundPredicate(StartTime, EndTime, Depth, Payload);
+				return true;
+			}
+		}
+
 		FContext Context(InParameters, InFilterPredicate, InPayloadMatchedPredicate);
 
 		InSearchPredicate(Context);
 
 		if(Context.IsMatchFound())
 		{
-			InFoundPredicate(Context.GetStartTimeFound(), Context.GetEndTimeFound(), Context.GetDepthFound());
+			InFoundPredicate(Context.GetStartTimeFound(), Context.GetEndTimeFound(), Context.GetDepthFound(), Context.GetPayloadFound());
+
+			if(InCache != nullptr && InParameters.SearchHandle != nullptr)
+			{
+				*InParameters.SearchHandle = InCache->Write(Context.GetStartTimeFound(), Context.GetEndTimeFound(), Context.GetDepthFound(), Context.GetPayloadFound());
+			}
 		}
 
 		return Context.IsMatchFound();
 	}
-};
\ No newline at end of file
+};
