Fix tooltips not showing up in some circumstances

Search handle being track-local meant that cached searches could reference incorrect search results across tracks

diff --git a/Engine/Source/Developer/TraceInsights/Public/Insights/ViewModels/TimingEventSearch.h b/Engine/Source/Developer/TraceInsights/Public/Insights/ViewModels/TimingEventSearch.h
--- a/Engine/Source/Developer/TraceInsights/Public/Insights/ViewModels/TimingEventSearch.h
+++ b/Engine/Source/Developer/TraceInsights/Public/Insights/ViewModels/TimingEventSearch.h
@@ -21,11 +21,21 @@ ENUM_CLASS_FLAGS(ETimingEventSearchFlags);
 // A handle to a timing event that was previously searched
 struct FTimingEventSearchHandle
 {
-	template<typename T, int32 N>
-	friend struct TTimingEventSearchCache;
-
 	bool IsValid() const { return Id != uint64(-1); }
 	void Reset() { Id = uint64(-1); }
+	bool operator==(const FTimingEventSearchHandle& InOther) const
+	{
+		return Id == InOther.Id;
+	}
+
+	static FTimingEventSearchHandle GenerateHandle()
+	{
+		static uint64 GlobalId = 0;
+
+		FTimingEventSearchHandle ReturnValue;
+		ReturnValue.Id = GlobalId++;
+		return ReturnValue;
+	}
 
 private:
 	uint64 Id = uint64(-1);
@@ -84,8 +94,7 @@ struct TTimingEventSearchCache
 {
 public:
 	TTimingEventSearchCache()
-		: CurrentHandle(0)
-		, WriteIndex(0)
+		: WriteIndex(0)
 	{}
 
 	// Write to the cache
@@ -94,7 +103,7 @@ public:
 	{
 		TPair<FTimingEventSearchHandle, FResultData>& WritePair = CachedValues[WriteIndex];
 
-		WritePair.Key.Id = CurrentHandle++;
+		WritePair.Key = FTimingEventSearchHandle::GenerateHandle();
 		WritePair.Value.Payload = InPayload;
 		WritePair.Value.StartTime = InStartTime;
 		WritePair.Value.EndTime = InEndTime;
@@ -119,7 +128,7 @@ public:
 				check(ReadIndex >= 0 && ReadIndex < Size);
 
 				const TPair<FTimingEventSearchHandle, FResultData>& CachedPair = CachedValues[ReadIndex];
-				if(InHandle.Id == CachedPair.Key.Id)
+				if(InHandle == CachedPair.Key)
 				{
 					OutStartTime = CachedPair.Value.StartTime;
 					OutEndTime = CachedPair.Value.EndTime;
@@ -143,9 +152,6 @@ public:
 	}
 
 private:
-	// Current handle
-	uint64 CurrentHandle;
-
 	// Current write index
 	int32 WriteIndex;
 
