Fix undefined behaviour resulting from incorrect lambda lifetimes

Fixes tooltips appearing incorrectly

diff --git a/Engine/Source/Developer/TraceInsights/Private/Insights/Widgets/STimingView.cpp b/Engine/Source/Developer/TraceInsights/Private/Insights/Widgets/STimingView.cpp
--- a/Engine/Source/Developer/TraceInsights/Private/Insights/Widgets/STimingView.cpp
+++ b/Engine/Source/Developer/TraceInsights/Private/Insights/Widgets/STimingView.cpp
@@ -2213,14 +2213,12 @@ void STimingView::UpdateHoveredTimingEvent(float MX, float MY)
 				const double StartTime = Viewport.SlateUnitsToTime(MX);
 				const double EndTime = StartTime + 2.0 / Viewport.GetScaleX(); // +2px
 
-				FTimingEventSearchParameters SearchParameters(StartTime, EndTime, ETimingEventSearchFlags::StopAtFirstMatch,
-					// Filter
-					[Depth](double, double, uint32 EventDepth)
-					{
-						return EventDepth == Depth;
-					});
-				
-				TimingEvent.Track->SearchTimingEvent(SearchParameters, TimingEvent);
+				auto EventFilter = [Depth](double, double, uint32 EventDepth)
+				{
+					return EventDepth == Depth;
+				};
+
+				TimingEvent.Track->SearchTimingEvent(FTimingEventSearchParameters(StartTime, EndTime, ETimingEventSearchFlags::StopAtFirstMatch, EventFilter), TimingEvent);
 			}
 		}
 	}
@@ -2289,15 +2287,13 @@ void STimingView::SelectLeftTimingEvent()
 		const double StartTime = SelectedTimingEvent.StartTime;
 		const double EndTime = SelectedTimingEvent.EndTime;
 
-		FTimingEventSearchParameters SearchParameters(0.0, StartTime, ETimingEventSearchFlags::SearchAll,
-			// Filter
-			[Depth, StartTime, EndTime](double EventStartTime, double EventEndTime, uint32 EventDepth)
-			{
-				return EventDepth == Depth &&
-					(EventStartTime < StartTime || EventEndTime < EndTime);
-			});
+		auto EventFilter = [Depth, StartTime, EndTime](double EventStartTime, double EventEndTime, uint32 EventDepth)
+		{
+			return EventDepth == Depth &&
+				(EventStartTime < StartTime || EventEndTime < EndTime);
+		};
 
-		if (SelectedTimingEvent.Track->SearchTimingEvent(SearchParameters, SelectedTimingEvent))
+		if (SelectedTimingEvent.Track->SearchTimingEvent(FTimingEventSearchParameters(0.0, StartTime, ETimingEventSearchFlags::SearchAll, EventFilter), SelectedTimingEvent))
 		{
 			SelectedTimingEvent.Track->ComputeTimingEventStats(SelectedTimingEvent);
 			BringIntoView(SelectedTimingEvent.StartTime, SelectedTimingEvent.EndTime);
@@ -2315,15 +2311,14 @@ void STimingView::SelectRightTimingEvent()
 		const uint32 Depth = SelectedTimingEvent.Depth;
 		const double StartTime = SelectedTimingEvent.StartTime;
 		const double EndTime = SelectedTimingEvent.EndTime;
-		FTimingEventSearchParameters SearchParameters(EndTime, Viewport.GetMaxValidTime(), ETimingEventSearchFlags::StopAtFirstMatch,
-			// Filter
-			[Depth, StartTime, EndTime](double EventStartTime, double EventEndTime, uint32 EventDepth)
-			{
-				return EventDepth == Depth &&
-					(EventStartTime > StartTime || EventEndTime > EndTime);
-			});
 
-		if (SelectedTimingEvent.Track->SearchTimingEvent(SearchParameters,SelectedTimingEvent))
+		auto EventFilter = [Depth, StartTime, EndTime](double EventStartTime, double EventEndTime, uint32 EventDepth)
+		{
+			return EventDepth == Depth &&
+				(EventStartTime > StartTime || EventEndTime > EndTime);
+		};
+
+		if (SelectedTimingEvent.Track->SearchTimingEvent(FTimingEventSearchParameters(EndTime, Viewport.GetMaxValidTime(), ETimingEventSearchFlags::StopAtFirstMatch, EventFilter), SelectedTimingEvent))
 		{
 			SelectedTimingEvent.Track->ComputeTimingEventStats(SelectedTimingEvent);
 			BringIntoView(SelectedTimingEvent.StartTime, SelectedTimingEvent.EndTime);
@@ -2342,16 +2337,15 @@ void STimingView::SelectUpTimingEvent()
 		const uint32 Depth = SelectedTimingEvent.Depth - 1;
 		const double StartTime = SelectedTimingEvent.StartTime;
 		const double EndTime = SelectedTimingEvent.EndTime;
-		FTimingEventSearchParameters SearchParameters(StartTime, EndTime, ETimingEventSearchFlags::StopAtFirstMatch,
-			// Filter
-			[Depth, StartTime, EndTime](double EventStartTime, double EventEndTime, uint32 EventDepth)
-			{
-				return EventDepth == Depth
-					&& EventStartTime <= EndTime
-					&& EventEndTime >= StartTime;
-			});
 
-		if (SelectedTimingEvent.Track->SearchTimingEvent(SearchParameters, SelectedTimingEvent))
+		auto EventFilter = [Depth, StartTime, EndTime](double EventStartTime, double EventEndTime, uint32 EventDepth)
+		{
+			return EventDepth == Depth
+				&& EventStartTime <= EndTime
+				&& EventEndTime >= StartTime;
+		};
+
+		if (SelectedTimingEvent.Track->SearchTimingEvent(FTimingEventSearchParameters(StartTime, EndTime, ETimingEventSearchFlags::StopAtFirstMatch, EventFilter), SelectedTimingEvent))
 		{
 			SelectedTimingEvent.Track->ComputeTimingEventStats(SelectedTimingEvent);
 			BringIntoView(SelectedTimingEvent.StartTime, SelectedTimingEvent.EndTime);
@@ -2370,9 +2364,8 @@ void STimingView::SelectDownTimingEvent()
 		const double StartTime = SelectedTimingEvent.StartTime;
 		const double EndTime = SelectedTimingEvent.EndTime;
 		double LargestDuration = 0.0;
-		FTimingEventSearchParameters SearchParameters(StartTime, EndTime, ETimingEventSearchFlags::SearchAll,
-		// Filter
-		[Depth, StartTime, EndTime, &LargestDuration](double EventStartTime, double EventEndTime, uint32 EventDepth)
+
+		auto EventFilter = [Depth, StartTime, EndTime, &LargestDuration](double EventStartTime, double EventEndTime, uint32 EventDepth)
 		{
 			const double Duration = EventEndTime - EventStartTime;
 
@@ -2380,15 +2373,15 @@ void STimingView::SelectDownTimingEvent()
 				&& EventDepth == Depth
 				&& EventStartTime <= EndTime
 				&& EventEndTime >= StartTime;
-		},
-		// Matched
-		[&LargestDuration](double EventStartTime, double EventEndTime, uint32 EventDepth)
+		};
+
+		auto EventMatched = [&LargestDuration](double EventStartTime, double EventEndTime, uint32 EventDepth)
 		{
 			const double Duration = EventEndTime - EventStartTime;
 			LargestDuration = Duration;
-		});
+		};
 		
-		if (SelectedTimingEvent.Track->SearchTimingEvent(SearchParameters, SelectedTimingEvent))
+		if (SelectedTimingEvent.Track->SearchTimingEvent(FTimingEventSearchParameters(StartTime, EndTime, ETimingEventSearchFlags::SearchAll, EventFilter, EventMatched), SelectedTimingEvent))
 		{
 			SelectedTimingEvent.Track->ComputeTimingEventStats(SelectedTimingEvent);
 			BringIntoView(SelectedTimingEvent.StartTime, SelectedTimingEvent.EndTime);
