[Insights] Timing: Fixed the tooltips not showing the hovered timing event correctly in some situations (when a previous event has EndTime equal with StartTime of hovered event).

diff --git a/Engine/Source/Developer/TraceInsights/Private/Insights/ViewModels/FileActivityTimingTrack.cpp b/Engine/Source/Developer/TraceInsights/Private/Insights/ViewModels/FileActivityTimingTrack.cpp
--- a/Engine/Source/Developer/TraceInsights/Private/Insights/ViewModels/FileActivityTimingTrack.cpp
+++ b/Engine/Source/Developer/TraceInsights/Private/Insights/ViewModels/FileActivityTimingTrack.cpp
@@ -397,12 +397,14 @@ void FFileActivitySharedState::Update()
 
 void FFileActivityTimingTrack::InitTooltip(FTooltipDrawState& Tooltip, const FTimingEvent& HoveredTimingEvent) const
 {
-	auto MatchDepth = [&HoveredTimingEvent](double, double, uint32 InDepth)
+	auto MatchEvent = [&HoveredTimingEvent](double InStartTime, double InEndTime, uint32 InDepth)
 	{ 
-		return InDepth == HoveredTimingEvent.Depth; 
+		return InDepth == HoveredTimingEvent.Depth
+			&& InStartTime == HoveredTimingEvent.StartTime
+			&& InEndTime == HoveredTimingEvent.EndTime;
 	};
 
-	FTimingEventSearchParameters SearchParameters(HoveredTimingEvent.StartTime, HoveredTimingEvent.EndTime, ETimingEventSearchFlags::StopAtFirstMatch, MatchDepth);
+	FTimingEventSearchParameters SearchParameters(HoveredTimingEvent.StartTime, HoveredTimingEvent.EndTime, ETimingEventSearchFlags::StopAtFirstMatch, MatchEvent);
 	FindIoTimingEvent(SearchParameters, false, [this, &Tooltip, &HoveredTimingEvent](double InFoundStartTime, double InFoundEndTime, uint32 InFoundDepth, const FFileActivitySharedState::FIoTimingEvent& InEvent)
 	{
 		Tooltip.ResetContent();
diff --git a/Engine/Source/Developer/TraceInsights/Private/Insights/ViewModels/LoadingTimingTrack.cpp b/Engine/Source/Developer/TraceInsights/Private/Insights/ViewModels/LoadingTimingTrack.cpp
--- a/Engine/Source/Developer/TraceInsights/Private/Insights/ViewModels/LoadingTimingTrack.cpp
+++ b/Engine/Source/Developer/TraceInsights/Private/Insights/ViewModels/LoadingTimingTrack.cpp
@@ -95,12 +95,14 @@ void FLoadingSharedState::SetColorSchema(int32 Schema)
 
 void FLoadingTimingTrack::InitTooltip(FTooltipDrawState& Tooltip, const FTimingEvent& HoveredTimingEvent) const
 {
-	auto MatchDepth = [&HoveredTimingEvent](double, double, uint32 InDepth)
+	auto MatchEvent = [&HoveredTimingEvent](double InStartTime, double InEndTime, uint32 InDepth)
 	{ 
-		return InDepth == HoveredTimingEvent.Depth; 
+		return InDepth == HoveredTimingEvent.Depth
+			&& InStartTime == HoveredTimingEvent.StartTime
+			&& InEndTime == HoveredTimingEvent.EndTime;
 	};
 
-	FTimingEventSearchParameters SearchParameters(HoveredTimingEvent.StartTime, HoveredTimingEvent.EndTime, ETimingEventSearchFlags::StopAtFirstMatch, MatchDepth);
+	FTimingEventSearchParameters SearchParameters(HoveredTimingEvent.StartTime, HoveredTimingEvent.EndTime, ETimingEventSearchFlags::StopAtFirstMatch, MatchEvent);
 	FindLoadTimeProfilerCpuEvent(SearchParameters, [this, &Tooltip, &HoveredTimingEvent](double InFoundStartTime, double InFoundEndTime, uint32 InFoundDepth, const Trace::FLoadTimeProfilerCpuEvent& InFoundEvent)
 	{
 		Tooltip.ResetContent();
diff --git a/Engine/Source/Developer/TraceInsights/Private/Insights/ViewModels/ThreadTimingTrack.cpp b/Engine/Source/Developer/TraceInsights/Private/Insights/ViewModels/ThreadTimingTrack.cpp
--- a/Engine/Source/Developer/TraceInsights/Private/Insights/ViewModels/ThreadTimingTrack.cpp
+++ b/Engine/Source/Developer/TraceInsights/Private/Insights/ViewModels/ThreadTimingTrack.cpp
@@ -70,7 +70,7 @@ void FThreadTimingTrack::DrawSelectedEventInfo(const FTimingEvent& SelectedTimin
 	FindTimingProfilerEvent(SelectedTimingEvent, [&SelectedTimingEvent, &Font, &Viewport, &DrawContext, &WhiteBrush](double InFoundStartTime, double InFoundEndTime, uint32 InFoundDepth, const Trace::FTimingProfilerEvent& InFoundEvent)
 	{	
 		const FTimerNodePtr TimerNodePtr = FTimingProfilerManager::Get()->GetTimerNode(InFoundEvent.TimerIndex);
-		if(TimerNodePtr.IsValid())
+		if (TimerNodePtr.IsValid())
 		{
 			FString Str = FString::Printf(TEXT("%s (Incl.: %s, Excl.: %s)"),
 				TimerNodePtr ? *(TimerNodePtr->GetName().ToString()) : TEXT("N/A"),
@@ -330,9 +330,11 @@ void FThreadTimingTrack::BuildContextMenu(FMenuBuilder& MenuBuilder)
 
 bool FThreadTimingTrack::FindTimingProfilerEvent(const FTimingEvent& InTimingEvent, TFunctionRef<void(double, double, uint32, const Trace::FTimingProfilerEvent&)> InFoundPredicate) const
 {
-	auto MatchDepth = [&InTimingEvent](double, double, uint32 InDepth)
+	auto MatchDepth = [&InTimingEvent](double InStartTime, double InEndTime, uint32 InDepth)
 	{ 
-		return InDepth == InTimingEvent.Depth; 
+		return InDepth == InTimingEvent.Depth
+			&& InStartTime == InTimingEvent.StartTime
+			&& InEndTime == InTimingEvent.EndTime;
 	};
 
 	FTimingEventSearchParameters SearchParameters(InTimingEvent.StartTime, InTimingEvent.EndTime, ETimingEventSearchFlags::StopAtFirstMatch, MatchDepth);
