[Insights] Minor changes / minor code refactoring.

diff --git a/Engine/Source/Developer/TraceInsights/Private/Insights/Table/Widgets/STableTreeView.cpp b/Engine/Source/Developer/TraceInsights/Private/Insights/Table/Widgets/STableTreeView.cpp
--- a/Engine/Source/Developer/TraceInsights/Private/Insights/Table/Widgets/STableTreeView.cpp
+++ b/Engine/Source/Developer/TraceInsights/Private/Insights/Table/Widgets/STableTreeView.cpp
@@ -1197,12 +1197,12 @@ TSharedRef<SWidget> STableTreeView::GetGroupingCrumbMenuContent(const TSharedPtr
 	MenuBuilder.BeginSection("InsertOrAdd");
 	{
 		const FText AddGroupingText = (CrumbGroupingDepth == CurrentGroupings.Num() - 1) ? // after last one
-			LOCTEXT("AddGrouping_Section", "Add Grouping...") :
-			LOCTEXT("InsertGrouping_Section", "Insert Grouping...");
+			LOCTEXT("GroupingMenu_Add", "Add Grouping...") :
+			LOCTEXT("GroupingMenu_Insert", "Insert Grouping...");
 		MenuBuilder.AddSubMenu
 		(
 			AddGroupingText,
-			LOCTEXT("GroupingMenu_Add_Desc", "Add or insert new grouping."),
+			LOCTEXT("GroupingMenu_AddOrInsert_Desc", "Add or insert new grouping."),
 			FNewMenuDelegate::CreateSP(this, &STableTreeView::BuildGroupingSubMenu_Add, CrumbGrouping),
 			false,
 			FSlateIcon()
@@ -1216,8 +1216,8 @@ TSharedRef<SWidget> STableTreeView::GetGroupingCrumbMenuContent(const TSharedPtr
 		{
 			MenuBuilder.AddSubMenu
 			(
-				LOCTEXT("ChangeGrouping_Section", "Change To..."),
-				LOCTEXT("ChangeGrouping_Desc", "Change selected grouping."),
+				LOCTEXT("GroupingMenu_Change", "Change To..."),
+				LOCTEXT("GroupingMenu_Change_Desc", "Change selected grouping."),
 				FNewMenuDelegate::CreateSP(this, &STableTreeView::BuildGroupingSubMenu_Change, CrumbGrouping),
 				false,
 				FSlateIcon()
diff --git a/Engine/Source/Developer/TraceInsights/Private/Insights/TraceInsightsModule.cpp b/Engine/Source/Developer/TraceInsights/Private/Insights/TraceInsightsModule.cpp
--- a/Engine/Source/Developer/TraceInsights/Private/Insights/TraceInsightsModule.cpp
+++ b/Engine/Source/Developer/TraceInsights/Private/Insights/TraceInsightsModule.cpp
@@ -9,7 +9,6 @@
 #include "Widgets/Docking/SDockTab.h"
 #include "Framework/Docking/LayoutService.h"
 #include "Widgets/SWidget.h"
-#include "Widgets/Text/STextBlock.h"
 #include "WorkspaceMenuStructure.h"
 #include "WorkspaceMenuStructureModule.h"
 
@@ -169,7 +168,7 @@ void FTraceInsightsModule::StartupModule()
 #endif
 
 #if WITH_EDITOR
-	// In editor, load the layout and display it here. In the standalone tool this is done in UserInterfaceCommand
+	// In editor, load the layout and display it here. In the standalone tool this is done in UserInterfaceCommand.
 	TSharedRef<FTabManager::FLayout> DefaultLayout = FTabManager::NewLayout("UnrealInsightsLayout_v1.0");
 	OnNewLayout(DefaultLayout);
 	PersistentLayout = FLayoutSaveRestore::LoadFromConfig(GEditorLayoutIni, DefaultLayout);
@@ -228,9 +227,9 @@ void FTraceInsightsModule::ShutdownModule()
 #if WITH_EDITOR
 void FTraceInsightsModule::HandleExit()
 {
-	// In editor, as module lifetimes are different, we need to shut down the recorder service before 
-	// threads get killed in the shutdown process otherwise we will hand forever waiting on the recorder server
-	if(TraceSessionService.IsValid())
+	// In editor, as module lifetimes are different, we need to shut down the recorder service before
+	// threads get killed in the shutdown process otherwise we will hang forever waiting on the recorder server
+	if (TraceSessionService.IsValid())
 	{
 		TraceSessionService->StopRecorderServer();
 	}
@@ -249,7 +248,7 @@ void FTraceInsightsModule::OnNewLayout(TSharedRef<FTabManager::FLayout> NewLayou
 		(
 			FTabManager::NewStack()
 #if WITH_EDITOR
-			// In editor we default to all tabs closed
+			// In editor, we default to all tabs closed.
 			->AddTab(FInsightsManagerTabs::StartPageTabId, ETabState::ClosedTab)
 #else
 			->AddTab(FInsightsManagerTabs::StartPageTabId, ETabState::OpenedTab)
diff --git a/Engine/Source/Developer/TraceInsights/Private/Insights/ViewModels/TimersViewColumn.h b/Engine/Source/Developer/TraceInsights/Private/Insights/ViewModels/TimersViewColumn.h
deleted file mode 100644
diff --git a/Engine/Source/Developer/TraceInsights/Private/Insights/Widgets/SLogView.cpp b/Engine/Source/Developer/TraceInsights/Private/Insights/Widgets/SLogView.cpp
--- a/Engine/Source/Developer/TraceInsights/Private/Insights/Widgets/SLogView.cpp
+++ b/Engine/Source/Developer/TraceInsights/Private/Insights/Widgets/SLogView.cpp
@@ -20,6 +20,7 @@
 #include "Insights/InsightsStyle.h"
 #include "Insights/TimingProfilerCommon.h"
 #include "Insights/TimingProfilerManager.h"
+#include "Insights/ViewModels/MarkersTimingTrack.h" // for FTimeMarkerTrackBuilder::GetColorBy*
 #include "Insights/ViewModels/TimingViewDrawHelper.h"
 #include "Insights/Widgets/STimingProfilerWindow.h"
 #include "Insights/Widgets/STimingView.h"
diff --git a/Engine/Source/Developer/TraceInsights/Private/Insights/Widgets/SStatsView.cpp b/Engine/Source/Developer/TraceInsights/Private/Insights/Widgets/SStatsView.cpp
--- a/Engine/Source/Developer/TraceInsights/Private/Insights/Widgets/SStatsView.cpp
+++ b/Engine/Source/Developer/TraceInsights/Private/Insights/Widgets/SStatsView.cpp
@@ -1648,10 +1648,7 @@ void TTimeCalculationHelper<double>::EnumerateValues(uint32 CounterId, const Tra
 
 	Counter.EnumerateFloatValues(IntervalStartTime, IntervalEndTime, false, [this, &StatsEx, Callback](double Time, double Value)
 	{
-		if (Time >= IntervalStartTime && Time < IntervalEndTime)
-		{
-			Callback(StatsEx, Value);
-		}
+		Callback(StatsEx, Value);
 	});
 }
 
@@ -1673,10 +1670,7 @@ void TTimeCalculationHelper<int64>::EnumerateValues(uint32 CounterId, const Trac
 
 	Counter.EnumerateValues(IntervalStartTime, IntervalEndTime, false, [this, &StatsEx, Callback](double Time, int64 Value)
 	{
-		if (Time >= IntervalStartTime && Time < IntervalEndTime)
-		{
-			Callback(StatsEx, Value);
-		}
+		Callback(StatsEx, Value);
 	});
 }
 
diff --git a/Engine/Source/Developer/TraceInsights/Public/Insights/ViewModels/TimingEventSearch.h b/Engine/Source/Developer/TraceInsights/Public/Insights/ViewModels/TimingEventSearch.h
--- a/Engine/Source/Developer/TraceInsights/Public/Insights/ViewModels/TimingEventSearch.h
+++ b/Engine/Source/Developer/TraceInsights/Public/Insights/ViewModels/TimingEventSearch.h
@@ -5,6 +5,7 @@
 #include "CoreMinimal.h"
 #include "Containers/Map.h"
 
+////////////////////////////////////////////////////////////////////////////////////////////////////
 // Search behavior flags
 enum class ETimingEventSearchFlags : int32
 {
@@ -15,9 +16,9 @@ enum class ETimingEventSearchFlags : int32
 	// Whether to stop at the first match
 	StopAtFirstMatch = (1 << 0),
 };
-
 ENUM_CLASS_FLAGS(ETimingEventSearchFlags);
 
+////////////////////////////////////////////////////////////////////////////////////////////////////
 // A handle to a timing event that was previously searched
 struct FTimingEventSearchHandle
 {
@@ -41,11 +42,12 @@ private:
 	uint64 Id = uint64(-1);
 };
 
+////////////////////////////////////////////////////////////////////////////////////////////////////
 // Parameters for a timing event search
 class FTimingEventSearchParameters
 {
 public:
-	// Predicate called to filter event matches. 
+	// Predicate called to filter event matches.
 	// Note it is assumed that the SearchPredicate (below) will perform basic range checks on StartTime and EndTime,
 	// therefore this is for any other filtering that needs to take place (e.g. left of an event, below an event etc.)
 	// Returns true to pass the filter.
@@ -75,7 +77,7 @@ public:
 	// Predicate called when we get a match
 	EventMatchedPredicate EventMatched;
 
-	// A handle to a previous search. This will be written during a search if a cache is utilized and a hit was not generated
+	// A handle to a previous search. This will be written during a search if a cache is utilized and a hit was not generated.
 	FTimingEventSearchHandle* SearchHandle;
 
 	// Start time of the search
@@ -88,6 +90,7 @@ public:
 	ETimingEventSearchFlags Flags;
 };
 
+////////////////////////////////////////////////////////////////////////////////////////////////////
 // Simple acceleration structure used to return previously searched results
 template <typename PayloadType, int32 Size = 3>
 struct TTimingEventSearchCache
@@ -120,10 +123,10 @@ public:
 	{
 		static_assert(Size > 1, "TTimingEventSearchCache only works with caches that have multiple values");
 
-		if(InHandle.IsValid())
+		if (InHandle.IsValid())
 		{
 			// Start read at the previous write, to catch most recently written values
-			for(int32 ReadIndex = (WriteIndex + (Size - 1)) % Size, ReadCount = 0; ReadCount < Size; ReadCount++, ReadIndex = (ReadIndex + 1) % Size)
+			for (int32 ReadIndex = (WriteIndex + (Size - 1)) % Size, ReadCount = 0; ReadCount < Size; ReadCount++, ReadIndex = (ReadIndex + 1) % Size)
 			{
 				check(ReadIndex >= 0 && ReadIndex < Size);
 
@@ -145,7 +148,7 @@ public:
 	// Reset the cache. Any reads after this will return false
 	void Reset()
 	{
-		for(TPair<FTimingEventSearchHandle, FResultData>& CachedPair : CachedValues)
+		for (TPair<FTimingEventSearchHandle, FResultData>& CachedPair : CachedValues)
 		{
 			CachedPair.Key.Reset();
 		}
@@ -168,9 +171,10 @@ private:
 	TPair<FTimingEventSearchHandle, FResultData> CachedValues[Size];
 };
 
+////////////////////////////////////////////////////////////////////////////////////////////////////
 // Helper used to orchestrate a search of a timing event track's events
 // Example of usage:
-// 
+//
 //		FMyStruct MatchedPayload;
 //
 //		TTimingEventSearch<FMyStruct>::Search(
@@ -190,7 +194,7 @@ private:
 // 		{
 // 			// Do something with MatchedPayload, e.g. call a captured lambda.
 // 		});
-// 
+//
 template <typename PayloadType>
 struct TTimingEventSearch
 {
@@ -204,7 +208,7 @@ public:
 	// Return true to pass the filer
 	typedef TFunctionRef<bool(double /*InStartTime*/, double /*InEndTime*/, uint32 /*InDepth*/, const PayloadType& /*InPayload*/)> PayloadFilterPredicate;
 
-	// Predicate called when a match has been found.
+	// Predicate called when a match has been found
 	typedef TFunctionRef<void(double /*InStartTime*/, double /*InEndTime*/, uint32 /*InDepth*/, const PayloadType& /*InPayload*/)> FoundPredicate;
 
 	// Predicate called to run the search, e.g. iterate over an array of events
@@ -293,7 +297,6 @@ public:
 	static bool NoFilter(double, double, uint32, const PayloadType&) { return true; }
 	static void NoMatch(double, double, uint32, const PayloadType&) {}
 
-public:
 	// Search using only the event filter and no match predicate
 	static bool Search(const FTimingEventSearchParameters& InParameters, SearchPredicate InSearchPredicate, FoundPredicate InFoundPredicate)
 	{
@@ -309,13 +312,13 @@ public:
 	// Search using a specific payload filter
 	static bool Search(const FTimingEventSearchParameters& InParameters, SearchPredicate InSearchPredicate, PayloadFilterPredicate InFilterPredicate, FoundPredicate InFoundPredicate, PayloadMatchedPredicate InPayloadMatchedPredicate, TTimingEventSearchCache<PayloadType>* InCache = nullptr)
 	{
-		if(InCache != nullptr && InParameters.SearchHandle != nullptr)
+		if (InCache != nullptr && InParameters.SearchHandle != nullptr)
 		{
 			double StartTime;
 			double EndTime;
 			uint32 Depth;
 			PayloadType Payload;
-			if(InCache->Read(*InParameters.SearchHandle, StartTime, EndTime, Depth, Payload))
+			if (InCache->Read(*InParameters.SearchHandle, StartTime, EndTime, Depth, Payload))
 			{
 				InParameters.EventMatched(StartTime, EndTime, Depth);
 				InPayloadMatchedPredicate(StartTime, EndTime, Depth, Payload);
@@ -328,11 +331,11 @@ public:
 
 		InSearchPredicate(Context);
 
-		if(Context.IsMatchFound())
+		if (Context.IsMatchFound())
 		{
 			InFoundPredicate(Context.GetStartTimeFound(), Context.GetEndTimeFound(), Context.GetDepthFound(), Context.GetPayloadFound());
 
-			if(InCache != nullptr && InParameters.SearchHandle != nullptr)
+			if (InCache != nullptr && InParameters.SearchHandle != nullptr)
 			{
 				*InParameters.SearchHandle = InCache->Write(Context.GetStartTimeFound(), Context.GetEndTimeFound(), Context.GetDepthFound(), Context.GetPayloadFound());
 			}
@@ -341,3 +344,5 @@ public:
 		return Context.IsMatchFound();
 	}
 };
+
+////////////////////////////////////////////////////////////////////////////////////////////////////
diff --git a/Engine/Source/Programs/UnrealInsights/UnrealInsights.Target.cs b/Engine/Source/Programs/UnrealInsights/UnrealInsights.Target.cs
--- a/Engine/Source/Programs/UnrealInsights/UnrealInsights.Target.cs
+++ b/Engine/Source/Programs/UnrealInsights/UnrealInsights.Target.cs
@@ -14,7 +14,6 @@ public class UnrealInsightsTarget : TargetRules
 		LaunchModuleName = "UnrealInsights";
 		ExtraModuleNames.Add("EditorStyle");
 
-		bBuildDeveloperTools = false;
 		bCompileAgainstEngine = false;
 		bCompileAgainstCoreUObject = true;
 		bForceBuildTargetPlatforms = true;
