[Insights]
- STimingView refactoring:
    * removed SGraphTrack (STimingView can be used as a host widget for Graph tracks).
    * moved code related to Thread tracks, Asset Loading tracks and File Activity tracks into TimingView Extenders (manages shared state and creation of tracks).
    * added proper support for generic tracks: top dockable, bottom dockable, scrollable and foreground.
    * changed search and lifetime for TimingEvent (now passed as SharedPtr)
    * moved Layout inside the FTimingTrackViewport
    * added dirty flags for FTimingTrackViewport
    * implemented draw state cashing for timing tracks (i.e. updated only when viewport is dirty or respective track is dirty); split FTimingEventsTrackDrawState and FTimingEventsTrackDrawStateBuilder out of FTimingViewDrawHelper class
    * made FBaseTimingTrack more customizable (PreUpdate/Update/PortUpdate + PreDraw/Draw/PostDraw); moved some of the API from FTimingEventsTrack into FBaseTimingTrack (ex. SearchEvent, OnEventSelected, OnClipboardCopyEvent, etc.)
    * made track id readonly and generation of ids private to FBaseTimingTrack
- Timing view: Fixed "auto snap to limits" for vertical scrolling to allow any position inside view when the total height of all tracks is less than the height of the view.

diff --git a/Engine/Source/Developer/TraceInsights/Private/Insights/NetworkingProfiler/ViewModels/PacketContentViewDrawHelper.cpp b/Engine/Source/Developer/TraceInsights/Private/Insights/NetworkingProfiler/ViewModels/PacketContentViewDrawHelper.cpp
--- a/Engine/Source/Developer/TraceInsights/Private/Insights/NetworkingProfiler/ViewModels/PacketContentViewDrawHelper.cpp
+++ b/Engine/Source/Developer/TraceInsights/Private/Insights/NetworkingProfiler/ViewModels/PacketContentViewDrawHelper.cpp
@@ -86,10 +86,10 @@ void FPacketContentViewDrawStateBuilder::AddEvent(const Trace::FNetProfilerConte
 	constexpr float EventDY = 2.0f;
 
 	const float EventY = Y0 + (EventH + EventDY) * Depth;
-	if (EventY < -EventH || EventY > Viewport.GetHeight())
-	{
-		return;
-	}
+	//if (EventY < -EventH || EventY > Viewport.GetHeight())
+	//{
+	//	return;
+	//}
 
 	// Ensure we have enough slots in array. See LastEventX2[Depth] usage.
 	while (LastEventX2.Num() <= Depth)
@@ -161,10 +161,7 @@ void FPacketContentViewDrawStateBuilder::AddEvent(const Trace::FNetProfilerConte
 		DrawBox.Y = EventY;
 		DrawBox.W = EventW;
 		DrawBox.H = EventH;
-		DrawBox.Color.R = EventColorFill.R * BorderColorFactor;
-		DrawBox.Color.G = EventColorFill.G * BorderColorFactor;
-		DrawBox.Color.B = EventColorFill.B * BorderColorFactor;
-		DrawBox.Color.A = 1.0f;
+		DrawBox.Color = FLinearColor(EventColorFill.R * BorderColorFactor, EventColorFill.G * BorderColorFactor, EventColorFill.B * BorderColorFactor, EventColorFill.A);
 	}
 	else // 1px or 2px boxes
 	{
@@ -191,12 +188,12 @@ void FPacketContentViewDrawStateBuilder::AddEvent(const Trace::FNetProfilerConte
 			Box.X1 = EventX1;
 			Box.X2 = EventX2;
 			Box.Color = Color;
-			Box.LinearColor = FLinearColor(EventColorFill.R * BorderColorFactor, EventColorFill.G * BorderColorFactor, EventColorFill.B * BorderColorFactor, 1.0f);
+			Box.LinearColor = FLinearColor(EventColorFill.R * BorderColorFactor, EventColorFill.G * BorderColorFactor, EventColorFill.B * BorderColorFactor, EventColorFill.A);
 		}
 	}
 
 	// Draw the name of the event.
-	if (EventW > 8.0f && EventH > 10.0f)
+	if (EventW > 8.0f)
 	{
 		FString Name = EventName ? EventName : TEXT("?");
 		if (EventW > Name.Len() * 2.0f + 48.0f)
@@ -363,11 +360,14 @@ void FPacketContentViewDrawHelper::Draw(const FPacketContentViewDrawState& DrawS
 	DrawContext.LayerId++;
 
 	// Draw texts.
-	for (const FPacketContentViewDrawState::FText& Text : DrawState.Texts)
+	//if (EventH > 10.0f)
 	{
-		DrawContext.DrawText(Text.X, Text.Y, Text.Text, EventFont, Text.bWhite ? FLinearColor::White : FLinearColor::Black);
+		for (const FPacketContentViewDrawState::FText& Text : DrawState.Texts)
+		{
+			DrawContext.DrawText(Text.X, Text.Y, Text.Text, EventFont, Text.bWhite ? FLinearColor::White : FLinearColor::Black);
+		}
+		DrawContext.LayerId++;
 	}
-	DrawContext.LayerId++;
 }
 
 ////////////////////////////////////////////////////////////////////////////////////////////////////
diff --git a/Engine/Source/Developer/TraceInsights/Private/Insights/TimingProfilerCommands.cpp b/Engine/Source/Developer/TraceInsights/Private/Insights/TimingProfilerCommands.cpp
--- a/Engine/Source/Developer/TraceInsights/Private/Insights/TimingProfilerCommands.cpp
+++ b/Engine/Source/Developer/TraceInsights/Private/Insights/TimingProfilerCommands.cpp
@@ -127,7 +127,6 @@ PRAGMA_ENABLE_OPTIMIZATION
 	}
 
 IMPLEMENT_TOGGLE_COMMAND(ToggleFramesTrackVisibility, IsFramesTrackVisible, ShowHideFramesTrack)
-IMPLEMENT_TOGGLE_COMMAND(ToggleGraphTrackVisibility, IsGraphTrackVisible, ShowHideGraphTrack)
 IMPLEMENT_TOGGLE_COMMAND(ToggleTimingViewVisibility, IsTimingViewVisible, ShowHideTimingView)
 IMPLEMENT_TOGGLE_COMMAND(ToggleTimersViewVisibility, IsTimersViewVisible, ShowHideTimersView)
 IMPLEMENT_TOGGLE_COMMAND(ToggleCallersTreeViewVisibility, IsCallersTreeViewVisible, ShowHideCallersTreeView)
diff --git a/Engine/Source/Developer/TraceInsights/Private/Insights/TimingProfilerManager.cpp b/Engine/Source/Developer/TraceInsights/Private/Insights/TimingProfilerManager.cpp
--- a/Engine/Source/Developer/TraceInsights/Private/Insights/TimingProfilerManager.cpp
+++ b/Engine/Source/Developer/TraceInsights/Private/Insights/TimingProfilerManager.cpp
@@ -10,7 +10,6 @@
 #include "Insights/InsightsManager.h"
 #include "Insights/TimingProfilerCommon.h"
 #include "Insights/Widgets/SFrameTrack.h"
-#include "Insights/Widgets/SGraphTrack.h"
 #include "Insights/Widgets/SLogView.h"
 #include "Insights/Widgets/SStatsView.h"
 #include "Insights/Widgets/STimersView.h"
@@ -68,7 +67,6 @@ void FTimingProfilerManager::PostConstructor()
 void FTimingProfilerManager::BindCommands()
 {
 	ActionManager.Map_ToggleFramesTrackVisibility_Global();
-	ActionManager.Map_ToggleGraphTrackVisibility_Global();
 	ActionManager.Map_ToggleTimingViewVisibility_Global();
 	ActionManager.Map_ToggleTimersViewVisibility_Global();
 	ActionManager.Map_ToggleCallersTreeViewVisibility_Global();
@@ -139,19 +137,6 @@ void FTimingProfilerManager::ShowHideFramesTrack(const bool bIsVisible)
 
 ////////////////////////////////////////////////////////////////////////////////////////////////////
 
-void FTimingProfilerManager::ShowHideGraphTrack(const bool bIsVisible)
-{
-	bIsGraphTrackVisible = bIsVisible;
-
-	TSharedPtr<STimingProfilerWindow> Wnd = GetProfilerWindow();
-	if (Wnd.IsValid())
-	{
-		Wnd->ShowHideTab(FTimingProfilerTabs::GraphTrackID, bIsGraphTrackVisible);
-	}
-}
-
-////////////////////////////////////////////////////////////////////////////////////////////////////
-
 void FTimingProfilerManager::ShowHideTimingView(const bool bIsVisible)
 {
 	bIsTimingViewVisible = bIsVisible;
diff --git a/Engine/Source/Developer/TraceInsights/Private/Insights/TimingProfilerManager.h b/Engine/Source/Developer/TraceInsights/Private/Insights/TimingProfilerManager.h
--- a/Engine/Source/Developer/TraceInsights/Private/Insights/TimingProfilerManager.h
+++ b/Engine/Source/Developer/TraceInsights/Private/Insights/TimingProfilerManager.h
@@ -103,11 +103,6 @@ public:
 	void SetFramesTrackVisible(const bool bIsVisible) { bIsFramesTrackVisible = bIsVisible; }
 	void ShowHideFramesTrack(const bool bIsVisible);
 
-	/** @return true, if Graph Track is visible */
-	const bool IsGraphTrackVisible() const { return bIsGraphTrackVisible; }
-	void SetGraphTrackVisible(const bool bIsVisible) { bIsGraphTrackVisible = bIsVisible; }
-	void ShowHideGraphTrack(const bool bIsVisible);
-
 	/** @return true, if Timing View is visible */
 	const bool IsTimingViewVisible() const { return bIsTimingViewVisible; }
 	void SetTimingViewVisible(const bool bIsVisible) { bIsTimingViewVisible = bIsVisible; }
diff --git a/Engine/Source/Developer/TraceInsights/Private/Insights/ViewModels/FileActivityTimingTrack.cpp b/Engine/Source/Developer/TraceInsights/Private/Insights/ViewModels/FileActivityTimingTrack.cpp
--- a/Engine/Source/Developer/TraceInsights/Private/Insights/ViewModels/FileActivityTimingTrack.cpp
+++ b/Engine/Source/Developer/TraceInsights/Private/Insights/ViewModels/FileActivityTimingTrack.cpp
@@ -3,21 +3,24 @@
 #include "FileActivityTimingTrack.h"
 
 #include "Fonts/FontMeasure.h"
+#include "Framework/MultiBox/MultiBoxBuilder.h"
 #include "Styling/SlateBrush.h"
 #include "TraceServices/AnalysisService.h"
 #include "TraceServices/SessionService.h"
 
 // Insights
-#include "Insights/Common/Stopwatch.h"
 #include "Insights/Common/PaintUtils.h"
+#include "Insights/Common/Stopwatch.h"
 #include "Insights/Common/TimeUtils.h"
 #include "Insights/InsightsManager.h"
+#include "Insights/ITimingViewSession.h"
 #include "Insights/TimingProfilerCommon.h"
 #include "Insights/ViewModels/TimingEvent.h"
 #include "Insights/ViewModels/TimingTrackViewport.h"
 #include "Insights/ViewModels/TimingViewDrawHelper.h"
 #include "Insights/ViewModels/TooltipDrawState.h"
 #include "Insights/ViewModels/TimingEventSearch.h"
+#include "Insights/Widgets/STimingView.h"
 
 #include <limits>
 
@@ -67,23 +70,78 @@ uint32 GetFileActivityTypeColor(Trace::EFileActivityType Type)
 // FFileActivitySharedState
 ////////////////////////////////////////////////////////////////////////////////////////////////////
 
-void FFileActivitySharedState::Reset()
+void FFileActivitySharedState::OnBeginSession(Insights::ITimingViewSession& InSession)
 {
-	bForceIoEventsUpdate = false;
+	if (&InSession != TimingView)
+	{
+		return;
+	}
 
+	IoOverviewTrack.Reset();
+	IoActivityTrack.Reset();
+
+	bShowHideAllIoTracks = false;
+	bForceIoEventsUpdate = false;
 	bMergeIoLanes = true;
 	bShowFileActivityBackgroundEvents = false;
 
 	FileActivities.Reset();
 	FileActivityMap.Reset();
+	AllIoEvents.Reset();
+}
 
+////////////////////////////////////////////////////////////////////////////////////////////////////
+
+void FFileActivitySharedState::OnEndSession(Insights::ITimingViewSession& InSession)
+{
+	if (&InSession != TimingView)
+	{
+		return;
+	}
+
+	IoOverviewTrack.Reset();
+	IoActivityTrack.Reset();
+
+	bShowHideAllIoTracks = false;
+	bForceIoEventsUpdate = false;
+	bMergeIoLanes = true;
+	bShowFileActivityBackgroundEvents = false;
+
+	FileActivities.Reset();
+	FileActivityMap.Reset();
 	AllIoEvents.Reset();
 }
 
 ////////////////////////////////////////////////////////////////////////////////////////////////////
 
-void FFileActivitySharedState::Update()
+void FFileActivitySharedState::Tick(Insights::ITimingViewSession& InSession, const Trace::IAnalysisSession& InAnalysisSession)
 {
+	if (&InSession != TimingView)
+	{
+		return;
+	}
+
+	if (!Trace::ReadFileActivityProvider(InAnalysisSession))
+	{
+		return;
+	}
+
+	if (!IoOverviewTrack.IsValid())
+	{
+		IoOverviewTrack = MakeShareable(new FOverviewFileActivityTimingTrack(*this));
+		IoOverviewTrack->SetOrder(-999999); // first track
+		IoOverviewTrack->SetVisibilityFlag(bShowHideAllIoTracks);
+		InSession.AddScrollableTrack(IoOverviewTrack);
+	}
+
+	if (!IoActivityTrack.IsValid())
+	{
+		IoActivityTrack = MakeShareable(new FDetailedFileActivityTimingTrack(*this));
+		IoActivityTrack->SetOrder(+999999); // last track
+		IoActivityTrack->SetVisibilityFlag(bShowHideAllIoTracks);
+		InSession.AddScrollableTrack(IoActivityTrack);
+	}
+
 	if (bForceIoEventsUpdate)
 	{
 		bForceIoEventsUpdate = false;
@@ -95,14 +153,10 @@ void FFileActivitySharedState::Update()
 		FStopwatch Stopwatch;
 		Stopwatch.Start();
 
-		TSharedPtr<const Trace::IAnalysisSession> Session = FInsightsManager::Get()->GetSession();
-		if (Session.IsValid() && Trace::ReadFileActivityProvider(*Session.Get()))
+		// Enumerate all IO events and cache them.
 		{
-			Trace::FAnalysisSessionReadScope SessionReadScope(*Session.Get());
-
-			const Trace::IFileActivityProvider& FileActivityProvider = *Trace::ReadFileActivityProvider(*Session.Get());
-
-			// Enumerate all IO events and cache them.
+			Trace::FAnalysisSessionReadScope SessionReadScope(InAnalysisSession);
+			const Trace::IFileActivityProvider& FileActivityProvider = *Trace::ReadFileActivityProvider(InAnalysisSession);
 			FileActivityProvider.EnumerateFileActivity([this](const Trace::FFileInfo& FileInfo, const Trace::IFileActivityProvider::Timeline& Timeline)
 			{
 				Timeline.EnumerateEvents(-std::numeric_limits<double>::infinity(), +std::numeric_limits<double>::infinity(),
@@ -392,58 +446,143 @@ void FFileActivitySharedState::Update()
 }
 
 ////////////////////////////////////////////////////////////////////////////////////////////////////
-// FFileActivityTimingTrack
+
+void FFileActivitySharedState::ExtendFilterMenu(Insights::ITimingViewSession& InSession, FMenuBuilder& InOutMenuBuilder)
+{
+	if (&InSession != TimingView)
+	{
+		return;
+	}
+
+	InOutMenuBuilder.BeginSection("File Activity", LOCTEXT("FileActivityHeading", "File Activity"));
+	{
+		InOutMenuBuilder.AddMenuEntry(
+			LOCTEXT("ShowAllIoTracks", "I/O Tracks - I"),
+			LOCTEXT("ShowAllIoTracks_Tooltip", "Show/hide the I/O (File Activity) tracks"),
+			FSlateIcon(),
+			FUIAction(FExecuteAction::CreateSP(this, &FFileActivitySharedState::ShowHideAllIoTracks_Execute),
+					  FCanExecuteAction(),
+					  FIsActionChecked::CreateSP(this, &FFileActivitySharedState::ShowHideAllIoTracks_IsChecked)),
+			"QuickFilterSeparator",
+			EUserInterfaceActionType::ToggleButton
+		);
+	}
+	InOutMenuBuilder.EndSection();
+}
+
 ////////////////////////////////////////////////////////////////////////////////////////////////////
 
-void FFileActivityTimingTrack::InitTooltip(FTooltipDrawState& Tooltip, const FTimingEvent& HoveredTimingEvent) const
+void FFileActivitySharedState::OnTracksChanged(Insights::ITimingViewSession& InSession, int32& InOutOrder)
 {
-	auto MatchEvent = [&HoveredTimingEvent](double InStartTime, double InEndTime, uint32 InDepth)
-	{ 
-		return InDepth == HoveredTimingEvent.Depth
-			&& InStartTime == HoveredTimingEvent.StartTime
-			&& InEndTime == HoveredTimingEvent.EndTime;
-	};
+	if (&InSession != TimingView)
+	{
+		return;
+	}
+}
+
+////////////////////////////////////////////////////////////////////////////////////////////////////
+
+bool FFileActivitySharedState::ShowHideAllIoTracks_IsChecked() const
+{
+	return bShowHideAllIoTracks;
+}
+
+////////////////////////////////////////////////////////////////////////////////////////////////////
+
+void FFileActivitySharedState::ShowHideAllIoTracks_Execute()
+{
+	bShowHideAllIoTracks = !bShowHideAllIoTracks;
+
+	if (IoOverviewTrack.IsValid())
+	{
+		IoOverviewTrack->SetVisibilityFlag(bShowHideAllIoTracks);
+	}
+	if (IoActivityTrack.IsValid())
+	{
+		IoActivityTrack->SetVisibilityFlag(bShowHideAllIoTracks);
+	}
+
+	if (TimingView)
+	{
+		TimingView->OnTrackVisibilityChanged();
+	}
 
-	FTimingEventSearchParameters SearchParameters(HoveredTimingEvent.StartTime, HoveredTimingEvent.EndTime, ETimingEventSearchFlags::StopAtFirstMatch, MatchEvent);
-	FindIoTimingEvent(SearchParameters, false, [this, &Tooltip, &HoveredTimingEvent](double InFoundStartTime, double InFoundEndTime, uint32 InFoundDepth, const FFileActivitySharedState::FIoTimingEvent& InEvent)
+	if (bShowHideAllIoTracks)
 	{
-		Tooltip.ResetContent();
+		RequestUpdate();
+	}
+}
 
-		const Trace::EFileActivityType ActivityType = static_cast<Trace::EFileActivityType>(InEvent.Type & 0x0F);
-		const bool bHasFailed = ((InEvent.Type & 0xF0) != 0);
+////////////////////////////////////////////////////////////////////////////////////////////////////
 
-		FString TypeStr;
-		uint32 TypeColor;
-		if (bHasFailed)
+void FFileActivitySharedState::ToggleBackgroundEvents()
+{
+	bShowFileActivityBackgroundEvents = !bShowFileActivityBackgroundEvents;
+
+	if (IoActivityTrack.IsValid())
+	{
+		IoActivityTrack->SetDirtyFlag();
+	}
+}
+
+////////////////////////////////////////////////////////////////////////////////////////////////////
+// FFileActivityTimingTrack
+////////////////////////////////////////////////////////////////////////////////////////////////////
+
+void FFileActivityTimingTrack::InitTooltip(FTooltipDrawState& InOutTooltip, const ITimingEvent& InTooltipEvent) const
+{
+	if (InTooltipEvent.CheckTrack(this) && FTimingEvent::CheckTypeName(InTooltipEvent))
+	{
+		const FTimingEvent& TooltipEvent = static_cast<const FTimingEvent&>(InTooltipEvent);
+
+		auto MatchEvent = [&TooltipEvent](double InStartTime, double InEndTime, uint32 InDepth)
 		{
-			TypeStr = TEXT("Failed ");
-			TypeStr += GetFileActivityTypeName(ActivityType);
-			TypeColor = 0xFFFF3333;
-		}
-		else
+			return InDepth == TooltipEvent.GetDepth()
+				&& InStartTime == TooltipEvent.GetStartTime()
+				&& InEndTime == TooltipEvent.GetEndTime();
+		};
+
+		FTimingEventSearchParameters SearchParameters(TooltipEvent.GetStartTime(), TooltipEvent.GetEndTime(), ETimingEventSearchFlags::StopAtFirstMatch, MatchEvent);
+		FindIoTimingEvent(SearchParameters, false, [this, &InOutTooltip, &TooltipEvent](double InFoundStartTime, double InFoundEndTime, uint32 InFoundDepth, const FFileActivitySharedState::FIoTimingEvent& InEvent)
 		{
-			TypeStr = GetFileActivityTypeName(ActivityType);
-			TypeColor = GetFileActivityTypeColor(ActivityType);
-		}
-		FLinearColor TypeLinearColor = FLinearColor(FColor(TypeColor));
-		TypeLinearColor.R *= 2.0f;
-		TypeLinearColor.G *= 2.0f;
-		TypeLinearColor.B *= 2.0f;
-		Tooltip.AddTitle(TypeStr, TypeLinearColor);
+			InOutTooltip.ResetContent();
 
-		Tooltip.AddTitle(InEvent.FileActivity->Path);
+			const Trace::EFileActivityType ActivityType = static_cast<Trace::EFileActivityType>(InEvent.Type & 0x0F);
+			const bool bHasFailed = ((InEvent.Type & 0xF0) != 0);
+
+			FString TypeStr;
+			uint32 TypeColor;
+			if (bHasFailed)
+			{
+				TypeStr = TEXT("Failed ");
+				TypeStr += GetFileActivityTypeName(ActivityType);
+				TypeColor = 0xFFFF3333;
+			}
+			else
+			{
+				TypeStr = GetFileActivityTypeName(ActivityType);
+				TypeColor = GetFileActivityTypeColor(ActivityType);
+			}
+			FLinearColor TypeLinearColor = FLinearColor(FColor(TypeColor));
+			TypeLinearColor.R *= 2.0f;
+			TypeLinearColor.G *= 2.0f;
+			TypeLinearColor.B *= 2.0f;
+			InOutTooltip.AddTitle(TypeStr, TypeLinearColor);
 
-		Tooltip.AddNameValueTextLine(TEXT("Duration:"), TimeUtils::FormatTimeAuto(HoveredTimingEvent.Duration()));
-		Tooltip.AddNameValueTextLine(TEXT("Depth:"), FString::Printf(TEXT("%d"), HoveredTimingEvent.Depth));
+			InOutTooltip.AddTitle(InEvent.FileActivity->Path);
 
-		if (ActivityType == Trace::FileActivityType_Read || ActivityType == Trace::FileActivityType_Write)
-		{
-			Tooltip.AddNameValueTextLine(TEXT("Offset:"), FText::AsNumber(InEvent.Offset).ToString() + TEXT(" bytes"));
-			Tooltip.AddNameValueTextLine(TEXT("Size:"), FText::AsNumber(InEvent.Size).ToString() + TEXT(" bytes"));
-		}
+			InOutTooltip.AddNameValueTextLine(TEXT("Duration:"), TimeUtils::FormatTimeAuto(TooltipEvent.GetDuration()));
+			InOutTooltip.AddNameValueTextLine(TEXT("Depth:"), FString::Printf(TEXT("%d"), TooltipEvent.GetDepth()));
 
-		Tooltip.UpdateLayout();
-	});
+			if (ActivityType == Trace::FileActivityType_Read || ActivityType == Trace::FileActivityType_Write)
+			{
+				InOutTooltip.AddNameValueTextLine(TEXT("Offset:"), FText::AsNumber(InEvent.Offset).ToString() + TEXT(" bytes"));
+				InOutTooltip.AddNameValueTextLine(TEXT("Size:"), FText::AsNumber(InEvent.Size).ToString() + TEXT(" bytes"));
+			}
+
+			InOutTooltip.UpdateLayout();
+		});
+	}
 }
 
 ////////////////////////////////////////////////////////////////////////////////////////////////////
@@ -461,14 +600,14 @@ bool FFileActivityTimingTrack::FindIoTimingEvent(const FTimingEventSearchParamet
 			{
 				Trace::FAnalysisSessionReadScope SessionReadScope(*Session.Get());
 
-				for (const FFileActivitySharedState::FIoTimingEvent& Event : State->GetAllEvents())
+				for (const FFileActivitySharedState::FIoTimingEvent& Event : SharedState.GetAllEvents())
 				{
 					if (Event.EndTime <= InContext.GetParameters().StartTime)
 					{
 						continue;
 					}
 
-					if (!InContext.ShouldContinueSearching() || InContext.GetParameters().StartTime >= InContext.GetParameters().EndTime)
+					if (!InContext.ShouldContinueSearching() || Event.StartTime >= InContext.GetParameters().EndTime)
 					{
 						break;
 					}
@@ -489,16 +628,19 @@ bool FFileActivityTimingTrack::FindIoTimingEvent(const FTimingEventSearchParamet
 // FOverviewFileActivityTimingTrack
 ////////////////////////////////////////////////////////////////////////////////////////////////////
 
-void FOverviewFileActivityTimingTrack::Draw(ITimingViewDrawHelper& Helper) const
+void FOverviewFileActivityTimingTrack::PreUpdate(const ITimingTrackUpdateContext& Context)
 {
-	FTimingEventsTrack& Track = *const_cast<FOverviewFileActivityTimingTrack*>(this);
-
-	if (Helper.BeginTimeline(Track))
+	const FTimingTrackViewport& Viewport = Context.GetViewport();
+	if (IsDirty() || Viewport.IsHorizontalViewportDirty())
 	{
-		// Draw the IO Overiew track using cached events (as those are sorted by Start Time).
-		for (const FFileActivitySharedState::FIoTimingEvent& Event : State->AllIoEvents)
+		ClearDirtyFlag();
+
+		FTimingEventsTrackDrawStateBuilder Builder(DrawState, Viewport);
+
+		for (const FFileActivitySharedState::FIoTimingEvent& Event : SharedState.AllIoEvents)
 		{
 			const Trace::EFileActivityType ActivityType = static_cast<Trace::EFileActivityType>(Event.Type & 0x0F);
+			const uint64 EventType = static_cast<uint64>(ActivityType);
 
 			if (ActivityType >= Trace::FileActivityType_Count)
 			{
@@ -509,84 +651,104 @@ void FOverviewFileActivityTimingTrack::Draw(ITimingViewDrawHelper& Helper) const
 			//const double EventEndTime = Event.EndTime; // keep duration of events
 			const double EventEndTime = Event.StartTime; // make all 0 duration events
 
-			if (EventEndTime <= Helper.GetViewport().GetStartTime())
+			if (EventEndTime <= Viewport.GetStartTime())
 			{
 				continue;
 			}
-			if (Event.StartTime >= Helper.GetViewport().GetEndTime())
+			if (Event.StartTime >= Viewport.GetEndTime())
 			{
 				break;
 			}
 
-			uint32 Color = GetFileActivityTypeColor(ActivityType);
-
 			const bool bHasFailed = ((Event.Type & 0xF0) != 0);
 
 			if (bHasFailed)
 			{
 				FString EventName = TEXT("Failed ");
 				EventName += GetFileActivityTypeName(ActivityType);
-				Color = 0xFFAA0000;
-				Helper.AddEvent(Event.StartTime, EventEndTime, 0, *EventName, Color);
+				const uint32 Color = 0xFFAA0000;
+				Builder.AddEvent(Event.StartTime, EventEndTime, 0, *EventName, EventType, Color);
 			}
 			else
 			{
-				Helper.AddEvent(Event.StartTime, EventEndTime, 0, GetFileActivityTypeName(ActivityType), Color);
+				const uint32 Color = GetFileActivityTypeColor(ActivityType);
+				Builder.AddEvent(Event.StartTime, EventEndTime, 0, GetFileActivityTypeName(ActivityType), EventType, Color);
 			}
 		}
 
-		Helper.EndTimeline(Track);
+		Builder.Flush();
+		SetNumLanes(1);
 	}
+
+	UpdateTrackHeight(Context);
 }
 
-bool FOverviewFileActivityTimingTrack::SearchTimingEvent(const FTimingEventSearchParameters& InSearchParameters, FTimingEvent& InOutTimingEvent) const
+////////////////////////////////////////////////////////////////////////////////////////////////////
+
+void FOverviewFileActivityTimingTrack::Draw(const ITimingTrackDrawContext& Context) const
+{
+	const FTimingViewDrawHelper& Helper = Context.GetHelper();
+	Helper.DrawEventsTrack(DrawState, *this);
+}
+
+////////////////////////////////////////////////////////////////////////////////////////////////////
+
+const TSharedPtr<const ITimingEvent> FOverviewFileActivityTimingTrack::SearchEvent(const FTimingEventSearchParameters& InSearchParameters) const
 {
+	TSharedPtr<const ITimingEvent> FoundEvent;
+
 	constexpr bool bIgnoreEventDepth = true;
-	return FindIoTimingEvent(InSearchParameters, bIgnoreEventDepth, [this, &InOutTimingEvent](double InFoundStartTime, double InFoundEndTime, uint32 InFoundDepth, const FFileActivitySharedState::FIoTimingEvent& InEvent)
+	FindIoTimingEvent(InSearchParameters, bIgnoreEventDepth, [this, &FoundEvent](double InFoundStartTime, double InFoundEndTime, uint32 InFoundDepth, const FFileActivitySharedState::FIoTimingEvent& InEvent)
 	{
-		InOutTimingEvent = FTimingEvent(this, InFoundStartTime, InFoundEndTime, InFoundDepth);
+		FoundEvent = MakeShareable(new FTimingEvent(SharedThis(this), InFoundStartTime, InFoundEndTime, InFoundDepth));
 	});
+
+	return FoundEvent;
 }
 
 ////////////////////////////////////////////////////////////////////////////////////////////////////
 // FDetailedFileActivityTimingTrack
 ////////////////////////////////////////////////////////////////////////////////////////////////////
 
-void FDetailedFileActivityTimingTrack::Draw(ITimingViewDrawHelper& Helper) const
+void FDetailedFileActivityTimingTrack::PreUpdate(const ITimingTrackUpdateContext& Context)
 {
-	FTimingEventsTrack& Track = *const_cast<FDetailedFileActivityTimingTrack*>(this);
-
-	// Draw IO track using cached events.
-	if (Helper.BeginTimeline(Track))
+	const FTimingTrackViewport& Viewport = Context.GetViewport();
+	if (IsDirty() || Viewport.IsHorizontalViewportDirty())
 	{
-		// Draw IO file activity background events.
-		if (State->bShowFileActivityBackgroundEvents)
+		ClearDirtyFlag();
+
+		FTimingEventsTrackDrawStateBuilder Builder(DrawState, Viewport);
+
+		// Add IO file activity background events.
+		if (SharedState.bShowFileActivityBackgroundEvents)
 		{
-			for (const TSharedPtr<FFileActivitySharedState::FIoFileActivity> Activity : State->FileActivities)
+			for (const TSharedPtr<FFileActivitySharedState::FIoFileActivity> Activity : SharedState.FileActivities)
 			{
-				if (Activity->EndTime <= Helper.GetViewport().GetStartTime())
+				if (Activity->EndTime <= Viewport.GetStartTime())
 				{
 					continue;
 				}
-				if (Activity->StartTime >= Helper.GetViewport().GetEndTime())
+				if (Activity->StartTime >= Viewport.GetEndTime())
 				{
 					break;
 				}
 
 				ensure(Activity->Depth <= 10000);
 
-				Helper.AddEvent(Activity->StartTime, Activity->EndTime, Activity->Depth, Activity->Path, 0x55333333);
+				const uint64 EventType = uint64(Trace::EFileActivityType::FileActivityType_Count);
+
+				Builder.AddEvent(Activity->StartTime, Activity->EndTime, Activity->Depth, Activity->Path, EventType, 0x55333333);
 			}
 		}
 
-		// Draw IO file activity foreground events.
-		for (const FFileActivitySharedState::FIoTimingEvent& Event : State->AllIoEvents)
+		// Add IO file activity foreground events.
+		for (const FFileActivitySharedState::FIoTimingEvent& Event : SharedState.AllIoEvents)
 		{
-			if (Event.EndTime <= Helper.GetViewport().GetStartTime())
+			if (Event.EndTime <= Viewport.GetStartTime())
 			{
 				continue;
 			}
-			if (Event.StartTime >= Helper.GetViewport().GetEndTime())
+			if (Event.StartTime >= Viewport.GetEndTime())
 			{
 				break;
 			}
@@ -594,8 +756,7 @@ void FDetailedFileActivityTimingTrack::Draw(ITimingViewDrawHelper& Helper) const
 			ensure(Event.Depth <= 10000);
 
 			const Trace::EFileActivityType ActivityType = static_cast<Trace::EFileActivityType>(Event.Type & 0x0F);
-
-			uint32 Color = GetFileActivityTypeColor(ActivityType);
+			const uint64 EventType = static_cast<uint64>(ActivityType);
 
 			if (ActivityType < Trace::FileActivityType_Count)
 			{
@@ -605,12 +766,13 @@ void FDetailedFileActivityTimingTrack::Draw(ITimingViewDrawHelper& Helper) const
 				{
 					FString EventName = TEXT("Failed ");
 					EventName += GetFileActivityTypeName(ActivityType);
-					Color = 0xFFAA0000;
-					Helper.AddEvent(Event.StartTime, Event.EndTime, Event.Depth, *EventName, Color);
+					const uint32 Color = 0xFFAA0000;
+					Builder.AddEvent(Event.StartTime, Event.EndTime, Event.Depth, *EventName, EventType, Color);
 				}
 				else
 				{
-					Helper.AddEvent(Event.StartTime, Event.EndTime, Event.Depth, GetFileActivityTypeName(ActivityType), Color);
+					const uint32 Color = GetFileActivityTypeColor(ActivityType);
+					Builder.AddEvent(Event.StartTime, Event.EndTime, Event.Depth, GetFileActivityTypeName(ActivityType), EventType, Color);
 				}
 			}
 			else
@@ -619,21 +781,39 @@ void FDetailedFileActivityTimingTrack::Draw(ITimingViewDrawHelper& Helper) const
 				EventName += " [";
 				EventName += Event.FileActivity->Path;
 				EventName += "]";
-				Helper.AddEvent(Event.StartTime, Event.EndTime, Event.Depth, *EventName, Color);
+				const uint32 Color = GetFileActivityTypeColor(ActivityType);
+				Builder.AddEvent(Event.StartTime, Event.EndTime, Event.Depth, *EventName, EventType, Color);
 			}
 		}
 
-		Helper.EndTimeline(Track);
+		Builder.Flush();
+		SetNumLanes(Builder.GetMaxDepth() + 1);
 	}
+
+	UpdateTrackHeight(Context);
+}
+
+////////////////////////////////////////////////////////////////////////////////////////////////////
+
+void FDetailedFileActivityTimingTrack::Draw(const ITimingTrackDrawContext& Context) const
+{
+	const FTimingViewDrawHelper& Helper = Context.GetHelper();
+	Helper.DrawEventsTrack(DrawState, *this);
 }
 
-bool FDetailedFileActivityTimingTrack::SearchTimingEvent(const FTimingEventSearchParameters& InSearchParameters, FTimingEvent& InOutTimingEvent) const
+////////////////////////////////////////////////////////////////////////////////////////////////////
+
+const TSharedPtr<const ITimingEvent> FDetailedFileActivityTimingTrack::SearchEvent(const FTimingEventSearchParameters& InSearchParameters) const
 {
+	TSharedPtr<const ITimingEvent> FoundEvent;
+
 	constexpr bool bIgnoreEventDepth = false;
-	return FindIoTimingEvent(InSearchParameters, bIgnoreEventDepth, [this, &InOutTimingEvent](double InFoundStartTime, double InFoundEndTime, uint32 InFoundDepth, const FFileActivitySharedState::FIoTimingEvent& InEvent)
+	FindIoTimingEvent(InSearchParameters, bIgnoreEventDepth, [this, &FoundEvent](double InFoundStartTime, double InFoundEndTime, uint32 InFoundDepth, const FFileActivitySharedState::FIoTimingEvent& InEvent)
 	{
-		InOutTimingEvent = FTimingEvent(this, InFoundStartTime, InFoundEndTime, InFoundDepth);
+		FoundEvent = MakeShareable(new FTimingEvent(SharedThis(this), InFoundStartTime, InFoundEndTime, InFoundDepth));
 	});
+
+	return FoundEvent;
 }
 
 ////////////////////////////////////////////////////////////////////////////////////////////////////
diff --git a/Engine/Source/Developer/TraceInsights/Private/Insights/ViewModels/FileActivityTimingTrack.h b/Engine/Source/Developer/TraceInsights/Private/Insights/ViewModels/FileActivityTimingTrack.h
--- a/Engine/Source/Developer/TraceInsights/Private/Insights/ViewModels/FileActivityTimingTrack.h
+++ b/Engine/Source/Developer/TraceInsights/Private/Insights/ViewModels/FileActivityTimingTrack.h
@@ -5,13 +5,16 @@
 #include "CoreMinimal.h"
 
 // Insights
+#include "Insights/ITimingViewExtender.h"
 #include "Insights/ViewModels/TimingEventsTrack.h"
+#include "Insights/ViewModels/TimingViewDrawHelper.h"
 
 class FTimingEventSearchParameters;
+class STimingView;
 
 ////////////////////////////////////////////////////////////////////////////////////////////////////
 
-class FFileActivitySharedState
+class FFileActivitySharedState : public Insights::ITimingViewExtender, public TSharedFromThis<FFileActivitySharedState>
 {
 	friend class FOverviewFileActivityTimingTrack;
 	friend class FDetailedFileActivityTimingTrack;
@@ -39,22 +42,36 @@ public:
 	};
 
 public:
-	FFileActivitySharedState() {}
+	explicit FFileActivitySharedState(STimingView* InTimingView) : TimingView(InTimingView) {}
+	virtual ~FFileActivitySharedState() = default;
 
-	const TArray<FIoTimingEvent>& GetAllEvents() const { return AllIoEvents; }
-
-	~FFileActivitySharedState() {}
+	// ITimingViewExtender
+	virtual void OnBeginSession(Insights::ITimingViewSession& InSession) override;
+	virtual void OnEndSession(Insights::ITimingViewSession& InSession) override;
+	virtual void Tick(Insights::ITimingViewSession& InSession, const Trace::IAnalysisSession& InAnalysisSession) override;
+	virtual void ExtendFilterMenu(Insights::ITimingViewSession& InSession, FMenuBuilder& InMenuBuilder) override;
+	virtual void OnTracksChanged(Insights::ITimingViewSession& InSession, int32& InOutOrder) override;
 
-	void Reset();
-	void Update();
+	const TArray<FIoTimingEvent>& GetAllEvents() const { return AllIoEvents; }
 
 	void RequestUpdate() { bForceIoEventsUpdate = true; }
 	void ToggleMergeLanes() { bMergeIoLanes = !bMergeIoLanes; }
-	void ToggleBackgroundEvents() { bShowFileActivityBackgroundEvents = !bShowFileActivityBackgroundEvents; }
+	void ToggleBackgroundEvents();
+
+	void ShowHideAllIoTracks() { ShowHideAllIoTracks_Execute(); }
 
 private:
-	bool bForceIoEventsUpdate;
+	bool ShowHideAllIoTracks_IsChecked() const;
+	void ShowHideAllIoTracks_Execute();
 
+private:
+	STimingView* TimingView;
+
+	TSharedPtr<FTimingEventsTrack> IoOverviewTrack;
+	TSharedPtr<FTimingEventsTrack> IoActivityTrack;
+
+	bool bShowHideAllIoTracks;
+	bool bForceIoEventsUpdate;
 	bool bMergeIoLanes;
 	bool bShowFileActivityBackgroundEvents;
 
@@ -70,20 +87,21 @@ private:
 class FFileActivityTimingTrack : public FTimingEventsTrack
 {
 public:
-	explicit FFileActivityTimingTrack(uint64 InTrackId, const FName SubType, const FString& InName, TSharedPtr<FFileActivitySharedState> InState)
-		: FTimingEventsTrack(InTrackId, FName(TEXT("FileActivity")), SubType, InName)
-		, State(InState)
+	explicit FFileActivityTimingTrack(FFileActivitySharedState& InSharedState, const FName SubType, const FString& InName)
+		: FTimingEventsTrack(FName(TEXT("FileActivity")), SubType, InName)
+		, SharedState(InSharedState)
 	{
 	}
 	virtual ~FFileActivityTimingTrack() {}
 
-	virtual void InitTooltip(FTooltipDrawState& Tooltip, const FTimingEvent& HoveredTimingEvent) const override;
+	virtual void InitTooltip(FTooltipDrawState& InOutTooltip, const ITimingEvent& InTooltipEvent) const override;
 
 protected:
 	bool FindIoTimingEvent(const FTimingEventSearchParameters& InParameters, bool bIgnoreEventDepth, TFunctionRef<void(double, double, uint32, const FFileActivitySharedState::FIoTimingEvent&)> InFoundPredicate) const;
 
 protected:
-	TSharedPtr<FFileActivitySharedState> State;
+	FFileActivitySharedState& SharedState;
+	FTimingEventsTrackDrawState DrawState;
 };
 
 ////////////////////////////////////////////////////////////////////////////////////////////////////
@@ -91,13 +109,14 @@ protected:
 class FOverviewFileActivityTimingTrack : public FFileActivityTimingTrack
 {
 public:
-	explicit FOverviewFileActivityTimingTrack(uint64 InTrackId, TSharedPtr<FFileActivitySharedState> InState)
-		: FFileActivityTimingTrack(InTrackId, FName(TEXT("Overview")), TEXT("I/O Overview"), InState)
+	explicit FOverviewFileActivityTimingTrack(FFileActivitySharedState& InSharedState)
+		: FFileActivityTimingTrack(InSharedState, FName(TEXT("Overview")), TEXT("I/O Overview"))
 	{
 	}
 
-	virtual void Draw(ITimingViewDrawHelper& Helper) const override;
-	virtual bool SearchTimingEvent(const FTimingEventSearchParameters& InSearchParameters,  FTimingEvent& InOutTimingEvent) const override;
+	virtual void PreUpdate(const ITimingTrackUpdateContext& Context) override;
+	virtual void Draw(const ITimingTrackDrawContext& Context) const override;
+	virtual const TSharedPtr<const ITimingEvent> SearchEvent(const FTimingEventSearchParameters& InSearchParameters) const override;
 };
 
 ////////////////////////////////////////////////////////////////////////////////////////////////////
@@ -105,13 +124,14 @@ public:
 class FDetailedFileActivityTimingTrack : public FFileActivityTimingTrack
 {
 public:
-	explicit FDetailedFileActivityTimingTrack(uint64 InTrackId, TSharedPtr<FFileActivitySharedState> InState)
-		: FFileActivityTimingTrack(InTrackId, FName(TEXT("Detailed")), TEXT("I/O Activity"), InState)
+	explicit FDetailedFileActivityTimingTrack(FFileActivitySharedState& InSharedState)
+		: FFileActivityTimingTrack(InSharedState, FName(TEXT("Detailed")), TEXT("I/O Activity"))
 	{
 	}
 
-	virtual void Draw(ITimingViewDrawHelper& Helper) const override;
-	virtual bool SearchTimingEvent(const FTimingEventSearchParameters& InSearchParameters, FTimingEvent& InOutTimingEvent) const override;
+	virtual void PreUpdate(const ITimingTrackUpdateContext& Context) override;
+	virtual void Draw(const ITimingTrackDrawContext& Context) const override;
+	virtual const TSharedPtr<const ITimingEvent> SearchEvent(const FTimingEventSearchParameters& InSearchParameters) const override;
 };
 
 ////////////////////////////////////////////////////////////////////////////////////////////////////
diff --git a/Engine/Source/Developer/TraceInsights/Private/Insights/ViewModels/GraphTrack.cpp b/Engine/Source/Developer/TraceInsights/Private/Insights/ViewModels/GraphTrack.cpp
--- a/Engine/Source/Developer/TraceInsights/Private/Insights/ViewModels/GraphTrack.cpp
+++ b/Engine/Source/Developer/TraceInsights/Private/Insights/ViewModels/GraphTrack.cpp
@@ -12,113 +12,21 @@
 #include "Insights/Common/TimeUtils.h"
 #include "Insights/InsightsStyle.h"
 #include "Insights/ViewModels/DrawHelpers.h"
+#include "Insights/ViewModels/GraphSeries.h"
+#include "Insights/ViewModels/GraphTrackBuilder.h"
+#include "Insights/ViewModels/GraphTrackEvent.h"
+#include "Insights/ViewModels/TimingEvent.h"
 #include "Insights/ViewModels/TimingTrackViewport.h"
+#include "Insights/ViewModels/TooltipDrawState.h"
 
 #define LOCTEXT_NAMESPACE "GraphTrack"
 
-////////////////////////////////////////////////////////////////////////////////////////////////////
-// FGraphSeries
-////////////////////////////////////////////////////////////////////////////////////////////////////
-
-FGraphSeries::FGraphSeries()
-	: Name()
-	, Description()
-	, bIsVisible(true)
-	, bIsDirty(false)
-	, bAutoZoom(false)
-	, TargetAutoZoomLowValue(0.0)
-	, TargetAutoZoomHighValue(1.0)
-	, AutoZoomLowValue(0.0)
-	, AutoZoomHighValue(1.0)
-	, BaselineY(0.0)
-	, ScaleY(1.0)
-	, Color(0.0f, 0.5f, 1.0f, 1.0f)
-	, BorderColor(0.3f, 0.8f, 1.0f, 1.0f)
-	//, Events()
-	//, Points()
-	//, LinePoints()
-	//, Boxes()
-{
-}
-
-////////////////////////////////////////////////////////////////////////////////////////////////////
-
-FGraphSeries::~FGraphSeries()
-{
-}
-
-////////////////////////////////////////////////////////////////////////////////////////////////////
-
-const FGraphSeries::FEvent* FGraphSeries::GetEvent(const float X, const float Y, const FTimingTrackViewport& Viewport, bool bCheckLine, bool bCheckBox) const
-{
-	const float LocalBaselineY = static_cast<float>(GetBaselineY());
-
-	for (const FGraphSeries::FEvent& Event : Events)
-	{
-		const float EventX1 = Viewport.TimeToSlateUnitsRounded(Event.Time);
-		const float EventX2 = Viewport.TimeToSlateUnitsRounded(Event.Time + Event.Duration);
-
-		const float EventY = GetRoundedYForValue(Event.Value);
-
-		// Check bounding box of the visual point.
-		constexpr float PointTolerance = 5.0f;
-		if (X >= EventX1 - PointTolerance && X <= EventX1 + PointTolerance &&
-			Y >= EventY - PointTolerance && Y <= EventY + PointTolerance)
-		{
-			return &Event;
-		}
-
-		if (bCheckLine)
-		{
-			// Check bounding box of the horizontal line.
-			constexpr float LineTolerance = 2.0f;
-			if (X >= EventX1 - LineTolerance && X <= EventX2 + LineTolerance &&
-				Y >= EventY - LineTolerance && Y <= EventY + LineTolerance)
-			{
-				return &Event;
-			}
-		}
-
-		if (bCheckBox)
-		{
-			// Check bounding box of the visual box.
-			constexpr float BoxTolerance = 1.0f;
-			if (X >= EventX1 - BoxTolerance && X <= EventX2 + BoxTolerance)
-			{
-				if (EventY < LocalBaselineY)
-				{
-					if (Y >= EventY - BoxTolerance && Y <= LocalBaselineY + BoxTolerance)
-					{
-						return &Event;
-					}
-				}
-				else
-				{
-					if (Y >= LocalBaselineY - BoxTolerance && Y <= EventY + BoxTolerance)
-					{
-						return &Event;
-					}
-				}
-			}
-		}
-	}
-
-	return nullptr;
-}
-
-////////////////////////////////////////////////////////////////////////////////////////////////////
-
-FString FGraphSeries::FormatValue(double Value) const
-{
-	return FString::Printf(TEXT("%g"), Value);
-}
-
 ////////////////////////////////////////////////////////////////////////////////////////////////////
 // FGraphTrack
 ////////////////////////////////////////////////////////////////////////////////////////////////////
 
-FGraphTrack::FGraphTrack(uint64 InTrackId)
-	: FBaseTimingTrack(InTrackId)
+FGraphTrack::FGraphTrack(const FName& InSubType)
+	: FBaseTimingTrack(FName(TEXT("Graph")), InSubType)
 	//, AllSeries()
 	, WhiteBrush(FInsightsStyle::Get().GetBrush("WhiteBrush"))
 	, PointBrush(FEditorStyle::GetBrush("Graph.ExecutionBubble"))
@@ -130,9 +38,6 @@ FGraphTrack::FGraphTrack(uint64 InTrackId)
 	, bDrawPolygon(true)
 	, bUseEventDuration(true)
 	, bDrawBoxes(true)
-	, HoveredGraphEvent()
-	, Tooltip()
-	, TooltipMousePosition(0.0f, 0.0f)
 	, NumAddedEvents(0)
 	, NumDrawPoints(0)
 	, NumDrawLines(0)
@@ -150,56 +55,28 @@ FGraphTrack::~FGraphTrack()
 
 void FGraphTrack::Reset()
 {
-	HoveredGraphEvent.Reset();
-	Tooltip.Reset();
-	TooltipMousePosition.Set(0.0f, 0.0f);
-
 	AllSeries.Reset();
 }
 
 ////////////////////////////////////////////////////////////////////////////////////////////////////
 
-void FGraphTrack::UpdateHoveredState(float MouseX, float MouseY, const FTimingTrackViewport& Viewport)
+void FGraphTrack::PostUpdate(const ITimingTrackUpdateContext& Context)
 {
 	constexpr float HeaderWidth = 100.0f;
 	constexpr float HeaderHeight = 14.0f;
 
-	HoveredGraphEvent.Reset();
+	const float MouseX = Context.GetMousePosition().X;
+	const float MouseY = Context.GetMousePosition().Y;
 
 	if (MouseY >= GetPosY() && MouseY < GetPosY() + GetHeight())
 	{
 		SetHoveredState(true);
 		SetHeaderHoveredState(MouseX < HeaderWidth && MouseY < GetPosY() + HeaderHeight);
-
-		GetEvent(MouseX, MouseY, Viewport, HoveredGraphEvent);
 	}
 	else
 	{
 		SetHoveredState(false);
 	}
-
-	if (HoveredGraphEvent.IsValid())
-	{
-		InitTooltip();
-		Tooltip.SetDesiredOpacity(1.0f);
-	}
-	else
-	{
-		Tooltip.SetDesiredOpacity(0.0f);
-	}
-	TooltipMousePosition.Set(MouseX, MouseY);
-}
-
-////////////////////////////////////////////////////////////////////////////////////////////////////
-
-void FGraphTrack::Update(const FTimingTrackViewport& Viewport)
-{
-	Tooltip.Update();
-
-	if (!TooltipMousePosition.IsZero())
-	{
-		Tooltip.SetPosition(TooltipMousePosition, 0.0f, Viewport.GetWidth(), 0.0f, Viewport.GetHeight());
-	}
 }
 
 ////////////////////////////////////////////////////////////////////////////////////////////////////
@@ -223,24 +100,21 @@ void FGraphTrack::UpdateStats()
 
 ////////////////////////////////////////////////////////////////////////////////////////////////////
 
-void FGraphTrack::DrawBackground(FDrawContext& DrawContext, const FTimingTrackViewport& Viewport) const
+void FGraphTrack::Draw(const ITimingTrackDrawContext& Context) const
 {
-	FDrawHelpers::DrawBackground(DrawContext, WhiteBrush, Viewport, GetPosY(), GetHeight());
-}
-
-////////////////////////////////////////////////////////////////////////////////////////////////////
+	FDrawContext& DrawContext = Context.GetDrawContext();
+	const FTimingTrackViewport& Viewport = Context.GetViewport();
 
-void FGraphTrack::Draw(FDrawContext& DrawContext, const FTimingTrackViewport& Viewport, const FVector2D& MousePosition) const
-{
-	DrawBackground(DrawContext, Viewport);
+	FDrawHelpers::DrawBackground(DrawContext, WhiteBrush, Viewport, GetPosY(), GetHeight());
 
+	// Set clipping area.
 	{
-		FVector2D AbsPos = DrawContext.Geometry.GetAbsolutePosition();
+		const FVector2D AbsPos = DrawContext.Geometry.GetAbsolutePosition();
 		const float L = AbsPos.X;
 		const float R = AbsPos.X + Viewport.GetWidth();
 		const float T = AbsPos.Y + GetPosY();
 		const float B = AbsPos.Y + GetPosY() + GetHeight();
-		FSlateClippingZone ClipZone(FVector2D(L, T), FVector2D(R, T), FVector2D(L, B), FVector2D(R, B));
+		const FSlateClippingZone ClipZone(FVector2D(L, T), FVector2D(R, T), FVector2D(L, B), FVector2D(R, B));
 		DrawContext.ElementList.PushClip(ClipZone);
 	}
 
@@ -267,23 +141,12 @@ void FGraphTrack::Draw(FDrawContext& DrawContext, const FTimingTrackViewport& Vi
 		}
 	}
 
-	if (HoveredGraphEvent.IsValid())
-	{
-		DrawHighlightedEvent(DrawContext, Viewport, HoveredGraphEvent);
-	}
-
+	// Restore clipping.
 	DrawContext.ElementList.PopClip();
 }
 
 ////////////////////////////////////////////////////////////////////////////////////////////////////
 
-void FGraphTrack::PostDraw(FDrawContext& DrawContext, const FTimingTrackViewport& Viewport, const FVector2D& MousePosition) const
-{
-	Tooltip.Draw(DrawContext);
-}
-
-////////////////////////////////////////////////////////////////////////////////////////////////////
-
 void FGraphTrack::DrawSeries(const FGraphSeries& Series, FDrawContext& DrawContext, const FTimingTrackViewport& Viewport) const
 {
 	if (bDrawBoxes)
@@ -510,14 +373,21 @@ void FGraphTrack::DrawSeries(const FGraphSeries& Series, FDrawContext& DrawConte
 
 ////////////////////////////////////////////////////////////////////////////////////////////////////
 
-void FGraphTrack::DrawHighlightedEvent(FDrawContext& DrawContext, const FTimingTrackViewport& Viewport, const FGraphTrack::FEvent& GraphEvent) const
+void FGraphTrack::DrawEvent(const ITimingTrackDrawContext& Context, const ITimingEvent& InTimingEvent, EDrawEventMode InDrawMode) const
 {
-	ensure(GraphEvent.IsValid());
+	ensure(InTimingEvent.CheckTrack(this));
+	ensure(FGraphTrackEvent::CheckTypeName(InTimingEvent));
+
+	const FGraphTrackEvent& GraphEvent = static_cast<const FGraphTrackEvent&>(InTimingEvent);
+	const TSharedPtr<const FGraphSeries> Series = GraphEvent.GetSeries();
+
+	const FTimingTrackViewport& Viewport = Context.GetViewport();
+	FDrawContext& DrawContext = Context.GetDrawContext();
 
-	const float EventX1 = Viewport.TimeToSlateUnitsRounded(GraphEvent.SeriesEvent.Time);
-	const float EventX2 = Viewport.TimeToSlateUnitsRounded(GraphEvent.SeriesEvent.Time + GraphEvent.SeriesEvent.Duration);
+	const float EventX1 = Viewport.TimeToSlateUnitsRounded(GraphEvent.GetStartTime());
+	const float EventX2 = Viewport.TimeToSlateUnitsRounded(GraphEvent.GetEndTime());
 
-	const float EventY1 = GetPosY() + GraphEvent.Series->GetRoundedYForValue(GraphEvent.SeriesEvent.Value);
+	const float EventY1 = GetPosY() + Series->GetRoundedYForValue(GraphEvent.GetValue());
 
 	const FLinearColor HighlightColor(1.0f, 1.0f, 0.0f, 1.0f);
 
@@ -533,14 +403,14 @@ void FGraphTrack::DrawHighlightedEvent(FDrawContext& DrawContext, const FTimingT
 			W = 1.0f;
 		}
 
-		const float EventY2 = GetPosY() + static_cast<float>(GraphEvent.Series->GetBaselineY());
+		const float EventY2 = GetPosY() + static_cast<float>(Series->GetBaselineY());
 
 		const float Y1 = FMath::Min(EventY1, EventY2);
 		const float DY = FMath::Abs(EventY2 - EventY1);
 
 		DrawContext.DrawBox(EventX1 - 1.0f, Y1 - 1.0f, W + 2.0f, DY + 3.0f, WhiteBrush, HighlightColor);
 		DrawContext.LayerId++;
-		DrawContext.DrawBox(EventX1, Y1, W, DY + 1.0f, WhiteBrush, GraphEvent.Series->Color);
+		DrawContext.DrawBox(EventX1, Y1, W, DY + 1.0f, WhiteBrush, Series->Color);
 		DrawContext.LayerId++;
 	}
 
@@ -558,35 +428,41 @@ void FGraphTrack::DrawHighlightedEvent(FDrawContext& DrawContext, const FTimingT
 
 		DrawContext.DrawBox(EventX1 - 1.0f, EventY1 - 1.0f, W + 2.0f, 3.0f, WhiteBrush, HighlightColor);
 		DrawContext.LayerId++;
-		DrawContext.DrawBox(EventX1, EventY1, W, 1.0f, WhiteBrush, GraphEvent.Series->Color);
+		DrawContext.DrawBox(EventX1, EventY1, W, 1.0f, WhiteBrush, Series->Color);
 		DrawContext.LayerId++;
 	}
 
 	// Draw highlighted point.
 	DrawContext.DrawBox(EventX1 - PointVisualSize / 2.0f - 1.5f, EventY1 - PointVisualSize / 2.0f - 1.5f, PointVisualSize + 4.0f, PointVisualSize + 4.0f, PointBrush, HighlightColor);
 	DrawContext.LayerId++;
-	DrawContext.DrawBox(EventX1 - PointVisualSize / 2.0f + 0.5f, EventY1 - PointVisualSize / 2.0f + 0.5f, PointVisualSize, PointVisualSize, PointBrush, GraphEvent.Series->Color);
+	DrawContext.DrawBox(EventX1 - PointVisualSize / 2.0f + 0.5f, EventY1 - PointVisualSize / 2.0f + 0.5f, PointVisualSize, PointVisualSize, PointBrush, Series->Color);
 	DrawContext.LayerId++;
 }
 
 ////////////////////////////////////////////////////////////////////////////////////////////////////
 
-void FGraphTrack::InitTooltip()
+void FGraphTrack::InitTooltip(FTooltipDrawState& InOutTooltip, const ITimingEvent& InTooltipEvent) const
 {
-	Tooltip.ResetContent();
-	Tooltip.AddTitle(HoveredGraphEvent.Series->GetName().ToString(), HoveredGraphEvent.Series->GetColor());
-	Tooltip.AddNameValueTextLine(TEXT("Time:"), TimeUtils::FormatTimeAuto(HoveredGraphEvent.SeriesEvent.Time));
-	Tooltip.AddNameValueTextLine(TEXT("Duration:"), TimeUtils::FormatTimeAuto(HoveredGraphEvent.SeriesEvent.Duration));
-	Tooltip.AddNameValueTextLine(TEXT("Value:"), HoveredGraphEvent.Series->FormatValue(HoveredGraphEvent.SeriesEvent.Value));
-	Tooltip.UpdateLayout();
+	if (InTooltipEvent.CheckTrack(this) && FGraphTrackEvent::CheckTypeName(InTooltipEvent))
+	{
+		const FGraphTrackEvent& TooltipEvent = static_cast<const FGraphTrackEvent&>(InTooltipEvent);
+		const TSharedRef<const FGraphSeries> Series = TooltipEvent.GetSeries();
+
+		InOutTooltip.ResetContent();
+		InOutTooltip.AddTitle(Series->GetName().ToString(), Series->GetColor());
+		InOutTooltip.AddNameValueTextLine(TEXT("Time:"), TimeUtils::FormatTimeAuto(TooltipEvent.GetStartTime()));
+		InOutTooltip.AddNameValueTextLine(TEXT("Duration:"), TimeUtils::FormatTimeAuto(TooltipEvent.GetDuration()));
+		InOutTooltip.AddNameValueTextLine(TEXT("Value:"), Series->FormatValue(TooltipEvent.GetValue()));
+		InOutTooltip.UpdateLayout();
+	}
 }
 
 ////////////////////////////////////////////////////////////////////////////////////////////////////
 
-const bool FGraphTrack::GetEvent(const float MouseX, const float MouseY, const FTimingTrackViewport& Viewport, FGraphTrack::FEvent& OutEvent) const
+const TSharedPtr<const ITimingEvent> FGraphTrack::GetEvent(float InPosX, float InPosY, const FTimingTrackViewport& Viewport) const
 {
-	const float LocalPosX = MouseX;
-	const float LocalPosY = MouseY - GetPosY();
+	const float LocalPosX = InPosX;
+	const float LocalPosY = InPosY - GetPosY();
 
 	const bool bCheckLine = bUseEventDuration && (bDrawLines || bDrawPolygon);
 
@@ -596,17 +472,15 @@ const bool FGraphTrack::GetEvent(const float MouseX, const float MouseY, const F
 		const TSharedPtr<FGraphSeries>& Series = AllSeries[SeriesIndex];
 		if (Series->IsVisible())
 		{
-			const FGraphSeries::FEvent* Event = Series->GetEvent(LocalPosX, LocalPosY, Viewport, bCheckLine, bDrawBoxes);
+			const FGraphSeriesEvent* Event = Series->GetEvent(LocalPosX, LocalPosY, Viewport, bCheckLine, bDrawBoxes);
 			if (Event != nullptr)
 			{
-				OutEvent.Series = Series;
-				OutEvent.SeriesEvent = *Event;
-				return true;
+				return MakeShareable(new FGraphTrackEvent(SharedThis(this), Series.ToSharedRef(), *Event));
 			}
 		}
 	}
 
-	return false;
+	return nullptr;
 }
 
 ////////////////////////////////////////////////////////////////////////////////////////////////////
@@ -895,8 +769,8 @@ bool FGraphTrack::ContextMenu_ShowSeries_IsChecked(FGraphSeries* Series)
 // FRandomGraphTrack
 ////////////////////////////////////////////////////////////////////////////////////////////////////
 
-FRandomGraphTrack::FRandomGraphTrack(uint64 InTrackId)
-	: FGraphTrack(InTrackId)
+FRandomGraphTrack::FRandomGraphTrack()
+	: FGraphTrack(FName(TEXT("Random")))
 {
 	bDrawPoints = true;
 	bDrawPointsWithBorder = true;
@@ -940,9 +814,9 @@ FRandomGraphTrack::~FRandomGraphTrack()
 
 ////////////////////////////////////////////////////////////////////////////////////////////////////
 
-void FRandomGraphTrack::Update(const FTimingTrackViewport& Viewport)
+void FRandomGraphTrack::Update(const ITimingTrackUpdateContext& Context)
 {
-	FGraphTrack::Update(Viewport);
+	FGraphTrack::Update(Context);
 
 	if (IsDirty())
 	{
@@ -950,6 +824,8 @@ void FRandomGraphTrack::Update(const FTimingTrackViewport& Viewport)
 
 		NumAddedEvents = 0;
 
+		const FTimingTrackViewport& Viewport = Context.GetViewport();
+
 		int32 Seed = 0;
 		for (TSharedPtr<FGraphSeries>& Series : AllSeries)
 		{
@@ -1028,341 +904,6 @@ void FRandomGraphTrack::GenerateSeries(FGraphSeries& Series, const FTimingTrackV
 	}
 }
 
-////////////////////////////////////////////////////////////////////////////////////////////////////
-// FGraphTrackBuilder
-////////////////////////////////////////////////////////////////////////////////////////////////////
-
-FGraphTrackBuilder::FGraphTrackBuilder(FGraphTrack& InTrack, FGraphSeries& InSeries, const FTimingTrackViewport& InViewport)
-	: Track(InTrack)
-	, Series(InSeries)
-	, Viewport(InViewport)
-{
-	Series.Events.Reset();
-	BeginPoints();
-	BeginConnectedLines();
-	BeginBoxes();
-}
-
-////////////////////////////////////////////////////////////////////////////////////////////////////
-
-FGraphTrackBuilder::~FGraphTrackBuilder()
-{
-	EndPoints();
-	EndConnectedLines();
-	EndBoxes();
-}
-
-////////////////////////////////////////////////////////////////////////////////////////////////////
-
-void FGraphTrackBuilder::AddEvent(double Time, double Duration, double Value)
-{
-	Track.NumAddedEvents++;
-
-	bool bKeepEvent = (Viewport.GetViewportDXForDuration(Duration) > 1.0f); // always keep the events wider than 1px
-
-	//if (Track.bDrawPoints)
-	{
-		bKeepEvent |= AddPoint(Time, Value);
-	}
-
-	if (Track.bDrawLines || Track.bDrawPolygon)
-	{
-		AddConnectedLine(Time, Value);
-
-		if (Track.bUseEventDuration && Duration != 0.0)
-		{
-			AddConnectedLine(Time + Duration, Value);
-		}
-	}
-
-	if (Track.bDrawBoxes)
-	{
-		AddBox(Time, Duration, Value);
-	}
-
-	if (bKeepEvent)
-	{
-		Series.Events.Add({ Time, Duration, Value });
-	}
-}
-
-////////////////////////////////////////////////////////////////////////////////////////////////////
-// FGraphTrackBuilder - Points
-////////////////////////////////////////////////////////////////////////////////////////////////////
-
-void FGraphTrackBuilder::BeginPoints()
-{
-	Series.Points.Reset();
-
-	PointsCurrentX = -DBL_MAX;
-
-	PointsAtCurrentX.Reset();
-	int32 MaxPointsPerLineScan = FMath::CeilToInt(Track.GetHeight() / FGraphTrack::PointSizeY) + 1;
-	if (MaxPointsPerLineScan > 0)
-	{
-		PointsAtCurrentX.AddDefaulted(MaxPointsPerLineScan);
-	}
-}
-
-////////////////////////////////////////////////////////////////////////////////////////////////////
-
-bool FGraphTrackBuilder::AddPoint(double Time, double Value)
-{
-	const float X = Viewport.TimeToSlateUnitsRounded(Time);
-	if (X < -FGraphTrack::PointVisualSize / 2.0f || X >= Viewport.GetWidth() + FGraphTrack::PointVisualSize / 2.0f)
-	{
-		return false;
-	}
-
-	// Align the X with a grid of GraphTrackPointDX pixels in size, in the global space (i.e. scroll independent).
-	const double AlignedX = FMath::RoundToDouble(Time * Viewport.GetScaleX() / FGraphTrack::PointSizeX) * FGraphTrack::PointSizeX;
-
-	if (AlignedX > PointsCurrentX + FGraphTrack::PointSizeX - 0.5)
-	{
-		FlushPoints();
-		PointsCurrentX = AlignedX;
-	}
-
-	const float Y = Series.GetRoundedYForValue(Value);
-
-	bool bIsNewVisiblePoint = false;
-
-	int32 Index = FMath::RoundToInt(Y / FGraphTrack::PointSizeY);
-	if (Index >= 0 && Index < PointsAtCurrentX.Num())
-	{
-		FPointInfo& Pt = PointsAtCurrentX[Index];
-		bIsNewVisiblePoint = !Pt.bValid;
-		Pt.bValid = true;
-		Pt.X = X;
-		Pt.Y = Y;
-	}
-
-	return bIsNewVisiblePoint;
-}
-
-////////////////////////////////////////////////////////////////////////////////////////////////////
-
-void FGraphTrackBuilder::FlushPoints()
-{
-	for (int32 Index = 0; Index < PointsAtCurrentX.Num(); ++Index)
-	{
-		FPointInfo& Pt = PointsAtCurrentX[Index];
-		if (Pt.bValid)
-		{
-			Pt.bValid = false;
-			Series.Points.Add(FVector2D(Pt.X, Pt.Y));
-		}
-	}
-}
-
-////////////////////////////////////////////////////////////////////////////////////////////////////
-
-void FGraphTrackBuilder::EndPoints()
-{
-	FlushPoints();
-}
-
-////////////////////////////////////////////////////////////////////////////////////////////////////
-// FGraphTrackBuilder - Connected Lines
-////////////////////////////////////////////////////////////////////////////////////////////////////
-
-void FGraphTrackBuilder::BeginConnectedLines()
-{
-	Series.LinePoints.Reset();
-
-	LinesCurrentX = -FLT_MAX;
-	LinesMinY = FLT_MAX;
-	LinesMaxY = -FLT_MAX;
-	LinesFirstY = FLT_MAX;
-	LinesLastY = FLT_MAX;
-	bIsLastLineAdded = false;
-}
-
-////////////////////////////////////////////////////////////////////////////////////////////////////
-
-void FGraphTrackBuilder::AddConnectedLine(double Time, double Value)
-{
-	if (bIsLastLineAdded)
-	{
-		return;
-	}
-
-	const float X = Viewport.TimeToSlateUnitsRounded(Viewport.RestrictEndTime(Time));
-	const float Y = Series.GetRoundedYForValue(Value);
-
-	ensure(X >= LinesCurrentX); // we are assuming events are already sorted by Time
-
-	if (X < 0)
-	{
-		LinesCurrentX = X;
-		LinesLastY = Y;
-		return;
-	}
-
-	if (X >= Viewport.GetWidth())
-	{
-		if (!bIsLastLineAdded)
-		{
-			bIsLastLineAdded = true;
-
-			if (LinesLastY != FLT_MAX)
-			{
-				FlushConnectedLine();
-
-				Series.LinePoints.Add(FVector2D(LinesCurrentX, LinesLastY));
-				Series.LinePoints.Add(FVector2D(X, Y));
-			}
-
-			// Reset the "reduction line" so last FlushConnectedLine() call will do nothing.
-			LinesMinY = Y;
-			LinesMaxY = Y;
-		}
-		return;
-	}
-
-	if (X > LinesCurrentX)
-	{
-		if (LinesLastY != FLT_MAX)
-		{
-			FlushConnectedLine();
-
-			Series.LinePoints.Add(FVector2D(LinesCurrentX, LinesLastY));
-			Series.LinePoints.Add(FVector2D(X, Y));
-		}
-
-		// Advance the "reduction line".
-		LinesCurrentX = X;
-		LinesMinY = Y;
-		LinesMaxY = Y;
-		LinesFirstY = Y;
-		LinesLastY = Y;
-	}
-	else
-	{
-		// Merge current line with the "reduction line".
-		if (Y < LinesMinY)
-		{
-			LinesMinY = Y;
-		}
-		if (Y > LinesMaxY)
-		{
-			LinesMaxY = Y;
-		}
-		LinesLastY = Y;
-	}
-}
-
-////////////////////////////////////////////////////////////////////////////////////////////////////
-
-void FGraphTrackBuilder::FlushConnectedLine()
-{
-	if (LinesCurrentX >= 0.0f && LinesMinY != LinesMaxY)
-	{
-		Series.LinePoints.Add(FVector2D(LinesCurrentX, LinesMaxY));
-		Series.LinePoints.Add(FVector2D(LinesCurrentX, LinesMinY));
-	}
-}
-
-////////////////////////////////////////////////////////////////////////////////////////////////////
-
-void FGraphTrackBuilder::EndConnectedLines()
-{
-	FlushConnectedLine();
-}
-
-////////////////////////////////////////////////////////////////////////////////////////////////////
-// FGraphTrackBuilder - Boxes
-////////////////////////////////////////////////////////////////////////////////////////////////////
-
-void FGraphTrackBuilder::BeginBoxes()
-{
-	Series.Boxes.Reset();
-}
-
-////////////////////////////////////////////////////////////////////////////////////////////////////
-
-void FGraphTrackBuilder::AddBox(double Time, double Duration, double Value)
-{
-	float X1 = Viewport.TimeToSlateUnitsRounded(Time);
-	if (X1 > Viewport.GetWidth())
-	{
-		return;
-	}
-
-	double EndTime = Viewport.RestrictEndTime(Time + Duration);
-	float X2 = Viewport.TimeToSlateUnitsRounded(EndTime);
-	if (X2 < 0)
-	{
-		return;
-	}
-
-	float W = X2 - X1;
-	ensure(W >= 0); // we expect events to be sorted
-
-	// Timing events are displayed with minimum 1px (including empty ones).
-	if (W == 0)
-	{
-		W = 1.0f;
-	}
-
-	const float Y = Series.GetRoundedYForValue(Value);
-
-	// simple reduction
-	if (W == 1.0f && Series.Boxes.Num() > 0)
-	{
-		FGraphSeries::FBox& LastBox = Series.Boxes.Last();
-		if (LastBox.W == 1.0f && LastBox.X == X1)
-		{
-			const float RoundedBaselineY = FMath::RoundToFloat(Series.GetBaselineY());
-
-			if (LastBox.Y < RoundedBaselineY)
-			{
-				if (Y < RoundedBaselineY)
-				{
-					// Merge current box with previous one.
-					if (Y < LastBox.Y)
-					{
-						LastBox.Y = Y;
-					}
-					return;
-				}
-			}
-			else
-			{
-				if (Y >= RoundedBaselineY)
-				{
-					// Merge current box with previous one.
-					if (Y > LastBox.Y)
-					{
-						LastBox.Y = Y;
-					}
-					return;
-				}
-			}
-		}
-	}
-
-	FGraphSeries::FBox Box;
-	Box.X = X1;
-	Box.W = W;
-	Box.Y = Y;
-	Series.Boxes.Add(Box);
-}
-
-////////////////////////////////////////////////////////////////////////////////////////////////////
-
-void FGraphTrackBuilder::FlushBox()
-{
-	// TODO: reduction algorithm
-}
-
-////////////////////////////////////////////////////////////////////////////////////////////////////
-
-void FGraphTrackBuilder::EndBoxes()
-{
-	FlushBox();
-}
-
 ////////////////////////////////////////////////////////////////////////////////////////////////////
 
 #undef LOCTEXT_NAMESPACE
diff --git a/Engine/Source/Developer/TraceInsights/Private/Insights/ViewModels/LoadingTimingTrack.cpp b/Engine/Source/Developer/TraceInsights/Private/Insights/ViewModels/LoadingTimingTrack.cpp
--- a/Engine/Source/Developer/TraceInsights/Private/Insights/ViewModels/LoadingTimingTrack.cpp
+++ b/Engine/Source/Developer/TraceInsights/Private/Insights/ViewModels/LoadingTimingTrack.cpp
@@ -3,19 +3,22 @@
 #include "LoadingTimingTrack.h"
 
 #include "Fonts/FontMeasure.h"
+#include "Framework/MultiBox/MultiBoxBuilder.h"
 #include "Styling/SlateBrush.h"
 #include "TraceServices/AnalysisService.h"
 #include "TraceServices/SessionService.h"
 
 // Insights
-#include "Insights/InsightsManager.h"
 #include "Insights/Common/PaintUtils.h"
 #include "Insights/Common/TimeUtils.h"
+#include "Insights/InsightsManager.h"
+#include "Insights/ITimingViewSession.h"
 #include "Insights/ViewModels/TimingEvent.h"
 #include "Insights/ViewModels/TimingTrackViewport.h"
 #include "Insights/ViewModels/TimingViewDrawHelper.h"
 #include "Insights/ViewModels/TooltipDrawState.h"
 #include "Insights/ViewModels/TimingEventSearch.h"
+#include "Insights/Widgets/STimingView.h"
 
 #define LOCTEXT_NAMESPACE "LoadingTimingTrack"
 
@@ -23,14 +26,144 @@
 // FLoadingSharedState
 ////////////////////////////////////////////////////////////////////////////////////////////////////
 
-void FLoadingSharedState::Reset()
+void FLoadingSharedState::OnBeginSession(Insights::ITimingViewSession& InSession)
+{
+	if (&InSession != TimingView)
+	{
+		return;
+	}
+
+	if (TimingView && TimingView->IsAssetLoadingModeEnabled())
+	{
+		bShowHideAllLoadingTracks = true;
+	}
+	else
+	{
+		bShowHideAllLoadingTracks = false;
+	}
+
+	LoadingTracks.Reset();
+
+	SetColorSchema(3);
+}
+
+////////////////////////////////////////////////////////////////////////////////////////////////////
+
+void FLoadingSharedState::OnEndSession(Insights::ITimingViewSession& InSession)
+{
+	if (&InSession != TimingView)
+	{
+		return;
+	}
+
+	bShowHideAllLoadingTracks = false;
+
+	LoadingTracks.Reset();
+
+	GetEventNameDelegate = nullptr;
+}
+
+////////////////////////////////////////////////////////////////////////////////////////////////////
+
+void FLoadingSharedState::Tick(Insights::ITimingViewSession& InSession, const Trace::IAnalysisSession& InAnalysisSession)
+{
+	if (&InSession != TimingView)
+	{
+		return;
+	}
+
+	const Trace::ILoadTimeProfilerProvider* LoadTimeProfilerProvider = Trace::ReadLoadTimeProfilerProvider(InAnalysisSession);
+	if (LoadTimeProfilerProvider)
+	{
+		Trace::FAnalysisSessionReadScope SessionReadScope(InAnalysisSession);
+
+		// Iterate through threads.
+		const Trace::IThreadProvider& ThreadProvider = Trace::ReadThreadProvider(InAnalysisSession);
+		ThreadProvider.EnumerateThreads([this, &InSession, LoadTimeProfilerProvider](const Trace::FThreadInfo& ThreadInfo)
+		{
+			// Check available Asset Loading tracks.
+			uint32 LoadingTimelineIndex;
+			if (LoadTimeProfilerProvider->GetCpuThreadTimelineIndex(ThreadInfo.Id, LoadingTimelineIndex))
+			{
+				if (!LoadingTracks.Contains(LoadingTimelineIndex))
+				{
+					const TCHAR* const GroupName = ThreadInfo.GroupName ? ThreadInfo.GroupName : ThreadInfo.Name;
+					const FString TrackName(ThreadInfo.Name && *ThreadInfo.Name ? FString::Printf(TEXT("Loading - %s"), ThreadInfo.Name) : FString::Printf(TEXT("Loading - Thread %u"), ThreadInfo.Id));
+
+					TSharedPtr<FLoadingTimingTrack> LoadingThreadTrack = MakeShareable(new FLoadingTimingTrack(*this, LoadingTimelineIndex, GroupName, TrackName));
+					LoadingThreadTrack->SetOrder(-100 + LoadingTracks.Num());
+					LoadingThreadTrack->SetVisibilityFlag(bShowHideAllLoadingTracks);
+					InSession.AddScrollableTrack(LoadingThreadTrack);
+					LoadingTracks.Add(LoadingTimelineIndex, LoadingThreadTrack);
+				}
+			}
+		});
+	}
+}
+
+////////////////////////////////////////////////////////////////////////////////////////////////////
+
+void FLoadingSharedState::ExtendFilterMenu(Insights::ITimingViewSession& InSession, FMenuBuilder& InOutMenuBuilder)
+{
+	if (&InSession != TimingView)
+	{
+		return;
+	}
+
+	InOutMenuBuilder.BeginSection("Asset Loading", LOCTEXT("AssetLoadingHeading", "Asset Loading"));
+	{
+		InOutMenuBuilder.AddMenuEntry(
+			LOCTEXT("ShowAllLoadingTracks", "Asset Loading Tracks - L"),
+			LOCTEXT("ShowAllLoadingTracks_Tooltip", "Show/hide the Asset Loading tracks"),
+			FSlateIcon(),
+			FUIAction(FExecuteAction::CreateSP(this, &FLoadingSharedState::ShowHideAllLoadingTracks_Execute),
+					  FCanExecuteAction(),
+					  FIsActionChecked::CreateSP(this, &FLoadingSharedState::ShowHideAllLoadingTracks_IsChecked)),
+			"QuickFilterSeparator",
+			EUserInterfaceActionType::ToggleButton
+		);
+	}
+	InOutMenuBuilder.EndSection();
+}
+
+////////////////////////////////////////////////////////////////////////////////////////////////////
+
+bool FLoadingSharedState::ShowHideAllLoadingTracks_IsChecked() const
+{
+	return bShowHideAllLoadingTracks;
+}
+
+////////////////////////////////////////////////////////////////////////////////////////////////////
+
+void FLoadingSharedState::ShowHideAllLoadingTracks_Execute()
+{
+	bShowHideAllLoadingTracks = !bShowHideAllLoadingTracks;
+
+	for (const auto& KV : LoadingTracks)
+	{
+		FLoadingTimingTrack& Track = *KV.Value;
+		Track.SetVisibilityFlag(bShowHideAllLoadingTracks);
+	}
+
+	if (TimingView)
+	{
+		TimingView->OnTrackVisibilityChanged();
+	}
+}
+
+////////////////////////////////////////////////////////////////////////////////////////////////////
+
+void FLoadingSharedState::OnTracksChanged(Insights::ITimingViewSession& InSession, int32& InOutOrder)
 {
-	LoadingGetEventNameFn = FLoadingTrackGetEventNameDelegate::CreateRaw(this, &FLoadingSharedState::GetEventName);
+	if (&InSession != TimingView)
+	{
+		return;
+	}
 }
 
 ////////////////////////////////////////////////////////////////////////////////////////////////////
 
-const TCHAR* FLoadingSharedState::GetEventNameByExportEventType(uint32 Depth, const Trace::FLoadTimeProfilerCpuEvent& Event) const
+const TCHAR* FLoadingSharedState::GetEventNameByEventType(uint32 Depth, const Trace::FLoadTimeProfilerCpuEvent& Event) const
 {
 	return Trace::GetLoadTimeProfilerObjectEventTypeString(Event.EventType);
 }
@@ -51,7 +184,7 @@ const TCHAR* FLoadingSharedState::GetEventNameByExportClassName(uint32 Depth, co
 
 ////////////////////////////////////////////////////////////////////////////////////////////////////
 
-const TCHAR* FLoadingSharedState::GetEventName(uint32 Depth, const Trace::FLoadTimeProfilerCpuEvent& Event) const
+const TCHAR* FLoadingSharedState::GetEventNameByPackageAndExportClassName(uint32 Depth, const Trace::FLoadTimeProfilerCpuEvent& Event) const
 {
 	if (Depth == 0)
 	{
@@ -71,9 +204,9 @@ const TCHAR* FLoadingSharedState::GetEventName(uint32 Depth, const Trace::FLoadT
 
 ////////////////////////////////////////////////////////////////////////////////////////////////////
 
-const TCHAR* FLoadingSharedState::GetEventNameEx(uint32 Depth, const Trace::FLoadTimeProfilerCpuEvent& Event) const
+const TCHAR* FLoadingSharedState::GetEventName(uint32 Depth, const Trace::FLoadTimeProfilerCpuEvent& Event) const
 {
-	return LoadingGetEventNameFn.Execute(Depth, Event);
+	return GetEventNameDelegate.Execute(Depth, Event);
 }
 
 ////////////////////////////////////////////////////////////////////////////////////////////////////
@@ -82,10 +215,10 @@ void FLoadingSharedState::SetColorSchema(int32 Schema)
 {
 	switch (Schema)
 	{
-		case 0: LoadingGetEventNameFn = FLoadingTrackGetEventNameDelegate::CreateRaw(this, &FLoadingSharedState::GetEventNameByExportEventType); break;
-		case 1: LoadingGetEventNameFn = FLoadingTrackGetEventNameDelegate::CreateRaw(this, &FLoadingSharedState::GetEventNameByPackageName); break;
-		case 2: LoadingGetEventNameFn = FLoadingTrackGetEventNameDelegate::CreateRaw(this, &FLoadingSharedState::GetEventNameByExportClassName); break;
-		case 3: LoadingGetEventNameFn = FLoadingTrackGetEventNameDelegate::CreateRaw(this, &FLoadingSharedState::GetEventName); break;
+		case 0: GetEventNameDelegate = FLoadingTrackGetEventNameDelegate::CreateRaw(this, &FLoadingSharedState::GetEventNameByEventType); break;
+		case 1: GetEventNameDelegate = FLoadingTrackGetEventNameDelegate::CreateRaw(this, &FLoadingSharedState::GetEventNameByPackageName); break;
+		case 2: GetEventNameDelegate = FLoadingTrackGetEventNameDelegate::CreateRaw(this, &FLoadingSharedState::GetEventNameByExportClassName); break;
+		case 3: GetEventNameDelegate = FLoadingTrackGetEventNameDelegate::CreateRaw(this, &FLoadingSharedState::GetEventNameByPackageAndExportClassName); break;
 	};
 }
 
@@ -93,57 +226,62 @@ void FLoadingSharedState::SetColorSchema(int32 Schema)
 // FLoadingTimingTrack
 ////////////////////////////////////////////////////////////////////////////////////////////////////
 
-void FLoadingTimingTrack::InitTooltip(FTooltipDrawState& Tooltip, const FTimingEvent& HoveredTimingEvent) const
+void FLoadingTimingTrack::InitTooltip(FTooltipDrawState& InOutTooltip, const ITimingEvent& InTooltipEvent) const
 {
-	auto MatchEvent = [&HoveredTimingEvent](double InStartTime, double InEndTime, uint32 InDepth)
-	{ 
-		return InDepth == HoveredTimingEvent.Depth
-			&& InStartTime == HoveredTimingEvent.StartTime
-			&& InEndTime == HoveredTimingEvent.EndTime;
-	};
-
-	FTimingEventSearchParameters SearchParameters(HoveredTimingEvent.StartTime, HoveredTimingEvent.EndTime, ETimingEventSearchFlags::StopAtFirstMatch, MatchEvent);
-	FindLoadTimeProfilerCpuEvent(SearchParameters, [this, &Tooltip, &HoveredTimingEvent](double InFoundStartTime, double InFoundEndTime, uint32 InFoundDepth, const Trace::FLoadTimeProfilerCpuEvent& InFoundEvent)
+	if (InTooltipEvent.CheckTrack(this) && FTimingEvent::CheckTypeName(InTooltipEvent))
 	{
-		Tooltip.ResetContent();
+		const FTimingEvent& TooltipEvent = static_cast<const FTimingEvent&>(InTooltipEvent);
 
-		Tooltip.AddTitle(State->GetEventNameEx(HoveredTimingEvent.Depth, InFoundEvent));
+		auto MatchEvent = [&TooltipEvent](double InStartTime, double InEndTime, uint32 InDepth)
+		{
+			return InDepth == TooltipEvent.GetDepth()
+				&& InStartTime == TooltipEvent.GetStartTime()
+				&& InEndTime == TooltipEvent.GetEndTime();
+		};
 
-		const Trace::FPackageInfo* Package = InFoundEvent.Package;
-		const Trace::FPackageExportInfo* Export = InFoundEvent.Export;
+		FTimingEventSearchParameters SearchParameters(TooltipEvent.GetStartTime(), TooltipEvent.GetEndTime(), ETimingEventSearchFlags::StopAtFirstMatch, MatchEvent);
+		FindLoadTimeProfilerCpuEvent(SearchParameters, [this, &InOutTooltip, &TooltipEvent](double InFoundStartTime, double InFoundEndTime, uint32 InFoundDepth, const Trace::FLoadTimeProfilerCpuEvent& InFoundEvent)
+		{
+			InOutTooltip.ResetContent();
 
-		Tooltip.AddNameValueTextLine(TEXT("Duration:"), TimeUtils::FormatTimeAuto(HoveredTimingEvent.Duration()));
-		Tooltip.AddNameValueTextLine(TEXT("Depth:"), FString::Printf(TEXT("%d"), HoveredTimingEvent.Depth));
+			InOutTooltip.AddTitle(SharedState.GetEventName(TooltipEvent.GetDepth(), InFoundEvent));
 
-		if (Package)
-		{
-			Tooltip.AddNameValueTextLine(TEXT("Package Name:"), Package->Name);
-			Tooltip.AddNameValueTextLine(TEXT("Header Size:"), FString::Printf(TEXT("%s bytes"), *FText::AsNumber(Package->Summary.TotalHeaderSize).ToString()));
-			Tooltip.AddNameValueTextLine(TEXT("Package Summary:"), FString::Printf(TEXT("%d imports, %d exports"), Package->Summary.ImportCount, Package->Summary.ExportCount));
-		}
+			const Trace::FPackageInfo* Package = InFoundEvent.Package;
+			const Trace::FPackageExportInfo* Export = InFoundEvent.Export;
 
-		Tooltip.AddNameValueTextLine(TEXT("Export Event:"), FString::Printf(TEXT("%s"), Trace::GetLoadTimeProfilerObjectEventTypeString(InFoundEvent.EventType)));
+			InOutTooltip.AddNameValueTextLine(TEXT("Duration:"), TimeUtils::FormatTimeAuto(TooltipEvent.GetDuration()));
+			InOutTooltip.AddNameValueTextLine(TEXT("Depth:"), FString::Printf(TEXT("%d"), TooltipEvent.GetDepth()));
 
-		if (Export)
-		{
-			Tooltip.AddNameValueTextLine(TEXT("Export Class:"), Export->Class ? Export->Class->Name : TEXT("N/A"));
-			Tooltip.AddNameValueTextLine(TEXT("Serial Size:"), FString::Printf(TEXT("%s bytes"), *FText::AsNumber(Export->SerialSize).ToString()));
-		}
+			if (Package)
+			{
+				InOutTooltip.AddNameValueTextLine(TEXT("Package Name:"), Package->Name);
+				InOutTooltip.AddNameValueTextLine(TEXT("Header Size:"), FString::Printf(TEXT("%s bytes"), *FText::AsNumber(Package->Summary.TotalHeaderSize).ToString()));
+				InOutTooltip.AddNameValueTextLine(TEXT("Package Summary:"), FString::Printf(TEXT("%d imports, %d exports"), Package->Summary.ImportCount, Package->Summary.ExportCount));
+			}
 
-		Tooltip.UpdateLayout();
-	});
+			InOutTooltip.AddNameValueTextLine(TEXT("Export Event:"), FString::Printf(TEXT("%s"), Trace::GetLoadTimeProfilerObjectEventTypeString(InFoundEvent.EventType)));
+
+			if (Export)
+			{
+				InOutTooltip.AddNameValueTextLine(TEXT("Export Class:"), Export->Class ? Export->Class->Name : TEXT("N/A"));
+				InOutTooltip.AddNameValueTextLine(TEXT("Serial Size:"), FString::Printf(TEXT("%s bytes"), *FText::AsNumber(Export->SerialSize).ToString()));
+			}
+
+			InOutTooltip.UpdateLayout();
+		});
+	}
 }
 
-////////////////////////////////////////////////////////////////////////////////////////////////////
-// FLoadingMainThreadTimingTrack
 ////////////////////////////////////////////////////////////////////////////////////////////////////
 
-void FLoadingTimingTrack::Draw(ITimingViewDrawHelper& Helper) const
+void FLoadingTimingTrack::PreUpdate(const ITimingTrackUpdateContext& Context)
 {
-	FTimingEventsTrack& Track = *const_cast<FLoadingTimingTrack*>(this);
-
-	if (Helper.BeginTimeline(Track))
+	if (IsDirty() || Context.GetViewport().IsHorizontalViewportDirty())
 	{
+		ClearDirtyFlag();
+
+		FTimingEventsTrackDrawStateBuilder Builder(DrawState, Context.GetViewport());
+
 		TSharedPtr<const Trace::IAnalysisSession> Session = FInsightsManager::Get()->GetSession();
 		if (Trace::ReadLoadTimeProfilerProvider(*Session.Get()))
 		{
@@ -151,40 +289,60 @@ void FLoadingTimingTrack::Draw(ITimingViewDrawHelper& Helper) const
 
 			const Trace::ILoadTimeProfilerProvider& LoadTimeProfilerProvider = *Trace::ReadLoadTimeProfilerProvider(*Session.Get());
 
-			LoadTimeProfilerProvider.ReadTimeline(TimelineIndex, [this, &Helper](const Trace::ILoadTimeProfilerProvider::CpuTimeline& Timeline)
+			LoadTimeProfilerProvider.ReadTimeline(TimelineIndex, [this, &Builder](const Trace::ILoadTimeProfilerProvider::CpuTimeline& Timeline)
 			{
+				const FTimingTrackViewport& Viewport = Builder.GetViewport();
 				if (FTimingEventsTrack::bUseDownSampling)
 				{
-					const double SecondsPerPixel = 1.0 / Helper.GetViewport().GetScaleX();
-					Timeline.EnumerateEventsDownSampled(Helper.GetViewport().GetStartTime(), Helper.GetViewport().GetEndTime(), SecondsPerPixel, [this, &Helper](double StartTime, double EndTime, uint32 Depth, const Trace::FLoadTimeProfilerCpuEvent& Event)
+					const double SecondsPerPixel = 1.0 / Viewport.GetScaleX();
+					Timeline.EnumerateEventsDownSampled(Viewport.GetStartTime(), Viewport.GetEndTime(), SecondsPerPixel, [this, &Builder](double StartTime, double EndTime, uint32 Depth, const Trace::FLoadTimeProfilerCpuEvent& Event)
 					{
-						const TCHAR* Name = State->GetEventNameEx(Depth, Event);
-						Helper.AddEvent(StartTime, EndTime, Depth, Name);
+						const TCHAR* Name = SharedState.GetEventName(Depth, Event);
+						const uint64 Type = static_cast<uint64>(Event.EventType);
+						const uint32 Color = 0;
+						Builder.AddEvent(StartTime, EndTime, Depth, Name, Type, Color);
 					});
 				}
 				else
 				{
-					Timeline.EnumerateEvents(Helper.GetViewport().GetStartTime(), Helper.GetViewport().GetEndTime(), [this, &Helper](double StartTime, double EndTime, uint32 Depth, const Trace::FLoadTimeProfilerCpuEvent& Event)
+					Timeline.EnumerateEvents(Viewport.GetStartTime(), Viewport.GetEndTime(), [this, &Builder](double StartTime, double EndTime, uint32 Depth, const Trace::FLoadTimeProfilerCpuEvent& Event)
 					{
-						const TCHAR* Name = State->GetEventNameEx(Depth, Event);
-						Helper.AddEvent(StartTime, EndTime, Depth, Name);
+						const TCHAR* Name = SharedState.GetEventName(Depth, Event);
+						const uint64 Type = static_cast<uint64>(Event.EventType);
+						const uint32 Color = 0;
+						Builder.AddEvent(StartTime, EndTime, Depth, Name, Type, Color);
 					});
 				}
 			});
 		}
 
-		Helper.EndTimeline(Track);
+		Builder.Flush();
+		SetNumLanes(Builder.GetMaxDepth() + 1);
 	}
+
+	UpdateTrackHeight(Context);
+}
+
+////////////////////////////////////////////////////////////////////////////////////////////////////
+
+void FLoadingTimingTrack::Draw(const ITimingTrackDrawContext& Context) const
+{
+	const FTimingViewDrawHelper& Helper = Context.GetHelper();
+	Helper.DrawEventsTrack(DrawState, *this);
 }
 
 ////////////////////////////////////////////////////////////////////////////////////////////////////
 
-bool FLoadingTimingTrack::SearchTimingEvent(const FTimingEventSearchParameters& InSearchParameters, FTimingEvent& InOutTimingEvent) const
+const TSharedPtr<const ITimingEvent> FLoadingTimingTrack::SearchEvent(const FTimingEventSearchParameters& InSearchParameters) const
 {
-	return FindLoadTimeProfilerCpuEvent(InSearchParameters, [this, &InOutTimingEvent](double InFoundStartTime, double InFoundEndTime, uint32 InFoundDepth, const Trace::FLoadTimeProfilerCpuEvent& InFoundEvent)
+	TSharedPtr<const ITimingEvent> FoundEvent;
+
+	FindLoadTimeProfilerCpuEvent(InSearchParameters, [this, &FoundEvent](double InFoundStartTime, double InFoundEndTime, uint32 InFoundDepth, const Trace::FLoadTimeProfilerCpuEvent& InFoundEvent)
 	{
-		InOutTimingEvent = FTimingEvent(this, InFoundStartTime, InFoundEndTime, InFoundDepth);
+		FoundEvent = MakeShareable(new FTimingEvent(SharedThis(this), InFoundStartTime, InFoundEndTime, InFoundDepth));
 	});
+
+	return FoundEvent;
 }
 
 ////////////////////////////////////////////////////////////////////////////////////////////////////
diff --git a/Engine/Source/Developer/TraceInsights/Private/Insights/ViewModels/LoadingTimingTrack.h b/Engine/Source/Developer/TraceInsights/Private/Insights/ViewModels/LoadingTimingTrack.h
--- a/Engine/Source/Developer/TraceInsights/Private/Insights/ViewModels/LoadingTimingTrack.h
+++ b/Engine/Source/Developer/TraceInsights/Private/Insights/ViewModels/LoadingTimingTrack.h
@@ -6,31 +6,61 @@
 #include "TraceServices/AnalysisService.h"
 
 // Insights
+#include "Insights/ITimingViewExtender.h"
 #include "Insights/ViewModels/TimingEventsTrack.h"
+#include "Insights/ViewModels/TimingViewDrawHelper.h"
 
 class FTimingEventSearchParameters;
+class STimingView;
+class FLoadingTimingTrack;
 
 /** Defines FLoadingTrackGetEventNameDelegate delegate interface. Returns the name for a timing event in a Loading track. */
 DECLARE_DELEGATE_RetVal_TwoParams(const TCHAR*, FLoadingTrackGetEventNameDelegate, uint32 /*Depth*/, const Trace::FLoadTimeProfilerCpuEvent& /*Event*/);
 
 ////////////////////////////////////////////////////////////////////////////////////////////////////
 
-class FLoadingSharedState
+class FLoadingSharedState : public Insights::ITimingViewExtender, public TSharedFromThis<FLoadingSharedState>
 {
 public:
-	void Reset();
+	explicit FLoadingSharedState(STimingView* InTimingView) : TimingView(InTimingView) {}
+	virtual ~FLoadingSharedState() = default;
+
+	// ITimingViewExtender
+	virtual void OnBeginSession(Insights::ITimingViewSession& InSession) override;
+	virtual void OnEndSession(Insights::ITimingViewSession& InSession) override;
+	virtual void Tick(Insights::ITimingViewSession& InSession, const Trace::IAnalysisSession& InAnalysisSession) override;
+	virtual void ExtendFilterMenu(Insights::ITimingViewSession& InSession, FMenuBuilder& InMenuBuilder) override;
+	virtual void OnTracksChanged(Insights::ITimingViewSession& InSession, int32& InOutOrder) override;
 
-	const TCHAR* GetEventNameByExportEventType(uint32 Depth, const Trace::FLoadTimeProfilerCpuEvent& Event) const;
-	const TCHAR* GetEventNameByPackageName(uint32 Depth, const Trace::FLoadTimeProfilerCpuEvent& Event) const;
-	const TCHAR* GetEventNameByExportClassName(uint32 Depth, const Trace::FLoadTimeProfilerCpuEvent& Event) const;
 	const TCHAR* GetEventName(uint32 Depth, const Trace::FLoadTimeProfilerCpuEvent& Event) const;
+	void SetColorSchema(int32 Schema);
 
-	const TCHAR* GetEventNameEx(uint32 Depth, const Trace::FLoadTimeProfilerCpuEvent& Event) const;
+	TSharedPtr<FLoadingTimingTrack> GetLoadingTrack(uint32 InThreadId)
+	{
+		TSharedPtr<FLoadingTimingTrack>* TrackPtrPtr = LoadingTracks.Find(InThreadId);
+		return TrackPtrPtr ? *TrackPtrPtr : nullptr;
+	}
 
-	void SetColorSchema(int32 Schema);
+	void ShowHideAllLoadingTracks() { ShowHideAllLoadingTracks_Execute(); }
 
 private:
-	FLoadingTrackGetEventNameDelegate LoadingGetEventNameFn;
+	const TCHAR* GetEventNameByEventType(uint32 Depth, const Trace::FLoadTimeProfilerCpuEvent& Event) const;
+	const TCHAR* GetEventNameByPackageName(uint32 Depth, const Trace::FLoadTimeProfilerCpuEvent& Event) const;
+	const TCHAR* GetEventNameByExportClassName(uint32 Depth, const Trace::FLoadTimeProfilerCpuEvent& Event) const;
+	const TCHAR* GetEventNameByPackageAndExportClassName(uint32 Depth, const Trace::FLoadTimeProfilerCpuEvent& Event) const;
+
+	bool ShowHideAllLoadingTracks_IsChecked() const;
+	void ShowHideAllLoadingTracks_Execute();
+
+private:
+	STimingView* TimingView;
+
+	bool bShowHideAllLoadingTracks;
+
+	/** Maps thread id to track pointer. */
+	TMap<uint32, TSharedPtr<FLoadingTimingTrack>> LoadingTracks;
+
+	FLoadingTrackGetEventNameDelegate GetEventNameDelegate;
 };
 
 ////////////////////////////////////////////////////////////////////////////////////////////////////
@@ -38,25 +68,27 @@ private:
 class FLoadingTimingTrack : public FTimingEventsTrack
 {
 public:
-	explicit FLoadingTimingTrack(uint64 InTrackId, uint32 InTimelineIndex, const FName& InSubType, const FString& InName, TSharedPtr<FLoadingSharedState> InState)
-		: FTimingEventsTrack(InTrackId, FName(TEXT("Loading")), InSubType, InName)
-		, State(InState)
+	explicit FLoadingTimingTrack(FLoadingSharedState& InSharedState, uint32 InTimelineIndex, const FName& InSubType, const FString& InName)
+		: FTimingEventsTrack(FName(TEXT("Loading")), InSubType, InName)
+		, SharedState(InSharedState)
 		, TimelineIndex(InTimelineIndex)
 	{
 	}
 	virtual ~FLoadingTimingTrack() {}
 
-	virtual void InitTooltip(FTooltipDrawState& Tooltip, const FTimingEvent& HoveredTimingEvent) const override;
-	virtual void Draw(ITimingViewDrawHelper& Helper) const override;
-	virtual bool SearchTimingEvent(const FTimingEventSearchParameters& InSearchParameters, FTimingEvent& InOutTimingEvent) const override;
+	virtual void InitTooltip(FTooltipDrawState& InOutTooltip, const ITimingEvent& InTooltipEvent) const override;
+	virtual void PreUpdate(const ITimingTrackUpdateContext& Context) override;
+	virtual void Draw(const ITimingTrackDrawContext& Context) const override;
+	virtual const TSharedPtr<const ITimingEvent> SearchEvent(const FTimingEventSearchParameters& InSearchParameters) const override;
 
 protected:
 	// Helper function to find an event given search parameters
 	bool FindLoadTimeProfilerCpuEvent(const FTimingEventSearchParameters& InParameters, TFunctionRef<void(double, double, uint32, const Trace::FLoadTimeProfilerCpuEvent&)> InFoundPredicate) const;
 
 protected:
-	TSharedPtr<FLoadingSharedState> State;
+	FLoadingSharedState& SharedState;
 	uint32 TimelineIndex;
+	FTimingEventsTrackDrawState DrawState;
 };
 
 ////////////////////////////////////////////////////////////////////////////////////////////////////
diff --git a/Engine/Source/Developer/TraceInsights/Private/Insights/ViewModels/MarkersTimingTrack.cpp b/Engine/Source/Developer/TraceInsights/Private/Insights/ViewModels/MarkersTimingTrack.cpp
--- a/Engine/Source/Developer/TraceInsights/Private/Insights/ViewModels/MarkersTimingTrack.cpp
+++ b/Engine/Source/Developer/TraceInsights/Private/Insights/ViewModels/MarkersTimingTrack.cpp
@@ -22,8 +22,8 @@
 // FMarkersTimingTrack
 ////////////////////////////////////////////////////////////////////////////////////////////////////
 
-FMarkersTimingTrack::FMarkersTimingTrack(uint64 InTrackId)
-	: FBaseTimingTrack(InTrackId)
+FMarkersTimingTrack::FMarkersTimingTrack()
+	: FBaseTimingTrack(FName("Markers"))
 	//, TimeMarkerBoxes()
 	//, TimeMarkerTexts()
 	, bIsCollapsed(false)
@@ -85,14 +85,16 @@ void FMarkersTimingTrack::UpdateHeight()
 
 ////////////////////////////////////////////////////////////////////////////////////////////////////
 
-void FMarkersTimingTrack::UpdateHoveredState(float MouseX, float MouseY, const FTimingTrackViewport& Viewport)
+void FMarkersTimingTrack::PostUpdate(const ITimingTrackUpdateContext& Context)
 {
 	constexpr float HeaderWidth = 80.0f;
 	constexpr float HeaderHeight = 14.0f;
 
+	const float MouseY = Context.GetMousePosition().Y;
 	if (MouseY >= GetPosY() && MouseY < GetPosY() + GetHeight())
 	{
 		SetHoveredState(true);
+		const float MouseX = Context.GetMousePosition().X;
 		SetHeaderHoveredState(MouseX < HeaderWidth && MouseY < GetPosY() + HeaderHeight);
 		TargetHoveredAnimPercent = 1.0f;
 	}
@@ -105,13 +107,13 @@ void FMarkersTimingTrack::UpdateHoveredState(float MouseX, float MouseY, const F
 
 ////////////////////////////////////////////////////////////////////////////////////////////////////
 
-void FMarkersTimingTrack::Tick(const double InCurrentTime, const float InDeltaTime)
+void FMarkersTimingTrack::Update(const ITimingTrackUpdateContext& InContext)
 {
 	if (CurrentHoveredAnimPercent != TargetHoveredAnimPercent)
 	{
 		if (CurrentHoveredAnimPercent < TargetHoveredAnimPercent)
 		{
-			const float ShowAnimSpeed = 2 * InDeltaTime;
+			const float ShowAnimSpeed = 2 * InContext.GetDeltaTime();
 			CurrentHoveredAnimPercent += ShowAnimSpeed;
 			if (CurrentHoveredAnimPercent > TargetHoveredAnimPercent)
 			{
@@ -120,7 +122,7 @@ void FMarkersTimingTrack::Tick(const double InCurrentTime, const float InDeltaTi
 		}
 		else
 		{
-			const float HideAnimSpeed = 3 * InDeltaTime;
+			const float HideAnimSpeed = 3 * InContext.GetDeltaTime();
 			CurrentHoveredAnimPercent -= HideAnimSpeed;
 			if (CurrentHoveredAnimPercent < TargetHoveredAnimPercent)
 			{
@@ -128,11 +130,19 @@ void FMarkersTimingTrack::Tick(const double InCurrentTime, const float InDeltaTi
 			}
 		}
 	}
+
+	const FTimingTrackViewport& Viewport = InContext.GetViewport();
+	if (IsDirty() || Viewport.IsHorizontalViewportDirty())
+	{
+		ClearDirtyFlag();
+
+		UpdateDrawState(Viewport);
+	}
 }
 
 ////////////////////////////////////////////////////////////////////////////////////////////////////
 
-void FMarkersTimingTrack::Update(const FTimingTrackViewport& InViewport)
+void FMarkersTimingTrack::UpdateDrawState(const FTimingTrackViewport& InViewport)
 {
 	FTimeMarkerTrackBuilder Builder(*this, InViewport);
 
@@ -155,19 +165,29 @@ void FMarkersTimingTrack::Update(const FTimingTrackViewport& InViewport)
 
 ////////////////////////////////////////////////////////////////////////////////////////////////////
 
-void FMarkersTimingTrack::Draw(FDrawContext& DrawContext, const FTimingTrackViewport& Viewport) const
+void FMarkersTimingTrack::Draw(const ITimingTrackDrawContext& Context) const
 {
-	const FLinearColor BackgroundColor(0.04f, 0.04f, 0.04f, 1.0f);
+	FDrawContext& DrawContext = Context.GetDrawContext();
+	const FTimingTrackViewport& Viewport = Context.GetViewport();
 
 	// Draw background.
+	const FLinearColor BackgroundColor(0.04f, 0.04f, 0.04f, 1.0f);
 	DrawContext.DrawBox(0.0f, GetPosY(), Viewport.GetWidth(), GetHeight(), WhiteBrush, BackgroundColor);
 	DrawContext.LayerId++;
 
 	// Draw the track's header, in background.
-	if (!IsHovered() || CurrentHoveredAnimPercent < 1.0)
+	if (!IsHovered() || CurrentHoveredAnimPercent < 1.0f)
 	{
 		DrawHeader(DrawContext, true);
 	}
+}
+
+////////////////////////////////////////////////////////////////////////////////////////////////////
+
+void FMarkersTimingTrack::PostDraw(const ITimingTrackDrawContext& Context) const
+{
+	FDrawContext& DrawContext = Context.GetDrawContext();
+	const FTimingTrackViewport& Viewport = Context.GetViewport();
 
 	//////////////////////////////////////////////////
 	// Draw vertical lines.
@@ -222,7 +242,7 @@ void FMarkersTimingTrack::Draw(FDrawContext& DrawContext, const FTimingTrackView
 	//////////////////////////////////////////////////
 
 	// When hovered, the track's header is draw on top.
-	if (IsHovered() || CurrentHoveredAnimPercent > 0.0)
+	if (IsHovered() || CurrentHoveredAnimPercent > 0.0f)
 	{
 		DrawHeader(DrawContext, false);
 	}
diff --git a/Engine/Source/Developer/TraceInsights/Private/Insights/ViewModels/MarkersTimingTrack.h b/Engine/Source/Developer/TraceInsights/Private/Insights/ViewModels/MarkersTimingTrack.h
--- a/Engine/Source/Developer/TraceInsights/Private/Insights/ViewModels/MarkersTimingTrack.h
+++ b/Engine/Source/Developer/TraceInsights/Private/Insights/ViewModels/MarkersTimingTrack.h
@@ -44,7 +44,7 @@ class FMarkersTimingTrack : public FBaseTimingTrack
 	friend class FTimeMarkerTrackBuilder;
 
 public:
-	explicit FMarkersTimingTrack(uint64 InTrackId);
+	FMarkersTimingTrack();
 	virtual ~FMarkersTimingTrack();
 
 	virtual void Reset() override;
@@ -64,11 +64,11 @@ public:
 	int32 GetNumBoxes() const { return TimeMarkerBoxes.Num(); }
 	int32 GetNumTexts() const { return TimeMarkerTexts.Num(); }
 
-	virtual void UpdateHoveredState(float MouseX, float MouseY, const FTimingTrackViewport& Viewport) override;
-	virtual void Tick(const double CurrentTime, const float DeltaTime) override;
-	virtual void Update(const FTimingTrackViewport& Viewport) override;
-
-	void Draw(FDrawContext& DrawContext, const FTimingTrackViewport& Viewport) const;
+	// FBaseTimingTrack
+	virtual void Update(const ITimingTrackUpdateContext& Context) override;
+	virtual void PostUpdate(const ITimingTrackUpdateContext& Context) override;
+	virtual void Draw(const ITimingTrackDrawContext& Context) const override;
+	virtual void PostDraw(const ITimingTrackDrawContext& Context) const override;
 
 private:
 	void ResetCache()
@@ -78,6 +78,7 @@ private:
 	}
 
 	void UpdateHeight();
+	void UpdateDrawState(const FTimingTrackViewport& Viewport);
 
 	void DrawHeader(FDrawContext& DrawContext, bool bFirstDraw) const;
 
