[Insights]

diff --git a/Engine/Source/Developer/TraceInsights/Private/Insights/ViewModels/ThreadTimingTrack.cpp b/Engine/Source/Developer/TraceInsights/Private/Insights/ViewModels/ThreadTimingTrack.cpp
--- a/Engine/Source/Developer/TraceInsights/Private/Insights/ViewModels/ThreadTimingTrack.cpp
+++ b/Engine/Source/Developer/TraceInsights/Private/Insights/ViewModels/ThreadTimingTrack.cpp
@@ -3,32 +3,395 @@
 #include "ThreadTimingTrack.h"
 
 #include "Fonts/FontMeasure.h"
-#include "Styling/SlateBrush.h"
 #include "Framework/MultiBox/MultiBoxBuilder.h"
 #include "HAL/PlatformApplicationMisc.h"
+#include "Styling/SlateBrush.h"
 
 // Insights
 #include "Insights/Common/PaintUtils.h"
 #include "Insights/Common/TimeUtils.h"
 #include "Insights/InsightsManager.h"
+#include "Insights/ITimingViewSession.h"
 #include "Insights/TimingProfilerManager.h"
 #include "Insights/ViewModels/TimerNode.h"
 #include "Insights/ViewModels/TimingEvent.h"
+#include "Insights/ViewModels/TimingEventSearch.h"
 #include "Insights/ViewModels/TimingTrackViewport.h"
 #include "Insights/ViewModels/TimingViewDrawHelper.h"
 #include "Insights/ViewModels/TooltipDrawState.h"
-#include "Insights/ViewModels/TimingEventSearch.h"
+#include "Insights/Widgets/STimingView.h"
 
 #define LOCTEXT_NAMESPACE "ThreadTimingTrack"
 
+////////////////////////////////////////////////////////////////////////////////////////////////////
+// FThreadTimingSharedState
+////////////////////////////////////////////////////////////////////////////////////////////////////
+
+TSharedPtr<FCpuTimingTrack> FThreadTimingSharedState::GetCpuTrack(uint32 InThreadId)
+{
+	TSharedPtr<FCpuTimingTrack>*const TrackPtrPtr = CpuTracks.Find(InThreadId);
+	return TrackPtrPtr ? *TrackPtrPtr : nullptr;
+}
+
+////////////////////////////////////////////////////////////////////////////////////////////////////
+
+bool FThreadTimingSharedState::IsGpuTrackVisible() const
+{
+	return GpuTrack != nullptr && GpuTrack->IsVisible();
+}
+
+////////////////////////////////////////////////////////////////////////////////////////////////////
+
+bool FThreadTimingSharedState::IsCpuTrackVisible(uint32 InThreadId) const
+{
+	const TSharedPtr<FCpuTimingTrack>*const TrackPtrPtr = CpuTracks.Find(InThreadId);
+	return TrackPtrPtr && (*TrackPtrPtr)->IsVisible();
+}
+
+////////////////////////////////////////////////////////////////////////////////////////////////////
+
+void FThreadTimingSharedState::OnBeginSession(Insights::ITimingViewSession& InSession)
+{
+	if (&InSession != TimingView)
+	{
+		return;
+	}
+
+	if (TimingView && TimingView->IsAssetLoadingModeEnabled())
+	{
+		bShowHideAllGpuTracks = false;
+		bShowHideAllCpuTracks = false;
+	}
+	else
+	{
+		bShowHideAllGpuTracks = true;
+		bShowHideAllCpuTracks = true;
+	}
+
+	GpuTrack = nullptr;
+	CpuTracks.Reset();
+	ThreadGroups.Reset();
+}
+
+////////////////////////////////////////////////////////////////////////////////////////////////////
+
+void FThreadTimingSharedState::OnEndSession(Insights::ITimingViewSession& InSession)
+{
+	if (&InSession != TimingView)
+	{
+		return;
+	}
+
+	bShowHideAllGpuTracks = false;
+	bShowHideAllCpuTracks = false;
+
+	GpuTrack = nullptr;
+	CpuTracks.Reset();
+	ThreadGroups.Reset();
+}
+
+////////////////////////////////////////////////////////////////////////////////////////////////////
+
+void FThreadTimingSharedState::Tick(Insights::ITimingViewSession& InSession, const Trace::IAnalysisSession& InAnalysisSession)
+{
+	if (&InSession != TimingView)
+	{
+		return;
+	}
+
+	const Trace::ITimingProfilerProvider* TimingProfilerProvider = Trace::ReadTimingProfilerProvider(InAnalysisSession);
+	if (TimingProfilerProvider)
+	{
+		Trace::FAnalysisSessionReadScope SessionReadScope(InAnalysisSession);
+
+		const Trace::ILoadTimeProfilerProvider* LoadTimeProfilerProvider = Trace::ReadLoadTimeProfilerProvider(InAnalysisSession);
+
+		// Check if we have a GPU track.
+		if (!GpuTrack.IsValid())
+		{
+			uint32 GpuTimelineIndex;
+			if (TimingProfilerProvider->GetGpuTimelineIndex(GpuTimelineIndex))
+			{
+				GpuTrack = MakeShared<FGpuTimingTrack>(*this, TEXT("GPU"), nullptr, GpuTimelineIndex, 0);
+				GpuTrack->SetOrder(1000);
+				GpuTrack->SetVisibilityFlag(bShowHideAllGpuTracks);
+				InSession.AddScrollableTrack(GpuTrack);
+			}
+		}
+
+		bool bTracksOrderChanged = false;
+		int32 Order = 2000;
+
+		// Iterate through threads.
+		const Trace::IThreadProvider& ThreadProvider = Trace::ReadThreadProvider(InAnalysisSession);
+		ThreadProvider.EnumerateThreads([this, &InSession, &bTracksOrderChanged, &Order, TimingProfilerProvider, LoadTimeProfilerProvider](const Trace::FThreadInfo& ThreadInfo)
+		{
+			// Check if this thread is part of a group?
+			bool bIsGroupVisible = bShowHideAllCpuTracks;
+			const TCHAR* const GroupName = ThreadInfo.GroupName ? ThreadInfo.GroupName : ThreadInfo.Name;
+			if (GroupName != nullptr)
+			{
+				if (!ThreadGroups.Contains(GroupName))
+				{
+					//UE_LOG(TimingProfiler, Log, TEXT("New CPU Thread Group (%d) : \"%s\""), ThreadGroups.Num() + 1, GroupName);
+					ThreadGroups.Add(GroupName, { GroupName, bIsGroupVisible, 0, Order });
+				}
+				else
+				{
+					FThreadGroup& ThreadGroup = ThreadGroups[GroupName];
+					bIsGroupVisible = ThreadGroup.bIsVisible;
+					ThreadGroup.Order = Order;
+				}
+			}
+
+			// Check if there is an available Asset Loading track for this thread.
+			bool bIsLoadingThread = false;
+			uint32 LoadingTimelineIndex;
+			if (LoadTimeProfilerProvider && LoadTimeProfilerProvider->GetCpuThreadTimelineIndex(ThreadInfo.Id, LoadingTimelineIndex))
+			{
+				bIsLoadingThread = true;
+			}
+
+			// Check if there is an available CPU track for this thread.
+			uint32 CpuTimelineIndex;
+			if (TimingProfilerProvider->GetCpuThreadTimelineIndex(ThreadInfo.Id, CpuTimelineIndex))
+			{
+				TSharedPtr<FCpuTimingTrack> Track;
+
+				if (!CpuTracks.Contains(ThreadInfo.Id))
+				{
+					FString TrackName(ThreadInfo.Name && *ThreadInfo.Name ? ThreadInfo.Name : FString::Printf(TEXT("Thread %u"), ThreadInfo.Id));
+
+					// Create new Timing Events track for the CPU thread.
+					Track = MakeShared<FCpuTimingTrack>(*this, TrackName, GroupName, CpuTimelineIndex, ThreadInfo.Id);
+					Track->SetOrder(Order);
+					CpuTracks.Add(ThreadInfo.Id, Track);
+
+					FThreadGroup& ThreadGroup = ThreadGroups[GroupName];
+					ThreadGroup.NumTimelines++;
+
+					if (TimingView && TimingView->IsAssetLoadingModeEnabled() && bIsLoadingThread)
+					{
+						Track->SetVisibilityFlag(true);
+						ThreadGroup.bIsVisible = true;
+					}
+					else
+					{
+						Track->SetVisibilityFlag(bIsGroupVisible);
+					}
+
+					InSession.AddScrollableTrack(Track);
+				}
+				else
+				{
+					Track = CpuTracks[ThreadInfo.Id];
+
+					if (Track->GetOrder() != Order)
+					{
+						Track->SetOrder(Order);
+						bTracksOrderChanged = true;
+					}
+				}
+			}
+
+			Order += 100;
+		});
+
+		if (bTracksOrderChanged && TimingView)
+		{
+			TimingView->UpdateScrollableTracksOrder();
+		}
+	}
+}
+
 ////////////////////////////////////////////////////////////////////////////////////////////////////
 
-void FThreadTimingTrack::Draw(ITimingViewDrawHelper& Helper) const
+void FThreadTimingSharedState::ExtendFilterMenu(Insights::ITimingViewSession& InSession, FMenuBuilder& InOutMenuBuilder)
 {
-	FTimingEventsTrack& Track = *const_cast<FThreadTimingTrack*>(this);
+	if (&InSession != TimingView)
+	{
+		return;
+	}
 
-	if (Helper.BeginTimeline(Track))
+	InOutMenuBuilder.BeginSection("ThreadProfiler", LOCTEXT("ThreadProfilerHeading", "Threads"));
 	{
+		//TODO: MenuBuilder.AddMenuEntry(Commands.ShowAllGpuTracks);
+		InOutMenuBuilder.AddMenuEntry(
+			LOCTEXT("ShowAllGpuTracks", "GPU Track - Y"),
+			LOCTEXT("ShowAllGpuTracks_Tooltip", "Show/hide the GPU track"),
+			FSlateIcon(),
+			FUIAction(FExecuteAction::CreateSP(this, &FThreadTimingSharedState::ShowHideAllGpuTracks_Execute),
+					  FCanExecuteAction(),
+					  FIsActionChecked::CreateSP(this, &FThreadTimingSharedState::ShowHideAllGpuTracks_IsChecked)),
+			NAME_None,
+			EUserInterfaceActionType::ToggleButton
+		);
+
+		//TODO: MenuBuilder.AddMenuEntry(Commands.ShowAllCpuTracks);
+		InOutMenuBuilder.AddMenuEntry(
+			LOCTEXT("ShowAllCpuTracks", "CPU Thread Tracks - U"),
+			LOCTEXT("ShowAllCpuTracks_Tooltip", "Show/hide all CPU tracks (and all CPU thread groups)"),
+			FSlateIcon(),
+			FUIAction(FExecuteAction::CreateSP(this, &FThreadTimingSharedState::ShowHideAllCpuTracks_Execute),
+					  FCanExecuteAction(),
+					  FIsActionChecked::CreateSP(this, &FThreadTimingSharedState::ShowHideAllCpuTracks_IsChecked)),
+			NAME_None,
+			EUserInterfaceActionType::ToggleButton
+		);
+	}
+	InOutMenuBuilder.EndSection();
+
+	InOutMenuBuilder.BeginSection("ThreadGroups", LOCTEXT("ThreadGroupsHeading", "CPU Thread Groups"));
+	CreateThreadGroupsMenu(InOutMenuBuilder);
+	InOutMenuBuilder.EndSection();
+}
+
+////////////////////////////////////////////////////////////////////////////////////////////////////
+
+void FThreadTimingSharedState::CreateThreadGroupsMenu(FMenuBuilder& InOutMenuBuilder)
+{
+	// Sort the list of thread groups.
+	TArray<const FThreadGroup*> SortedThreadGroups;
+	SortedThreadGroups.Reserve(ThreadGroups.Num());
+	for (const auto& KV : ThreadGroups)
+	{
+		SortedThreadGroups.Add(&KV.Value);
+	}
+	Algo::SortBy(SortedThreadGroups, &FThreadGroup::GetOrder);
+
+	for (const FThreadGroup* ThreadGroupPtr : SortedThreadGroups)
+	{
+		const FThreadGroup& ThreadGroup = *ThreadGroupPtr;
+		if (ThreadGroup.NumTimelines > 0)
+		{
+			InOutMenuBuilder.AddMenuEntry(
+				//FText::FromString(ThreadGroup.Name),
+				FText::Format(LOCTEXT("ThreadGroupFmt", "{0} ({1})"), FText::FromString(ThreadGroup.Name), ThreadGroup.NumTimelines),
+				TAttribute<FText>(), // no tooltip
+				FSlateIcon(),
+				FUIAction(FExecuteAction::CreateSP(this, &FThreadTimingSharedState::ToggleTrackVisibilityByGroup_Execute, ThreadGroup.Name),
+						  FCanExecuteAction::CreateLambda([] { return true; }),
+						  FIsActionChecked::CreateSP(this, &FThreadTimingSharedState::ToggleTrackVisibilityByGroup_IsChecked, ThreadGroup.Name)),
+				NAME_None,
+				EUserInterfaceActionType::ToggleButton
+			);
+		}
+	}
+}
+
+////////////////////////////////////////////////////////////////////////////////////////////////////
+
+bool FThreadTimingSharedState::ShowHideAllCpuTracks_IsChecked() const
+{
+	return bShowHideAllCpuTracks;
+}
+
+////////////////////////////////////////////////////////////////////////////////////////////////////
+
+void FThreadTimingSharedState::ShowHideAllCpuTracks_Execute()
+{
+	bShowHideAllCpuTracks = !bShowHideAllCpuTracks;
+
+	for (const auto& KV : CpuTracks)
+	{
+		FCpuTimingTrack& Track = *KV.Value;
+		Track.SetVisibilityFlag(bShowHideAllCpuTracks);
+	}
+
+	for (auto& KV : ThreadGroups)
+	{
+		KV.Value.bIsVisible = bShowHideAllCpuTracks;
+	}
+
+	if (TimingView)
+	{
+		TimingView->OnTrackVisibilityChanged();
+	}
+}
+
+////////////////////////////////////////////////////////////////////////////////////////////////////
+
+bool FThreadTimingSharedState::ShowHideAllGpuTracks_IsChecked() const
+{
+	return bShowHideAllGpuTracks;
+}
+
+////////////////////////////////////////////////////////////////////////////////////////////////////
+
+void FThreadTimingSharedState::ShowHideAllGpuTracks_Execute()
+{
+	bShowHideAllGpuTracks = !bShowHideAllGpuTracks;
+
+	if (GpuTrack.IsValid())
+	{
+		GpuTrack->SetVisibilityFlag(bShowHideAllGpuTracks);
+
+		if (TimingView)
+		{
+			TimingView->OnTrackVisibilityChanged();
+		}
+	}
+}
+
+////////////////////////////////////////////////////////////////////////////////////////////////////
+
+bool FThreadTimingSharedState::ToggleTrackVisibilityByGroup_IsChecked(const TCHAR* InGroupName) const
+{
+	if (ThreadGroups.Contains(InGroupName))
+	{
+		const FThreadGroup& ThreadGroup = ThreadGroups[InGroupName];
+		return ThreadGroup.bIsVisible;
+	}
+	return false;
+}
+
+////////////////////////////////////////////////////////////////////////////////////////////////////
+
+void FThreadTimingSharedState::ToggleTrackVisibilityByGroup_Execute(const TCHAR* InGroupName)
+{
+	if (ThreadGroups.Contains(InGroupName))
+	{
+		FThreadGroup& ThreadGroup = ThreadGroups[InGroupName];
+		ThreadGroup.bIsVisible = !ThreadGroup.bIsVisible;
+
+		for (const auto& KV : CpuTracks)
+		{
+			FCpuTimingTrack& Track = *KV.Value;
+			if (Track.GetGroupName() == InGroupName)
+			{
+				Track.SetVisibilityFlag(ThreadGroup.bIsVisible);
+			}
+		}
+
+		if (TimingView)
+		{
+			TimingView->OnTrackVisibilityChanged();
+		}
+	}
+}
+
+////////////////////////////////////////////////////////////////////////////////////////////////////
+
+void FThreadTimingSharedState::OnTracksChanged(Insights::ITimingViewSession& InSession, int32& InOutOrder)
+{
+	if (&InSession != TimingView)
+	{
+		return;
+	}
+}
+
+////////////////////////////////////////////////////////////////////////////////////////////////////
+// FThreadTimingTrack
+////////////////////////////////////////////////////////////////////////////////////////////////////
+
+void FThreadTimingTrack::PreUpdate(const ITimingTrackUpdateContext& Context)
+{
+	if (IsDirty() || Context.GetViewport().IsHorizontalViewportDirty())
+	{
+		ClearDirtyFlag();
+
+		FTimingEventsTrackDrawStateBuilder Builder(GetDrawState(), Context.GetViewport());
+
 		TSharedPtr<const Trace::IAnalysisSession> Session = FInsightsManager::Get()->GetSession();
 		if (Session.IsValid() && Trace::ReadTimingProfilerProvider(*Session.Get()))
 		{
@@ -36,46 +399,79 @@ void FThreadTimingTrack::Draw(ITimingViewDrawHelper& Helper) const
 
 			const Trace::ITimingProfilerProvider& TimingProfilerProvider = *Trace::ReadTimingProfilerProvider(*Session.Get());
 
-			TimingProfilerProvider.ReadTimers([this, &Helper, &Track, &TimingProfilerProvider](const Trace::FTimingProfilerTimer* Timers, uint64 TimersCount)
+			TimingProfilerProvider.ReadTimers([this, &Builder, &TimingProfilerProvider](const Trace::FTimingProfilerTimer* Timers, uint64 TimersCount)
 			{
-				TimingProfilerProvider.ReadTimeline(Track.GetId(), [this, &Helper, &Track, Timers](const Trace::ITimingProfilerProvider::Timeline& Timeline)
+				TimingProfilerProvider.ReadTimeline(TimelineIndex, [this, &Builder, Timers](const Trace::ITimingProfilerProvider::Timeline& Timeline)
 				{
+					const FTimingTrackViewport& Viewport = Builder.GetViewport();
 					if (FTimingEventsTrack::bUseDownSampling)
 					{
-						const double SecondsPerPixel = 1.0 / Helper.GetViewport().GetScaleX();
-						Timeline.EnumerateEventsDownSampled(Helper.GetViewport().GetStartTime(), Helper.GetViewport().GetEndTime(), SecondsPerPixel, [this, &Helper, Timers](double StartTime, double EndTime, uint32 Depth, const Trace::FTimingProfilerEvent& Event)
-						{
-							Helper.AddEvent(StartTime, EndTime, Depth, Timers[Event.TimerIndex].Name);
-						});
+						const double SecondsPerPixel = 1.0 / Viewport.GetScaleX();
+						Timeline.EnumerateEventsDownSampled(Viewport.GetStartTime(), Viewport.GetEndTime(), SecondsPerPixel,
+							[this, &Builder, Timers](double StartTime, double EndTime, uint32 Depth, const Trace::FTimingProfilerEvent& Event)
+							{
+								const uint64 Type = Timers[Event.TimerIndex].Id;
+								const uint32 Color = 0;
+								Builder.AddEvent(StartTime, EndTime, Depth, Timers[Event.TimerIndex].Name, Type, Color);
+							});
 					}
 					else
 					{
-						Timeline.EnumerateEvents(Helper.GetViewport().GetStartTime(), Helper.GetViewport().GetEndTime(), [this, &Helper, Timers](double StartTime, double EndTime, uint32 Depth, const Trace::FTimingProfilerEvent& Event)
-						{
-							Helper.AddEvent(StartTime, EndTime, Depth, Timers[Event.TimerIndex].Name);
-						});
+						Timeline.EnumerateEvents(Viewport.GetStartTime(), Viewport.GetEndTime(),
+							[this, &Builder, Timers](double StartTime, double EndTime, uint32 Depth, const Trace::FTimingProfilerEvent& Event)
+							{
+								const uint64 Type = Timers[Event.TimerIndex].Id;
+								const uint32 Color = 0;
+								Builder.AddEvent(StartTime, EndTime, Depth, Timers[Event.TimerIndex].Name, Type, Color);
+							});
 					}
 				});
 			});
-
-			Helper.EndTimeline(Track);
 		}
+
+		Builder.Flush();
+		SetNumLanes(Builder.GetMaxDepth() + 1);
 	}
+
+	UpdateTrackHeight(Context);
 }
 
 ////////////////////////////////////////////////////////////////////////////////////////////////////
 
-void FThreadTimingTrack::DrawSelectedEventInfo(const FTimingEvent& SelectedTimingEvent, const FTimingTrackViewport& Viewport, const FDrawContext& DrawContext, const FSlateBrush* WhiteBrush, const FSlateFontInfo& Font) const
+void FThreadTimingTrack::Draw(const ITimingTrackDrawContext& Context) const
 {
-	FindTimingProfilerEvent(SelectedTimingEvent, [&SelectedTimingEvent, &Font, &Viewport, &DrawContext, &WhiteBrush](double InFoundStartTime, double InFoundEndTime, uint32 InFoundDepth, const Trace::FTimingProfilerEvent& InFoundEvent)
-	{	
+	const FTimingViewDrawHelper& Helper = Context.GetHelper();
+	Helper.DrawEventsTrack(DrawState, *this);
+}
+
+////////////////////////////////////////////////////////////////////////////////////////////////////
+
+void FThreadTimingTrack::PostDraw(const ITimingTrackDrawContext& Context) const
+{
+	const TSharedPtr<const ITimingEvent> SelectedEventPtr = Context.GetSelectedEvent();
+	if (SelectedEventPtr.IsValid() &&
+		SelectedEventPtr->CheckTrack(this) &&
+		FTimingEvent::CheckTypeName(*SelectedEventPtr))
+	{
+		const FTimingEvent& SelectedEvent = static_cast<const FTimingEvent&>(*SelectedEventPtr);
+		const FTimingViewDrawHelper& Helper = Context.GetHelper();
+		DrawSelectedEventInfo(SelectedEvent, Context.GetViewport(), Context.GetDrawContext(), Helper.GetWhiteBrush(), Helper.GetEventFont());
+	}
+}
+
+////////////////////////////////////////////////////////////////////////////////////////////////////
+
+void FThreadTimingTrack::DrawSelectedEventInfo(const FTimingEvent& SelectedEvent, const FTimingTrackViewport& Viewport, const FDrawContext& DrawContext, const FSlateBrush* WhiteBrush, const FSlateFontInfo& Font) const
+{
+	FindTimingProfilerEvent(SelectedEvent, [&SelectedEvent, &Font, &Viewport, &DrawContext, &WhiteBrush](double InFoundStartTime, double InFoundEndTime, uint32 InFoundDepth, const Trace::FTimingProfilerEvent& InFoundEvent)
+	{
 		const FTimerNodePtr TimerNodePtr = FTimingProfilerManager::Get()->GetTimerNode(InFoundEvent.TimerIndex);
 		if (TimerNodePtr.IsValid())
 		{
 			FString Str = FString::Printf(TEXT("%s (Incl.: %s, Excl.: %s)"),
 				TimerNodePtr ? *(TimerNodePtr->GetName().ToString()) : TEXT("N/A"),
-				*TimeUtils::FormatTimeAuto(SelectedTimingEvent.Duration()),
-				*TimeUtils::FormatTimeAuto(SelectedTimingEvent.ExclusiveTime));
+				*TimeUtils::FormatTimeAuto(SelectedEvent.GetDuration()),
+				*TimeUtils::FormatTimeAuto(SelectedEvent.GetExclusiveTime()));
 
 			const TSharedRef<FSlateFontMeasure> FontMeasureService = FSlateApplication::Get().GetRenderer()->GetFontMeasureService();
 			const FVector2D Size = FontMeasureService->Measure(Str, Font);
@@ -96,68 +492,69 @@ void FThreadTimingTrack::DrawSelectedEventInfo(const FTimingEvent& SelectedTimin
 
 ////////////////////////////////////////////////////////////////////////////////////////////////////
 
-void FThreadTimingTrack::InitTooltip(FTooltipDrawState& Tooltip, const FTimingEvent& HoveredTimingEvent) const
+void FThreadTimingTrack::InitTooltip(FTooltipDrawState& InOutTooltip, const ITimingEvent& InTooltipEvent) const
 {
-	FindTimingProfilerEvent(HoveredTimingEvent, [this, &Tooltip, &HoveredTimingEvent](double InFoundStartTime, double InFoundEndTime, uint32 InFoundDepth, const Trace::FTimingProfilerEvent& InFoundEvent)
-	{	
-		FTimingEvent ParentTimingEvent;
-		Trace::FTimingProfilerEvent ParentEvent;
-		FTimingEvent RootTimingEvent;
-		Trace::FTimingProfilerEvent RootEvent;
-		if (HoveredTimingEvent.IsValid())
+	if (InTooltipEvent.CheckTrack(this) && FTimingEvent::CheckTypeName(InTooltipEvent))
+	{
+		const FTimingEvent& TooltipEvent = static_cast<const FTimingEvent&>(InTooltipEvent);
+
+		FindTimingProfilerEvent(TooltipEvent, [this, &InOutTooltip, &TooltipEvent](double InFoundStartTime, double InFoundEndTime, uint32 InFoundDepth, const Trace::FTimingProfilerEvent& InFoundEvent)
 		{
-			GetParentAndRoot(HoveredTimingEvent, ParentTimingEvent, ParentEvent, RootTimingEvent, RootEvent);
-		}
+			TSharedPtr<FTimingEvent> ParentTimingEvent;
+			Trace::FTimingProfilerEvent ParentEvent;
+			TSharedPtr<FTimingEvent> RootTimingEvent;
+			Trace::FTimingProfilerEvent RootEvent;
+			GetParentAndRoot(TooltipEvent, ParentTimingEvent, ParentEvent, RootTimingEvent, RootEvent);
 
-		Tooltip.ResetContent();
+			InOutTooltip.ResetContent();
 
-		const FTimerNodePtr TimerNodePtr = FTimingProfilerManager::Get()->GetTimerNode(InFoundEvent.TimerIndex);
-		FString TimerName = TimerNodePtr ? TimerNodePtr->GetName().ToString() : TEXT("N/A");
-		Tooltip.AddTitle(TimerName);
+			const FTimerNodePtr TimerNodePtr = FTimingProfilerManager::Get()->GetTimerNode(InFoundEvent.TimerIndex);
+			FString TimerName = TimerNodePtr ? TimerNodePtr->GetName().ToString() : TEXT("N/A");
+			InOutTooltip.AddTitle(TimerName);
 
-		if (ParentTimingEvent.IsValid() && HoveredTimingEvent.Depth > 0)
-		{
-			const FTimerNodePtr ParentTimerNodePtr = FTimingProfilerManager::Get()->GetTimerNode(ParentEvent.TimerIndex);
-			const FString ParentTimerName = ParentTimerNodePtr.IsValid() ? ParentTimerNodePtr->GetName().ToString() : TEXT("N/A");
-			FNumberFormattingOptions FormattingOptions;
-			FormattingOptions.MaximumFractionalDigits = 2;
-			const FString ValueStr = FString::Printf(TEXT("%s %s"), *FText::AsPercent(HoveredTimingEvent.Duration() / ParentTimingEvent.Duration(), &FormattingOptions).ToString(), *ParentTimerName);
-			Tooltip.AddNameValueTextLine(TEXT("% of Parent:"), ValueStr);
-		}
+			if (ParentTimingEvent.IsValid() && TooltipEvent.GetDepth() > 0)
+			{
+				const FTimerNodePtr ParentTimerNodePtr = FTimingProfilerManager::Get()->GetTimerNode(ParentEvent.TimerIndex);
+				const FString ParentTimerName = ParentTimerNodePtr.IsValid() ? ParentTimerNodePtr->GetName().ToString() : TEXT("N/A");
+				FNumberFormattingOptions FormattingOptions;
+				FormattingOptions.MaximumFractionalDigits = 2;
+				const FString ValueStr = FString::Printf(TEXT("%s %s"), *FText::AsPercent(TooltipEvent.GetDuration() / ParentTimingEvent->GetDuration(), &FormattingOptions).ToString(), *ParentTimerName);
+				InOutTooltip.AddNameValueTextLine(TEXT("% of Parent:"), ValueStr);
+			}
 
-		if (RootTimingEvent.IsValid() && HoveredTimingEvent.Depth > 1)
-		{
-			const FTimerNodePtr RootTimerNodePtr = FTimingProfilerManager::Get()->GetTimerNode(RootEvent.TimerIndex);
-			const FString RootTimerName = RootTimerNodePtr.IsValid() ? RootTimerNodePtr->GetName().ToString() : TEXT("N/A");
-			FNumberFormattingOptions FormattingOptions;
-			FormattingOptions.MaximumFractionalDigits = 2;
-			const FString ValueStr = FString::Printf(TEXT("%s %s"), *FText::AsPercent(HoveredTimingEvent.Duration() / RootTimingEvent.Duration(), &FormattingOptions).ToString(), *RootTimerName);
-			Tooltip.AddNameValueTextLine(TEXT("% of Root:"), ValueStr);
-		}
+			if (RootTimingEvent.IsValid() && TooltipEvent.GetDepth() > 1)
+			{
+				const FTimerNodePtr RootTimerNodePtr = FTimingProfilerManager::Get()->GetTimerNode(RootEvent.TimerIndex);
+				const FString RootTimerName = RootTimerNodePtr.IsValid() ? RootTimerNodePtr->GetName().ToString() : TEXT("N/A");
+				FNumberFormattingOptions FormattingOptions;
+				FormattingOptions.MaximumFractionalDigits = 2;
+				const FString ValueStr = FString::Printf(TEXT("%s %s"), *FText::AsPercent(TooltipEvent.GetDuration() / RootTimingEvent->GetDuration(), &FormattingOptions).ToString(), *RootTimerName);
+				InOutTooltip.AddNameValueTextLine(TEXT("% of Root:"), ValueStr);
+			}
 
-		Tooltip.AddNameValueTextLine(TEXT("Inclusive Time:"), TimeUtils::FormatTimeAuto(HoveredTimingEvent.Duration()));
+			InOutTooltip.AddNameValueTextLine(TEXT("Inclusive Time:"), TimeUtils::FormatTimeAuto(TooltipEvent.GetDuration()));
 
-		//Tooltip.AddNameValueTextLine(TEXT("Exclusive Time:"), TimeUtils::FormatTimeAuto(HoveredTimingEvent.ExclusiveTime));
-		{
-			FNumberFormattingOptions FormattingOptions;
-			FormattingOptions.MaximumFractionalDigits = 2;
-			const FString ExclStr = FString::Printf(TEXT("%s (%s)"), *TimeUtils::FormatTimeAuto(HoveredTimingEvent.ExclusiveTime), *FText::AsPercent(HoveredTimingEvent.ExclusiveTime / HoveredTimingEvent.Duration(), &FormattingOptions).ToString());
-			Tooltip.AddNameValueTextLine(TEXT("Exclusive Time:"), ExclStr);
-		}
+			{
+				FNumberFormattingOptions FormattingOptions;
+				FormattingOptions.MaximumFractionalDigits = 2;
+				const FString ExclStr = FString::Printf(TEXT("%s (%s)"), *TimeUtils::FormatTimeAuto(TooltipEvent.GetExclusiveTime()), *FText::AsPercent(TooltipEvent.GetExclusiveTime() / TooltipEvent.GetDuration(), &FormattingOptions).ToString());
+				InOutTooltip.AddNameValueTextLine(TEXT("Exclusive Time:"), ExclStr);
+			}
 
-		Tooltip.AddNameValueTextLine(TEXT("Depth:"), FString::Printf(TEXT("%d"), HoveredTimingEvent.Depth));
+			InOutTooltip.AddNameValueTextLine(TEXT("Depth:"), FString::Printf(TEXT("%d"), TooltipEvent.GetDepth()));
 
-		Tooltip.UpdateLayout();
-	});
+			InOutTooltip.UpdateLayout();
+		});
+	}
 }
 
 ////////////////////////////////////////////////////////////////////////////////////////////////////
 
-void FThreadTimingTrack::GetParentAndRoot(const FTimingEvent& TimingEvent, FTimingEvent& OutParentTimingEvent, Trace::FTimingProfilerEvent& OutParentEvent, FTimingEvent& OutRootTimingEvent, Trace::FTimingProfilerEvent& OutRootEvent) const
+void FThreadTimingTrack::GetParentAndRoot(const FTimingEvent& TimingEvent, TSharedPtr<FTimingEvent>& OutParentTimingEvent, Trace::FTimingProfilerEvent& OutParentEvent, TSharedPtr<FTimingEvent>& OutRootTimingEvent, Trace::FTimingProfilerEvent& OutRootEvent) const
 {
 	// Note: This function does not compute Exclusive Time for parent and root events.
 
-	if (TimingEvent.Depth > 0)
+	if (TimingEvent.GetDepth() > 0)
 	{
 		TSharedPtr<const Trace::IAnalysisSession> Session = FInsightsManager::Get()->GetSession();
 		if (Session.IsValid())
@@ -168,19 +565,19 @@ void FThreadTimingTrack::GetParentAndRoot(const FTimingEvent& TimingEvent, FTimi
 			{
 				const Trace::ITimingProfilerProvider& TimingProfilerProvider = *Trace::ReadTimingProfilerProvider(*Session.Get());
 
-				TimingProfilerProvider.ReadTimeline(GetId(), [&TimingEvent, &OutParentTimingEvent, &OutParentEvent, &OutRootTimingEvent, &OutRootEvent](const Trace::ITimingProfilerProvider::Timeline& Timeline)
+				TimingProfilerProvider.ReadTimeline(GetTimelineIndex(), [&TimingEvent, &OutParentTimingEvent, &OutParentEvent, &OutRootTimingEvent, &OutRootEvent](const Trace::ITimingProfilerProvider::Timeline& Timeline)
 				{
-					Timeline.EnumerateEvents(TimingEvent.StartTime, TimingEvent.EndTime, [&TimingEvent, &OutParentTimingEvent, &OutParentEvent, &OutRootTimingEvent, &OutRootEvent](double EventStartTime, double EventEndTime, uint32 EventDepth, const Trace::FTimingProfilerEvent& Event)
+					Timeline.EnumerateEvents(TimingEvent.GetStartTime(), TimingEvent.GetEndTime(), [&TimingEvent, &OutParentTimingEvent, &OutParentEvent, &OutRootTimingEvent, &OutRootEvent](double EventStartTime, double EventEndTime, uint32 EventDepth, const Trace::FTimingProfilerEvent& Event)
 					{
 						if (EventDepth == 0)
 						{
 							OutRootEvent = Event;
-							OutRootTimingEvent = FTimingEvent(TimingEvent.Track, EventStartTime, EventEndTime, EventDepth);
+							OutRootTimingEvent = MakeShareable(new FTimingEvent(TimingEvent.GetTrack(), EventStartTime, EventEndTime, EventDepth));
 						}
-						if (EventDepth == TimingEvent.Depth - 1)
+						if (EventDepth == TimingEvent.GetDepth() - 1)
 						{
 							OutParentEvent = Event;
-							OutParentTimingEvent = FTimingEvent(TimingEvent.Track, EventStartTime, EventEndTime, EventDepth);
+							OutParentTimingEvent = MakeShareable(new FTimingEvent(TimingEvent.GetTrack(), EventStartTime, EventEndTime, EventDepth));
 						}
 					});
 				});
@@ -191,20 +588,26 @@ void FThreadTimingTrack::GetParentAndRoot(const FTimingEvent& TimingEvent, FTimi
 
 ////////////////////////////////////////////////////////////////////////////////////////////////////
 
-bool FThreadTimingTrack::SearchTimingEvent(const FTimingEventSearchParameters& InSearchParameters, FTimingEvent& InOutTimingEvent) const
+const TSharedPtr<const ITimingEvent> FThreadTimingTrack::SearchEvent(const FTimingEventSearchParameters& InSearchParameters) const
 {
-	return FindTimingProfilerEvent(InSearchParameters, [this, &InOutTimingEvent](double InFoundStartTime, double InFoundEndTime, uint32 InFoundDepth, const Trace::FTimingProfilerEvent& InFoundEvent)
+	TSharedPtr<const ITimingEvent> FoundEvent;
+
+	FindTimingProfilerEvent(InSearchParameters, [this, &FoundEvent](double InFoundStartTime, double InFoundEndTime, uint32 InFoundDepth, const Trace::FTimingProfilerEvent& InFoundEvent)
 	{
-		InOutTimingEvent = FTimingEvent(this, InFoundStartTime, InFoundEndTime, InFoundDepth);
+		FoundEvent = MakeShareable(new FTimingEvent(SharedThis(this), InFoundStartTime, InFoundEndTime, InFoundDepth));
 	});
+
+	return FoundEvent;
 }
 
 ////////////////////////////////////////////////////////////////////////////////////////////////////
 
-void FThreadTimingTrack::ComputeTimingEventStats(FTimingEvent& InOutTimingEvent) const
+void FThreadTimingTrack::UpdateEventStats(ITimingEvent& InOutEvent) const
 {
-	if (InOutTimingEvent.IsValid())
+	if (InOutEvent.CheckTrack(this) && FTimingEvent::CheckTypeName(InOutEvent))
 	{
+		FTimingEvent& TrackEvent = static_cast<FTimingEvent&>(InOutEvent);
+
 		TSharedPtr<const Trace::IAnalysisSession> Session = FInsightsManager::Get()->GetSession();
 		if (Session.IsValid())
 		{
@@ -215,7 +618,7 @@ void FThreadTimingTrack::ComputeTimingEventStats(FTimingEvent& InOutTimingEvent)
 				const Trace::ITimingProfilerProvider& TimingProfilerProvider = *Trace::ReadTimingProfilerProvider(*Session.Get());
 
 				// Compute Exclusive Time.
-				TimingProfilerProvider.ReadTimeline(InOutTimingEvent.Track->GetId(), [&InOutTimingEvent](const Trace::ITimingProfilerProvider::Timeline& Timeline)
+				TimingProfilerProvider.ReadTimeline(GetTimelineIndex(), [&TrackEvent](const Trace::ITimingProfilerProvider::Timeline& Timeline)
 				{
 					struct FEnumerationState
 					{
@@ -230,16 +633,16 @@ void FThreadTimingTrack::ComputeTimingEventStats(FTimingEvent& InOutTimingEvent)
 					};
 					FEnumerationState State;
 
-					State.EventStartTime = InOutTimingEvent.StartTime;
-					State.EventEndTime = InOutTimingEvent.EndTime;
-					State.EventDepth = InOutTimingEvent.Depth;
+					State.EventStartTime = TrackEvent.GetStartTime();
+					State.EventEndTime = TrackEvent.GetEndTime();
+					State.EventDepth = TrackEvent.GetDepth();
 
 					State.CurrentDepth = 0;
 					State.LastTime = 0.0;
 					State.ExclusiveTime = 0.0;
 					State.IsInEventScope = false;
 
-					Timeline.EnumerateEvents(InOutTimingEvent.StartTime, InOutTimingEvent.EndTime, [&State](bool IsEnter, double Time, const Trace::FTimingProfilerEvent& Event)
+					Timeline.EnumerateEvents(TrackEvent.GetStartTime(), TrackEvent.GetEndTime(), [&State](bool IsEnter, double Time, const Trace::FTimingProfilerEvent& Event)
 					{
 						if (IsEnter)
 						{
@@ -265,7 +668,7 @@ void FThreadTimingTrack::ComputeTimingEventStats(FTimingEvent& InOutTimingEvent)
 						State.LastTime = Time;
 					});
 
-					InOutTimingEvent.ExclusiveTime = State.ExclusiveTime;
+					TrackEvent.SetExclusiveTime(State.ExclusiveTime);
 				});
 			}
 		}
@@ -274,28 +677,38 @@ void FThreadTimingTrack::ComputeTimingEventStats(FTimingEvent& InOutTimingEvent)
 
 ////////////////////////////////////////////////////////////////////////////////////////////////////
 
-void FThreadTimingTrack::OnEventSelected(const FTimingEvent& SelectedTimingEvent) const
+void FThreadTimingTrack::OnEventSelected(const ITimingEvent& InSelectedEvent) const
 {
-	FindTimingProfilerEvent(SelectedTimingEvent, [](double InFoundStartTime, double InFoundEndTime, uint32 InFoundDepth, const Trace::FTimingProfilerEvent& InFoundEvent)
+	if (InSelectedEvent.CheckTrack(this) && FTimingEvent::CheckTypeName(InSelectedEvent))
 	{
-		// Select the timer node corresponding to timing event type of selected timing event.
-		FTimingProfilerManager::Get()->SetSelectedTimer(InFoundEvent.TimerIndex);
-	});
+		const FTimingEvent& TrackEvent = static_cast<const FTimingEvent&>(InSelectedEvent);
+
+		FindTimingProfilerEvent(TrackEvent, [](double InFoundStartTime, double InFoundEndTime, uint32 InFoundDepth, const Trace::FTimingProfilerEvent& InFoundEvent)
+		{
+			// Select the timer node corresponding to timing event type of selected timing event.
+			FTimingProfilerManager::Get()->SetSelectedTimer(InFoundEvent.TimerIndex);
+		});
+	}
 }
 
 ////////////////////////////////////////////////////////////////////////////////////////////////////
 
-void FThreadTimingTrack::OnClipboardCopyEvent(const FTimingEvent& SelectedTimingEvent) const
+void FThreadTimingTrack::OnClipboardCopyEvent(const ITimingEvent& InSelectedEvent) const
 {
-	FindTimingProfilerEvent(SelectedTimingEvent, [](double InFoundStartTime, double InFoundEndTime, uint32 InFoundDepth, const Trace::FTimingProfilerEvent& InFoundEvent)
+	if (InSelectedEvent.CheckTrack(this) && FTimingEvent::CheckTypeName(InSelectedEvent))
 	{
-		const FTimerNodePtr TimerNodePtr = FTimingProfilerManager::Get()->GetTimerNode(InFoundEvent.TimerIndex);
-		if (TimerNodePtr.IsValid())
+		const FTimingEvent& TrackEvent = static_cast<const FTimingEvent&>(InSelectedEvent);
+
+		FindTimingProfilerEvent(TrackEvent, [](double InFoundStartTime, double InFoundEndTime, uint32 InFoundDepth, const Trace::FTimingProfilerEvent& InFoundEvent)
 		{
-			// Copy name of selected timing event to clipboard.
-			FPlatformApplicationMisc::ClipboardCopy(*TimerNodePtr->GetName().ToString());
-		}
-	});
+			const FTimerNodePtr TimerNodePtr = FTimingProfilerManager::Get()->GetTimerNode(InFoundEvent.TimerIndex);
+			if (TimerNodePtr.IsValid())
+			{
+				// Copy name of selected timing event to clipboard.
+				FPlatformApplicationMisc::ClipboardCopy(*TimerNodePtr->GetName().ToString());
+			}
+		});
+	}
 }
 
 ////////////////////////////////////////////////////////////////////////////////////////////////////
@@ -303,7 +716,7 @@ void FThreadTimingTrack::OnClipboardCopyEvent(const FTimingEvent& SelectedTiming
 void FThreadTimingTrack::BuildContextMenu(FMenuBuilder& MenuBuilder)
 {
 	FText SectionTitle;
-	if(GetGroupName() != nullptr)
+	if (GetGroupName() != nullptr)
 	{
 		SectionTitle = FText::Format(LOCTEXT("TrackTitleGroupFmt", "{0} (Group: {1})"), FText::FromString(GetName()), FText::FromString(GetGroupName()));
 	}
@@ -330,15 +743,15 @@ void FThreadTimingTrack::BuildContextMenu(FMenuBuilder& MenuBuilder)
 
 bool FThreadTimingTrack::FindTimingProfilerEvent(const FTimingEvent& InTimingEvent, TFunctionRef<void(double, double, uint32, const Trace::FTimingProfilerEvent&)> InFoundPredicate) const
 {
-	auto MatchDepth = [&InTimingEvent](double InStartTime, double InEndTime, uint32 InDepth)
+	auto MatchEvent = [&InTimingEvent](double InStartTime, double InEndTime, uint32 InDepth)
 	{ 
-		return InDepth == InTimingEvent.Depth
-			&& InStartTime == InTimingEvent.StartTime
-			&& InEndTime == InTimingEvent.EndTime;
+		return InDepth == InTimingEvent.GetDepth()
+			&& InStartTime == InTimingEvent.GetStartTime()
+			&& InEndTime == InTimingEvent.GetEndTime();
 	};
 
-	FTimingEventSearchParameters SearchParameters(InTimingEvent.StartTime, InTimingEvent.EndTime, ETimingEventSearchFlags::StopAtFirstMatch, MatchDepth);
-	SearchParameters.SearchHandle = &InTimingEvent.SearchHandle;
+	FTimingEventSearchParameters SearchParameters(InTimingEvent.GetStartTime(), InTimingEvent.GetEndTime(), ETimingEventSearchFlags::StopAtFirstMatch, MatchEvent);
+	SearchParameters.SearchHandle = &InTimingEvent.GetSearchHandle();
 	return FindTimingProfilerEvent(SearchParameters, InFoundPredicate);
 }
 
@@ -360,7 +773,7 @@ bool FThreadTimingTrack::FindTimingProfilerEvent(const FTimingEventSearchParamet
 				{
 					const Trace::ITimingProfilerProvider& TimingProfilerProvider = *Trace::ReadTimingProfilerProvider(*Session.Get());
 
-					TimingProfilerProvider.ReadTimeline(GetId(), [&InContext](const Trace::ITimingProfilerProvider::Timeline& Timeline)
+					TimingProfilerProvider.ReadTimeline(GetTimelineIndex(), [&InContext](const Trace::ITimingProfilerProvider::Timeline& Timeline)
 					{
 						Timeline.EnumerateEvents(InContext.GetParameters().StartTime, InContext.GetParameters().EndTime, [&InContext](double EventStartTime, double EventEndTime, uint32 EventDepth, const Trace::FTimingProfilerEvent& Event)
 						{
@@ -379,7 +792,6 @@ bool FThreadTimingTrack::FindTimingProfilerEvent(const FTimingEventSearchParamet
 		SearchCache);
 }
 
-
 ////////////////////////////////////////////////////////////////////////////////////////////////////
 
 #undef LOCTEXT_NAMESPACE
diff --git a/Engine/Source/Developer/TraceInsights/Private/Insights/ViewModels/ThreadTimingTrack.h b/Engine/Source/Developer/TraceInsights/Private/Insights/ViewModels/ThreadTimingTrack.h
--- a/Engine/Source/Developer/TraceInsights/Private/Insights/ViewModels/ThreadTimingTrack.h
+++ b/Engine/Source/Developer/TraceInsights/Private/Insights/ViewModels/ThreadTimingTrack.h
@@ -3,23 +3,97 @@
 #pragma once
 
 #include "CoreMinimal.h"
+#include "TraceServices/Model/TimingProfiler.h"
 
 // Insights
+#include "Insights/ITimingViewExtender.h"
+#include "Insights/ViewModels/TimingEventSearch.h" // for TTimingEventSearchCache
 #include "Insights/ViewModels/TimingEventsTrack.h"
-#include "Insights/ViewModels/TimingEventSearch.h"
-#include "TraceServices/Model/TimingProfiler.h"
+#include "Insights/ViewModels/TimingViewDrawHelper.h" // for FTimingEventsTrackDrawState
 
+class FTimingEvent;
 class FTimingEventSearchParameters;
+class FGpuTimingTrack;
+class FCpuTimingTrack;
+class STimingView;
+
+////////////////////////////////////////////////////////////////////////////////////////////////////
+
+class FThreadTimingSharedState : public Insights::ITimingViewExtender, public TSharedFromThis<FThreadTimingSharedState>
+{
+private:
+	struct FThreadGroup
+	{
+		const TCHAR* Name; /**< The thread group name; pointer to string owned by ThreadProvider. */
+		bool bIsVisible;  /**< Toggle to show/hide all thread timelines associated with this group at once. Used also as default for new thread timelines. */
+		uint32 NumTimelines; /**< Number of thread timelines associated with this group. */
+		int32 Order; //**< Order index used for sorting. Inherited from last thread timeline associated with this group. **/
+
+		int32 GetOrder() const { return Order; }
+	};
+
+public:
+	explicit FThreadTimingSharedState(STimingView* InTimingView) : TimingView(InTimingView) {}
+	virtual ~FThreadTimingSharedState() = default;
+
+	TSharedPtr<FGpuTimingTrack> GetGpuTrack() { return GpuTrack; }
+	TSharedPtr<FCpuTimingTrack> GetCpuTrack(uint32 InThreadId);
+
+	bool IsGpuTrackVisible() const;
+	bool IsCpuTrackVisible(uint32 InThreadId) const;
+
+	//////////////////////////////////////////////////
+	// ITimingViewExtender interface
+
+	virtual void OnBeginSession(Insights::ITimingViewSession& InSession) override;
+	virtual void OnEndSession(Insights::ITimingViewSession& InSession) override;
+	virtual void Tick(Insights::ITimingViewSession& InSession, const Trace::IAnalysisSession& InAnalysisSession) override;
+	virtual void ExtendFilterMenu(Insights::ITimingViewSession& InSession, FMenuBuilder& InMenuBuilder) override;
+	virtual void OnTracksChanged(Insights::ITimingViewSession& InSession, int32& InOutOrder) override;
+
+	//////////////////////////////////////////////////
+
+	void ShowHideAllGpuTracks() { ShowHideAllGpuTracks_Execute(); }
+	void ShowHideAllCpuTracks() { ShowHideAllCpuTracks_Execute(); }
+
+private:
+	void CreateThreadGroupsMenu(FMenuBuilder& MenuBuilder);
+
+	bool ShowHideAllGpuTracks_IsChecked() const;
+	void ShowHideAllGpuTracks_Execute();
+
+	bool ShowHideAllCpuTracks_IsChecked() const;
+	void ShowHideAllCpuTracks_Execute();
+
+	bool ToggleTrackVisibilityByGroup_IsChecked(const TCHAR* InGroupName) const;
+	void ToggleTrackVisibilityByGroup_Execute(const TCHAR* InGroupName);
+
+private:
+	STimingView* TimingView;
+
+	bool bShowHideAllGpuTracks;
+	bool bShowHideAllCpuTracks;
+
+	TSharedPtr<FGpuTimingTrack> GpuTrack;
+
+	/** Maps thread id to track pointer. */
+	TMap<uint32, TSharedPtr<FCpuTimingTrack>> CpuTracks;
+
+	/** Maps thread group name to thread group info. */
+	TMap<const TCHAR*, FThreadGroup> ThreadGroups;
+};
 
 ////////////////////////////////////////////////////////////////////////////////////////////////////
 
 class FThreadTimingTrack : public FTimingEventsTrack
 {
 public:
-	explicit FThreadTimingTrack(uint64 InTrackId, const FName& InSubType, const FString& InName, const TCHAR* InGroupName)
-		: FTimingEventsTrack(InTrackId, FName(TEXT("Thread")), InSubType, InName)
+	explicit FThreadTimingTrack(FThreadTimingSharedState& InSharedState, const FName& InSubType, const FString& InName, const TCHAR* InGroupName, uint32 InTimelineIndex, uint32 InThreadId)
+		: FTimingEventsTrack(FName(TEXT("Thread")), InSubType, InName)
+		, SharedState(InSharedState)
 		, GroupName(InGroupName)
-		, ThreadId(0)
+		, TimelineIndex(InTimelineIndex)
+		, ThreadId(InThreadId)
 	{
 	}
 
@@ -27,28 +101,50 @@ public:
 
 	const TCHAR* GetGroupName() const { return GroupName; };
 
-	void SetThreadId(uint32 InThreadId) { ThreadId = InThreadId; }
+	uint32 GetTimelineIndex() const { return TimelineIndex; }
+	//void SetTimelineIndex(uint32 InTimelineIndex) { TimelineIndex = InTimelineIndex; }
+
 	uint32 GetThreadId() const { return ThreadId; }
+	//void SetThreadId(uint32 InThreadId) { ThreadId = InThreadId; }
+
+	virtual void PreUpdate(const ITimingTrackUpdateContext& Context) override;
+
+	virtual void Draw(const ITimingTrackDrawContext& Context) const override;
+	virtual void PostDraw(const ITimingTrackDrawContext& Context) const override;
 
-	virtual void Draw(ITimingViewDrawHelper& Helper) const override;
-	virtual void DrawSelectedEventInfo(const FTimingEvent& SelectedTimingEvent, const FTimingTrackViewport& Viewport, const FDrawContext& DrawContext, const FSlateBrush* WhiteBrush, const FSlateFontInfo& Font) const override;
-	virtual void InitTooltip(FTooltipDrawState& Tooltip, const FTimingEvent& HoveredTimingEvent) const override;
+	virtual void InitTooltip(FTooltipDrawState& InOutTooltip, const ITimingEvent& InTooltipEvent) const override;
 
-	virtual bool SearchTimingEvent(const FTimingEventSearchParameters& InSearchParameters, FTimingEvent& InOutTimingEvent) const override;
-	virtual void ComputeTimingEventStats(FTimingEvent& InOutTimingEvent) const override;
-	virtual void OnEventSelected(const FTimingEvent& SelectedTimingEvent) const override;
-	virtual void OnClipboardCopyEvent(const FTimingEvent& SelectedTimingEvent) const override;
+	virtual const TSharedPtr<const ITimingEvent> SearchEvent(const FTimingEventSearchParameters& InSearchParameters) const override;
+
+	virtual void UpdateEventStats(ITimingEvent& InOutEvent) const override;
+	virtual void OnEventSelected(const ITimingEvent& InSelectedEvent) const override;
+	virtual void OnClipboardCopyEvent(const ITimingEvent& InSelectedEvent) const override;
 	virtual void BuildContextMenu(FMenuBuilder& MenuBuilder) override;
 
+protected:
+	FTimingEventsTrackDrawState& GetDrawState() { return DrawState; }
+
 private:
+	void DrawSelectedEventInfo(const FTimingEvent& SelectedEvent, const FTimingTrackViewport& Viewport, const FDrawContext& DrawContext, const FSlateBrush* WhiteBrush, const FSlateFontInfo& Font) const;
+
 	bool FindTimingProfilerEvent(const FTimingEvent& InTimingEvent, TFunctionRef<void(double, double, uint32, const Trace::FTimingProfilerEvent&)> InFoundPredicate) const;
 	bool FindTimingProfilerEvent(const FTimingEventSearchParameters& InParameters, TFunctionRef<void(double, double, uint32, const Trace::FTimingProfilerEvent&)> InFoundPredicate) const;
-	void GetParentAndRoot(const FTimingEvent& TimingEvent, FTimingEvent& OutParentTimingEvent, Trace::FTimingProfilerEvent& OutParentEvent, FTimingEvent& OutRootTimingEvent, Trace::FTimingProfilerEvent& OutRootEvent) const;
+
+	void GetParentAndRoot(const FTimingEvent& TimingEvent,
+						  TSharedPtr<FTimingEvent>& OutParentTimingEvent,
+						  Trace::FTimingProfilerEvent& OutParentEvent,
+						  TSharedPtr<FTimingEvent>& OutRootTimingEvent,
+						  Trace::FTimingProfilerEvent& OutRootEvent) const;
 
 private:
+	FThreadTimingSharedState& SharedState;
+
 	const TCHAR* GroupName;
+	uint32 TimelineIndex;
 	uint32 ThreadId;
 
+	FTimingEventsTrackDrawState DrawState;
+
 	// Search cache
 	mutable TTimingEventSearchCache<Trace::FTimingProfilerEvent> SearchCache;
 };
@@ -58,8 +154,8 @@ private:
 class FCpuTimingTrack : public FThreadTimingTrack
 {
 public:
-	explicit FCpuTimingTrack(uint64 InTrackId, const FString& InName, const TCHAR* InGroupName)
-		: FThreadTimingTrack(InTrackId, FName(TEXT("CPU")), InName, InGroupName)
+	explicit FCpuTimingTrack(FThreadTimingSharedState& InSharedState, const FString& InName, const TCHAR* InGroupName, uint32 InTimelineIndex, uint32 InThreadId)
+		: FThreadTimingTrack(InSharedState, FName(TEXT("CPU")), InName, InGroupName, InTimelineIndex, InThreadId)
 	{
 	}
 };
@@ -69,8 +165,8 @@ public:
 class FGpuTimingTrack : public FThreadTimingTrack
 {
 public:
-	explicit FGpuTimingTrack(uint64 InTrackId, const FString& InName, const TCHAR* InGroupName)
-		: FThreadTimingTrack(InTrackId, FName(TEXT("GPU")), InName, InGroupName)
+	explicit FGpuTimingTrack(FThreadTimingSharedState& InSharedState, const FString& InName, const TCHAR* InGroupName, uint32 InTimelineIndex, uint32 InThreadId)
+		: FThreadTimingTrack(InSharedState, FName(TEXT("GPU")), InName, InGroupName, InTimelineIndex, InThreadId)
 	{
 	}
 };
diff --git a/Engine/Source/Developer/TraceInsights/Private/Insights/ViewModels/TimeRulerTrack.cpp b/Engine/Source/Developer/TraceInsights/Private/Insights/ViewModels/TimeRulerTrack.cpp
--- a/Engine/Source/Developer/TraceInsights/Private/Insights/ViewModels/TimeRulerTrack.cpp
+++ b/Engine/Source/Developer/TraceInsights/Private/Insights/ViewModels/TimeRulerTrack.cpp
@@ -7,6 +7,7 @@
 #include "Framework/Application/SlateApplication.h"
 #include "Rendering/DrawElements.h"
 #include "Styling/CoreStyle.h"
+#include <limits>
 
 // Insights
 #include "Insights/Common/PaintUtils.h"
@@ -19,8 +20,8 @@
 
 ////////////////////////////////////////////////////////////////////////////////////////////////////
 
-FTimeRulerTrack::FTimeRulerTrack(uint64 InTrackId)
-	: FBaseTimingTrack(InTrackId)
+FTimeRulerTrack::FTimeRulerTrack()
+	: FBaseTimingTrack(FName(TEXT("TimeRuler")))
 	, WhiteBrush(FInsightsStyle::Get().GetBrush("WhiteBrush"))
 	, Font(FCoreStyle::GetDefaultFontStyle("Regular", 8))
 	, CrtMousePosTextWidth(0.0f)
@@ -40,20 +41,46 @@ void FTimeRulerTrack::Reset()
 {
 	FBaseTimingTrack::Reset();
 
+	bIsSelecting = false;
+	SelectionStartTime = 0.0;
+	SelectionEndTime = 0.0;
+
+	bIsDragging = false;
+	TimeMarker = std::numeric_limits<double>::infinity();
+
 	constexpr float TimeRulerHeight = 24.0f;
 	SetHeight(TimeRulerHeight);
 }
 
 ////////////////////////////////////////////////////////////////////////////////////////////////////
 
-void FTimeRulerTrack::UpdateHoveredState(float MouseX, float MouseY, const FTimingTrackViewport& Viewport)
+void FTimeRulerTrack::SetSelection(const bool bInIsSelecting, const double InSelectionStartTime, const double InSelectionEndTime)
+{
+	bIsSelecting = bInIsSelecting;
+	SelectionStartTime = InSelectionStartTime;
+	SelectionEndTime = InSelectionEndTime;
+}
+
+////////////////////////////////////////////////////////////////////////////////////////////////////
+
+void FTimeRulerTrack::SetTimeMarker(const bool bInIsDragging, const double InTimeMarker)
+{
+	bIsDragging = bInIsDragging;
+	TimeMarker = InTimeMarker;
+}
+
+////////////////////////////////////////////////////////////////////////////////////////////////////
+
+void FTimeRulerTrack::PostUpdate(const ITimingTrackUpdateContext& Context)
 {
 	//constexpr float HeaderWidth = 100.0f;
 	//constexpr float HeaderHeight = 14.0f;
 
+	const float MouseY = Context.GetMousePosition().Y;
 	if (MouseY >= GetPosY() && MouseY < GetPosY() + GetHeight())
 	{
 		SetHoveredState(true);
+		//const float MouseX = Context.GetMousePosition().X;
 		//SetHeaderHoveredState(MouseX < HeaderWidth && MouseY < GetPosY() + HeaderHeight);
 	}
 	else
@@ -64,15 +91,11 @@ void FTimeRulerTrack::UpdateHoveredState(float MouseX, float MouseY, const FTimi
 
 ////////////////////////////////////////////////////////////////////////////////////////////////////
 
-void FTimeRulerTrack::DrawBackground(FDrawContext& DrawContext, const FTimingTrackViewport& Viewport) const
+void FTimeRulerTrack::Draw(const ITimingTrackDrawContext& Context) const
 {
-	FDrawHelpers::DrawBackground(DrawContext, WhiteBrush, Viewport, GetPosY(), GetHeight());
-}
+	FDrawContext& DrawContext = Context.GetDrawContext();
+	const FTimingTrackViewport& Viewport = Context.GetViewport();
 
-////////////////////////////////////////////////////////////////////////////////////////////////////
-
-void FTimeRulerTrack::Draw(FDrawContext& DrawContext, const FTimingTrackViewport& Viewport, const FVector2D& MousePosition, const bool bIsSelecting, const double SelectionStartTime, const double SelectionEndTime, const double TimeMarker) const
-{
 	const float MinorTickMark = 5.0f;
 	const float MajorTickMark = 20 * MinorTickMark;
 
@@ -91,7 +114,7 @@ void FTimeRulerTrack::Draw(FDrawContext& DrawContext, const FTimingTrackViewport
 	float MajorOX = static_cast<float>(FMath::RoundToDouble(MajorN * static_cast<double>(MajorTickMark) - VX));
 
 	// Draw the time ruler's background.
-	DrawBackground(DrawContext, Viewport);
+	FDrawHelpers::DrawBackground(DrawContext, WhiteBrush, Viewport, GetPosY(), GetHeight());
 
 	// Draw the minor tick marks.
 	for (float X = MinorOX; X < Viewport.GetWidth(); X += MinorTickMark)
@@ -123,10 +146,21 @@ void FTimeRulerTrack::Draw(FDrawContext& DrawContext, const FTimingTrackViewport
 			(T < Viewport.GetMinValidTime() || T >= Viewport.GetMaxValidTime()) ? FLinearColor(0.7f, 0.5f, 0.5f, 1.0f) : FLinearColor(0.8f, 0.8f, 0.8f, 1.0f));
 	}
 	DrawContext.LayerId++;
+}
+
+////////////////////////////////////////////////////////////////////////////////////////////////////
+
+void FTimeRulerTrack::PostDraw(const ITimingTrackDrawContext& Context) const
+{
+	FDrawContext& DrawContext = Context.GetDrawContext();
+	const FTimingTrackViewport& Viewport = Context.GetViewport();
+	const FVector2D& MousePosition = Context.GetMousePosition();
 
 	bool bShowMousePos = !MousePosition.IsZero();
 	if (bShowMousePos)
 	{
+		const double DT = 100.0 / Viewport.GetScaleX();
+
 		const FLinearColor MousePosLineColor(0.9f, 0.9f, 0.9f, 0.1f);
 		const FLinearColor MousePosTextBackgroundColor(0.9f, 0.9f, 0.9f, 1.0f);
 		FLinearColor MousePosTextForegroundColor(0.1f, 0.1f, 0.1f, 1.0f);
@@ -147,6 +181,8 @@ void FTimeRulerTrack::Draw(FDrawContext& DrawContext, const FTimingTrackViewport
 			MousePosText = TimeUtils::FormatTime(MousePosTime, MousePosPrecision);
 		}
 
+		const TSharedRef<FSlateFontMeasure> FontMeasureService = FSlateApplication::Get().GetRenderer()->GetFontMeasureService();
+
 		const float MousePosTextWidth = FMath::RoundToFloat(FontMeasureService->Measure(MousePosText, Font).X);
 
 		if (!FMath::IsNearlyEqual(CrtMousePosTextWidth, MousePosTextWidth))
@@ -155,6 +191,8 @@ void FTimeRulerTrack::Draw(FDrawContext& DrawContext, const FTimingTrackViewport
 			CrtMousePosTextWidth = CrtMousePosTextWidth * 0.6f + MousePosTextWidth * 0.4f;
 		}
 
+		const float TextY = GetPosY() + 11.0f;
+
 		float X = MousePosition.X;
 		float W = CrtMousePosTextWidth + 4.0f;
 		if (bIsSelecting && SelectionStartTime < SelectionEndTime)
@@ -193,20 +231,23 @@ void FTimeRulerTrack::Draw(FDrawContext& DrawContext, const FTimingTrackViewport
 		DrawContext.LayerId++;
 	}
 
-	// Draw the time marker (orange vertical line).
-	float TimeMarkerX = Viewport.TimeToSlateUnitsRounded(TimeMarker);
+	// Draw the time marker.
+	const float TimeMarkerX = Viewport.TimeToSlateUnitsRounded(TimeMarker);
 	if (TimeMarkerX >= 0.0f && TimeMarkerX < Viewport.GetWidth())
 	{
 		const FLinearColor TimeMarkerColor(0.85f, 0.5f, 0.03f, 0.5f);
 		const FLinearColor TimeMarkerTextBackgroundColor(TimeMarkerColor.CopyWithNewOpacity(1.0f));
-		FLinearColor TimeMarkerTextForegroundColor(0.1f, 0.1f, 0.1f, 1.0f);
+		const FLinearColor TimeMarkerTextForegroundColor(0.1f, 0.1f, 0.1f, 1.0f);
 
+		// Draw the orange vertical line.
 		DrawContext.DrawBox(TimeMarkerX, 0.0f, 1.0f, Viewport.GetHeight(), WhiteBrush, TimeMarkerColor);
 		DrawContext.LayerId++;
 
 		// Time at current marker
 		FString TimeMarkerText;
 
+		const double DT = 100.0 / Viewport.GetScaleX();
+
 		const double MousePosPrecision = FMath::Max(DT / 100.0, TimeUtils::Nanosecond);
 		if (!MousePosition.IsZero() && MousePosition.Y >= GetPosY() && MousePosition.Y < GetPosY() + GetHeight())
 		{
@@ -219,6 +260,8 @@ void FTimeRulerTrack::Draw(FDrawContext& DrawContext, const FTimingTrackViewport
 			TimeMarkerText = TimeUtils::FormatTime(TimeMarker, MousePosPrecision);
 		}
 
+		const TSharedRef<FSlateFontMeasure> FontMeasureService = FSlateApplication::Get().GetRenderer()->GetFontMeasureService();
+
 		const float TimeMarkerTextWidth = FMath::RoundToFloat(FontMeasureService->Measure(TimeMarkerText, Font).X);
 
 		if (!FMath::IsNearlyEqual(CrtTimeMarkerTextWidth, TimeMarkerTextWidth))
diff --git a/Engine/Source/Developer/TraceInsights/Private/Insights/ViewModels/TimeRulerTrack.h b/Engine/Source/Developer/TraceInsights/Private/Insights/ViewModels/TimeRulerTrack.h
--- a/Engine/Source/Developer/TraceInsights/Private/Insights/ViewModels/TimeRulerTrack.h
+++ b/Engine/Source/Developer/TraceInsights/Private/Insights/ViewModels/TimeRulerTrack.h
@@ -17,21 +17,17 @@ class FTimingTrackViewport;
 class FTimeRulerTrack : public FBaseTimingTrack
 {
 public:
-	explicit FTimeRulerTrack(uint64 InTrackId);
+	FTimeRulerTrack();
 	virtual ~FTimeRulerTrack();
 
 	virtual void Reset() override;
-	virtual void UpdateHoveredState(float MouseX, float MouseY, const FTimingTrackViewport& Viewport);
 
-	void Draw(FDrawContext& DrawContext, const FTimingTrackViewport& Viewport,
-				const FVector2D& MousePosition = FVector2D(0.0f, 0.0f),
-				const bool bIsSelecting = false,
-				const double SelectionStartTime = 0.0,
-				const double SelectionEndTime = 0.0,
-				const double TimeMarker = 0.0) const;
+	void SetSelection(const bool bInIsSelecting, const double InSelectionStartTime, const double InSelectionEndTime);
+	void SetTimeMarker(const bool bInIsDragging, const double InTimeMarker);
 
-private:
-	void DrawBackground(FDrawContext& DrawContext, const FTimingTrackViewport& Viewport) const;
+	virtual void PostUpdate(const ITimingTrackUpdateContext& Context) override;
+	void Draw(const ITimingTrackDrawContext& Context) const override;
+	void PostDraw(const ITimingTrackDrawContext& Context) const override;
 
 public:
 	// Slate resources
@@ -43,6 +39,13 @@ public:
 
 	// Smoothed time marker text width to avoid flickering
 	mutable float CrtTimeMarkerTextWidth;
+
+private:
+	bool bIsSelecting;
+	double SelectionStartTime;
+	double SelectionEndTime;
+	bool bIsDragging;
+	double TimeMarker;
 };
 
 ////////////////////////////////////////////////////////////////////////////////////////////////////
diff --git a/Engine/Source/Developer/TraceInsights/Private/Insights/ViewModels/TimingEventsTrack.cpp b/Engine/Source/Developer/TraceInsights/Private/Insights/ViewModels/TimingEventsTrack.cpp
--- a/Engine/Source/Developer/TraceInsights/Private/Insights/ViewModels/TimingEventsTrack.cpp
+++ b/Engine/Source/Developer/TraceInsights/Private/Insights/ViewModels/TimingEventsTrack.cpp
@@ -2,125 +2,141 @@
 
 #include "Insights/ViewModels/TimingEventsTrack.h"
 
+#include "Insights/ViewModels/TimingEvent.h"
+#include "Insights/ViewModels/TimingTrackViewport.h"
+#include "Insights/ViewModels/TimingViewDrawHelper.h"
+
 #define LOCTEXT_NAMESPACE "TimingEventsTrack"
 
 ////////////////////////////////////////////////////////////////////////////////////////////////////
-// FTimingEventsTrackLayout
+
+const FName FTimingEvent::TypeName = FName(TEXT("FTimingEvent"));
+
+bool FTimingEventsTrack::bUseDownSampling = true;
+
 ////////////////////////////////////////////////////////////////////////////////////////////////////
 
-void FTimingEventsTrackLayout::ForceNormalMode()
+FTimingEventsTrack::FTimingEventsTrack(const FName& InType, const FName& InSubType, const FString& InName)
+	: FBaseTimingTrack(InType, InSubType, InName)
+	, NumLanes(0)
+	, DrawState(MakeShareable(new FTimingEventsTrackDrawState()))
 {
-	bIsCompactMode = false;
-	EventH = NormalLayoutEventH;
-	EventDY = NormalLayoutEventDY;
-	TimelineDY = NormalLayoutTimelineDY;
-	MinTimelineH = NormalLayoutMinTimelineH;
-	TargetMinTimelineH = NormalLayoutMinTimelineH;
 }
 
 ////////////////////////////////////////////////////////////////////////////////////////////////////
 
-void FTimingEventsTrackLayout::ForceCompactMode()
+FTimingEventsTrack::~FTimingEventsTrack()
 {
-	bIsCompactMode = true;
-	EventH = CompactLayoutEventH;
-	EventDY = CompactLayoutEventDY;
-	TimelineDY = CompactLayoutTimelineDY;
-	MinTimelineH = CompactLayoutMinTimelineH;
-	TargetMinTimelineH = CompactLayoutMinTimelineH;
 }
 
 ////////////////////////////////////////////////////////////////////////////////////////////////////
 
-void FTimingEventsTrackLayout::Update()
+void FTimingEventsTrack::Reset()
 {
-	constexpr float LayoutTransitionSpeed = 0.25f;
+	FBaseTimingTrack::Reset();
 
-	if (bIsCompactMode)
-	{
-		if (EventH > CompactLayoutEventH)
-		{
-			EventH -= LayoutTransitionSpeed;
-		}
-		if (EventDY > CompactLayoutEventDY)
-		{
-			EventDY -= LayoutTransitionSpeed;
-		}
-		if (TimelineDY > CompactLayoutTimelineDY)
-		{
-			TimelineDY -= LayoutTransitionSpeed;
-		}
-	}
-	else
-	{
-		if (EventH < NormalLayoutEventH)
-		{
-			EventH += LayoutTransitionSpeed;
-		}
-		if (EventDY < NormalLayoutEventDY)
-		{
-			EventDY += LayoutTransitionSpeed;
-		}
-		if (TimelineDY < NormalLayoutTimelineDY)
-		{
-			TimelineDY += LayoutTransitionSpeed;
-		}
-	}
+	NumLanes = 0;
+	DrawState->Reset();
+}
+
+////////////////////////////////////////////////////////////////////////////////////////////////////
 
-	if (MinTimelineH > TargetMinTimelineH)
+void FTimingEventsTrack::PostUpdate(const ITimingTrackUpdateContext& Context)
+{
+	FBaseTimingTrack::PostUpdate(Context);
+
+	constexpr float HeaderWidth = 100.0f;
+	constexpr float HeaderHeight = 14.0f;
+
+	const float MouseY = Context.GetMousePosition().Y;
+	if (MouseY >= GetPosY() && MouseY < GetPosY() + GetHeight())
 	{
-		MinTimelineH -= LayoutTransitionSpeed;
+		SetHoveredState(true);
+		
+		const float MouseX = Context.GetMousePosition().X;
+		SetHeaderHoveredState(MouseX < HeaderWidth && MouseY < GetPosY() + HeaderHeight);
 	}
-	else if (MinTimelineH < TargetMinTimelineH)
+	else
 	{
-		MinTimelineH += LayoutTransitionSpeed;
+		SetHoveredState(false);
 	}
 }
 
-////////////////////////////////////////////////////////////////////////////////////////////////////
-// FTimingEventsTrack
-////////////////////////////////////////////////////////////////////////////////////////////////////
-
-bool FTimingEventsTrack::bUseDownSampling = true;
-
 ////////////////////////////////////////////////////////////////////////////////////////////////////
 
-FTimingEventsTrack::FTimingEventsTrack(uint64 InTrackId, const FName& InType, const FName& InSubType, const FString& InName)
-	: FBaseTimingTrack(InTrackId, InType, InSubType, InName)
-	, NumLanes(0)
+void FTimingEventsTrack::DrawEvent(const ITimingTrackDrawContext& Context, const ITimingEvent& InTimingEvent, EDrawEventMode InDrawMode) const
 {
+	if (InTimingEvent.CheckTrack(this) && FTimingEvent::CheckTypeName(InTimingEvent))
+	{
+		const FTimingEvent& TrackEvent = static_cast<const FTimingEvent&>(InTimingEvent);
+		const FTimingViewLayout& Layout = Context.GetViewport().GetLayout();
+		const float Y = TrackEvent.GetTrack()->GetPosY() + Layout.GetLaneY(TrackEvent.GetDepth());
+		Context.GetHelper().DrawTimingEventHighlight(TrackEvent.GetStartTime(), TrackEvent.GetEndTime(), Y, InDrawMode);
+	}
 }
 
 ////////////////////////////////////////////////////////////////////////////////////////////////////
 
-FTimingEventsTrack::~FTimingEventsTrack()
+const TSharedPtr<const ITimingEvent> FTimingEventsTrack::GetEvent(float InPosX, float InPosY, const FTimingTrackViewport& Viewport) const
 {
-}
+	const FTimingViewLayout& Layout = Viewport.GetLayout();
 
-////////////////////////////////////////////////////////////////////////////////////////////////////
+	const float TopLaneY = GetPosY() + 1.0f + Layout.TimelineDY; // +1.0f is for horizontal line between timelines
+	const float DY = InPosY - TopLaneY;
 
-void FTimingEventsTrack::Reset()
-{
-	FBaseTimingTrack::Reset();
+	// If mouse is not above first sub-track or below last sub-track...
+	if (DY >= 0 && DY < GetHeight() - 1.0f - 2 * Layout.TimelineDY)
+	{
+		const int32 Depth = DY / (Layout.EventH + Layout.EventDY);
 
-	NumLanes = 0;
+		const double StartTime = Viewport.SlateUnitsToTime(InPosX) - 1.0 / Viewport.GetScaleX(); // -1px
+		const double EndTime = StartTime + 3.0 / Viewport.GetScaleX(); // +3px
+
+		auto EventFilter = [Depth](double, double, uint32 EventDepth)
+		{
+			return EventDepth == Depth;
+		};
+
+		return SearchEvent(FTimingEventSearchParameters(StartTime, EndTime, ETimingEventSearchFlags::StopAtFirstMatch, EventFilter));
+	}
+
+	return nullptr;
 }
 
 ////////////////////////////////////////////////////////////////////////////////////////////////////
 
-void FTimingEventsTrack::UpdateHoveredState(float MouseX, float MouseY, const FTimingTrackViewport& Viewport)
+void FTimingEventsTrack::UpdateTrackHeight(const ITimingTrackUpdateContext& Context)
 {
-	constexpr float HeaderWidth = 100.0f;
-	constexpr float HeaderHeight = 14.0f;
+	const FTimingTrackViewport& Viewport = Context.GetViewport();
 
-	if (MouseY >= GetPosY() && MouseY < GetPosY() + GetHeight())
+	const float CurrentTrackHeight = GetHeight();
+	const float DesiredTrackHeight = Viewport.GetLayout().ComputeTrackHeight(NumLanes);
+
+	if (CurrentTrackHeight < DesiredTrackHeight)
 	{
-		SetHoveredState(true);
-		SetHeaderHoveredState(MouseX < HeaderWidth && MouseY < GetPosY() + HeaderHeight);
+		float NewTrackHeight;
+		if (Viewport.IsDirty(ETimingTrackViewportDirtyFlags::VLayoutChanged))
+		{
+			NewTrackHeight = DesiredTrackHeight;
+		}
+		else
+		{
+			NewTrackHeight = FMath::CeilToFloat(CurrentTrackHeight * 0.9f + DesiredTrackHeight * 0.1f);
+		}
+		SetHeight(NewTrackHeight);
 	}
-	else
+	else if (CurrentTrackHeight > DesiredTrackHeight)
 	{
-		SetHoveredState(false);
+		float NewTrackHeight;
+		if (Viewport.IsDirty(ETimingTrackViewportDirtyFlags::VLayoutChanged))
+		{
+			NewTrackHeight = DesiredTrackHeight;
+		}
+		else
+		{
+			NewTrackHeight = FMath::FloorToFloat(CurrentTrackHeight * 0.9f + DesiredTrackHeight * 0.1f);
+		}
+		SetHeight(NewTrackHeight);
 	}
 }
 
diff --git a/Engine/Source/Developer/TraceInsights/Private/Insights/ViewModels/TimingGraphTrack.cpp b/Engine/Source/Developer/TraceInsights/Private/Insights/ViewModels/TimingGraphTrack.cpp
--- a/Engine/Source/Developer/TraceInsights/Private/Insights/ViewModels/TimingGraphTrack.cpp
+++ b/Engine/Source/Developer/TraceInsights/Private/Insights/ViewModels/TimingGraphTrack.cpp
@@ -11,6 +11,7 @@
 #include "Insights/Common/TimeUtils.h"
 #include "Insights/InsightsManager.h"
 #include "Insights/TimingProfilerManager.h"
+#include "Insights/ViewModels/GraphTrackBuilder.h"
 #include "Insights/ViewModels/TimingTrackViewport.h"
 
 #define LOCTEXT_NAMESPACE "GraphTrack"
@@ -54,8 +55,8 @@ FString FTimingGraphSeries::FormatValue(double Value) const
 // FTimingGraphTrack
 ////////////////////////////////////////////////////////////////////////////////////////////////////
 
-FTimingGraphTrack::FTimingGraphTrack(uint64 InTrackId)
-	: FGraphTrack(InTrackId)
+FTimingGraphTrack::FTimingGraphTrack()
+	: FGraphTrack(FName(TEXT("Random")))
 {
 	bDrawPoints = true;
 	bDrawPointsWithBorder = true;
@@ -172,11 +173,11 @@ FTimingGraphTrack::~FTimingGraphTrack()
 
 ////////////////////////////////////////////////////////////////////////////////////////////////////
 
-void FTimingGraphTrack::Update(const FTimingTrackViewport& Viewport)
+void FTimingGraphTrack::Update(const ITimingTrackUpdateContext& Context)
 {
-	FGraphTrack::Update(Viewport);
+	FGraphTrack::Update(Context);
 
-	const bool bIsEntireGraphTrackDirty = IsDirty();
+	const bool bIsEntireGraphTrackDirty = IsDirty() || Context.GetViewport().IsHorizontalViewportDirty();
 	bool bNeedsUpdate = bIsEntireGraphTrackDirty;
 
 	if (!bNeedsUpdate)
@@ -198,6 +199,8 @@ void FTimingGraphTrack::Update(const FTimingTrackViewport& Viewport)
 
 		NumAddedEvents = 0;
 
+		const FTimingTrackViewport& Viewport = Context.GetViewport();
+
 		for (TSharedPtr<FGraphSeries>& Series : AllSeries)
 		{
 			if (Series->IsVisible() && (bIsEntireGraphTrackDirty || Series->IsDirty()))
diff --git a/Engine/Source/Developer/TraceInsights/Private/Insights/ViewModels/TimingGraphTrack.h b/Engine/Source/Developer/TraceInsights/Private/Insights/ViewModels/TimingGraphTrack.h
--- a/Engine/Source/Developer/TraceInsights/Private/Insights/ViewModels/TimingGraphTrack.h
+++ b/Engine/Source/Developer/TraceInsights/Private/Insights/ViewModels/TimingGraphTrack.h
@@ -5,6 +5,7 @@
 #include "CoreMinimal.h"
 
 // Insights
+#include "Insights/ViewModels/GraphSeries.h"
 #include "Insights/ViewModels/GraphTrack.h"
 
 ////////////////////////////////////////////////////////////////////////////////////////////////////
@@ -34,10 +35,10 @@ public:
 class FTimingGraphTrack : public FGraphTrack
 {
 public:
-	explicit FTimingGraphTrack(uint64 InTrackId);
+	FTimingGraphTrack();
 	virtual ~FTimingGraphTrack();
 
-	virtual void Update(const FTimingTrackViewport& Viewport) override;
+	virtual void Update(const ITimingTrackUpdateContext& Context) override;
 
 	void AddDefaultFrameSeries();
 
diff --git a/Engine/Source/Developer/TraceInsights/Private/Insights/ViewModels/TimingTrackViewport.cpp b/Engine/Source/Developer/TraceInsights/Private/Insights/ViewModels/TimingTrackViewport.cpp
--- a/Engine/Source/Developer/TraceInsights/Private/Insights/ViewModels/TimingTrackViewport.cpp
+++ b/Engine/Source/Developer/TraceInsights/Private/Insights/ViewModels/TimingTrackViewport.cpp
@@ -10,14 +10,21 @@
 
 bool FTimingTrackViewport::UpdateSize(const float InWidth, const float InHeight)
 {
-	if (Width != InWidth || Height != InHeight)
+	bool bSizeChanged = false;
+	if (Width != InWidth)
 	{
 		Width = InWidth;
-		Height = InHeight;
 		EndTime = SlateUnitsToTime(Width);
-		return true;
+		AddDirtyFlags(ETimingTrackViewportDirtyFlags::HSizeChanged);
+		bSizeChanged = true;
 	}
-	return false;
+	if (Height != InHeight)
+	{
+		Height = InHeight;
+		AddDirtyFlags(ETimingTrackViewportDirtyFlags::VSizeChanged);
+		bSizeChanged = true;
+	}
+	return bSizeChanged;
 }
 
 ////////////////////////////////////////////////////////////////////////////////////////////////////
@@ -29,6 +36,7 @@ bool FTimingTrackViewport::ScrollAtTime(const double Time)
 	{
 		StartTime = NewStartTime;
 		EndTime = SlateUnitsToTime(Width);
+		AddDirtyFlags(ETimingTrackViewportDirtyFlags::HPositionChanged);
 		return true;
 	}
 	return false;
@@ -64,6 +72,7 @@ bool FTimingTrackViewport::ZoomOnTimeInterval(const double Time, const double Du
 		StartTime = NewStartTime;
 		ScaleX = NewScaleX;
 		EndTime = SlateUnitsToTime(Width);
+		AddDirtyFlags(ETimingTrackViewportDirtyFlags::HPositionChanged | ETimingTrackViewportDirtyFlags::HScaleChanged);
 		return true;
 	}
 	return false;
@@ -102,6 +111,7 @@ bool FTimingTrackViewport::ZoomWithFixedX(const double NewScaleX, const float X)
 		StartTime += static_cast<double>(X) * (1.0 / ScaleX - 1.0 / LocalNewScaleX);
 		ScaleX = LocalNewScaleX;
 		EndTime = SlateUnitsToTime(Width);
+		AddDirtyFlags(ETimingTrackViewportDirtyFlags::HPositionChanged | ETimingTrackViewportDirtyFlags::HScaleChanged);
 		return true;
 	}
 	return false;
@@ -186,11 +196,21 @@ bool FTimingTrackViewport::EnforceHorizontalScrollLimits(const double U)
 
 ////////////////////////////////////////////////////////////////////////////////////////////////////
 
+void FTimingTrackViewport::UpdateLayout()
+{
+	if (Layout.Update())
+	{
+		AddDirtyFlags(ETimingTrackViewportDirtyFlags::VLayoutChanged);
+	}
+}
+
+////////////////////////////////////////////////////////////////////////////////////////////////////
+
 bool FTimingTrackViewport::OnUserScrolled(TSharedPtr<SScrollBar> ScrollBar, float ScrollOffset)
 {
 	const double S = 1.0 / (MaxValidTime - MinValidTime);
-	const double D = EndTime - StartTime;
-	const float ThumbSizeFraction = FMath::Clamp<float>(D * S, 0.0f, 1.0f);
+	const double Page = EndTime - StartTime;
+	const float ThumbSizeFraction = FMath::Clamp<float>(Page * S, 0.0f, 1.0f);
 	const float OffsetFraction = FMath::Clamp<float>(ScrollOffset, 0.0f, 1.0f - ThumbSizeFraction);
 
 	ScrollBar->SetState(OffsetFraction, ThumbSizeFraction);
@@ -210,8 +230,8 @@ bool FTimingTrackViewport::OnUserScrolled(TSharedPtr<SScrollBar> ScrollBar, floa
 void FTimingTrackViewport::UpdateScrollBar(TSharedPtr<SScrollBar> ScrollBar) const
 {
 	const double S = 1.0 / (MaxValidTime - MinValidTime);
-	const double D = EndTime - StartTime;
-	const float ThumbSizeFraction = FMath::Clamp<float>(D * S, 0.0f, 1.0f);
+	const double Page = EndTime - StartTime;
+	const float ThumbSizeFraction = FMath::Clamp<float>(Page * S, 0.0f, 1.0f);
 	const float ScrollOffset = static_cast<float>((StartTime - MinValidTime) * S);
 	const float OffsetFraction = FMath::Clamp<float>(ScrollOffset, 0.0f, 1.0f - ThumbSizeFraction);
 
@@ -222,9 +242,9 @@ void FTimingTrackViewport::UpdateScrollBar(TSharedPtr<SScrollBar> ScrollBar) con
 
 bool FTimingTrackViewport::OnUserScrolledY(TSharedPtr<SScrollBar> ScrollBar, float ScrollOffset)
 {
-	const float S = 1.0 / ScrollHeight;
-	const float D = Height - TopOffset;
-	const float ThumbSizeFraction = FMath::Clamp<float>(D * S, 0.0f, 1.0f);
+	const float S = 1.0f / ScrollHeight;
+	const float Page = Height - TopOffset - BottomOffset;
+	const float ThumbSizeFraction = FMath::Clamp<float>(Page * S, 0.0f, 1.0f);
 	const float OffsetFraction = FMath::Clamp<float>(ScrollOffset, 0.0f, 1.0f - ThumbSizeFraction);
 
 	ScrollBar->SetState(OffsetFraction, ThumbSizeFraction);
@@ -233,6 +253,7 @@ bool FTimingTrackViewport::OnUserScrolledY(TSharedPtr<SScrollBar> ScrollBar, flo
 	if (NewScrollPosY != ScrollPosY)
 	{
 		ScrollPosY = NewScrollPosY;
+		AddDirtyFlags(ETimingTrackViewportDirtyFlags::VPositionChanged);
 		return true;
 	}
 
@@ -243,9 +264,9 @@ bool FTimingTrackViewport::OnUserScrolledY(TSharedPtr<SScrollBar> ScrollBar, flo
 
 void FTimingTrackViewport::UpdateScrollBarY(TSharedPtr<SScrollBar> ScrollBar) const
 {
-	const float S = 1.0 / ScrollHeight;
-	const float D = Height - TopOffset;
-	const float ThumbSizeFraction = FMath::Clamp<float>(D * S, 0.0f, 1.0f);
+	const float S = 1.0f / ScrollHeight;
+	const float Page = Height - TopOffset - BottomOffset;
+	const float ThumbSizeFraction = FMath::Clamp<float>(Page * S, 0.0f, 1.0f);
 	const float ScrollOffset = ScrollPosY * S;
 	const float OffsetFraction = FMath::Clamp<float>(ScrollOffset, 0.0f, 1.0f - ThumbSizeFraction);
 
diff --git a/Engine/Source/Developer/TraceInsights/Private/Insights/ViewModels/TimingViewDrawHelper.cpp b/Engine/Source/Developer/TraceInsights/Private/Insights/ViewModels/TimingViewDrawHelper.cpp
--- a/Engine/Source/Developer/TraceInsights/Private/Insights/ViewModels/TimingViewDrawHelper.cpp
+++ b/Engine/Source/Developer/TraceInsights/Private/Insights/ViewModels/TimingViewDrawHelper.cpp
@@ -15,208 +15,30 @@
 #include "Insights/InsightsStyle.h"
 #include "Insights/ViewModels/DrawHelpers.h"
 #include "Insights/ViewModels/MarkersTimingTrack.h"
+#include "Insights/ViewModels/TimingEvent.h"
 #include "Insights/ViewModels/TimingEventsTrack.h"
 #include "Insights/ViewModels/TimingTrackViewport.h"
 
+////////////////////////////////////////////////////////////////////////////////////////////////////
+// FTimingEventsTrackDrawStateBuilder
 ////////////////////////////////////////////////////////////////////////////////////////////////////
 
-FTimingViewDrawHelper::FTimingViewDrawHelper(const FDrawContext& InDC, const FTimingTrackViewport& InViewport, const FTimingEventsTrackLayout& InLayout)
-	: DrawContext(InDC)
+FTimingEventsTrackDrawStateBuilder::FTimingEventsTrackDrawStateBuilder(FTimingEventsTrackDrawState& InState, const FTimingTrackViewport& InViewport)
+	: DrawState(InState)
 	, Viewport(InViewport)
-	, Layout(InLayout)
-	, WhiteBrush(FInsightsStyle::Get().GetBrush("WhiteBrush"))
-	, EventBorderBrush(FInsightsStyle::Get().GetBrush("EventBorder"))
-	, HoveredEventBorderBrush(FInsightsStyle::Get().GetBrush("HoveredEventBorder"))
-	, SelectedEventBorderBrush(FInsightsStyle::Get().GetBrush("SelectedEventBorder"))
-	, BackgroundAreaBrush(WhiteBrush)
-	, ValidAreaColor(0.07f, 0.07f, 0.07f, 1.0f)
-	, InvalidAreaColor(0.1f, 0.07f, 0.07f, 1.0f)
-	, EdgeColor(0.05f, 0.05f, 0.05f, 1.0f)
-	, EventFont(FCoreStyle::GetDefaultFontStyle("Regular", 8))
-	, ValidAreaX(0.0f)
-	, ValidAreaW(0.0f)
-	, TimelineTopY(0.0f)
-	, TimelineY(0.0f)
-	, MaxDepth(0)
-	, TimelineIndex(0)
+	, MaxDepth(-1)
 	, LastEventX2()
 	, LastBox()
-	, Stats()
-{
-}
-
-////////////////////////////////////////////////////////////////////////////////////////////////////
-
-FTimingViewDrawHelper::~FTimingViewDrawHelper()
-{
-}
-
-////////////////////////////////////////////////////////////////////////////////////////////////////
-
-void FTimingViewDrawHelper::DrawBackground() const
-{
-	const float Y = 0.0f;
-	const float H = FMath::CeilToFloat(Viewport.GetHeight());
-	FDrawHelpers::DrawBackground(DrawContext, BackgroundAreaBrush, Viewport, Y, H, ValidAreaX, ValidAreaW);
-}
-
-////////////////////////////////////////////////////////////////////////////////////////////////////
-
-void FTimingViewDrawHelper::BeginTimelines()
-{
-	TimelineIndex = -1;
-	TimelineTopY = -Viewport.GetScrollPosY();
-
-	float Y = TimelineTopY + Viewport.GetTopOffset();
-
-	if (Y > 0.0f && ValidAreaW > 0.0f)
-	{
-		Y = FMath::Min(Y, Viewport.GetHeight());
-
-		// Draw invalid area (top).
-		DrawContext.DrawBox(DrawContext.LayerId + ToInt32(EDrawLayer::TimelineHeader), ValidAreaX, 0.0f, ValidAreaW, Y, BackgroundAreaBrush, InvalidAreaColor);
-	}
-}
-
-////////////////////////////////////////////////////////////////////////////////////////////////////
-
-bool FTimingViewDrawHelper::BeginTimeline(FTimingEventsTrack& Track)
-{
-	TimelineIndex++;
-
-	Track.SetPosY(TimelineTopY + Viewport.GetScrollPosY());
-
-	if (Track.GetHeight() < Layout.MinTimelineH)
-	{
-		Track.SetHeight(Layout.MinTimelineH);
-	}
-
-	if (TimelineTopY + Track.GetHeight() < -1.0f)
-	{
-		TimelineTopY += Track.GetHeight();
-		return false;
-	}
-
-	if (TimelineTopY > Viewport.GetHeight() - Viewport.GetTopOffset())
-	{
-		return false;
-	}
-
-	MaxDepth = -1;
-
-	// +1.0f is for horizontal line between timelines
-	TimelineY = Viewport.GetTopOffset() + TimelineTopY + 1.0f + Layout.TimelineDY;
-
-	// Reset "last event X2" for all depths.
-	LastEventX2.Reset();
-
-	// Reset "last box" for all depths.
-	LastBox.Reset();
-
-	return true;
-}
-
-////////////////////////////////////////////////////////////////////////////////////////////////////
-
-void FTimingViewDrawHelper::EndTimeline(FTimingEventsTrack& Track)
+	, EventFont(FCoreStyle::GetDefaultFontStyle("Regular", 8))
 {
-	// Flush merged boxes.
-	for (int32 Depth = 0; Depth <= MaxDepth; ++Depth)
-	{
-		const FBoxData& Box = LastBox[Depth];
-		if (Box.X1 < Box.X2)
-		{
-			const float EventY = TimelineY + (Layout.EventH + Layout.EventDY) * Depth;
-			DrawBox(Box, EventY, Layout.EventH);
-		}
-	}
-
-	Track.SetDepth(MaxDepth);
-
-	float NewH;
-	if (MaxDepth < 0)
-	{
-		NewH = Layout.MinTimelineH;
-	}
-	else //if (MaxDepth >= 0)
-	{
-		// 1.0f is for horizontal line between timelines
-		NewH = 1.0f + Layout.TimelineDY + Layout.EventH * (MaxDepth + 1) + Layout.EventDY * MaxDepth + Layout.TimelineDY;
-
-		if (NewH < FTimingEventsTrackLayout::RealMinTimelineH)
-		{
-			NewH = FTimingEventsTrackLayout::RealMinTimelineH;
-		}
-	}
-
-	if (Track.GetHeight() < NewH)
-	{
-		float NewTrackHeight;
-		if (Layout.bIsCompactMode)
-		{
-			NewTrackHeight = FMath::CeilToFloat(Track.GetHeight() * 0.5f + NewH * 0.5f);
-		}
-		else
-		{
-			NewTrackHeight = FMath::CeilToFloat(Track.GetHeight() * 0.9f + NewH * 0.1f);
-		}
-		Track.SetHeight(NewTrackHeight);
-	}
-	else if (Track.GetHeight() > NewH)
-	{
-		float NewTrackHeight;
-		if (Layout.bIsCompactMode)
-		{
-			NewTrackHeight = FMath::FloorToFloat(Track.GetHeight() * 0.5f + NewH * 0.5f);
-		}
-		else
-		{
-			NewTrackHeight = FMath::FloorToFloat(Track.GetHeight() * 0.9f + NewH * 0.1f);
-		}
-		Track.SetHeight(NewTrackHeight);
-	}
-
-	if (Track.GetHeight() > 0)
-	{
-		const float Y = Viewport.GetTopOffset() + TimelineTopY;
-
-		FLinearColor TextColor(1.0f, 1.0f, 1.0f, 1.0f);
-
-		if (Track.IsSelected())
-		{
-			if (Track.IsHovered())
-			{
-				TextColor = FLinearColor(1.0f, 1.0f, 0.0f, 1.0f);
-			}
-			else
-			{
-				TextColor = FLinearColor(1.0f, 1.0f, 0.5f, 1.0f);
-			}
-		}
-		else if (Track.IsHovered())
-		{
-			TextColor = FLinearColor(1.0f, 1.0f, 0.0f, 1.0f);
-		}
-
-		// Draw a horizontal line between timelines.
-		DrawContext.DrawBox(DrawContext.LayerId + ToInt32(EDrawLayer::TimelineHeader), 0, Y, Viewport.GetWidth(), 1.0f, WhiteBrush, EdgeColor);
-
-		// Draw name of timeline.
-		//const FString Name = FString::Printf(TEXT("%d Y: %g, H: %g, VO: %g, VY: %g"), TimelineIndex, Track.Y, Track.H, Viewport.TopOffset, Viewport.ScrollPosY); //debug
-		const TSharedRef<FSlateFontMeasure> FontMeasureService = FSlateApplication::Get().GetRenderer()->GetFontMeasureService();
-		float NameWidth = FontMeasureService->Measure(Track.GetName(), EventFont).X;
-		DrawContext.DrawBox(DrawContext.LayerId + ToInt32(EDrawLayer::TimelineHeader), 0.0f, Y + 1.0f, NameWidth + 4.0f, 12.0f, WhiteBrush, EdgeColor);
-		DrawContext.DrawText(DrawContext.LayerId + ToInt32(EDrawLayer::TimelineText), 2.0f, Y, Track.GetName(), EventFont, TextColor);
-	}
-
-	TimelineTopY += Track.GetHeight();
+	DrawState.Reset();
 }
 
 ////////////////////////////////////////////////////////////////////////////////////////////////////
 
-void FTimingViewDrawHelper::AddEvent(double EventStartTime, double EventEndTime, uint32 EventDepth, const TCHAR* EventName, uint32 Color)
+void FTimingEventsTrackDrawStateBuilder::AddEvent(double EventStartTime, double EventEndTime, uint32 EventDepth, const TCHAR* EventName, uint64 TypeId, uint32 Color)
 {
-	Stats.NumEvents++;
+	DrawState.NumEvents++;
 
 	float EventX1 = Viewport.TimeToSlateUnitsRounded(EventStartTime);
 	if (EventX1 > Viewport.GetWidth())
@@ -249,12 +71,6 @@ void FTimingViewDrawHelper::AddEvent(double EventStartTime, double EventEndTime,
 		LastBox.AddDefaulted(Depth + 1 - LastBox.Num());
 	}
 
-	float EventY = TimelineY + (Layout.EventH + Layout.EventDY) * Depth;
-	if (EventY < -Layout.EventH || EventY > Viewport.GetHeight())
-	{
-		return;
-	}
-
 	// Ensure we have enough slots in array. See LastEventX2[Depth] usage.
 	while (LastEventX2.Num() <= Depth)
 	{
@@ -366,31 +182,40 @@ void FTimingViewDrawHelper::AddEvent(double EventStartTime, double EventEndTime,
 	//constexpr float BorderColorFactor = 0.75f; // darker border
 	constexpr float BorderColorFactor = 1.25f; // brighter border
 
-	//////////////////////////////////////////////////
+	//TODO if (HighlightedEventTypeId != FTimingEvent::InvalidTypeId && TypeId != HighlightedEventTypeId)
+	//{
+	//	EventColorFill.A = 0.1f;
+	//}
 
-	if (EventW > 2.0f && Layout.EventH > 2.0f)
-	{
-		// Fill inside of the timing event box.
-		DrawContext.DrawBox(DrawContext.LayerId + ToInt32(EDrawLayer::EventFill), EventX1 + 1.0f, EventY + 1.0f, EventW - 2.0f, Layout.EventH - 2.0f, WhiteBrush, EventColorFill);
-		Stats.NumDrawBoxes++;
-	}
+	//////////////////////////////////////////////////
 
 	// Save X2, for current depth.
 	LastEventX2[Depth] = EventX2;
 
-	// Draw border around the timing event box.
 	if (EventW > 2.0f)
 	{
 		FBoxData& Box = LastBox[Depth];
 		if (Box.X1 < Box.X2)
 		{
-			DrawBox(Box, EventY, Layout.EventH);
+			FlushBox(Box, Depth);
 			Box.Reset();
 		}
 
-		const FLinearColor EventColorBorder(EventColorFill.R * BorderColorFactor, EventColorFill.G * BorderColorFactor, EventColorFill.B * BorderColorFactor, 1.0f);
-		DrawContext.DrawBox(DrawContext.LayerId + ToInt32(EDrawLayer::EventBorder), EventX1, EventY, EventW, Layout.EventH, EventBorderBrush, EventColorBorder);
-		Stats.NumDrawBorders++;
+		// Fill inside of the timing event box.
+		DrawState.InsideBoxes.AddUninitialized();
+		FTimingEventsTrackDrawState::FBoxPrimitive& InsideBox = DrawState.InsideBoxes.Last();
+		InsideBox.Depth = Depth;
+		InsideBox.X = EventX1 + 1.0f;
+		InsideBox.W = EventW - 2.0f;
+		InsideBox.Color = EventColorFill;
+
+		// Add border around the timing event box.
+		DrawState.Borders.AddUninitialized();
+		FTimingEventsTrackDrawState::FBoxPrimitive& BorderBox = DrawState.Borders.Last();
+		BorderBox.Depth = Depth;
+		BorderBox.X = EventX1;
+		BorderBox.W = EventW;
+		BorderBox.Color = FLinearColor(EventColorFill.R * BorderColorFactor, EventColorFill.G * BorderColorFactor, EventColorFill.B * BorderColorFactor, EventColorFill.A);
 	}
 	else // 1px or 2px boxes
 	{
@@ -403,26 +228,26 @@ void FTimingViewDrawHelper::AddEvent(double EventStartTime, double EventEndTime,
 		{
 			// Merge it with previous box.
 			Box.X2 = EventX2;
-			Stats.NumMergedBoxes++;
+			DrawState.NumMergedBoxes++;
 		}
 		else
 		{
 			// Flush previous box, if any.
 			if (Box.X1 < Box.X2)
 			{
-				DrawBox(Box, EventY, Layout.EventH);
+				FlushBox(Box, Depth);
 			}
 
 			// Start new "merge box".
 			Box.X1 = EventX1;
 			Box.X2 = EventX2;
 			Box.Color = Color;
-			Box.LinearColor = FLinearColor(EventColorFill.R * BorderColorFactor, EventColorFill.G * BorderColorFactor, EventColorFill.B * BorderColorFactor, 1.0f);
+			Box.LinearColor = FLinearColor(EventColorFill.R * BorderColorFactor, EventColorFill.G * BorderColorFactor, EventColorFill.B * BorderColorFactor, EventColorFill.A);
 		}
 	}
 
 	// Draw the name of the timing event.
-	if (EventW > 8.0f && Layout.EventH > 10.0f)
+	if (EventW > 8.0f)
 	{
 		FString Name = EventName;// +TEXT(" [") + FText::AsNumber(Event.Type->Id).ToString() + TEXT("]");
 		if (EventW > Name.Len() * 2.0f + 48.0f)
@@ -443,44 +268,227 @@ void FTimingViewDrawHelper::AddEvent(double EventStartTime, double EventEndTime,
 			// Ref: https://books.google.ro/books?id=0pVr7dhmdWYC
 			const bool bIsDarkColor = (EventColorFill.ComputeLuminance() < 0.4f);
 
-			DrawContext.DrawText(DrawContext.LayerId + ToInt32(EDrawLayer::EventText), EventX1 + 2.0f, EventY + 1.0f, Name, 0, LastWholeCharacterIndex + 1, EventFont, bIsDarkColor ? FLinearColor::White : FLinearColor::Black);
-			Stats.NumDrawTexts++;
+			DrawState.Texts.AddDefaulted();
+			FTimingEventsTrackDrawState::FTextPrimitive& DrawText = DrawState.Texts.Last();
+			DrawText.Depth = Depth;
+			DrawText.X = EventX1 + 2.0f;
+			DrawText.Text = Name.Left(LastWholeCharacterIndex + 1);
+			DrawText.bWhite = bIsDarkColor;
 		}
 	}
 }
 
 ////////////////////////////////////////////////////////////////////////////////////////////////////
 
-void FTimingViewDrawHelper::DrawBox(const FBoxData& Box, const float EventY, const float EventH)
+void FTimingEventsTrackDrawStateBuilder::FlushBox(const FBoxData& Box, int32 Depth)
 {
-	DrawContext.DrawBox(DrawContext.LayerId + ToInt32(EDrawLayer::EventBorder), Box.X1, EventY, Box.X2 - Box.X1, EventH, WhiteBrush, Box.LinearColor);
-	Stats.NumDrawBoxes++;
+	DrawState.Boxes.AddUninitialized();
+	FTimingEventsTrackDrawState::FBoxPrimitive& DrawBox = DrawState.Boxes.Last();
+	DrawBox.Depth = Depth;
+	DrawBox.X = Box.X1;
+	DrawBox.W = Box.X2 - Box.X1;
+	DrawBox.Color = Box.LinearColor;
 }
 
 ////////////////////////////////////////////////////////////////////////////////////////////////////
 
-void FTimingViewDrawHelper::EndTimelines()
+void FTimingEventsTrackDrawStateBuilder::Flush()
 {
-	float Y = Viewport.GetTopOffset() + TimelineTopY;
+	// Flush merged boxes.
+	for (int32 Depth = 0; Depth <= MaxDepth; ++Depth)
+	{
+		const FBoxData& Box = LastBox[Depth];
+		if (Box.X1 < Box.X2)
+		{
+			FlushBox(Box, Depth);
+		}
+	}
+}
 
-	// Draw a last horizontal line.
-	DrawContext.DrawBox(DrawContext.LayerId + ToInt32(EDrawLayer::TimelineHeader), 0, Y, Viewport.GetWidth(), 1.0f, WhiteBrush, EdgeColor);
-	Y += 1.0f;
+////////////////////////////////////////////////////////////////////////////////////////////////////
+// FTimingViewDrawHelper
+////////////////////////////////////////////////////////////////////////////////////////////////////
 
-	if (Y < Viewport.GetHeight() && ValidAreaW > 0.0f)
+FTimingViewDrawHelper::FTimingViewDrawHelper(const FDrawContext& InDrawContext, const FTimingTrackViewport& InViewport)
+	: DrawContext(InDrawContext)
+	, Viewport(InViewport)
+	, WhiteBrush(FInsightsStyle::Get().GetBrush("WhiteBrush"))
+	, EventBorderBrush(FInsightsStyle::Get().GetBrush("EventBorder"))
+	, HoveredEventBorderBrush(FInsightsStyle::Get().GetBrush("HoveredEventBorder"))
+	, SelectedEventBorderBrush(FInsightsStyle::Get().GetBrush("SelectedEventBorder"))
+	, BackgroundAreaBrush(WhiteBrush)
+	, ValidAreaColor(0.07f, 0.07f, 0.07f, 1.0f)
+	, InvalidAreaColor(0.1f, 0.07f, 0.07f, 1.0f)
+	, EdgeColor(0.05f, 0.05f, 0.05f, 1.0f)
+	, EventFont(FCoreStyle::GetDefaultFontStyle("Regular", 8))
+	, ValidAreaX(0.0f)
+	, ValidAreaW(0.0f)
+	, HighlightedEventTypeId(uint64(-1))
+	, NumEvents(0)
+	, NumMergedBoxes(0)
+	, NumDrawBoxes(0)
+	, NumDrawBorders(0)
+	, NumDrawTexts(0)
+{
+}
+
+////////////////////////////////////////////////////////////////////////////////////////////////////
+
+FTimingViewDrawHelper::~FTimingViewDrawHelper()
+{
+}
+
+////////////////////////////////////////////////////////////////////////////////////////////////////
+
+void FTimingViewDrawHelper::DrawBackground() const
+{
+	const float Y = 0.0f;
+	const float H = FMath::CeilToFloat(Viewport.GetHeight());
+	FDrawHelpers::DrawBackground(DrawContext, BackgroundAreaBrush, Viewport, Y, H, ValidAreaX, ValidAreaW); // also computes ValidAreaX and ValidAreaW
+}
+
+////////////////////////////////////////////////////////////////////////////////////////////////////
+
+void FTimingViewDrawHelper::DrawEventsTrack(const FTimingEventsTrackDrawState& DrawState, const FTimingEventsTrack& Track) const
+{
+	const float TrackY = Track.GetPosY();
+
+	if (Track.GetHeight() > 0.0f &&
+		TrackY + Track.GetHeight() > Viewport.GetTopOffset() &&
+		TrackY < Viewport.GetHeight() - Viewport.GetBottomOffset())
 	{
-		Y = FMath::Max(Y, 0.0f);
+		const FTimingViewLayout& Layout = Viewport.GetLayout();
+
+		NumEvents += DrawState.GetNumEvents();
+		NumMergedBoxes += DrawState.GetNumMergedBoxes();
+
+		const float TopLaneY = TrackY + 1.0f + Layout.TimelineDY; // +1.0f is for horizontal line between timelines
+
+		// Draw filled boxes.
+		//if (Layout.EventH > 2.0f)
+		{
+			const int32 EventFillLayerId = ReservedLayerId + ToInt32(EDrawLayer::EventFill);
+			const float EventFillH = Layout.EventH;
+			for (const FTimingEventsTrackDrawState::FBoxPrimitive& Box : DrawState.Boxes)
+			{
+				const float Y = TopLaneY + (Layout.EventH + Layout.EventDY) * Box.Depth;
+				DrawContext.DrawBox(EventFillLayerId, Box.X, Y, Box.W, EventFillH, WhiteBrush, Box.Color);
+			}
+			NumDrawBoxes += DrawState.Boxes.Num();
+		}
+
+		// Draw merged borders.
+		if (Layout.EventH > 2.0f)
+		{
+			const int32 EventFillLayerId = ReservedLayerId + ToInt32(EDrawLayer::EventFill);
+			const float EventFillH = Layout.EventH - 2.0f;
+			for (const FTimingEventsTrackDrawState::FBoxPrimitive& Box : DrawState.InsideBoxes)
+			{
+				const float Y = TopLaneY + (Layout.EventH + Layout.EventDY) * Box.Depth + 1.0f;
+				DrawContext.DrawBox(EventFillLayerId, Box.X, Y, Box.W, EventFillH, WhiteBrush, Box.Color);
+			}
+			NumDrawBoxes += DrawState.InsideBoxes.Num();
+		}
+
+		// Draw borders.
+		//if (Layout.EventH > 0.0f)
+		{
+			const int32 EventBorderLayerId = ReservedLayerId + ToInt32(EDrawLayer::EventBorder);
+			const float EventBorderH = Layout.EventH;
+			for (const FTimingEventsTrackDrawState::FBoxPrimitive& Box : DrawState.Borders)
+			{
+				const float Y = TopLaneY + (Layout.EventH + Layout.EventDY) * Box.Depth;
+				DrawContext.DrawBox(EventBorderLayerId, Box.X, Y, Box.W, EventBorderH, EventBorderBrush, Box.Color);
+			}
+			NumDrawBorders += DrawState.Borders.Num();
+		}
+
+		// Draw texts.
+		if (Layout.EventH > 10.0f)
+		{
+			const int32 EventTextLayerId = ReservedLayerId + ToInt32(EDrawLayer::EventText);
+			for (const FTimingEventsTrackDrawState::FTextPrimitive& Text : DrawState.Texts)
+			{
+				const float Y = TopLaneY + (Layout.EventH + Layout.EventDY) * Text.Depth + 1.0f;
+				DrawContext.DrawText(EventTextLayerId, Text.X, Y, Text.Text, EventFont, Text.bWhite ? FLinearColor::White : FLinearColor::Black);
+			}
+			NumDrawTexts += DrawState.Texts.Num();
+		}
+
+		// Draw a horizontal line between timelines (top line of a track).
+		const int32 HeaderLayerId = ReservedLayerId + ToInt32(EDrawLayer::HeaderBackground);
+		DrawContext.DrawBox(HeaderLayerId, 0.0f, TrackY, Viewport.GetWidth(), 1.0f, WhiteBrush, EdgeColor);
+
+		// Draw header with name of timeline.
+		if (Track.GetHeight() > 4.0f)
+		{
+			const TSharedRef<FSlateFontMeasure> FontMeasureService = FSlateApplication::Get().GetRenderer()->GetFontMeasureService();
+			float NameWidth = FontMeasureService->Measure(Track.GetName(), EventFont).X;
+
+			const float HeaderH = FMath::Min(12.0f, Track.GetHeight() - 1.0f);
+			DrawContext.DrawBox(HeaderLayerId, 0.0f, TrackY + 1.0f, NameWidth + 4.0f, HeaderH, WhiteBrush, EdgeColor);
 
-		// Draw invalid area (bottom).
-		DrawContext.DrawBox(DrawContext.LayerId + ToInt32(EDrawLayer::TimelineHeader), ValidAreaX, Y, ValidAreaW, Viewport.GetHeight() - Y, BackgroundAreaBrush, InvalidAreaColor);
+			const FLinearColor TextColor = Track.IsHovered() ?  FLinearColor(1.0f, 1.0f, 0.0f, 1.0f) :
+			                               Track.IsSelected() ? FLinearColor(1.0f, 1.0f, 0.5f, 1.0f) :
+			                                                    FLinearColor(1.0f, 1.0f, 1.0f, 1.0f);
+
+			const float TextY = TrackY + FMath::Min(0.0f, (Track.GetHeight() - 13.0f) / 2.0f);
+			const int32 HeaderTextLayerId = ReservedLayerId + ToInt32(EDrawLayer::HeaderText);
+			DrawContext.DrawText(HeaderTextLayerId, 2.0f, TextY, Track.GetName(), EventFont, TextColor);
+		}
 	}
+}
 
+////////////////////////////////////////////////////////////////////////////////////////////////////
+
+void FTimingViewDrawHelper::BeginDrawTracks() const
+{
+	// Reserve layers.
+	ReservedLayerId = DrawContext.LayerId;
 	DrawContext.LayerId += ToInt32(EDrawLayer::Count);
 }
 
 ////////////////////////////////////////////////////////////////////////////////////////////////////
 
-void FTimingViewDrawHelper::DrawTimingEventHighlight(double StartTime, double EndTime, float Y, EHighlightMode Mode)
+void FTimingViewDrawHelper::EndDrawTracks() const
+{
+	if (Viewport.GetWidth() > 0.0f)
+	{
+		// Y position of the first pixel below the last track.
+		const float Y = Viewport.GetTopOffset() + Viewport.GetScrollHeight() - Viewport.GetScrollPosY();
+
+		const float TopY = Viewport.GetTopOffset();
+		const float BottomY = Viewport.GetHeight() - Viewport.GetBottomOffset();
+
+		if (Y >= Viewport.GetTopOffset() && Y < BottomY)
+		{
+			// Draw a last horizontal line.
+			DrawContext.DrawBox(ReservedLayerId + ToInt32(EDrawLayer::HeaderBackground), 0.0f, Y, Viewport.GetWidth(), 1.0f, WhiteBrush, EdgeColor);
+		}
+
+		// Note: ValidAreaX and ValidAreaW are computed in DrawBackground.
+		if (ValidAreaW > 0.0f)
+		{
+			const float TopInvalidAreaH = FMath::Min(0.0f - Viewport.GetScrollPosY(), Viewport.GetScrollableAreaHeight());
+			if (TopInvalidAreaH > 0.0f)
+			{
+				// Draw invalid area (top).
+				DrawContext.DrawBox(ReservedLayerId + ToInt32(EDrawLayer::HeaderBackground), ValidAreaX, TopY, ValidAreaW, TopInvalidAreaH, BackgroundAreaBrush, InvalidAreaColor);
+			}
+
+			const float BottomInvalidAreaH = FMath::Min(BottomY - Y - 1.0f, Viewport.GetScrollableAreaHeight());
+			if (BottomInvalidAreaH > 0.0f)
+			{
+				// Draw invalid area (bottom).
+				DrawContext.DrawBox(ReservedLayerId + ToInt32(EDrawLayer::HeaderBackground), ValidAreaX, BottomY - BottomInvalidAreaH, ValidAreaW, BottomInvalidAreaH, BackgroundAreaBrush, InvalidAreaColor);
+			}
+		}
+	}
+}
+
+////////////////////////////////////////////////////////////////////////////////////////////////////
+
+void FTimingViewDrawHelper::DrawTimingEventHighlight(double StartTime, double EndTime, float Y, EDrawEventMode Mode) const
 {
 	float EventX1 = Viewport.TimeToSlateUnitsRounded(StartTime);
 	if (EventX1 > Viewport.GetWidth())
@@ -517,26 +525,29 @@ void FTimingViewDrawHelper::DrawTimingEventHighlight(double StartTime, double En
 
 	const float EventW = EventX2 - EventX1;
 
-	if (Mode == EHighlightMode::Hovered)
+	const FTimingViewLayout& Layout = Viewport.GetLayout();
+
+	const int32 LayerId = ReservedLayerId + ToInt32(EDrawLayer::EventHighlight);
+
+	if (Mode == EDrawEventMode::Hovered)
 	{
 		const FLinearColor Color(1.0f, 1.0f, 0.0f, 1.0f); // yellow
 
 		// Draw border around the timing event box.
-		DrawContext.DrawBox(EventX1 - 2.0f, Y - 2.0f, EventW + 4.0f, Layout.EventH + 4.0f, HoveredEventBorderBrush, Color);
+		DrawContext.DrawBox(LayerId, EventX1 - 2.0f, Y - 2.0f, EventW + 4.0f, Layout.EventH + 4.0f, HoveredEventBorderBrush, Color);
 	}
-	else // EHighlightMode::Selected or EHighlightMode::SelectedAndHovered
+	else // EDrawEventMode::Selected or EDrawEventMode::SelectedAndHovered
 	{
 		// Animate color from white (if selected and hovered) or yellow (if only selected) to black, using a squared sine function.
 		const double Time = static_cast<double>(FPlatformTime::Cycles64()) * FPlatformTime::GetSecondsPerCycle64();
 		float S = FMath::Sin(2.0 * Time);
 		S = S * S; // squared, to ensure only positive [0 - 1] values
-		const float Blue = (Mode == EHighlightMode::SelectedAndHovered) ? 0.0f : S;
+		const float Blue = (Mode == EDrawEventMode::SelectedAndHovered) ? 0.0f : S;
 		const FLinearColor Color(S, S, Blue, 1.0f);
 
 		// Draw border around the timing event box.
-		DrawContext.DrawBox(EventX1 - 2.0f, Y - 2.0f, EventW + 4.0f, Layout.EventH + 4.0f, SelectedEventBorderBrush, Color);
+		DrawContext.DrawBox(LayerId, EventX1 - 2.0f, Y - 2.0f, EventW + 4.0f, Layout.EventH + 4.0f, SelectedEventBorderBrush, Color);
 	}
-	DrawContext.LayerId++;
 }
 
 ////////////////////////////////////////////////////////////////////////////////////////////////////
