[Insights]

diff --git a/Engine/Source/Developer/TraceInsights/Private/Insights/ViewModels/TimingViewDrawHelper.h b/Engine/Source/Developer/TraceInsights/Private/Insights/ViewModels/TimingViewDrawHelper.h
--- a/Engine/Source/Developer/TraceInsights/Private/Insights/ViewModels/TimingViewDrawHelper.h
+++ b/Engine/Source/Developer/TraceInsights/Private/Insights/ViewModels/TimingViewDrawHelper.h
@@ -5,42 +5,85 @@
 #include "CoreMinimal.h"
 #include "Fonts/SlateFontInfo.h"
 #include "Math/Color.h"
-#include "Insights/ITimingViewDrawHelper.h"
 
 ////////////////////////////////////////////////////////////////////////////////////////////////////
 
 struct FSlateBrush;
 struct FDrawContext;
 
+class ITimingTrackDrawContext;
 class FTimingEventsTrack;
-struct FTimingEventsTrackLayout;
 class FTimingTrackViewport;
 
+enum class EDrawEventMode : uint32;
+
+struct FTimingViewLayout;
+
 ////////////////////////////////////////////////////////////////////////////////////////////////////
 
-class FTimingViewDrawHelper final : public ITimingViewDrawHelper
+struct FTimingEventsTrackDrawState
 {
-public:
-	enum class EHighlightMode : uint32
+	struct FBoxPrimitive
 	{
-		Hovered = 1,
-		Selected = 2,
-		SelectedAndHovered = 3
+		int32 Depth;
+		float X;
+		float W;
+		FLinearColor Color;
 	};
 
-private:
-	enum class EDrawLayer : int32
+	struct FTextPrimitive
 	{
-		EventBorder,
-		EventFill,
-		EventText,
-		TimelineHeader,
-		TimelineText,
-
-		Count,
+		int32 Depth;
+		float X;
+		FString Text;
+		bool bWhite;
 	};
-	static int32 ToInt32(EDrawLayer Layer) { return static_cast<int32>(Layer); }
 
+	FTimingEventsTrackDrawState()
+		: Boxes()
+		, InsideBoxes()
+		, Borders()
+		, Texts()
+		, NumLanes(0)
+		, NumEvents(0)
+		, NumMergedBoxes(0)
+	{
+	}
+
+	void Reset()
+	{
+		Boxes.Reset();
+		InsideBoxes.Reset();
+		Borders.Reset();
+		Texts.Reset();
+		NumLanes = 0;
+		NumEvents = 0;
+		NumMergedBoxes = 0;
+	}
+
+	int32 GetNumLanes() const { return NumLanes; }
+
+	int32 GetNumEvents() const { return NumEvents; }
+	int32 GetNumMergedBoxes() const { return NumMergedBoxes; }
+	int32 GetTotalNumBoxes() const { return Boxes.Num() + InsideBoxes.Num(); }
+
+	TArray<FBoxPrimitive> Boxes;
+	TArray<FBoxPrimitive> InsideBoxes;
+	TArray<FBoxPrimitive> Borders;
+	TArray<FTextPrimitive> Texts;
+
+	int32 NumLanes;
+
+	// Debug stats.
+	int32 NumEvents;
+	int32 NumMergedBoxes;
+};
+
+////////////////////////////////////////////////////////////////////////////////////////////////////
+
+class FTimingEventsTrackDrawStateBuilder
+{
+private:
 	struct FBoxData
 	{
 		float X1;
@@ -52,41 +95,58 @@ private:
 		void Reset() { X1 = 0.0f; X2 = 0.0f; Color = 0; }
 	};
 
-	//struct FEventBoxInfo
-	//{
-	//	float X1;
-	//	float X2;
-	//	int32 Depth;
-	//	uint32 Color;
-	//}
-
-	//struct FEventTextInfo
-	//{
-	//	float X;
-	//	float Y;
-	//	FString Text;
-	//	bool bUseDarkTextColor; // true if text needs to be displayed in Black, otherwise will be displayed in White
-	//};
-
-	struct FStats
+public:
+	explicit FTimingEventsTrackDrawStateBuilder(FTimingEventsTrackDrawState& InState, const FTimingTrackViewport& InViewport);
+
+	/**
+	 * Non-copyable
+	 */
+	FTimingEventsTrackDrawStateBuilder(const FTimingEventsTrackDrawStateBuilder&) = delete;
+	FTimingEventsTrackDrawStateBuilder& operator=(const FTimingEventsTrackDrawStateBuilder&) = delete;
+
+	const FTimingTrackViewport& GetViewport() const { return Viewport; }
+
+	void AddEvent(double InEventStartTime, double InEventEndTime, uint32 InEventDepth, const TCHAR* InEventName, uint64 InEventType, uint32 InEventColor);
+	void Flush();
+
+	int32 GetMaxDepth() const { return MaxDepth; }
+
+private:
+	void FlushBox(const FBoxData& Box, const int32 Depth);
+
+private:
+	FTimingEventsTrackDrawState& DrawState; // cached draw state to build
+
+	const FTimingTrackViewport& Viewport;
+
+	int32 MaxDepth;
+
+	TArray<float> LastEventX2; // X2 value for last event on each depth
+	TArray<FBoxData> LastBox;
+
+	const FSlateFontInfo EventFont;
+};
+
+////////////////////////////////////////////////////////////////////////////////////////////////////
+
+class FTimingViewDrawHelper final
+{
+private:
+	enum class EDrawLayer : int32
 	{
-		int32 NumEvents;
-		int32 NumDrawBoxes;
-		int32 NumMergedBoxes;
-		int32 NumDrawBorders;
-		int32 NumDrawTexts;
-
-		FStats()
-			: NumEvents(0)
-			, NumDrawBoxes(0)
-			, NumMergedBoxes(0)
-			, NumDrawBorders(0)
-			, NumDrawTexts(0)
-		{}
+		EventBorder,
+		EventFill,
+		EventText,
+		EventHighlight,
+		HeaderBackground,
+		HeaderText,
+
+		Count,
 	};
+	static int32 ToInt32(EDrawLayer Layer) { return static_cast<int32>(Layer); }
 
 public:
-	explicit FTimingViewDrawHelper(const FDrawContext& InDrawContext, const FTimingTrackViewport& InViewport, const FTimingEventsTrackLayout& InLayout);
+	explicit FTimingViewDrawHelper(const FDrawContext& InDrawContext, const FTimingTrackViewport& InViewport);
 	~FTimingViewDrawHelper();
 
 	/**
@@ -95,37 +155,31 @@ public:
 	FTimingViewDrawHelper(const FTimingViewDrawHelper&) = delete;
 	FTimingViewDrawHelper& operator=(const FTimingViewDrawHelper&) = delete;
 
-	// ITimingViewDrawHelper interface
-	virtual const FDrawContext& GetDrawContext() const override { return DrawContext; }
-	virtual const FTimingTrackViewport& GetViewport() const override { return Viewport; }
-	virtual const FTimingEventsTrackLayout& GetLayout() const override { return Layout; }
-	virtual bool BeginTimeline(FTimingEventsTrack& Track) override;
-	virtual void AddEvent(double StartTime, double EndTime, uint32 Depth, const TCHAR* EventName, uint32 Color = 0) override;
-	virtual void EndTimeline(FTimingEventsTrack& Track) override;
+	const FDrawContext& GetDrawContext() const { return DrawContext; }
+	const FTimingTrackViewport& GetViewport() const { return Viewport; }
 
 	const FSlateBrush* GetWhiteBrush() const { return WhiteBrush; }
 	const FSlateFontInfo& GetEventFont() const { return EventFont; }
 
-	int32 GetNumEvents() const      { return Stats.NumEvents; }
-	int32 GetNumDrawBoxes() const   { return Stats.NumDrawBoxes; }
-	int32 GetNumMergedBoxes() const { return Stats.NumMergedBoxes; }
-	int32 GetNumDrawBorders() const { return Stats.NumDrawBorders; }
-	int32 GetNumDrawTexts() const   { return Stats.NumDrawTexts; }
-
 	void DrawBackground() const;
-	void DrawTimingEventHighlight(double StartTime, double EndTime, float Y, EHighlightMode Mode);
 
-	//TODO: move the following in a Builder class
-	void BeginTimelines();
-	void EndTimelines();
+	void BeginDrawTracks() const;
+	void DrawEventsTrack(const FTimingEventsTrackDrawState& DrawState, const FTimingEventsTrack& Track) const;
+	void EndDrawTracks() const;
 
-private:
-	void DrawBox(const FBoxData& Box, const float EventY, const float EventH);
+	void DrawTimingEventHighlight(double StartTime, double EndTime, float Y, EDrawEventMode Mode) const;
+
+	void SetHighlightedEventTypeId(uint64 InHighlightedEventTypeId) { HighlightedEventTypeId = InHighlightedEventTypeId; }
+
+	int32 GetNumEvents() const { return NumEvents; }
+	int32 GetNumMergedBoxes() const { return NumMergedBoxes; }
+	int32 GetNumDrawBoxes() const { return NumDrawBoxes; }
+	int32 GetNumDrawBorders() const { return NumDrawBorders; }
+	int32 GetNumDrawTexts() const { return NumDrawTexts; }
 
 private:
 	const FDrawContext& DrawContext;
 	const FTimingTrackViewport& Viewport;
-	const FTimingEventsTrackLayout& Layout;
 
 	const FSlateBrush* WhiteBrush;
 	const FSlateBrush* EventBorderBrush;
@@ -137,24 +191,18 @@ private:
 	const FLinearColor EdgeColor;
 	const FSlateFontInfo EventFont;
 
+	mutable int32 ReservedLayerId;
+
 	mutable float ValidAreaX;
 	mutable float ValidAreaW;
 
-	//////////////////////////////////////////////////
-	// Builder state
-
-	float TimelineTopY;
-	float TimelineY;
-	int32 MaxDepth;
-	int32 TimelineIndex;
-
-	TArray<float> LastEventX2; // X2 value for last event on each depth, for current timeline
-	TArray<FBoxData> LastBox;
-
-	//////////////////////////////////////////////////
+	uint64 HighlightedEventTypeId;
 
-	/** Debug stats */
-	mutable FStats Stats;
+	mutable int32 NumEvents;
+	mutable int32 NumMergedBoxes;
+	mutable int32 NumDrawBoxes;
+	mutable int32 NumDrawBorders;
+	mutable int32 NumDrawTexts;
 };
 
 ////////////////////////////////////////////////////////////////////////////////////////////////////
diff --git a/Engine/Source/Developer/TraceInsights/Private/Insights/Widgets/STimersView.cpp b/Engine/Source/Developer/TraceInsights/Private/Insights/Widgets/STimersView.cpp
--- a/Engine/Source/Developer/TraceInsights/Private/Insights/Widgets/STimersView.cpp
+++ b/Engine/Source/Developer/TraceInsights/Private/Insights/Widgets/STimersView.cpp
@@ -818,17 +818,12 @@ void STimersView::TreeView_OnMouseButtonDoubleClick(FTimerNodePtr TimerNodePtr)
 	}
 	else
 	{
-		//im:TODO: const bool bIsTracked = FTimingProfilerManager::Get()->IsTimerTracked(TimerNodePtr->GetId());
-	//	if (!bIsTracked)
-	//	{
-	//		// Add a new graph series.
-	//		FTimingProfilerManager::Get()->TrackTimer(TimerNodePtr->GetId());
-	//	}
-	//	else
-	//	{
-	//		// Remove the corresponding graph series.
-	//		FTimingProfilerManager::Get()->UntrackTimer(TimerNodePtr->GetId());
-	//	}
+		TSharedPtr<STimingProfilerWindow> Wnd = FTimingProfilerManager::Get()->GetProfilerWindow();
+		TSharedPtr<STimingView> TimingView = Wnd.IsValid() ? Wnd->GetTimingView() : nullptr;
+		if (TimingView.IsValid())
+		{
+			TimingView->SetHighlightedEventTypeId(TimerNodePtr->GetId());
+		}
 	}
 }
 
diff --git a/Engine/Source/Developer/TraceInsights/Private/Insights/Widgets/STimingProfilerWindow.cpp b/Engine/Source/Developer/TraceInsights/Private/Insights/Widgets/STimingProfilerWindow.cpp
--- a/Engine/Source/Developer/TraceInsights/Private/Insights/Widgets/STimingProfilerWindow.cpp
+++ b/Engine/Source/Developer/TraceInsights/Private/Insights/Widgets/STimingProfilerWindow.cpp
@@ -27,7 +27,6 @@
 #include "Insights/TimingProfilerManager.h"
 #include "Insights/Version.h"
 #include "Insights/Widgets/SFrameTrack.h"
-#include "Insights/Widgets/SGraphTrack.h"
 #include "Insights/Widgets/SInsightsSettings.h"
 #include "Insights/Widgets/SLogView.h"
 #include "Insights/Widgets/SStatsView.h"
@@ -44,7 +43,6 @@
 
 const FName FTimingProfilerTabs::ToolbarID(TEXT("Toolbar"));
 const FName FTimingProfilerTabs::FramesTrackID(TEXT("Frames"));
-const FName FTimingProfilerTabs::GraphTrackID(TEXT("Graph"));
 const FName FTimingProfilerTabs::TimingViewID(TEXT("TimingView"));
 const FName FTimingProfilerTabs::TimersID(TEXT("Timers"));
 const FName FTimingProfilerTabs::CallersID(TEXT("Callers"));
@@ -80,11 +78,6 @@ void STimingProfilerWindow::Reset()
 		FrameTrack->Reset();
 	}
 
-	if (GraphTrack)
-	{
-		GraphTrack->Reset();
-	}
-
 	if (TimingView)
 	{
 		TimingView->Reset();
@@ -168,32 +161,6 @@ void STimingProfilerWindow::OnFramesTrackTabClosed(TSharedRef<SDockTab> TabBeing
 
 ////////////////////////////////////////////////////////////////////////////////////////////////////
 
-TSharedRef<SDockTab> STimingProfilerWindow::SpawnTab_GraphTrack(const FSpawnTabArgs& Args)
-{
-	FTimingProfilerManager::Get()->SetGraphTrackVisible(true);
-
-	const TSharedRef<SDockTab> DockTab = SNew(SDockTab)
-		.ShouldAutosize(false)
-		.TabRole(ETabRole::PanelTab)
-		[
-			SAssignNew(GraphTrack, SGraphTrack)
-		];
-
-	DockTab->SetOnTabClosed(SDockTab::FOnTabClosedCallback::CreateRaw(this, &STimingProfilerWindow::OnGraphTrackTabClosed));
-
-	return DockTab;
-}
-
-////////////////////////////////////////////////////////////////////////////////////////////////////
-
-void STimingProfilerWindow::OnGraphTrackTabClosed(TSharedRef<SDockTab> TabBeingClosed)
-{
-	FTimingProfilerManager::Get()->SetGraphTrackVisible(false);
-	GraphTrack = nullptr;
-}
-
-////////////////////////////////////////////////////////////////////////////////////////////////////
-
 TSharedRef<SDockTab> STimingProfilerWindow::SpawnTab_TimingView(const FSpawnTabArgs& Args)
 {
 	FTimingProfilerManager::Get()->SetTimingViewVisible(true);
@@ -366,11 +333,6 @@ void STimingProfilerWindow::Construct(const FArguments& InArgs, const TSharedRef
 		.SetIcon(FSlateIcon(FInsightsStyle::GetStyleSetName(), "FramesTrack.Icon.Small"))
 		.SetGroup(AppMenuGroup);
 
-	TabManager->RegisterTabSpawner(FTimingProfilerTabs::GraphTrackID, FOnSpawnTab::CreateRaw(this, &STimingProfilerWindow::SpawnTab_GraphTrack))
-		.SetDisplayName(LOCTEXT("GraphTrackTabTitle", "Graph"))
-		.SetIcon(FSlateIcon(FInsightsStyle::GetStyleSetName(), "GraphTrack.Icon.Small"))
-		.SetGroup(AppMenuGroup);
-
 	TabManager->RegisterTabSpawner(FTimingProfilerTabs::TimingViewID, FOnSpawnTab::CreateRaw(this, &STimingProfilerWindow::SpawnTab_TimingView))
 		.SetDisplayName(LOCTEXT("TimingViewTabTitle", "Timing View"))
 		.SetIcon(FSlateIcon(FInsightsStyle::GetStyleSetName(), "TimingView.Icon.Small"))
@@ -434,13 +396,6 @@ void STimingProfilerWindow::Construct(const FArguments& InArgs, const TSharedRef
 						->AddTab(FTimingProfilerTabs::FramesTrackID, TimingProfilerManager->IsFramesTrackVisible() ? ETabState::OpenedTab : ETabState::ClosedTab)
 					)
 					->Split
-					(
-						FTabManager::NewStack()
-						->SetSizeCoefficient(0.2f)
-						->SetHideTabWell(true)
-						->AddTab(FTimingProfilerTabs::GraphTrackID, TimingProfilerManager->IsGraphTrackVisible() ? ETabState::OpenedTab : ETabState::ClosedTab)
-					)
-					->Split
 					(
 						FTabManager::NewStack()
 						->SetSizeCoefficient(0.5f)
diff --git a/Engine/Source/Developer/TraceInsights/Private/Insights/Widgets/STimingProfilerWindow.h b/Engine/Source/Developer/TraceInsights/Private/Insights/Widgets/STimingProfilerWindow.h
--- a/Engine/Source/Developer/TraceInsights/Private/Insights/Widgets/STimingProfilerWindow.h
+++ b/Engine/Source/Developer/TraceInsights/Private/Insights/Widgets/STimingProfilerWindow.h
@@ -21,7 +21,6 @@ class FActiveTimerHandle;
 class FMenuBuilder;
 
 class SFrameTrack;
-class SGraphTrack;
 class SLogView;
 class SStatsView;
 class STimersView;
@@ -36,7 +35,6 @@ struct FTimingProfilerTabs
 	// Tab identifiers
 	static const FName ToolbarID;
 	static const FName FramesTrackID;
-	static const FName GraphTrackID;
 	static const FName TimingViewID;
 	static const FName TimersID;
 	static const FName CallersID;
@@ -85,9 +83,6 @@ private:
 	TSharedRef<SDockTab> SpawnTab_FramesTrack(const FSpawnTabArgs& Args);
 	void OnFramesTrackTabClosed(TSharedRef<SDockTab> TabBeingClosed);
 
-	TSharedRef<SDockTab> SpawnTab_GraphTrack(const FSpawnTabArgs& Args);
-	void OnGraphTrackTabClosed(TSharedRef<SDockTab> TabBeingClosed);
-
 	TSharedRef<SDockTab> SpawnTab_TimingView(const FSpawnTabArgs& Args);
 	void OnTimingViewTabClosed(TSharedRef<SDockTab> TabBeingClosed);
 
@@ -181,9 +176,6 @@ private:
 	/** The Frame track widget */
 	TSharedPtr<SFrameTrack> FrameTrack;
 
-	/** The graph track widget */
-	TSharedPtr<SGraphTrack> GraphTrack;
-
 	/** The Timing view (multi-track) widget */
 	TSharedPtr<STimingView> TimingView;
 
