[Insights]

diff --git a/Engine/Source/Developer/TraceInsights/Private/Insights/Widgets/STimingView.cpp b/Engine/Source/Developer/TraceInsights/Private/Insights/Widgets/STimingView.cpp
--- a/Engine/Source/Developer/TraceInsights/Private/Insights/Widgets/STimingView.cpp
+++ b/Engine/Source/Developer/TraceInsights/Private/Insights/Widgets/STimingView.cpp
@@ -5,6 +5,7 @@
 #include "Containers/ArrayBuilder.h"
 #include "Containers/MapBuilder.h"
 #include "EditorStyleSet.h"
+#include "Features/IModularFeatures.h"
 #include "Fonts/FontMeasure.h"
 #include "Fonts/SlateFontInfo.h"
 #include "Framework/Application/MenuStack.h"
@@ -25,12 +26,12 @@
 #include "Widgets/Layout/SSpacer.h"
 #include "Widgets/SOverlay.h"
 #include "Widgets/Text/STextBlock.h"
-#include "Features/IModularFeatures.h"
 
 // Insights
 #include "Insights/Common/PaintUtils.h"
 #include "Insights/Common/Stopwatch.h"
 #include "Insights/Common/TimeUtils.h"
+#include "Insights/ITimingViewExtender.h"
 #include "Insights/InsightsManager.h"
 #include "Insights/InsightsStyle.h"
 #include "Insights/LoadingProfiler/LoadingProfilerManager.h"
@@ -39,23 +40,25 @@
 #include "Insights/TimingProfilerCommon.h"
 #include "Insights/TimingProfilerManager.h"
 #include "Insights/ViewModels/BaseTimingTrack.h"
+#include "Insights/ViewModels/DrawHelpers.h"
 #include "Insights/ViewModels/FileActivityTimingTrack.h"
+#include "Insights/ViewModels/GraphSeries.h"
+#include "Insights/ViewModels/GraphTrack.h"
 #include "Insights/ViewModels/LoadingTimingTrack.h"
+#include "Insights/ViewModels/MarkersTimingTrack.h"
 #include "Insights/ViewModels/ThreadTimingTrack.h"
+#include "Insights/ViewModels/TimeRulerTrack.h"
+#include "Insights/ViewModels/TimingEventSearch.h"
 #include "Insights/ViewModels/TimingGraphTrack.h"
-#include "Insights/ViewModels/DrawHelpers.h"
 #include "Insights/ViewModels/TimingViewDrawHelper.h"
 #include "Insights/Widgets/SStatsView.h"
 #include "Insights/Widgets/STimersView.h"
 #include "Insights/Widgets/STimingProfilerWindow.h"
-#include "Insights/ITimingViewExtender.h"
-#include "Insights/ViewModels/TimingEventSearch.h"
 
 #include <limits>
 
 ////////////////////////////////////////////////////////////////////////////////////////////////////
 
-#define DEBUG_TIMING_TRACK 0
 #define LOCTEXT_NAMESPACE "STimingView"
 
 // start auto generated ids from a big number (MSB set to 1) to avoid collisions with ids for gpu/cpu tracks based on 32bit timeline index
@@ -69,15 +72,20 @@ namespace Insights { const FName TimingViewExtenderFeatureName(TEXT("TimingViewE
 ////////////////////////////////////////////////////////////////////////////////////////////////////
 
 STimingView::STimingView()
-	: LoadingSharedState(MakeShareable(new FLoadingSharedState()))
+	: ThreadTimingSharedState(MakeShareable(new FThreadTimingSharedState(this)))
+	, LoadingSharedState(MakeShareable(new FLoadingSharedState(this)))
 	, bAssetLoadingMode(false)
-	, FileActivitySharedState(MakeShareable(new FFileActivitySharedState()))
-	, TimeRulerTrack(FBaseTimingTrack::GenerateId())
-	, MarkersTrack(FBaseTimingTrack::GenerateId())
-	, GraphTrack(MakeShareable(new FTimingGraphTrack(FBaseTimingTrack::GenerateId())))
+	, FileActivitySharedState(MakeShareable(new FFileActivitySharedState(this)))
+	, TimeRulerTrack(MakeShareable(new FTimeRulerTrack()))
+	, MarkersTrack(MakeShareable(new FMarkersTimingTrack()))
+	, GraphTrack(MakeShareable(new FTimingGraphTrack()))
 	, WhiteBrush(FInsightsStyle::Get().GetBrush("WhiteBrush"))
 	, MainFont(FCoreStyle::GetDefaultFontStyle("Regular", 8))
 {
+	IModularFeatures::Get().RegisterModularFeature(Insights::TimingViewExtenderFeatureName, ThreadTimingSharedState.Get());
+	IModularFeatures::Get().RegisterModularFeature(Insights::TimingViewExtenderFeatureName, LoadingSharedState.Get());
+	IModularFeatures::Get().RegisterModularFeature(Insights::TimingViewExtenderFeatureName, FileActivitySharedState.Get());
+
 	Reset();
 }
 
@@ -85,16 +93,20 @@ STimingView::STimingView()
 
 STimingView::~STimingView()
 {
-	for (const auto& KV : CachedTimelines)
-	{
-		delete KV.Value;
-	}
-	CachedTimelines.Reset();
+	AllTracks.Reset();
+	TopDockedTracks.Reset();
+	BottomDockedTracks.Reset();
+	ScrollableTracks.Reset();
+	ForegroundTracks.Reset();
 
 	for (Insights::ITimingViewExtender* Extender : GetExtenders())
 	{
 		Extender->OnEndSession(*this);
 	}
+
+	IModularFeatures::Get().UnregisterModularFeature(Insights::TimingViewExtenderFeatureName, ThreadTimingSharedState.Get());
+	IModularFeatures::Get().UnregisterModularFeature(Insights::TimingViewExtenderFeatureName, LoadingSharedState.Get());
+	IModularFeatures::Get().UnregisterModularFeature(Insights::TimingViewExtenderFeatureName, FileActivitySharedState.Get());
 }
 
 ////////////////////////////////////////////////////////////////////////////////////////////////////
@@ -181,76 +193,41 @@ void STimingView::Construct(const FArguments& InArgs)
 
 void STimingView::Reset()
 {
-	bShowHideAllGpuTracks = true;
-	bShowHideAllCpuTracks = true;
-	bShowHideAllLoadingTracks = false;
-	bShowHideAllIoTracks = false;
-
-	//////////////////////////////////////////////////
-
-	Viewport.Reset();
-
-	bIsViewportDirty = true;
-	bIsVerticalViewportDirty = true;
-
-	//////////////////////////////////////////////////
-
-	for (const auto& KV : CachedTimelines)
-	{
-		delete KV.Value;
-	}
-	CachedTimelines.Reset();
-
-	//////////////////////////////////////////////////
-
 	for (Insights::ITimingViewExtender* Extender : GetExtenders())
 	{
 		Extender->OnEndSession(*this);
-		Extender->OnBeginSession(*this);
 	}
 
-	//TODO: TopTracks.Reset();
-	//TODO: ScrollableTracks.Reset();
-	//TODO: BottomTracks.Reset();
-	//TODO: ForegroundTracks.Reset();
-
-	//////////////////////////////////////////////////
-
-	TimingEventsTracks.Reset();
-
-	bTimingEventsNeedUpdate = true;
-	bTimingEventsTracksChanged = true;
-
-	FTimingEventsTrack::bUseDownSampling = true;
-
 	//////////////////////////////////////////////////
 
-	GpuTrack = nullptr;
-	CpuTracks.Reset();
-
-	ThreadGroups.Reset();
+	Viewport.Reset();
 
 	//////////////////////////////////////////////////
 
-	LoadingSharedState->Reset();
+	AllTracks.Reset();
+	TopDockedTracks.Reset();
+	BottomDockedTracks.Reset();
+	ScrollableTracks.Reset();
+	ForegroundTracks.Reset();
 
-	LoadingTracks.Reset();
+	FTimingEventsTrack::bUseDownSampling = true;
 
 	//////////////////////////////////////////////////
 
-	FileActivitySharedState->Reset();
-	IoOverviewTrack = nullptr;
-	IoActivityTrack = nullptr;
-
-	//////////////////////////////////////////////////
+	TimeRulerTrack->Reset();
+	TimeRulerTrack->SetOrder(-999);
+	AddTopDockedTrack(TimeRulerTrack);
 
-	TimeRulerTrack.Reset();
-	MarkersTrack.Reset();
+	MarkersTrack->Reset();
+	MarkersTrack->SetOrder(-100);
+	AddTopDockedTrack(MarkersTrack);
 
 	GraphTrack->Reset();
+	GraphTrack->SetOrder(0);
 	GraphTrack->SetHeight(200.0f);
 	GraphTrack->AddDefaultFrameSeries();
 	GraphTrack->SetVisibilityFlag(false);
+	AddTopDockedTrack(GraphTrack);
 
 	//////////////////////////////////////////////////
 
@@ -275,8 +252,27 @@ void STimingView::Reset()
 	SelectionStartTime = 0.0;
 	SelectionEndTime = 0.0;
 
-	HoveredTimingEvent.Reset();
-	SelectedTimingEvent.Reset();
+	if (HoveredTrack.IsValid())
+	{
+		HoveredTrack.Reset();
+		OnHoveredTrackChangedDelegate.Broadcast(HoveredTrack);
+	}
+	if (HoveredEvent.IsValid())
+	{
+		HoveredEvent.Reset();
+		OnHoveredEventChangedDelegate.Broadcast(HoveredEvent);
+	}
+	if (SelectedTrack.IsValid())
+	{
+		SelectedTrack.Reset();
+		OnSelectedTrackChangedDelegate.Broadcast(SelectedTrack);
+	}
+	if (SelectedEvent.IsValid())
+	{
+		SelectedEvent.Reset();
+		OnSelectedEventChangedDelegate.Broadcast(SelectedEvent);
+	}
+
 	Tooltip.Reset();
 
 	LastSelectionType = ESelectionType::None;
@@ -284,36 +280,51 @@ void STimingView::Reset()
 	TimeMarker = std::numeric_limits<double>::infinity();
 	bIsScrubbing = false;
 
-	//ThisGeometry
+	HighlightedEventTypeId = uint64(-1);
 
-	Layout.ForceNormalMode();
+	//ThisGeometry
 
 	//////////////////////////////////////////////////
 
 	NumUpdatedEvents = 0;
-	TimelineCacheUpdateDurationHistory.Reset();
-	TimelineCacheUpdateDurationHistory.AddValue(0);
-	TimeMarkerCacheUpdateDurationHistory.Reset();
-	TimeMarkerCacheUpdateDurationHistory.AddValue(0);
-	DrawDurationHistory.Reset();
-	OnPaintDurationHistory.Reset();
+	PreUpdateTracksDurationHistory.Reset();
+	PreUpdateTracksDurationHistory.AddValue(0);
+	UpdateTracksDurationHistory.Reset();
+	UpdateTracksDurationHistory.AddValue(0);
+	PostUpdateTracksDurationHistory.Reset();
+	PostUpdateTracksDurationHistory.AddValue(0);
+	TickDurationHistory.Reset();
+	TickDurationHistory.AddValue(0);
+	PreDrawTracksDurationHistory.Reset();
+	PreDrawTracksDurationHistory.AddValue(0);
+	DrawTracksDurationHistory.Reset();
+	DrawTracksDurationHistory.AddValue(0);
+	PostDrawTracksDurationHistory.Reset();
+	PostDrawTracksDurationHistory.AddValue(0);
+	OnPaintDeltaTimeHistory.Reset();
+	OnPaintDeltaTimeHistory.AddValue(0);
 	LastOnPaintTime = FPlatformTime::Cycles64();
 
-	if (bAssetLoadingMode)
+	//////////////////////////////////////////////////
+
+	for (Insights::ITimingViewExtender* Extender : GetExtenders())
 	{
-		EnableAssetLoadingMode();
+		Extender->OnBeginSession(*this);
 	}
 }
 
 ////////////////////////////////////////////////////////////////////////////////////////////////////
 
-void STimingView::EnableAssetLoadingMode()
+bool STimingView::IsGpuTrackVisible() const
 {
-	bAssetLoadingMode = true;
+	return ThreadTimingSharedState && ThreadTimingSharedState->IsGpuTrackVisible();
+}
 
-	bShowHideAllGpuTracks = false;
-	bShowHideAllCpuTracks = false;
-	bShowHideAllLoadingTracks = true;
+////////////////////////////////////////////////////////////////////////////////////////////////////
+
+bool STimingView::IsCpuTrackVisible(uint32 InThreadId) const
+{
+	return ThreadTimingSharedState && ThreadTimingSharedState->IsCpuTrackVisible(InThreadId);
 }
 
 ////////////////////////////////////////////////////////////////////////////////////////////////////
@@ -322,419 +333,429 @@ void STimingView::Tick(const FGeometry& AllottedGeometry, const double InCurrent
 {
 	//SCompoundWidget::Tick(AllottedGeometry, InCurrentTime, InDeltaTime);
 
-	//FStopwatch TickStopwatch;
-	//TickStopwatch.Start();
+	FStopwatch TickStopwatch;
+	TickStopwatch.Start();
 
 	ThisGeometry = AllottedGeometry;
 
 	const float ViewWidth = AllottedGeometry.GetLocalSize().X;
 	const float ViewHeight = AllottedGeometry.GetLocalSize().Y;
-	if (Viewport.UpdateSize(ViewWidth, ViewHeight))
+
+	////////////////////////////////////////////////////////////////////////////////////////////////////
+	// Update viewport.
+
+	Viewport.UpdateSize(ViewWidth, ViewHeight);
+
+	if (!bIsPanning)
 	{
-		bIsViewportDirty = true;
+		// Elastic snap to horizontal time limits.
+		if (Viewport.EnforceHorizontalScrollLimits(0.5)) // 0.5 is the interpolation factor
+		{
+			UpdateHorizontalScrollBar();
+		}
 	}
 
-	//////////////////////////////////////////////////
-	// Update Y postion of non-scrollable tracks (docked on top).
-	// Also compute total height (TopOffset) of top docked tracks.
+	////////////////////////////////////////////////////////////////////////////////////////////////////
+	// Check the analysis session time.
 
-	float TopOffset = 0.0f;
-	if (TimeRulerTrack.IsVisible())
+	TSharedPtr<const Trace::IAnalysisSession> Session = FInsightsManager::Get()->GetSession();
+	if (Session)
 	{
-		TimeRulerTrack.SetPosY(TopOffset);
-		TopOffset += TimeRulerTrack.GetHeight();
+		double SessionTime = 0.0;
+		{
+			Trace::FAnalysisSessionReadScope SessionReadScope(*Session.Get());
+			SessionTime = Session->GetDurationSeconds();
+		}
+
+		// Check if horizontal scroll area has changed.
+		if (SessionTime > Viewport.GetMaxValidTime() &&
+			SessionTime != DBL_MAX &&
+			SessionTime != std::numeric_limits<double>::infinity())
+		{
+			const double PreviousSessionTime = Viewport.GetMaxValidTime();
+			if ((PreviousSessionTime >= Viewport.GetStartTime() && PreviousSessionTime <= Viewport.GetEndTime()) ||
+				(SessionTime >= Viewport.GetStartTime() && SessionTime <= Viewport.GetEndTime()))
+			{
+				Viewport.AddDirtyFlags(ETimingTrackViewportDirtyFlags::HClippedSessionTimeChanged);
+			}
+
+			//UE_LOG(TimingProfiler, Log, TEXT("Session Duration: %g"), DT);
+			Viewport.SetMaxValidTime(SessionTime);
+			UpdateHorizontalScrollBar();
+		}
 	}
-	if (MarkersTrack.IsVisible())
+
+	////////////////////////////////////////////////////////////////////////////////////////////////////
+
+	if (Session)
 	{
-		MarkersTrack.SetPosY(TopOffset);
-		TopOffset += MarkersTrack.GetHeight();
+		// Tick plugin extenders.
+		// Each extender can add/remove tracks and/or change order of tracks.
+		for (Insights::ITimingViewExtender* Extender : GetExtenders())
+		{
+			Extender->Tick(*this, *Session.Get());
+		}
 	}
-	if (GraphTrack->IsVisible())
+
+	////////////////////////////////////////////////////////////////////////////////////////////////////
+
+	// Animate the (vertical) layout transition (i.e. compact mode <-> normal mode).
+	Viewport.UpdateLayout();
+
+	TimeRulerTrack->SetSelection(bIsSelecting, SelectionStartTime, SelectionEndTime);
+	TimeRulerTrack->SetTimeMarker(bIsDragging, TimeMarker);
+
+	////////////////////////////////////////////////////////////////////////////////////////////////////
+
+	class FTimingTrackUpdateContext : public ITimingTrackUpdateContext
 	{
-		GraphTrack->SetPosY(TopOffset);
-		TopOffset += GraphTrack->GetHeight();
-	}
-	Viewport.SetTopOffset(TopOffset);
+	public:
+		explicit FTimingTrackUpdateContext(STimingView* InTimingView, double InCurrentTime, float InDeltaTime) : TimingView(InTimingView), CurrentTime(InCurrentTime), DeltaTime(InDeltaTime) {}
 
-	//////////////////////////////////////////////////
+		virtual const FTimingTrackViewport& GetViewport() const override { return TimingView->GetViewport(); }
+		virtual const FVector2D& GetMousePosition() const override { return TimingView->GetMousePosition(); }
+		virtual const TSharedPtr<const ITimingEvent> GetHoveredEvent() const override { return TimingView->GetHoveredEvent(); }
+		virtual const TSharedPtr<const ITimingEvent> GetSelectedEvent() const override { return TimingView->GetSelectedEvent(); }
+		virtual double GetCurrentTime() const override { return CurrentTime; }
+		virtual float GetDeltaTime() const override { return DeltaTime; }
+
+	public:
+		STimingView* TimingView;
+		double CurrentTime;
+		float DeltaTime;
+	};
+
+	FTimingTrackUpdateContext UpdateContext(this, InCurrentTime, InDeltaTime);
+
+	////////////////////////////////////////////////////////////////////////////////////////////////////
+	// Pre-Update.
+	// The tracks needs to update their size.
 
-	if (!bIsPanning)
 	{
-		//////////////////////////////////////////////////
-		// Elastic snap to vertical scroll limits.
+		FStopwatch PreUpdateTracksStopwatch;
+		PreUpdateTracksStopwatch.Start();
 
-		const float MinY = 0.0f;// -0.5f * Viewport.Height;
-		const float MaxY = Viewport.GetScrollHeight() - Viewport.GetHeight() + TopOffset + 7.0f;
-		const float U = 0.5f;
+		for (TSharedPtr<FBaseTimingTrack>& TrackPtr : TopDockedTracks)
+		{
+			if (TrackPtr->IsVisible())
+			{
+				TrackPtr->PreUpdate(UpdateContext);
+			}
+		}
 
-		float ScrollPosY = Viewport.GetScrollPosY();
-		if (ScrollPosY < MinY)
+		for (TSharedPtr<FBaseTimingTrack>& TrackPtr : BottomDockedTracks)
 		{
-			ScrollPosY = ScrollPosY * U + (1.0f - U) * MinY;
-			if (FMath::IsNearlyEqual(ScrollPosY, MinY, 0.5f))
+			if (TrackPtr->IsVisible())
 			{
-				ScrollPosY = MinY;
+				TrackPtr->PreUpdate(UpdateContext);
 			}
 		}
-		else if (ScrollPosY > MaxY)
+
+		for (TSharedPtr<FBaseTimingTrack>& TrackPtr : ScrollableTracks)
 		{
-			ScrollPosY = ScrollPosY * U + (1.0f - U) * MaxY;
-			if (FMath::IsNearlyEqual(ScrollPosY, MaxY, 0.5f))
+			if (TrackPtr->IsVisible())
 			{
-				ScrollPosY = MaxY;
+				TrackPtr->PreUpdate(UpdateContext);
 			}
-			if (ScrollPosY < MinY)
+		}
+
+		for (TSharedPtr<FBaseTimingTrack>& TrackPtr : ForegroundTracks)
+		{
+			if (TrackPtr->IsVisible())
 			{
-				ScrollPosY = MinY;
+				TrackPtr->PreUpdate(UpdateContext);
 			}
 		}
 
-		ScrollAtPosY(ScrollPosY);
+		PreUpdateTracksStopwatch.Stop();
+		PreUpdateTracksDurationHistory.AddValue(PreUpdateTracksStopwatch.AccumulatedTime);
+	}
 
-		//////////////////////////////////////////////////
-		// Elastic snap to horizontal time limits.
+	////////////////////////////////////////////////////////////////////////////////////////////////////
+	// Update Y postion for the visible top/bottom docked tracks.
+	// Compute the total height of top/bottom docked areas.
 
-		if (Viewport.EnforceHorizontalScrollLimits(0.5)) // 0.5 is the interpolation factor
+	float TopOffset = 0.0f;
+	for (TSharedPtr<FBaseTimingTrack>& TrackPtr : TopDockedTracks)
+	{
+		TrackPtr->SetPosY(TopOffset);
+		if (TrackPtr->IsVisible())
 		{
-			UpdateHorizontalScrollBar();
-			bIsViewportDirty = true;
+			TopOffset += TrackPtr->GetHeight();
 		}
 	}
+	Viewport.SetTopOffset(TopOffset);
 
-	//TODO: Move this check into STimersView/SStatsView or in STimingProfilerWindow or in FTimingProfilerManager.
-	// We need to check if TimersView or StatsView needs to update their lists of timers / counters.
-	// But, ensure we do not check too often.
-	static uint64 NextTimestamp = 0;
-	uint64 Time = FPlatformTime::Cycles64();
-	if (Time > NextTimestamp)
+	float BottomOffset = 0.0f;
+	for (TSharedPtr<FBaseTimingTrack>& TrackPtr : BottomDockedTracks)
 	{
-		const uint64 WaitTime = static_cast<uint64>(0.2 / FPlatformTime::GetSecondsPerCycle64()); // 200ms
-		NextTimestamp = Time + WaitTime;
-
-		TSharedPtr<STimingProfilerWindow> Wnd = FTimingProfilerManager::Get()->GetProfilerWindow();
-		if (Wnd)
+		if (TrackPtr->IsVisible())
 		{
-			TSharedPtr<STimersView> TimersView = Wnd->GetTimersView();
-			if (TimersView)
-			{
-				TimersView->RebuildTree(false);
-			}
-			TSharedPtr<SStatsView> StatsView = Wnd->GetStatsView();
-			if (StatsView)
-			{
-				StatsView->RebuildTree(false);
-			}
+			BottomOffset += TrackPtr->GetHeight();
 		}
 	}
-
-	FileActivitySharedState->Update();
-
-	TSharedPtr<const Trace::IAnalysisSession> Session = FInsightsManager::Get()->GetSession();
-	if (Session)
+	float BottomOffsetY = Viewport.GetHeight() - BottomOffset;
+	for (TSharedPtr<FBaseTimingTrack>& TrackPtr : BottomDockedTracks)
 	{
-		Trace::FAnalysisSessionReadScope SessionReadScope(*Session.Get());
-
-		// Check if horizontal scroll area has changed.
-		double SessionTime = Session->GetDurationSeconds();
-		if (SessionTime > Viewport.GetMaxValidTime() &&
-			SessionTime != DBL_MAX &&
-			SessionTime != std::numeric_limits<double>::infinity())
+		TrackPtr->SetPosY(BottomOffsetY);
+		if (TrackPtr->IsVisible())
 		{
-			//UE_LOG(TimingProfiler, Log, TEXT("Session Duration: %g"), DT);
-			Viewport.SetMaxValidTime(SessionTime);
-			UpdateHorizontalScrollBar();
-			//bIsMaxValidTimeDirty = true;
-
-			if (SessionTime >= Viewport.GetStartTime() && SessionTime <= Viewport.GetEndTime())
-			{
-				bTimingEventsNeedUpdate = true;
-				MarkersTrack.SetDirtyFlag();
-			}
+			BottomOffsetY += TrackPtr->GetHeight();
 		}
+	}
+	Viewport.SetBottomOffset(BottomOffset);
 
-		// After this point we can add new tracks, so reset the dirty flag
-		bTimingEventsTracksChanged = false;
+	////////////////////////////////////////////////////////////////////////////////////////////////////
+	// Compute the total height of visible scrollable tracks.
 
-		if (Trace::ReadFileActivityProvider(*Session.Get()))
+	float ScrollHeight = 0.0f;
+	for (TSharedPtr<FBaseTimingTrack>& TrackPtr : ScrollableTracks)
+	{
+		if (TrackPtr->IsVisible())
 		{
-			if (IoOverviewTrack == nullptr)
-			{
-				// Note: The I/O timelines are just prototypes for now (will be removed once the functionality is moved in analyzer).
-				const uint64 TrackId = FBaseTimingTrack::GenerateId();
-				IoOverviewTrack = new FOverviewFileActivityTimingTrack(TrackId, FileActivitySharedState);
-				IoOverviewTrack->SetOrder(-1);
-				IoOverviewTrack->SetVisibilityFlag(bShowHideAllIoTracks);
-				AddTimingEventsTrack(IoOverviewTrack);
-			}
+			ScrollHeight += TrackPtr->GetHeight();
 		}
+	}
+	ScrollHeight += 1.0f; // allow 1 pixel at the bottom (for last horizontal line)
 
-		int32 Order = 1;
+	////////////////////////////////////////////////////////////////////////////////////////////////////
+	// Check if vertical scroll area has changed.
 
-		const Trace::ITimingProfilerProvider* TimingProfilerProvider = Trace::ReadTimingProfilerProvider(*Session.Get());
-		const Trace::ILoadTimeProfilerProvider* LoadTimeProfilerProvider = Trace::ReadLoadTimeProfilerProvider(*Session.Get());
+	bool bScrollHeightChanged = false;
+	if (ScrollHeight != Viewport.GetScrollHeight())
+	{
+		bScrollHeightChanged = true;
+		Viewport.SetScrollHeight(ScrollHeight);
+		UpdateVerticalScrollBar();
+	}
 
-		if (TimingProfilerProvider || LoadTimeProfilerProvider)
+	////////////////////////////////////////////////////////////////////////////////////////////////////
+	// Elastic snap to vertical scroll limits.
+
+	if (!bIsPanning)
+	{
+		const float DY = Viewport.GetScrollHeight() - Viewport.GetScrollableAreaHeight() + 7.0f; // +7 is to allow some space for the horizontal scrollbar
+		const float MinY = FMath::Min(DY, 0.0f);
+		const float MaxY = DY - MinY;
+
+		float ScrollPosY = Viewport.GetScrollPosY();
+
+		if (ScrollPosY < MinY)
 		{
-			// Check if we have a GPU track.
-			uint32 GpuTimelineIndex;
-			if (TimingProfilerProvider && TimingProfilerProvider->GetGpuTimelineIndex(GpuTimelineIndex))
+			if (bScrollHeightChanged || Viewport.IsDirty(ETimingTrackViewportDirtyFlags::VLayoutChanged))
 			{
-				const uint64 TrackId = static_cast<uint64>(GpuTimelineIndex);
-				if (!CachedTimelines.Contains(TrackId))
-				{
-					GpuTrack = new FGpuTimingTrack(TrackId, TEXT("GPU"), nullptr);
-					GpuTrack->SetOrder(0);
-					GpuTrack->SetVisibilityFlag(bShowHideAllGpuTracks);
-					AddTimingEventsTrack(GpuTrack);
-				}
+				ScrollPosY = MinY;
 			}
-
-			// Check available CPU tracks.
-			const Trace::IThreadProvider& ThreadProvider = Trace::ReadThreadProvider(*Session.Get());
-			ThreadProvider.EnumerateThreads([this, &Order, TimingProfilerProvider, LoadTimeProfilerProvider](const Trace::FThreadInfo& ThreadInfo)
+			else
 			{
-				const TCHAR* GroupName = ThreadInfo.GroupName;
-				if (GroupName == nullptr)
+				constexpr float U = 0.5f;
+				ScrollPosY = ScrollPosY * U + (1.0f - U) * MinY;
+				if (FMath::IsNearlyEqual(ScrollPosY, MinY, 0.5f))
 				{
-					GroupName = ThreadInfo.Name;
+					ScrollPosY = MinY;
 				}
-
-				bool bIsGroupVisible = bShowHideAllCpuTracks;
-				if (GroupName != nullptr)
+			}
+		}
+		else if (ScrollPosY > MaxY)
+		{
+			if (bScrollHeightChanged || Viewport.IsDirty(ETimingTrackViewportDirtyFlags::VLayoutChanged))
+			{
+				ScrollPosY = MaxY;
+			}
+			else
+			{
+				constexpr float U = 0.5f;
+				ScrollPosY = ScrollPosY * U + (1.0f - U) * MaxY;
+				if (FMath::IsNearlyEqual(ScrollPosY, MaxY, 0.5f))
 				{
-					if (!ThreadGroups.Contains(GroupName))
-					{
-						UE_LOG(TimingProfiler, Log, TEXT("New CPU Thread Group (%d) : \"%s\""), ThreadGroups.Num() + 1, GroupName);
-						ThreadGroups.Add(GroupName, { GroupName, bIsGroupVisible, 0, Order });
-					}
-					else
-					{
-						FThreadGroup& ThreadGroup = ThreadGroups[GroupName];
-						bIsGroupVisible = ThreadGroup.bIsVisible;
-						ThreadGroup.Order = Order;
-					}
+					ScrollPosY = MaxY;
 				}
+			}
+			if (ScrollPosY < MinY)
+			{
+				ScrollPosY = MinY;
+			}
+		}
 
-				bool bIsLoadingThread = false;
-				uint32 LoadingTimelineIndex;
-				if (LoadTimeProfilerProvider && LoadTimeProfilerProvider->GetCpuThreadTimelineIndex(ThreadInfo.Id, LoadingTimelineIndex))
-				{
-					bIsLoadingThread = true;
-					if (!LoadingTracks.Contains(LoadingTimelineIndex))
-					{
-						FString TrackName(ThreadInfo.Name && *ThreadInfo.Name ? FString::Printf(TEXT("Loading - %s"), ThreadInfo.Name) : FString::Printf(TEXT("Loading - Thread %u"), ThreadInfo.Id));
-
-						uint64 TrackId = FBaseTimingTrack::GenerateId();
-						FLoadingTimingTrack* LoadingThreadTrack = new FLoadingTimingTrack(TrackId, LoadingTimelineIndex, GroupName, TrackName, LoadingSharedState);
-						LoadingThreadTrack->SetOrder(-100 + LoadingTracks.Num());
-						LoadingThreadTrack->SetVisibilityFlag(bShowHideAllLoadingTracks);
-						AddTimingEventsTrack(LoadingThreadTrack);
-						LoadingTracks.Add(LoadingTimelineIndex, LoadingThreadTrack);
-					}
-				}
+		ScrollAtPosY(ScrollPosY);
+	}
 
-				uint32 CpuTimelineIndex;
-				if (TimingProfilerProvider && TimingProfilerProvider->GetCpuThreadTimelineIndex(ThreadInfo.Id, CpuTimelineIndex))
-				{
-					FTimingEventsTrack* Track = nullptr;
-					const uint64 TrackId = static_cast<uint64>(CpuTimelineIndex);
+	////////////////////////////////////////////////////////////////////////////////////////////////////
+	// Update Y position for visible scrollable tracks.
 
-					if (!CachedTimelines.Contains(TrackId))
-					{
-						FString TrackName(ThreadInfo.Name && *ThreadInfo.Name ? ThreadInfo.Name : FString::Printf(TEXT("Thread %u"), ThreadInfo.Id));
-
-						// Create new Timing Events track for the CPU thread.
-						Track = new FCpuTimingTrack(TrackId, TrackName, GroupName);
-						Track->SetOrder(Order);
-						((FCpuTimingTrack*)Track)->SetThreadId(ThreadInfo.Id);
-						CpuTracks.Add(ThreadInfo.Id, Track);
-
-						FThreadGroup& ThreadGroup = ThreadGroups[GroupName];
-						ThreadGroup.NumTimelines++;
-
-						if (bAssetLoadingMode && bIsLoadingThread)
-						{
-							Track->SetVisibilityFlag(true);
-							ThreadGroup.bIsVisible = true;
-						}
-						else
-						{
-							Track->SetVisibilityFlag(bIsGroupVisible);
-						}
-
-						AddTimingEventsTrack(Track);
-					}
-					else
-					{
-						Track = CachedTimelines[TrackId];
+	float ScrollableTrackPosY = Viewport.GetTopOffset() - Viewport.GetScrollPosY();
+	for (TSharedPtr<FBaseTimingTrack>& TrackPtr : ScrollableTracks)
+	{
+		TrackPtr->SetPosY(ScrollableTrackPosY); // set pos y also for the hidden tracks
+		if (TrackPtr->IsVisible())
+		{
+			ScrollableTrackPosY += TrackPtr->GetHeight();
+		}
+	}
 
-						if (Track->GetOrder() != Order)
-						{
-							Track->SetOrder(Order);
-							bTimingEventsTracksChanged = true;
-						}
-					}
+	////////////////////////////////////////////////////////////////////////////////////////////////////
+	// At this point it is assumed all tracks have proper position and size.
+	// Update.
+	{
+		FStopwatch UpdateTracksStopwatch;
+		UpdateTracksStopwatch.Start();
 
-					Order++;
-				}
-			});
+		for (TSharedPtr<FBaseTimingTrack>& TrackPtr : TopDockedTracks)
+		{
+			if (TrackPtr->IsVisible())
+			{
+				TrackPtr->Update(UpdateContext);
+			}
 		}
 
-		if (Trace::ReadFileActivityProvider(*Session.Get()))
+		for (TSharedPtr<FBaseTimingTrack>& TrackPtr : BottomDockedTracks)
 		{
-			if (IoActivityTrack == nullptr)
+			if (TrackPtr->IsVisible())
 			{
-				// Note: The I/O timelines are just prototypes for now (will be removed once the functionality is moved in analyzer).
-				uint64 TrackId = FBaseTimingTrack::GenerateId();
-				IoActivityTrack = new FDetailedFileActivityTimingTrack(TrackId, FileActivitySharedState);
-				IoActivityTrack->SetOrder(999999);
-				IoActivityTrack->SetVisibilityFlag(bShowHideAllIoTracks);
-				AddTimingEventsTrack(IoActivityTrack);
+				TrackPtr->Update(UpdateContext);
 			}
 		}
 
-		// Tick plugin extenders
-		for (Insights::ITimingViewExtender* Extender : GetExtenders())
+		for (TSharedPtr<FBaseTimingTrack>& TrackPtr : ScrollableTracks)
 		{
-			Extender->Tick(*this, *Session.Get());
+			if (TrackPtr->IsVisible())
+			{
+				TrackPtr->Update(UpdateContext);
+			}
 		}
 
-		if (bTimingEventsTracksChanged)
+		for (TSharedPtr<FBaseTimingTrack>& TrackPtr : ForegroundTracks)
 		{
-			// Let extenders re-order their track data
-			for (Insights::ITimingViewExtender* Extender : GetExtenders())
+			if (TrackPtr->IsVisible())
 			{
-				Extender->OnTracksChanged(*this, Order);
+				TrackPtr->Update(UpdateContext);
 			}
+		}
 
-			// The list has changed. Sort the list again.
-			Algo::SortBy(TimingEventsTracks, &FTimingEventsTrack::GetOrder);
+		UpdateTracksStopwatch.Stop();
+		UpdateTracksDurationHistory.AddValue(UpdateTracksStopwatch.AccumulatedTime);
+	}
+	//////////////////////////////////////////////////
+	// Post-Update.
+	{
+		FStopwatch PostUpdateTracksStopwatch;
+		PostUpdateTracksStopwatch.Start();
 
-			bTimingEventsTracksChanged = false;
+		for (TSharedPtr<FBaseTimingTrack>& TrackPtr : TopDockedTracks)
+		{
+			if (TrackPtr->IsVisible())
+			{
+				TrackPtr->PostUpdate(UpdateContext);
+			}
+		}
 
-			bTimingEventsNeedUpdate = true;
+		for (TSharedPtr<FBaseTimingTrack>& TrackPtr : BottomDockedTracks)
+		{
+			if (TrackPtr->IsVisible())
+			{
+				TrackPtr->PostUpdate(UpdateContext);
+			}
 		}
-	}
 
-	// Compute total height of visible tracks.
-	float TotalScrollHeight = 0.0f;
-	for (FTimingEventsTrack* TrackPtr : TimingEventsTracks)
-	{
-		ensure(TrackPtr != nullptr);
-		FTimingEventsTrack& Track = *TrackPtr;
+		for (TSharedPtr<FBaseTimingTrack>& TrackPtr : ScrollableTracks)
+		{
+			if (TrackPtr->IsVisible())
+			{
+				TrackPtr->PostUpdate(UpdateContext);
+			}
+		}
 
-		if (Track.IsVisible())
+		for (TSharedPtr<FBaseTimingTrack>& TrackPtr : ForegroundTracks)
 		{
-			Track.SetPosY(TotalScrollHeight);
-			TotalScrollHeight += Track.GetHeight();
+			if (TrackPtr->IsVisible())
+			{
+				TrackPtr->PostUpdate(UpdateContext);
+			}
 		}
 
-		// Reset track's hovered/selected state.
-		Track.SetHoveredState(false);
-		Track.SetSelectedFlag(false);
+		PostUpdateTracksStopwatch.Stop();
+		PostUpdateTracksDurationHistory.AddValue(PostUpdateTracksStopwatch.AccumulatedTime);
 	}
-	TotalScrollHeight += 1.0f; // allow 1 pixel at the bottom (for last horizontal line)
+	//////////////////////////////////////////////////
 
-	// Set hovered/selected state for actual hovered/selected tracks, if any.
-	if (HoveredTimingEvent.Track != nullptr)
-	{
-		const_cast<FTimingEventsTrack*>(HoveredTimingEvent.Track)->SetHoveredState(true);
-	}
-	if (SelectedTimingEvent.Track != nullptr)
+	Tooltip.Update();
+	if (!MousePosition.IsZero())
 	{
-		const_cast<FTimingEventsTrack*>(SelectedTimingEvent.Track)->SetSelectedFlag(true);
+		Tooltip.SetPosition(MousePosition, 0.0f, Viewport.GetWidth() - 12.0f, 0.0f, Viewport.GetHeight() - 12.0f); // -12.0f is to avoid overlaping the scrollbars
 	}
 
-	//TimeRulerTrack.Tick(InCurrentTime, InDeltaTime);
-	MarkersTrack.Tick(InCurrentTime, InDeltaTime);
+	//////////////////////////////////////////////////
 
-	// Check if vertical scroll area has changed.
-	if (TotalScrollHeight != Viewport.GetScrollHeight())
+	// Reset hovered/selected flags for all tracks.
+	for (TSharedPtr<FBaseTimingTrack>& TrackPtr : ScrollableTracks)
 	{
-		Viewport.SetScrollHeight(TotalScrollHeight);
-		UpdateVerticalScrollBar();
-		bIsVerticalViewportDirty = true;
+		if (TrackPtr->IsVisible())
+		{
+			TrackPtr->SetHoveredState(false);
+			TrackPtr->SetSelectedFlag(false);
+		}
 	}
 
-	// Update layout change animation (i.e. compact mode <-> normal mode).
-	Layout.Update();
-
-	//////////////////////////////////////////////////
-	// Update dirty tracks.
-
-	if (bIsViewportDirty)
+	// Set the hovered flag for the actual hovered track, if any.
+	if (HoveredTrack.IsValid())
 	{
-		bIsViewportDirty = false;
-
-		//TimeRuler.SetDirtyFlag();
-		MarkersTrack.SetDirtyFlag();
-		bTimingEventsNeedUpdate = true;
+		HoveredTrack->SetHoveredState(true);
 	}
 
-	if (bIsVerticalViewportDirty)
+	// Set the selected flag for the actual selected track, if any.
+	if (SelectedTrack.IsValid())
 	{
-		bIsVerticalViewportDirty = false;
-
-		bTimingEventsNeedUpdate = true;
+		SelectedTrack->SetSelectedFlag(true);
 	}
 
-	//if (TimeRuler.IsVisible() && TimeRuler.IsDirty())
-	//{
-	//	TimeRuler.ClearDirtyFlag();
-	//
-	//	TimeRuler.OnViewportChanged(Viewport);
-	//}
+	//////////////////////////////////////////////////
 
-	if (MarkersTrack.IsVisible() && MarkersTrack.IsDirty())
-	{
-		MarkersTrack.ClearDirtyFlag();
+	Viewport.ResetDirtyFlags();
 
-		FStopwatch Stopwatch;
-		Stopwatch.Start();
+	UpdateOtherViews();
 
-		MarkersTrack.Update(Viewport);
+	TickStopwatch.Stop();
+	TickDurationHistory.AddValue(TickStopwatch.AccumulatedTime);
+}
 
-		Stopwatch.Stop();
-		TimeMarkerCacheUpdateDurationHistory.AddValue(Stopwatch.AccumulatedTime);
-	}
+////////////////////////////////////////////////////////////////////////////////////////////////////
 
-	if (bTimingEventsNeedUpdate)
+//TODO: Move this check into STimersView/SStatsView or in STimingProfilerWindow or in FTimingProfilerManager.
+void STimingView::UpdateOtherViews()
+{
+	// We need to check if TimersView or StatsView needs to update their lists of timers / counters.
+	// But, ensure we do not check too often.
+	static uint64 NextTimestamp = 0;
+	uint64 Time = FPlatformTime::Cycles64();
+	if (Time > NextTimestamp)
 	{
-		bTimingEventsNeedUpdate = false;
-
-		GraphTrack->SetDirtyFlag();
-
-		FStopwatch Stopwatch;
-		Stopwatch.Start();
+		const uint64 WaitTime = static_cast<uint64>(0.2 / FPlatformTime::GetSecondsPerCycle64()); // 200ms
+		NextTimestamp = Time + WaitTime;
 
-		for (int32 TrackIndex = 0; TrackIndex < TimingEventsTracks.Num(); ++TrackIndex)
+		TSharedPtr<STimingProfilerWindow> Wnd = FTimingProfilerManager::Get()->GetProfilerWindow();
+		if (Wnd)
 		{
-			TimingEventsTracks[TrackIndex]->Update(Viewport);
+			TSharedPtr<STimersView> TimersView = Wnd->GetTimersView();
+			if (TimersView)
+			{
+				TimersView->RebuildTree(false);
+			}
+			TSharedPtr<SStatsView> StatsView = Wnd->GetStatsView();
+			if (StatsView)
+			{
+				StatsView->RebuildTree(false);
+			}
 		}
-
-		Stopwatch.Stop();
-		TimelineCacheUpdateDurationHistory.AddValue(Stopwatch.AccumulatedTime);
 	}
-
-	if (GraphTrack->IsVisible())
-	{
-		GraphTrack->Update(Viewport);
-	}
-
-	Tooltip.Update();
-	if (!MousePosition.IsZero())
-	{
-		Tooltip.SetPosition(MousePosition, 0.0f, Viewport.GetWidth() - 12.0f, 0.0f, Viewport.GetHeight() - 12.0f); // -12.0f is to avoid overlaping the scrollbars
-	}
-
-	//TickStopwatch.Stop();
-	//TODO: TotalUpdateDurationHistory.AddValue(TickStopwatch.AccumulatedTime);
 }
 
 ////////////////////////////////////////////////////////////////////////////////////////////////////
 
 int32 STimingView::OnPaint(const FPaintArgs& Args, const FGeometry& AllottedGeometry, const FSlateRect& MyCullingRect, FSlateWindowElementList& OutDrawElements, int32 LayerId, const FWidgetStyle& InWidgetStyle, bool bParentEnabled) const
 {
+	FStopwatch Stopwatch;
+	Stopwatch.Start();
+
 	const bool bEnabled = ShouldBeEnabled(bParentEnabled);
 	const ESlateDrawEffect DrawEffects = bEnabled ? ESlateDrawEffect::None : ESlateDrawEffect::DisabledEffect;
 	FDrawContext DrawContext(AllottedGeometry, MyCullingRect, InWidgetStyle, DrawEffects, OutDrawElements, LayerId);
@@ -770,91 +791,216 @@ int32 STimingView::OnPaint(const FPaintArgs& Args, const FGeometry& AllottedGeom
 	}
 #endif
 
-	//////////////////////////////////////////////////
+	////////////////////////////////////////////////////////////////////////////////////////////////////
 
-	FStopwatch Stopwatch;
-	Stopwatch.Start();
+	class FTimingTrackDrawContext : public ITimingTrackDrawContext
+	{
+	public:
+		explicit FTimingTrackDrawContext(const STimingView* InTimingView, FDrawContext& InDrawContext, const FTimingViewDrawHelper& InHelper)
+			: TimingView(InTimingView)
+			, DrawContext(InDrawContext)
+			, Helper(InHelper)
+		{}
+
+		virtual const FTimingTrackViewport& GetViewport() const override { return TimingView->GetViewport(); }
+		virtual const FVector2D& GetMousePosition() const override { return TimingView->GetMousePosition(); }
+		virtual const TSharedPtr<const ITimingEvent> GetHoveredEvent() const override { return TimingView->GetHoveredEvent(); }
+		virtual const TSharedPtr<const ITimingEvent> GetSelectedEvent() const override { return TimingView->GetSelectedEvent(); }
+		virtual FDrawContext& GetDrawContext() const override { return DrawContext; }
+		virtual const FTimingViewDrawHelper& GetHelper() const override { return Helper; }
+
+	public:
+		const STimingView* TimingView;
+		FDrawContext& DrawContext;
+		const FTimingViewDrawHelper& Helper;
+	};
 
-	FTimingViewDrawHelper Helper(DrawContext, Viewport, Layout);
+	FTimingViewDrawHelper Helper(DrawContext, Viewport);
+	FTimingTrackDrawContext TimingDrawContext(this, DrawContext, Helper);
+
+	////////////////////////////////////////////////////////////////////////////////////////////////////
+
+	Helper.SetHighlightedEventTypeId(HighlightedEventTypeId);
 
 	// Draw background.
 	Helper.DrawBackground();
 
-	// Draw scrollable tracks.
-	TSharedPtr<const Trace::IAnalysisSession> Session = FInsightsManager::Get()->GetSession();
-	if(Session.IsValid())
+	//////////////////////////////////////////////////
+	// Pre-Draw
 	{
-		Helper.BeginTimelines();
+		FStopwatch PreDrawTracksStopwatch;
+		PreDrawTracksStopwatch.Start();
 
-		for (int32 TrackIndex = 0; TrackIndex < TimingEventsTracks.Num(); ++TrackIndex)
+		for (const TSharedPtr<FBaseTimingTrack>& TrackPtr : ScrollableTracks)
 		{
-			FTimingEventsTrack& Track = *TimingEventsTracks[TrackIndex];
-			if (Track.IsVisible())
+			if (TrackPtr->IsVisible())
+			{
+				TrackPtr->PreDraw(TimingDrawContext);
+			}
+		}
+
+		for (const TSharedPtr<FBaseTimingTrack>& TrackPtr : TopDockedTracks)
+		{
+			if (TrackPtr->IsVisible())
+			{
+				TrackPtr->PreDraw(TimingDrawContext);
+			}
+		}
+
+		for (const TSharedPtr<FBaseTimingTrack>& TrackPtr : BottomDockedTracks)
+		{
+			if (TrackPtr->IsVisible())
 			{
-				Track.Draw(Helper);
+				TrackPtr->PreDraw(TimingDrawContext);
 			}
 		}
 
-		Helper.EndTimelines();
+		for (const TSharedPtr<FBaseTimingTrack>& TrackPtr : ForegroundTracks)
+		{
+			if (TrackPtr->IsVisible())
+			{
+				TrackPtr->PreDraw(TimingDrawContext);
+			}
+		}
+
+		PreDrawTracksStopwatch.Stop();
+		PreDrawTracksDurationHistory.AddValue(PreDrawTracksStopwatch.AccumulatedTime);
 	}
 
-	if (!FTimingEvent::AreEquals(SelectedTimingEvent, HoveredTimingEvent))
+	//////////////////////////////////////////////////
+	// Draw
 	{
-		// Highlight the selected timing event (if any).
-		if (SelectedTimingEvent.IsValid())
+		FStopwatch DrawTracksStopwatch;
+		DrawTracksStopwatch.Start();
+
+		Helper.BeginDrawTracks();
+
+		for (const TSharedPtr<FBaseTimingTrack>& TrackPtr : ScrollableTracks)
+		{
+			if (TrackPtr->IsVisible())
+			{
+				TrackPtr->Draw(TimingDrawContext);
+			}
+		}
+
+		for (const TSharedPtr<FBaseTimingTrack>& TrackPtr : TopDockedTracks)
 		{
-			const float Y = Viewport.GetViewportY(SelectedTimingEvent.Track->GetPosY()) + Layout.GetLaneY(SelectedTimingEvent.Depth);
-			Helper.DrawTimingEventHighlight(SelectedTimingEvent.StartTime, SelectedTimingEvent.EndTime, Y, FTimingViewDrawHelper::EHighlightMode::Selected);
+			if (TrackPtr->IsVisible())
+			{
+				TrackPtr->Draw(TimingDrawContext);
+			}
 		}
 
-		// Highlight the hovered timing event (if any).
-		if (HoveredTimingEvent.IsValid())
+		for (const TSharedPtr<FBaseTimingTrack>& TrackPtr : BottomDockedTracks)
 		{
-			const float Y = Viewport.GetViewportY(HoveredTimingEvent.Track->GetPosY()) + Layout.GetLaneY(HoveredTimingEvent.Depth);
-			Helper.DrawTimingEventHighlight(HoveredTimingEvent.StartTime, HoveredTimingEvent.EndTime, Y, FTimingViewDrawHelper::EHighlightMode::Hovered);
+			if (TrackPtr->IsVisible())
+			{
+				TrackPtr->Draw(TimingDrawContext);
+			}
 		}
+
+		for (const TSharedPtr<FBaseTimingTrack>& TrackPtr : ForegroundTracks)
+		{
+			if (TrackPtr->IsVisible())
+			{
+				TrackPtr->Draw(TimingDrawContext);
+			}
+		}
+
+		Helper.EndDrawTracks();
+
+		DrawTracksStopwatch.Stop();
+		DrawTracksDurationHistory.AddValue(DrawTracksStopwatch.AccumulatedTime);
 	}
-	else
+
+	//////////////////////////////////////////////////
+	// Draw the selected and/or hovered event.
+
+	if (ITimingEvent::AreValidAndEquals(SelectedEvent, HoveredEvent))
 	{
+		const TSharedRef<const FBaseTimingTrack> TrackPtr = SelectedEvent->GetTrack();
+
 		// Highlight the selected and hovered timing event (if any).
-		if (SelectedTimingEvent.IsValid())
+		if (TrackPtr->IsVisible())
 		{
-			const float Y = Viewport.GetViewportY(SelectedTimingEvent.Track->GetPosY()) + Layout.GetLaneY(SelectedTimingEvent.Depth);
-			Helper.DrawTimingEventHighlight(SelectedTimingEvent.StartTime, SelectedTimingEvent.EndTime, Y, FTimingViewDrawHelper::EHighlightMode::SelectedAndHovered);
+			SelectedEvent->GetTrack()->DrawEvent(TimingDrawContext, *SelectedEvent, EDrawEventMode::SelectedAndHovered);
 		}
 	}
-
-	if (GraphTrack->IsVisible())
+	else
 	{
-		GraphTrack->Draw(DrawContext, Viewport, MousePosition);
-	}
+		// Highlight the selected timing event (if any).
+		if (SelectedEvent.IsValid())
+		{
+			const TSharedRef<const FBaseTimingTrack> TrackPtr = SelectedEvent->GetTrack();
+			if (TrackPtr->IsVisible())
+			{
+				SelectedEvent->GetTrack()->DrawEvent(TimingDrawContext, *SelectedEvent, EDrawEventMode::Selected);
+			}
+		}
 
-	if (MarkersTrack.IsVisible())
-	{
-		// Draw the time markers (vertical lines + category & message texts).
-		MarkersTrack.Draw(DrawContext, Viewport);
+		// Highlight the hovered timing event (if any).
+		if (HoveredEvent.IsValid())
+		{
+			const TSharedRef<const FBaseTimingTrack> TrackPtr = HoveredEvent->GetTrack();
+			if (TrackPtr->IsVisible())
+			{
+				HoveredEvent->GetTrack()->DrawEvent(TimingDrawContext, *HoveredEvent, EDrawEventMode::Hovered);
+			}
+		}
 	}
 
-	//////////////////////////////////////////////////
+	// Draw the time range selection.
+	FDrawHelpers::DrawTimeRangeSelection(DrawContext, Viewport, SelectionStartTime, SelectionEndTime, WhiteBrush, MainFont);
 
-	// Draw the time ruler.
-	if (TimeRulerTrack.IsVisible())
+	//////////////////////////////////////////////////
+	// Post-Draw
 	{
-		TimeRulerTrack.Draw(DrawContext, Viewport, MousePosition, bIsSelecting, SelectionStartTime, SelectionEndTime, TimeMarker);
-	}
+		FStopwatch PostDrawTracksStopwatch;
+		PostDrawTracksStopwatch.Start();
 
-	// Fill background for the Tracks filter combobox.
-	DrawContext.DrawBox(0, 0, 66, 22, WhiteBrush, FLinearColor(0.05f, 0.05f, 0.05f, 1.0f));
+		for (const TSharedPtr<FBaseTimingTrack>& TrackPtr : ScrollableTracks)
+		{
+			if (TrackPtr->IsVisible())
+			{
+				TrackPtr->PostDraw(TimingDrawContext);
+			}
+		}
 
-	//////////////////////////////////////////////////
+		for (const TSharedPtr<FBaseTimingTrack>& TrackPtr : TopDockedTracks)
+		{
+			if (TrackPtr->IsVisible())
+			{
+				TrackPtr->PostDraw(TimingDrawContext);
+			}
+		}
 
-	// Draw the time range selection.
-	FDrawHelpers::DrawTimeRangeSelection(DrawContext, Viewport, SelectionStartTime, SelectionEndTime, WhiteBrush, MainFont);
+		for (const TSharedPtr<FBaseTimingTrack>& TrackPtr : BottomDockedTracks)
+		{
+			if (TrackPtr->IsVisible())
+			{
+				TrackPtr->PostDraw(TimingDrawContext);
+			}
+		}
+
+		for (const TSharedPtr<FBaseTimingTrack>& TrackPtr : ForegroundTracks)
+		{
+			if (TrackPtr->IsVisible())
+			{
+				TrackPtr->PostDraw(TimingDrawContext);
+			}
+		}
+
+		PostDrawTracksStopwatch.Stop();
+		PostDrawTracksDurationHistory.AddValue(PostDrawTracksStopwatch.AccumulatedTime);
+	}
 
 	//////////////////////////////////////////////////
 
-	Stopwatch.Stop();
-	DrawDurationHistory.AddValue(Stopwatch.AccumulatedTime);
+	// Draw tooltip with info about hovered event.
+	Tooltip.Draw(DrawContext);
+
+	// Fill background for the "Tracks" filter combobox.
+	DrawContext.DrawBox(0.0f, 0.0f, 66.0f, 24.0f, WhiteBrush, FLinearColor(0.05f, 0.05f, 0.05f, 1.0f));
 
 	//////////////////////////////////////////////////
 
@@ -883,55 +1029,51 @@ int32 STimingView::OnPaint(const FPaintArgs& Args, const FGeometry& AllottedGeom
 
 		// Time interval since last OnPaint call.
 		const uint64 CurrentTime = FPlatformTime::Cycles64();
-		const uint64 OnPaintDuration = CurrentTime - LastOnPaintTime;
+		const uint64 OnPaintDeltaTime = CurrentTime - LastOnPaintTime;
 		LastOnPaintTime = CurrentTime;
-		OnPaintDurationHistory.AddValue(OnPaintDuration); // saved for last 32 OnPaint calls
-		const uint64 AvgOnPaintDuration = OnPaintDurationHistory.ComputeAverage();
-		const uint64 AvgOnPaintDurationMs = FStopwatch::Cycles64ToMilliseconds(AvgOnPaintDuration);
-		const double AvgOnPaintFps = AvgOnPaintDurationMs != 0 ? 1.0 / FStopwatch::Cycles64ToSeconds(AvgOnPaintDuration) : 0.0;
+		OnPaintDeltaTimeHistory.AddValue(OnPaintDeltaTime); // saved for last 32 OnPaint calls
+		const uint64 AvgOnPaintDeltaTime = OnPaintDeltaTimeHistory.ComputeAverage();
+		const uint64 AvgOnPaintDeltaTimeMs = FStopwatch::Cycles64ToMilliseconds(AvgOnPaintDeltaTime);
+		const double AvgOnPaintFps = AvgOnPaintDeltaTimeMs != 0 ? 1.0 / FStopwatch::Cycles64ToSeconds(AvgOnPaintDeltaTime) : 0.0;
 
-		const uint64 AvgDrawDurationMs = FStopwatch::Cycles64ToMilliseconds(DrawDurationHistory.ComputeAverage());
+		const uint64 AvgPreDrawTracksDurationMs = FStopwatch::Cycles64ToMilliseconds(PreDrawTracksDurationHistory.ComputeAverage());
+		const uint64 AvgDrawTracksDurationMs = FStopwatch::Cycles64ToMilliseconds(DrawTracksDurationHistory.ComputeAverage());
+		const uint64 AvgPostDrawTracksDurationMs = FStopwatch::Cycles64ToMilliseconds(PostDrawTracksDurationHistory.ComputeAverage());
+		const uint64 AvgTotalDrawDurationMs = FStopwatch::Cycles64ToMilliseconds(TotalDrawDurationHistory.ComputeAverage());
 
 		DrawContext.DrawText
 		(
 			DbgX, DbgY,
-			FString::Printf(TEXT("D: %llu ms + %llu ms = %llu ms (%d fps)"),
-				AvgDrawDurationMs, // drawing time
-				AvgOnPaintDurationMs - AvgDrawDurationMs, // other overhead to OnPaint calls
-				AvgOnPaintDurationMs, // average time between two OnPaint calls
+			FString::Printf(TEXT("D: %llu ms + %llu ms + %llu ms + %llu ms = %llu ms | + %llu ms = %llu ms (%d fps)"),
+				AvgPreDrawTracksDurationMs, // pre-draw tracks time
+				AvgDrawTracksDurationMs, // draw tracks time
+				AvgPostDrawTracksDurationMs, // post-draw tracks time
+				AvgTotalDrawDurationMs - AvgPreDrawTracksDurationMs - AvgDrawTracksDurationMs - AvgPostDrawTracksDurationMs, // other draw code
+				AvgTotalDrawDurationMs,
+				AvgOnPaintDeltaTimeMs - AvgTotalDrawDurationMs, // other overhead to OnPaint calls
+				AvgOnPaintDeltaTimeMs, // average time between two OnPaint calls
 				FMath::RoundToInt(AvgOnPaintFps)), // framerate of OnPaint calls
 			SummaryFont, DbgTextColor
 		);
 		DbgY += DbgDY;
 
 		//////////////////////////////////////////////////
-		// Display the "timelines cache update" info.
+		// Display the "update" performance info.
 
-		const uint64 AvgTimelineUpdateDurationMs = FStopwatch::Cycles64ToMilliseconds(TimelineCacheUpdateDurationHistory.ComputeAverage());
-		const uint64 LastTimelineUpdateDurationMs = FStopwatch::Cycles64ToMilliseconds(TimelineCacheUpdateDurationHistory.GetValue(0));
+		const uint64 AvgPreUpdateTracksDurationMs = FStopwatch::Cycles64ToMilliseconds(PreUpdateTracksDurationHistory.ComputeAverage());
+		const uint64 AvgUpdateTracksDurationMs = FStopwatch::Cycles64ToMilliseconds(UpdateTracksDurationHistory.ComputeAverage());
+		const uint64 AvgPostUpdateTracksDurationMs = FStopwatch::Cycles64ToMilliseconds(PostUpdateTracksDurationHistory.ComputeAverage());
+		const uint64 AvgTickDurationMs = FStopwatch::Cycles64ToMilliseconds(TickDurationHistory.ComputeAverage());
 
 		DrawContext.DrawText
 		(
 			DbgX, DbgY,
-			FString::Printf(TEXT("U1 avg: %llu ms, last: %llu ms"),
-				AvgTimelineUpdateDurationMs, // time to update the Timeline cache
-				LastTimelineUpdateDurationMs), // last time to update the Timeline cache
-			SummaryFont, DbgTextColor
-		);
-		DbgY += DbgDY;
-
-		//////////////////////////////////////////////////
-		// Display the "time markers cache update" info.
-
-		const uint64 AvgTimeMarkerUpdateDurationMs = FStopwatch::Cycles64ToMilliseconds(TimeMarkerCacheUpdateDurationHistory.ComputeAverage());
-		const uint64 LastTimeMarkerUpdateDurationMs = FStopwatch::Cycles64ToMilliseconds(TimeMarkerCacheUpdateDurationHistory.GetValue(0));
-
-		DrawContext.DrawText
-		(
-			DbgX, DbgY,
-			FString::Printf(TEXT("U2 avg: %llu ms, last: %llu ms"),
-				AvgTimeMarkerUpdateDurationMs, // time to update the TimeMarker cache
-				LastTimeMarkerUpdateDurationMs), // last time to update the Timeline cache
+			FString::Printf(TEXT("U avg: %llu ms + %llu ms + %llu ms + %llu ms = %llu ms"),
+				AvgPreUpdateTracksDurationMs,
+				AvgUpdateTracksDurationMs,
+				AvgPostUpdateTracksDurationMs,
+				AvgTickDurationMs - AvgPreUpdateTracksDurationMs - AvgUpdateTracksDurationMs - AvgPostUpdateTracksDurationMs,
+				AvgTickDurationMs),
 			SummaryFont, DbgTextColor
 		);
 		DbgY += DbgDY;
@@ -958,36 +1100,42 @@ int32 STimingView::OnPaint(const FPaintArgs& Args, const FGeometry& AllottedGeom
 		//////////////////////////////////////////////////
 		// Display time markers stats.
 
-		DrawContext.DrawText
-		(
-			DbgX, DbgY,
-			FString::Format(TEXT("{0}{1} logs : {2} boxes, {3} texts"),
-			{
-				MarkersTrack.IsVisible() ? TEXT("") : TEXT("*"),
-				FText::AsNumber(MarkersTrack.GetNumLogMessages()).ToString(),
-				FText::AsNumber(MarkersTrack.GetNumBoxes()).ToString(),
-				FText::AsNumber(MarkersTrack.GetNumTexts()).ToString(),
-			}),
-			SummaryFont, DbgTextColor
-		);
-		DbgY += DbgDY;
+		if (MarkersTrack)
+		{
+			DrawContext.DrawText
+			(
+				DbgX, DbgY,
+				FString::Format(TEXT("{0}{1} logs : {2} boxes, {3} texts"),
+				{
+					MarkersTrack->IsVisible() ? TEXT("") : TEXT("*"),
+					FText::AsNumber(MarkersTrack->GetNumLogMessages()).ToString(),
+					FText::AsNumber(MarkersTrack->GetNumBoxes()).ToString(),
+					FText::AsNumber(MarkersTrack->GetNumTexts()).ToString(),
+				}),
+				SummaryFont, DbgTextColor
+			);
+			DbgY += DbgDY;
+		}
 
 		//////////////////////////////////////////////////
 		// Display Graph track stats.
 
-		DrawContext.DrawText
-		(
-			DbgX, DbgY,
-			FString::Format(TEXT("{0} events : {1} points, {2} lines, {3} boxes"),
-				{
-					FText::AsNumber(GraphTrack->GetNumAddedEvents()).ToString(),
-					FText::AsNumber(GraphTrack->GetNumDrawPoints()).ToString(),
-					FText::AsNumber(GraphTrack->GetNumDrawLines()).ToString(),
-					FText::AsNumber(GraphTrack->GetNumDrawBoxes()).ToString(),
-				}),
-				SummaryFont, DbgTextColor
-				);
-		DbgY += DbgDY;
+		if (GraphTrack)
+		{
+			DrawContext.DrawText
+			(
+				DbgX, DbgY,
+				FString::Format(TEXT("{0} events : {1} points, {2} lines, {3} boxes"),
+					{
+						FText::AsNumber(GraphTrack->GetNumAddedEvents()).ToString(),
+						FText::AsNumber(GraphTrack->GetNumDrawPoints()).ToString(),
+						FText::AsNumber(GraphTrack->GetNumDrawLines()).ToString(),
+						FText::AsNumber(GraphTrack->GetNumDrawBoxes()).ToString(),
+					}),
+					SummaryFont, DbgTextColor
+			);
+			DbgY += DbgDY;
+		}
 
 		//////////////////////////////////////////////////
 		// Display viewport's horizontal info.
@@ -1027,59 +1175,88 @@ int32 STimingView::OnPaint(const FPaintArgs& Args, const FGeometry& AllottedGeom
 		if (bIsPanning) InputStr += " Panning";
 		if (bIsSelecting) InputStr += " Selecting";
 		if (bIsDragging) InputStr += " Dragging";
+		if (bIsScrubbing) InputStr += " Scrubbing";
 		DrawContext.DrawText(DbgX, DbgY, InputStr, SummaryFont, DbgTextColor);
 		DbgY += DbgDY;
 	}
 
 	//////////////////////////////////////////////////
 
-	{
-		// Draw info about selected event (bottom-right corner).
-		if (SelectedTimingEvent.IsValid())
-		{
-			SelectedTimingEvent.Track->DrawSelectedEventInfo(SelectedTimingEvent, Viewport, DrawContext, WhiteBrush, MainFont);
-		}
+	Stopwatch.Stop();
+	TotalDrawDurationHistory.AddValue(Stopwatch.AccumulatedTime);
 
-		// Draw tooltip with info about hovered event.
-		Tooltip.Draw(DrawContext);
+	return SCompoundWidget::OnPaint(Args, AllottedGeometry, MyCullingRect, OutDrawElements, LayerId, InWidgetStyle, bParentEnabled && IsEnabled());
+}
 
-		if (GraphTrack->IsVisible())
-		{
-			GraphTrack->PostDraw(DrawContext, Viewport, MousePosition);
-		}
-	}
+////////////////////////////////////////////////////////////////////////////////////////////////////
 
-	return SCompoundWidget::OnPaint(Args, AllottedGeometry, MyCullingRect, OutDrawElements, LayerId, InWidgetStyle, bParentEnabled && IsEnabled());
+void STimingView::AddTopDockedTrack(TSharedPtr<FBaseTimingTrack> Track)
+{
+	check(Track.IsValid());
+	UE_LOG(TimingProfiler, Log, TEXT("New Top Docked Track (%d) : %s - %s (\"%s\")"),
+		TopDockedTracks.Num() + 1,
+		*Track->GetType().ToString(),
+		*Track->GetSubType().ToString(),
+		*Track->GetName());
+	check(!AllTracks.Contains(Track->GetId()));
+	AllTracks.Add(Track->GetId(), Track);
+	TopDockedTracks.Add(Track);
+	Algo::SortBy(TopDockedTracks, &FBaseTimingTrack::GetOrder);
 }
 
 ////////////////////////////////////////////////////////////////////////////////////////////////////
 
-void STimingView::AddTimingEventsTrack(FTimingEventsTrack* Track)
+void STimingView::AddBottomDockedTrack(TSharedPtr<FBaseTimingTrack> Track)
 {
-	UE_LOG(TimingProfiler, Log, TEXT("New Timing Events Track (%d) : %s - %s (\"%s\")"),
-		TimingEventsTracks.Num() + 1,
+	check(Track.IsValid());
+	UE_LOG(TimingProfiler, Log, TEXT("New Bottom Docked Track (%d) : %s - %s (\"%s\")"),
+		BottomDockedTracks.Num() + 1,
 		*Track->GetType().ToString(),
 		*Track->GetSubType().ToString(),
 		*Track->GetName());
+	check(!AllTracks.Contains(Track->GetId()));
+	AllTracks.Add(Track->GetId(), Track);
+	BottomDockedTracks.Add(Track);
+	Algo::SortBy(BottomDockedTracks, &FBaseTimingTrack::GetOrder);
+}
 
-	CachedTimelines.Add(Track->GetId(), Track);
-	TimingEventsTracks.Add(Track);
+////////////////////////////////////////////////////////////////////////////////////////////////////
 
-	bTimingEventsTracksChanged = true;
+void STimingView::AddScrollableTrack(TSharedPtr<FBaseTimingTrack> Track)
+{
+	check(Track.IsValid());
+	UE_LOG(TimingProfiler, Log, TEXT("New Scrollable Track (%d) : %s - %s (\"%s\")"),
+		ScrollableTracks.Num() + 1,
+		*Track->GetType().ToString(),
+		*Track->GetSubType().ToString(),
+		*Track->GetName());
+	check(!AllTracks.Contains(Track->GetId()));
+	AllTracks.Add(Track->GetId(), Track);
+	ScrollableTracks.Add(Track);
+	UpdateScrollableTracksOrder();
 }
 
 ////////////////////////////////////////////////////////////////////////////////////////////////////
 
-bool STimingView::IsGpuTrackVisible() const
+void STimingView::UpdateScrollableTracksOrder()
 {
-	return GpuTrack != nullptr && GpuTrack->IsVisible();
+	Algo::SortBy(ScrollableTracks, &FBaseTimingTrack::GetOrder);
 }
 
 ////////////////////////////////////////////////////////////////////////////////////////////////////
 
-bool STimingView::IsCpuTrackVisible(uint32 ThreadId) const
+void STimingView::AddForegroundTrack(TSharedPtr<FBaseTimingTrack> Track)
 {
-	return CpuTracks.Contains(ThreadId) && CpuTracks[ThreadId]->IsVisible();
+	check(Track.IsValid());
+	UE_LOG(TimingProfiler, Log, TEXT("New Foreground Track (%d) : %s - %s (\"%s\")"),
+		ForegroundTracks.Num() + 1,
+		*Track->GetType().ToString(),
+		*Track->GetSubType().ToString(),
+		*Track->GetName());
+	check(!AllTracks.Contains(Track->GetId()));
+	AllTracks.Add(Track->GetId(), Track);
+	ForegroundTracks.Add(Track);
+	Algo::SortBy(ForegroundTracks, &FBaseTimingTrack::GetOrder);
 }
 
 ////////////////////////////////////////////////////////////////////////////////////////////////////
@@ -1091,38 +1268,23 @@ FReply STimingView::OnMouseButtonDown(const FGeometry& MyGeometry, const FPointe
 	MousePositionOnButtonDown = MyGeometry.AbsoluteToLocal(MouseEvent.GetScreenSpacePosition());
 	MousePosition = MousePositionOnButtonDown;
 
+	bool bStartPanningSelectingOrScrubbing = false;
 	bool bStartPanning = false;
 	bool bStartSelecting = false;
 	bool bStartScrubbing = false;
 
 	if (MouseEvent.GetEffectingButton() == EKeys::LeftMouseButton)
 	{
-		if (MarkersTrack.IsVisible() && MarkersTrack.IsHeaderHovered())
+		if (MarkersTrack->IsVisible() && MarkersTrack->IsHeaderHovered())
 		{
-			MarkersTrack.ToggleCollapsed();
+			MarkersTrack->ToggleCollapsed();
 		}
 
 		if (!bIsRMB_Pressed)
-		{
-			bIsLMB_Pressed = true;
-
-			if(MousePositionOnButtonDown.Y < TimeRulerTrack.GetHeight() && (MouseEvent.GetModifierKeys().IsControlDown() || MouseEvent.GetModifierKeys().IsShiftDown()))
-			{
-				bStartScrubbing = true;
-			}
-			else if (!bIsSpaceBarKeyPressed &&
-				(MousePositionOnButtonDown.Y < TimeRulerTrack.GetHeight() ||
-				 (MouseEvent.GetModifierKeys().IsControlDown() && MouseEvent.GetModifierKeys().IsShiftDown())))
-			{
-				bStartSelecting = true;
-			}
-			else
-			{
-				bStartPanning = true;
-			}
-
-			// Capture mouse, so we can drag outside this widget.
-			Reply = FReply::Handled().CaptureMouse(SharedThis(this));
+		{
+			bIsLMB_Pressed = true;
+			bStartPanningSelectingOrScrubbing = true;
+			SelectHoveredTimingTrack();
 		}
 	}
 	else if (MouseEvent.GetEffectingButton() == EKeys::RightMouseButton)
@@ -1130,24 +1292,37 @@ FReply STimingView::OnMouseButtonDown(const FGeometry& MyGeometry, const FPointe
 		if (!bIsLMB_Pressed)
 		{
 			bIsRMB_Pressed = true;
+			bStartPanningSelectingOrScrubbing = true;
+			SelectHoveredTimingTrack();
+		}
+	}
 
-			if (!bIsSpaceBarKeyPressed &&
-				(MousePositionOnButtonDown.Y < TimeRulerTrack.GetHeight() ||
-				(MouseEvent.GetModifierKeys().IsControlDown() && MouseEvent.GetModifierKeys().IsShiftDown())))
-			{
-				bStartSelecting = true;
-			}
-			else
-			{
-				bStartPanning = true;
-			}
-
-			// Capture mouse, so we can drag outside this widget.
-			Reply = FReply::Handled().CaptureMouse(SharedThis(this));
+	if (bStartPanningSelectingOrScrubbing)
+	{
+		if (bIsSpaceBarKeyPressed)
+		{
+			bStartPanning = true;
+		}
+		else if (MousePositionOnButtonDown.Y < TimeRulerTrack->GetHeight() &&
+			(MouseEvent.GetModifierKeys().IsControlDown() || MouseEvent.GetModifierKeys().IsShiftDown()))
+		{
+			bStartScrubbing = true;
+		}
+		else if (MousePositionOnButtonDown.Y < TimeRulerTrack->GetHeight() ||
+			(MouseEvent.GetModifierKeys().IsControlDown() && MouseEvent.GetModifierKeys().IsShiftDown()))
+		{
+			bStartSelecting = true;
 		}
+		else
+		{
+			bStartPanning = true;
+		}
+
+		// Capture mouse, so we can drag outside this widget.
+		Reply = FReply::Handled().CaptureMouse(SharedThis(this));
 	}
 
-	if(bStartScrubbing)
+	if (bStartScrubbing)
 	{
 		bIsPanning = false;
 		bIsDragging = false;
@@ -1211,7 +1386,6 @@ FReply STimingView::OnMouseButtonUp(const FGeometry& MyGeometry, const FPointerE
 			if (bIsPanning)
 			{
 				PanningMode = EPanningMode::None;
-
 				bIsPanning = false;
 			}
 			else if (bIsSelecting)
@@ -1219,7 +1393,7 @@ FReply STimingView::OnMouseButtonUp(const FGeometry& MyGeometry, const FPointerE
 				RaiseSelectionChanged();
 				bIsSelecting = false;
 			}
-			else if(bIsScrubbing)
+			else if (bIsScrubbing)
 			{
 				RaiseTimeMarkerChanged();
 				bIsScrubbing = false;
@@ -1229,10 +1403,11 @@ FReply STimingView::OnMouseButtonUp(const FGeometry& MyGeometry, const FPointerE
 			{
 				// Select the hovered timing event (if any).
 				UpdateHoveredTimingEvent(MousePositionOnButtonUp.X, MousePositionOnButtonUp.Y);
+				SelectHoveredTimingTrack();
 				SelectHoveredTimingEvent();
 
 				// When clicking on an empty space...
-				if (!SelectedTimingEvent.IsValid())
+				if (!SelectedEvent.IsValid())
 				{
 					// ...reset selection.
 					SelectionEndTime = SelectionStartTime = 0.0;
@@ -1256,7 +1431,6 @@ FReply STimingView::OnMouseButtonUp(const FGeometry& MyGeometry, const FPointerE
 			if (bIsPanning)
 			{
 				PanningMode = EPanningMode::None;
-
 				bIsPanning = false;
 			}
 			else if (bIsSelecting)
@@ -1264,7 +1438,7 @@ FReply STimingView::OnMouseButtonUp(const FGeometry& MyGeometry, const FPointerE
 				RaiseSelectionChanged();
 				bIsSelecting = false;
 			}
-			else if(bIsScrubbing)
+			else if (bIsScrubbing)
 			{
 				RaiseTimeMarkerChanged();
 				bIsScrubbing = false;
@@ -1272,6 +1446,7 @@ FReply STimingView::OnMouseButtonUp(const FGeometry& MyGeometry, const FPointerE
 
 			if (bIsValidForMouseClick)
 			{
+				SelectHoveredTimingTrack();
 				ShowContextMenu(MouseEvent);
 			}
 
@@ -1295,14 +1470,21 @@ FReply STimingView::OnMouseButtonDoubleClick(const FGeometry& MyGeometry, const
 
 	if (MouseEvent.GetEffectingButton() == EKeys::LeftMouseButton)
 	{
-		if (MarkersTrack.IsVisible() && MarkersTrack.IsHovered())
+		if (MarkersTrack->IsVisible() && MarkersTrack->IsHovered())
 		{
-			MarkersTrack.ToggleCollapsed();
+			MarkersTrack->ToggleCollapsed();
 		}
-		else if (HoveredTimingEvent.IsValid())
+		else if (HoveredEvent.IsValid())
 		{
-			double EndTime = Viewport.RestrictEndTime(HoveredTimingEvent.EndTime);
-			SelectTimeInterval(HoveredTimingEvent.StartTime, EndTime - HoveredTimingEvent.StartTime);
+			if (MouseEvent.GetModifierKeys().IsControlDown())
+			{
+				const double EndTime = Viewport.RestrictEndTime(HoveredEvent->GetEndTime());
+				SelectTimeInterval(HoveredEvent->GetStartTime(), EndTime - HoveredEvent->GetStartTime());
+			}
+			else
+			{
+				//TODO: HighlightedEventTypeId = HoveredTimingEvent.TypeId;
+			}
 		}
 
 		Reply = FReply::Handled();
@@ -1370,16 +1552,6 @@ FReply STimingView::OnMouseMove(const FGeometry& MyGeometry, const FPointerEvent
 		}
 		else
 		{
-			if (MarkersTrack.IsVisible())
-			{
-				MarkersTrack.UpdateHoveredState(MousePosition.X, MousePosition.Y, Viewport);
-			}
-
-			if (GraphTrack->IsVisible())
-			{
-				GraphTrack->UpdateHoveredState(MousePosition.X, MousePosition.Y, Viewport);
-			}
-
 			UpdateHoveredTimingEvent(MousePosition.X, MousePosition.Y);
 		}
 
@@ -1411,17 +1583,16 @@ void STimingView::OnMouseLeave(const FPointerEvent& MouseEvent)
 
 		MousePosition = FVector2D::ZeroVector;
 
-		if (MarkersTrack.IsVisible())
+		if (HoveredTrack.IsValid())
 		{
-			MarkersTrack.UpdateHoveredState(MousePosition.X, MousePosition.Y, Viewport);
+			HoveredTrack.Reset();
+			OnHoveredTrackChangedDelegate.Broadcast(HoveredTrack);
 		}
-
-		if (GraphTrack->IsVisible())
+		if (HoveredEvent.IsValid())
 		{
-			GraphTrack->UpdateHoveredState(MousePosition.X, MousePosition.Y, Viewport);
+			HoveredEvent.Reset();
+			OnHoveredEventChangedDelegate.Broadcast(HoveredEvent);
 		}
-
-		HoveredTimingEvent.Reset();
 		Tooltip.SetDesiredOpacity(0.0f);
 	}
 }
@@ -1432,6 +1603,7 @@ FReply STimingView::OnMouseWheel(const FGeometry& MyGeometry, const FPointerEven
 {
 	if (MouseEvent.GetModifierKeys().IsShiftDown())
 	{
+		//MousePosition = MyGeometry.AbsoluteToLocal(MouseEvent.GetScreenSpacePosition());
 		if (GraphTrack->IsVisible() &&
 			MousePosition.Y >= GraphTrack->GetPosY() &&
 			MousePosition.Y < GraphTrack->GetPosY() + GraphTrack->GetHeight())
@@ -1477,12 +1649,11 @@ FReply STimingView::OnMouseWheel(const FGeometry& MyGeometry, const FPointerEven
 	{
 		// Zoom in/out horizontally.
 		const double Delta = MouseEvent.GetWheelDelta();
-		MousePosition = MyGeometry.AbsoluteToLocal(MouseEvent.GetScreenSpacePosition());
+		//MousePosition = MyGeometry.AbsoluteToLocal(MouseEvent.GetScreenSpacePosition());
 		if (Viewport.RelativeZoomWithFixedX(Delta, MousePosition.X))
 		{
 			//Viewport.EnforceHorizontalScrollLimits(1.0);
 			UpdateHorizontalScrollBar();
-			bIsViewportDirty = true;
 		}
 	}
 
@@ -1569,9 +1740,9 @@ FReply STimingView::OnKeyDown(const FGeometry& MyGeometry, const FKeyEvent& InKe
 	if (InKeyEvent.GetKey() == EKeys::B)
 	{
 		// Toggle Bookmarks.
-		if (MarkersTrack.IsVisible())
+		if (MarkersTrack->IsVisible())
 		{
-			if (!MarkersTrack.IsBookmarksTrack())
+			if (!MarkersTrack->IsBookmarksTrack())
 			{
 				SetDrawOnlyBookmarks(true);
 			}
@@ -1590,9 +1761,9 @@ FReply STimingView::OnKeyDown(const FGeometry& MyGeometry, const FKeyEvent& InKe
 	else if (InKeyEvent.GetKey() == EKeys::M)
 	{
 		// Toggle Time Markers.
-		if (MarkersTrack.IsVisible())
+		if (MarkersTrack->IsVisible())
 		{
-			if (MarkersTrack.IsBookmarksTrack())
+			if (MarkersTrack->IsBookmarksTrack())
 			{
 				SetDrawOnlyBookmarks(false);
 			}
@@ -1618,15 +1789,14 @@ FReply STimingView::OnKeyDown(const FGeometry& MyGeometry, const FKeyEvent& InKe
 	{
 		if (InKeyEvent.GetModifierKeys().IsControlDown())
 		{
-			if (SelectedTimingEvent.IsValid())
+			if (SelectedEvent.IsValid())
 			{
-				SelectedTimingEvent.Track->OnClipboardCopyEvent(SelectedTimingEvent);
+				SelectedEvent->GetTrack()->OnClipboardCopyEvent(*SelectedEvent);
 			}
 		}
 		else
 		{
-			Layout.bIsCompactMode = !Layout.bIsCompactMode;
-			bTimingEventsNeedUpdate = true;
+			Viewport.SwitchLayoutCompactMode();
 			return FReply::Handled();
 		}
 	}
@@ -1644,7 +1814,6 @@ FReply STimingView::OnKeyDown(const FGeometry& MyGeometry, const FKeyEvent& InKe
 		{
 			//Viewport.EnforceHorizontalScrollLimits(1.0);
 			UpdateHorizontalScrollBar();
-			bIsViewportDirty = true;
 		}
 		return FReply::Handled();
 	}
@@ -1657,7 +1826,6 @@ FReply STimingView::OnKeyDown(const FGeometry& MyGeometry, const FKeyEvent& InKe
 		{
 			//Viewport.EnforceHorizontalScrollLimits(1.0);
 			UpdateHorizontalScrollBar();
-			bIsViewportDirty = true;
 		}
 		return FReply::Handled();
 	}
@@ -1672,7 +1840,6 @@ FReply STimingView::OnKeyDown(const FGeometry& MyGeometry, const FKeyEvent& InKe
 			{
 				//Viewport.EnforceHorizontalScrollLimits(1.0);
 				UpdateHorizontalScrollBar();
-				bIsViewportDirty = true;
 			}
 		}
 		else
@@ -1692,7 +1859,6 @@ FReply STimingView::OnKeyDown(const FGeometry& MyGeometry, const FKeyEvent& InKe
 			{
 				//Viewport.EnforceHorizontalScrollLimits(1.0);
 				UpdateHorizontalScrollBar();
-				bIsViewportDirty = true;
 			}
 		}
 		else
@@ -1729,10 +1895,10 @@ FReply STimingView::OnKeyDown(const FGeometry& MyGeometry, const FKeyEvent& InKe
 	}
 	else if (InKeyEvent.GetKey() == EKeys::Enter)
 	{
-		if (SelectedTimingEvent.IsValid())
+		if (SelectedEvent.IsValid())
 		{
-			const double Duration = Viewport.RestrictDuration(SelectedTimingEvent.StartTime, SelectedTimingEvent.EndTime);
-			SelectTimeInterval(SelectedTimingEvent.StartTime, Duration);
+			const double Duration = Viewport.RestrictDuration(SelectedEvent->GetStartTime(), SelectedEvent->GetEndTime());
+			SelectTimeInterval(SelectedEvent->GetStartTime(), Duration);
 		}
 		return FReply::Handled();
 	}
@@ -1745,7 +1911,7 @@ FReply STimingView::OnKeyDown(const FGeometry& MyGeometry, const FKeyEvent& InKe
 	else if (InKeyEvent.GetKey() == EKeys::D) // debug: toggles down-sampling on/off
 	{
 		FTimingEventsTrack::bUseDownSampling = !FTimingEventsTrack::bUseDownSampling;
-		bTimingEventsNeedUpdate = true;
+		Viewport.AddDirtyFlags(ETimingTrackViewportDirtyFlags::HInvalidated);
 		return FReply::Handled();
 	}
 	else if (InKeyEvent.GetKey() == EKeys::G) // debug: toggles Graph track on/off
@@ -1755,22 +1921,22 @@ FReply STimingView::OnKeyDown(const FGeometry& MyGeometry, const FKeyEvent& InKe
 	}
 	else if (InKeyEvent.GetKey() == EKeys::Y) // debug: toggles GPU track on/off
 	{
-		ShowHideAllGpuTracks_Execute();
+		ThreadTimingSharedState->ShowHideAllGpuTracks();
 		return FReply::Handled();
 	}
 	else if (InKeyEvent.GetKey() == EKeys::U) // debug: toggles CPU tracks on/off
 	{
-		ShowHideAllCpuTracks_Execute();
+		ThreadTimingSharedState->ShowHideAllCpuTracks();
 		return FReply::Handled();
 	}
 	else if (InKeyEvent.GetKey() == EKeys::L)  // debug: toggles Loading tracks on/off
 	{
-		ShowHideAllLoadingTracks_Execute();
+		LoadingSharedState->ShowHideAllLoadingTracks();
 		return FReply::Handled();
 	}
 	else if (InKeyEvent.GetKey() == EKeys::I)  // debug: toggles IO tracks on/off
 	{
-		ShowHideAllIoTracks_Execute();
+		FileActivitySharedState->ShowHideAllIoTracks();
 		return FReply::Handled();
 	}
 	else if (InKeyEvent.GetKey() == EKeys::O)  // debug: toggles IO merge lanes algorithm
@@ -1844,9 +2010,9 @@ void STimingView::ShowContextMenu(const FPointerEvent& MouseEvent)
 			FSlateIcon()
 		);
 	}
-	else if(HoveredTimingEvent.Track != nullptr)
+	else if (HoveredTrack.IsValid())
 	{
-		const_cast<FTimingEventsTrack*>(HoveredTimingEvent.Track)->BuildContextMenu(MenuBuilder);
+		HoveredTrack->BuildContextMenu(MenuBuilder);
 	}
 	else
 	{
@@ -1884,32 +2050,29 @@ void STimingView::ShowContextMenu(const FPointerEvent& MouseEvent)
 
 void STimingView::BindCommands()
 {
-	FTimingViewCommands::Register();
-	const FTimingViewCommands& Commands = FTimingViewCommands::Get();
-
-	TSharedPtr<FUICommandList> CommandList = FTimingProfilerManager::Get()->GetCommandList();
+	//FTimingViewCommands::Register();
+	//const FTimingViewCommands& Commands = FTimingViewCommands::Get();
+	//
+	//TSharedPtr<FUICommandList> CommandList = FTimingProfilerManager::Get()->GetCommandList();
 
-	CommandList->MapAction(
-		Commands.ShowAllGpuTracks,
-		FExecuteAction::CreateSP(this, &STimingView::ShowHideAllGpuTracks_Execute),
-		FCanExecuteAction(), //FCanExecuteAction::CreateLambda([] { return true; }),
-		FIsActionChecked::CreateSP(this, &STimingView::ShowHideAllGpuTracks_IsChecked));
+	//CommandList->MapAction(
+	//	Commands.ShowAllGpuTracks,
+	//	FExecuteAction::CreateSP(this, &STimingView::ShowHideAllGpuTracks_Execute),
+	//	FCanExecuteAction(), //FCanExecuteAction::CreateLambda([] { return true; }),
+	//	FIsActionChecked::CreateSP(this, &STimingView::ShowHideAllGpuTracks_IsChecked));
 
-	CommandList->MapAction(
-		Commands.ShowAllCpuTracks,
-		FExecuteAction::CreateSP(this, &STimingView::ShowHideAllCpuTracks_Execute),
-		FCanExecuteAction(), //FCanExecuteAction::CreateLambda([] { return true; }),
-		FIsActionChecked::CreateSP(this, &STimingView::ShowHideAllCpuTracks_IsChecked));
+	//CommandList->MapAction(
+	//	Commands.ShowAllCpuTracks,
+	//	FExecuteAction::CreateSP(this, &STimingView::ShowHideAllCpuTracks_Execute),
+	//	FCanExecuteAction(), //FCanExecuteAction::CreateLambda([] { return true; }),
+	//	FIsActionChecked::CreateSP(this, &STimingView::ShowHideAllCpuTracks_IsChecked));
 }
 
 ////////////////////////////////////////////////////////////////////////////////////////////////////
 
 void STimingView::HorizontalScrollBar_OnUserScrolled(float ScrollOffset)
 {
-	if (Viewport.OnUserScrolled(HorizontalScrollBar, ScrollOffset))
-	{
-		bIsViewportDirty = true;
-	}
+	Viewport.OnUserScrolled(HorizontalScrollBar, ScrollOffset);
 }
 
 ////////////////////////////////////////////////////////////////////////////////////////////////////
@@ -1923,10 +2086,7 @@ void STimingView::UpdateHorizontalScrollBar()
 
 void STimingView::VerticalScrollBar_OnUserScrolled(float ScrollOffset)
 {
-	if (Viewport.OnUserScrolledY(VerticalScrollBar, ScrollOffset))
-	{
-		bIsVerticalViewportDirty = true;
-	}
+	Viewport.OnUserScrolledY(VerticalScrollBar, ScrollOffset);
 }
 
 ////////////////////////////////////////////////////////////////////////////////////////////////////
@@ -1943,9 +2103,7 @@ void STimingView::ScrollAtPosY(float ScrollPosY)
 	if (ScrollPosY != Viewport.GetScrollPosY())
 	{
 		Viewport.SetScrollPosY(ScrollPosY);
-
 		UpdateVerticalScrollBar();
-		bIsVerticalViewportDirty = true;
 	}
 }
 
@@ -1956,7 +2114,6 @@ void STimingView::ScrollAtTime(double StartTime)
 	if (Viewport.ScrollAtTime(StartTime))
 	{
 		UpdateHorizontalScrollBar();
-		bIsViewportDirty = true;
 	}
 }
 
@@ -1967,9 +2124,7 @@ void STimingView::CenterOnTimeInterval(double IntervalStartTime, double Interval
 	if (Viewport.CenterOnTimeInterval(IntervalStartTime, IntervalDuration))
 	{
 		Viewport.EnforceHorizontalScrollLimits(1.0);
-
 		UpdateHorizontalScrollBar();
-		bIsViewportDirty = true;
 	}
 }
 
@@ -2095,22 +2250,22 @@ void STimingView::SelectToTimeMarker(double Time)
 
 void STimingView::SetTimeMarkersVisible(bool bIsMarkersTrackVisible)
 {
-	if (MarkersTrack.IsVisible() != bIsMarkersTrackVisible)
+	if (MarkersTrack->IsVisible() != bIsMarkersTrackVisible)
 	{
-		MarkersTrack.SetVisibilityFlag(bIsMarkersTrackVisible);
+		MarkersTrack->SetVisibilityFlag(bIsMarkersTrackVisible);
 
-		if (MarkersTrack.IsVisible())
+		if (MarkersTrack->IsVisible())
 		{
 			if (Viewport.GetScrollPosY() != 0.0f)
 			{
-				Viewport.SetScrollPosY(Viewport.GetScrollPosY() + MarkersTrack.GetHeight());
+				Viewport.SetScrollPosY(Viewport.GetScrollPosY() + MarkersTrack->GetHeight());
 			}
 
-			MarkersTrack.SetDirtyFlag();
+			MarkersTrack->SetDirtyFlag();
 		}
 		else
 		{
-			Viewport.SetScrollPosY(Viewport.GetScrollPosY() - MarkersTrack.GetHeight());
+			Viewport.SetScrollPosY(Viewport.GetScrollPosY() - MarkersTrack->GetHeight());
 		}
 	}
 }
@@ -2119,90 +2274,125 @@ void STimingView::SetTimeMarkersVisible(bool bIsMarkersTrackVisible)
 
 void STimingView::SetDrawOnlyBookmarks(bool bIsBookmarksTrack)
 {
-	if (MarkersTrack.IsBookmarksTrack() != bIsBookmarksTrack)
+	if (MarkersTrack->IsBookmarksTrack() != bIsBookmarksTrack)
 	{
-		const float PrevHeight = MarkersTrack.GetHeight();
-		MarkersTrack.SetBookmarksTrackFlag(bIsBookmarksTrack);
+		const float PrevHeight = MarkersTrack->GetHeight();
+		MarkersTrack->SetBookmarksTrackFlag(bIsBookmarksTrack);
 
-		if (MarkersTrack.IsVisible())
+		if (MarkersTrack->IsVisible())
 		{
 			if (Viewport.GetScrollPosY() != 0.0f)
 			{
-				Viewport.SetScrollPosY(Viewport.GetScrollPosY() + MarkersTrack.GetHeight() - PrevHeight);
+				Viewport.SetScrollPosY(Viewport.GetScrollPosY() + MarkersTrack->GetHeight() - PrevHeight);
 			}
 
-			MarkersTrack.SetDirtyFlag();
+			MarkersTrack->SetDirtyFlag();
 		}
 	}
 }
 
 ////////////////////////////////////////////////////////////////////////////////////////////////////
 
-void STimingView::UpdateHoveredTimingEvent(float MX, float MY)
+const TSharedPtr<FBaseTimingTrack> STimingView::GetTrackAt(float InPosX, float InPosY) const
 {
-	FTimingEvent TimingEvent;
+	TSharedPtr<FBaseTimingTrack> FoundTrack;
 
-	if (MY >= Viewport.GetTopOffset() && MY < Viewport.GetHeight())
+	if (InPosY < 0.0f)
+	{
+		// above viewport
+	}
+	else if (InPosY < Viewport.GetTopOffset())
+	{
+		// Top Docked Tracks
+		for (const TSharedPtr<FBaseTimingTrack>& TrackPtr : TopDockedTracks)
+		{
+			const FBaseTimingTrack& Track = *TrackPtr;
+			if (TrackPtr->IsVisible())
+			{
+				if (InPosY >= Track.GetPosY() && InPosY < Track.GetPosY() + Track.GetHeight())
+				{
+					FoundTrack = TrackPtr;
+					break;
+				}
+			}
+		}
+	}
+	else if (InPosY < Viewport.GetHeight() - Viewport.GetBottomOffset())
 	{
-		for (const auto& KV : CachedTimelines)
+		// Scrollable Tracks
+		for (const TSharedPtr<FBaseTimingTrack>& TrackPtr : ScrollableTracks)
 		{
-			const FTimingEventsTrack& Track = *KV.Value;
+			const FBaseTimingTrack& Track = *TrackPtr;
 			if (Track.IsVisible())
 			{
-				const float Y = Viewport.GetTopOffset() + Track.GetPosY() - Viewport.GetScrollPosY();
-				if (MY >= Y && MY < Y + Track.GetHeight())
+				if (InPosY >= Track.GetPosY() && InPosY < Track.GetPosY() + Track.GetHeight())
 				{
-					TimingEvent.Track = &Track;
+					FoundTrack = TrackPtr;
 					break;
 				}
 			}
 		}
-
-		if (TimingEvent.Track)
+	}
+	else if (InPosY < Viewport.GetHeight())
+	{
+		// Bottom Docked Tracks
+		for (const TSharedPtr<FBaseTimingTrack>& TrackPtr : BottomDockedTracks)
 		{
-			const float Y0 = Viewport.GetTopOffset() + TimingEvent.Track->GetPosY() - Viewport.GetScrollPosY() + 1.0f + Layout.TimelineDY;
-
-			// If mouse is not above first sub-track or below last sub-track...
-			if (MY >= Y0 && MY < Y0 + TimingEvent.Track->GetHeight() + Layout.TimelineDY)
+			const FBaseTimingTrack& Track = *TrackPtr;
+			if (TrackPtr->IsVisible())
 			{
-				int32 Depth = (MY - Y0) / (Layout.EventDY + Layout.EventH);
-				float EventMY = (MY - Y0) - Depth * (Layout.EventDY + Layout.EventH);
-
-				const double StartTime = Viewport.SlateUnitsToTime(MX);
-				const double EndTime = StartTime + 2.0 / Viewport.GetScaleX(); // +2px
-
-				auto EventFilter = [Depth](double, double, uint32 EventDepth)
+				if (InPosY >= Track.GetPosY() && InPosY < Track.GetPosY() + Track.GetHeight())
 				{
-					return EventDepth == Depth;
-				};
-
-				TimingEvent.Track->SearchTimingEvent(FTimingEventSearchParameters(StartTime, EndTime, ETimingEventSearchFlags::StopAtFirstMatch, EventFilter), TimingEvent);
+					FoundTrack = TrackPtr;
+					break;
+				}
 			}
 		}
 	}
+	else
+	{
+		// below viewport
+	}
+
+	return FoundTrack;
+}
+
+////////////////////////////////////////////////////////////////////////////////////////////////////
+
+void STimingView::UpdateHoveredTimingEvent(float InMousePosX, float InMousePosY)
+{
+	TSharedPtr<FBaseTimingTrack> NewHoveredTrack = GetTrackAt(InMousePosX, InMousePosY);
+	if (NewHoveredTrack != HoveredTrack)
+	{
+		HoveredTrack = NewHoveredTrack;
+		OnHoveredTrackChangedDelegate.Broadcast(HoveredTrack);
+	}
 
-	if (TimingEvent.IsValid())
+	TSharedPtr<const ITimingEvent> NewHoveredEvent;
+	if (HoveredTrack.IsValid())
 	{
-		if (!TimingEvent.Equals(HoveredTimingEvent))
-		{
-			HoveredTimingEvent = TimingEvent;
-			HoveredTimingEvent.Track->ComputeTimingEventStats(HoveredTimingEvent);
-			HoveredTimingEvent.Track->InitTooltip(Tooltip, HoveredTimingEvent);
+		NewHoveredEvent = HoveredTrack->GetEvent(InMousePosX, InMousePosY, Viewport);
+	}
 
-			OnHoveredEventChangedDelegate.Broadcast(HoveredTimingEvent);
+	if (NewHoveredEvent.IsValid())
+	{
+		if (!HoveredEvent.IsValid() || !NewHoveredEvent->Equals(*HoveredEvent))
+		{
+			HoveredEvent = NewHoveredEvent;
+			ensure(HoveredTrack == HoveredEvent->GetTrack());
+			HoveredTrack->UpdateEventStats(const_cast<ITimingEvent&>(*HoveredEvent));
+			HoveredTrack->InitTooltip(Tooltip, *HoveredEvent);
+			OnHoveredEventChangedDelegate.Broadcast(HoveredEvent);
 		}
 		Tooltip.SetDesiredOpacity(1.0f);
 	}
 	else
 	{
-		HoveredTimingEvent.Reset();
-
-		if(TimingEvent.Track != nullptr)
+		if (HoveredEvent.IsValid())
 		{
-			HoveredTimingEvent.Track = TimingEvent.Track;
+			HoveredEvent.Reset();
+			OnHoveredEventChangedDelegate.Broadcast(HoveredEvent);
 		}
-
-		OnHoveredEventChangedDelegate.Broadcast(HoveredTimingEvent);
 		Tooltip.SetDesiredOpacity(0.0f);
 	}
 }
@@ -2213,48 +2403,73 @@ void STimingView::OnSelectedTimingEventChanged()
 {
 	if (!bAssetLoadingMode)
 	{
-		if (SelectedTimingEvent.IsValid())
+		if (SelectedEvent.IsValid())
 		{
-			SelectedTimingEvent.Track->OnEventSelected(SelectedTimingEvent);
+			SelectedEvent->GetTrack()->UpdateEventStats(const_cast<ITimingEvent&>(*SelectedEvent));
+			SelectedEvent->GetTrack()->OnEventSelected(*SelectedEvent);
 		}
 	}
 
-	OnSelectedEventChangedDelegate.Broadcast(SelectedTimingEvent);
+	if (!SelectedEvent.IsValid())
+	{
+		HighlightedEventTypeId = uint64(-1);
+	}
+
+	OnSelectedEventChangedDelegate.Broadcast(SelectedEvent);
+}
+
+////////////////////////////////////////////////////////////////////////////////////////////////////
+
+void STimingView::SelectHoveredTimingTrack()
+{
+	if (SelectedTrack != HoveredTrack)
+	{
+		SelectedTrack = HoveredTrack;
+		OnSelectedTrackChangedDelegate.Broadcast(SelectedTrack);
+	}
 }
 
 ////////////////////////////////////////////////////////////////////////////////////////////////////
 
 void STimingView::SelectHoveredTimingEvent()
 {
-	SelectedTimingEvent = HoveredTimingEvent;
-	if (SelectedTimingEvent.IsValid())
+	if (SelectedEvent != HoveredEvent)
 	{
-		LastSelectionType = ESelectionType::TimingEvent;
-		BringIntoView(SelectedTimingEvent.StartTime, SelectedTimingEvent.EndTime);
+		SelectedEvent = HoveredEvent;
+
+		if (SelectedEvent.IsValid())
+		{
+			LastSelectionType = ESelectionType::TimingEvent;
+			BringIntoView(SelectedEvent->GetStartTime(), SelectedEvent->GetEndTime());
+		}
+
+		OnSelectedTimingEventChanged();
 	}
-	OnSelectedTimingEventChanged();
 }
 
 ////////////////////////////////////////////////////////////////////////////////////////////////////
 
 void STimingView::SelectLeftTimingEvent()
 {
-	if (SelectedTimingEvent.IsValid())
+	if (SelectedEvent.IsValid())
 	{
-		const uint32 Depth = SelectedTimingEvent.Depth;
-		const double StartTime = SelectedTimingEvent.StartTime;
-		const double EndTime = SelectedTimingEvent.EndTime;
+		const uint32 Depth = SelectedEvent->GetDepth();
+		const double StartTime = SelectedEvent->GetStartTime();
+		const double EndTime = SelectedEvent->GetEndTime();
 
 		auto EventFilter = [Depth, StartTime, EndTime](double EventStartTime, double EventEndTime, uint32 EventDepth)
 		{
-			return EventDepth == Depth &&
-				(EventStartTime < StartTime || EventEndTime < EndTime);
+			return EventDepth == Depth
+				&& (EventStartTime < StartTime || EventEndTime < EndTime);
 		};
 
-		if (SelectedTimingEvent.Track->SearchTimingEvent(FTimingEventSearchParameters(0.0, StartTime, ETimingEventSearchFlags::SearchAll, EventFilter), SelectedTimingEvent))
+		const TSharedPtr<const ITimingEvent> LeftEvent = SelectedEvent->GetTrack()->SearchEvent(
+			FTimingEventSearchParameters(0.0, StartTime, ETimingEventSearchFlags::SearchAll, EventFilter));
+
+		if (LeftEvent.IsValid())
 		{
-			SelectedTimingEvent.Track->ComputeTimingEventStats(SelectedTimingEvent);
-			BringIntoView(SelectedTimingEvent.StartTime, SelectedTimingEvent.EndTime);
+			SelectedEvent = LeftEvent;
+			BringIntoView(SelectedEvent->GetStartTime(), SelectedEvent->GetEndTime());
 			OnSelectedTimingEventChanged();
 		}
 	}
@@ -2264,22 +2479,25 @@ void STimingView::SelectLeftTimingEvent()
 
 void STimingView::SelectRightTimingEvent()
 {
-	if (SelectedTimingEvent.IsValid())
+	if (SelectedEvent.IsValid())
 	{
-		const uint32 Depth = SelectedTimingEvent.Depth;
-		const double StartTime = SelectedTimingEvent.StartTime;
-		const double EndTime = SelectedTimingEvent.EndTime;
+		const uint32 Depth = SelectedEvent->GetDepth();
+		const double StartTime = SelectedEvent->GetStartTime();
+		const double EndTime = SelectedEvent->GetEndTime();
 
 		auto EventFilter = [Depth, StartTime, EndTime](double EventStartTime, double EventEndTime, uint32 EventDepth)
 		{
-			return EventDepth == Depth &&
-				(EventStartTime > StartTime || EventEndTime > EndTime);
+			return EventDepth == Depth
+				&& (EventStartTime > StartTime || EventEndTime > EndTime);
 		};
 
-		if (SelectedTimingEvent.Track->SearchTimingEvent(FTimingEventSearchParameters(EndTime, Viewport.GetMaxValidTime(), ETimingEventSearchFlags::StopAtFirstMatch, EventFilter), SelectedTimingEvent))
+		const TSharedPtr<const ITimingEvent> RightEvent = SelectedEvent->GetTrack()->SearchEvent(
+			FTimingEventSearchParameters(EndTime, Viewport.GetMaxValidTime(), ETimingEventSearchFlags::StopAtFirstMatch, EventFilter));
+
+		if (RightEvent.IsValid())
 		{
-			SelectedTimingEvent.Track->ComputeTimingEventStats(SelectedTimingEvent);
-			BringIntoView(SelectedTimingEvent.StartTime, SelectedTimingEvent.EndTime);
+			SelectedEvent = RightEvent;
+			BringIntoView(SelectedEvent->GetStartTime(), SelectedEvent->GetEndTime());
 			OnSelectedTimingEventChanged();
 		}
 	}
@@ -2289,12 +2507,12 @@ void STimingView::SelectRightTimingEvent()
 
 void STimingView::SelectUpTimingEvent()
 {
-	if (SelectedTimingEvent.IsValid() &&
-		SelectedTimingEvent.Depth > 0)
+	if (SelectedEvent.IsValid() &&
+		SelectedEvent->GetDepth() > 0)
 	{
-		const uint32 Depth = SelectedTimingEvent.Depth - 1;
-		const double StartTime = SelectedTimingEvent.StartTime;
-		const double EndTime = SelectedTimingEvent.EndTime;
+		const uint32 Depth = SelectedEvent->GetDepth() - 1;
+		const double StartTime = SelectedEvent->GetStartTime();
+		const double EndTime = SelectedEvent->GetEndTime();
 
 		auto EventFilter = [Depth, StartTime, EndTime](double EventStartTime, double EventEndTime, uint32 EventDepth)
 		{
@@ -2303,10 +2521,13 @@ void STimingView::SelectUpTimingEvent()
 				&& EventEndTime >= StartTime;
 		};
 
-		if (SelectedTimingEvent.Track->SearchTimingEvent(FTimingEventSearchParameters(StartTime, EndTime, ETimingEventSearchFlags::StopAtFirstMatch, EventFilter), SelectedTimingEvent))
+		const TSharedPtr<const ITimingEvent> UpEvent = SelectedEvent->GetTrack()->SearchEvent(
+			FTimingEventSearchParameters(StartTime, EndTime, ETimingEventSearchFlags::StopAtFirstMatch, EventFilter));
+
+		if (UpEvent.IsValid())
 		{
-			SelectedTimingEvent.Track->ComputeTimingEventStats(SelectedTimingEvent);
-			BringIntoView(SelectedTimingEvent.StartTime, SelectedTimingEvent.EndTime);
+			SelectedEvent = UpEvent;
+			BringIntoView(SelectedEvent->GetStartTime(), SelectedEvent->GetEndTime());
 			OnSelectedTimingEventChanged();
 		}
 	}
@@ -2316,17 +2537,16 @@ void STimingView::SelectUpTimingEvent()
 
 void STimingView::SelectDownTimingEvent()
 {
-	if (SelectedTimingEvent.IsValid())
+	if (SelectedEvent.IsValid())
 	{
-		const uint32 Depth = SelectedTimingEvent.Depth + 1;
-		const double StartTime = SelectedTimingEvent.StartTime;
-		const double EndTime = SelectedTimingEvent.EndTime;
+		const uint32 Depth = SelectedEvent->GetDepth() + 1;
+		const double StartTime = SelectedEvent->GetStartTime();
+		const double EndTime = SelectedEvent->GetEndTime();
 		double LargestDuration = 0.0;
 
 		auto EventFilter = [Depth, StartTime, EndTime, &LargestDuration](double EventStartTime, double EventEndTime, uint32 EventDepth)
 		{
 			const double Duration = EventEndTime - EventStartTime;
-
 			return Duration > LargestDuration
 				&& EventDepth == Depth
 				&& EventStartTime <= EndTime
@@ -2338,11 +2558,14 @@ void STimingView::SelectDownTimingEvent()
 			const double Duration = EventEndTime - EventStartTime;
 			LargestDuration = Duration;
 		};
-		
-		if (SelectedTimingEvent.Track->SearchTimingEvent(FTimingEventSearchParameters(StartTime, EndTime, ETimingEventSearchFlags::SearchAll, EventFilter, EventMatched), SelectedTimingEvent))
+
+		const TSharedPtr<const ITimingEvent> DownEvent = SelectedEvent->GetTrack()->SearchEvent(
+			FTimingEventSearchParameters(StartTime, EndTime, ETimingEventSearchFlags::SearchAll, EventFilter, EventMatched));
+
+		if (DownEvent.IsValid())
 		{
-			SelectedTimingEvent.Track->ComputeTimingEventStats(SelectedTimingEvent);
-			BringIntoView(SelectedTimingEvent.StartTime, SelectedTimingEvent.EndTime);
+			SelectedEvent = DownEvent;
+			BringIntoView(SelectedEvent->GetStartTime(), SelectedEvent->GetEndTime());
 			OnSelectedTimingEventChanged();
 		}
 	}
@@ -2359,7 +2582,7 @@ void STimingView::FrameSelection()
 	if (LastSelectionType == ESelectionType::TimingEvent)
 	{
 		// Try framing the selected timing event.
-		if (SelectedTimingEvent.IsValid())
+		if (SelectedEvent.IsValid())
 		{
 			Type = ESelectionType::TimingEvent;
 		}
@@ -2383,7 +2606,7 @@ void STimingView::FrameSelection()
 	if (LastSelectionType == ESelectionType::None || Type == ESelectionType::None)
 	{
 		// First, try framing the selected timing event...
-		if (SelectedTimingEvent.IsValid())
+		if (SelectedEvent.IsValid())
 		{
 			Type = ESelectionType::TimingEvent;
 		}
@@ -2396,8 +2619,8 @@ void STimingView::FrameSelection()
 	if (Type == ESelectionType::TimingEvent)
 	{
 		// Frame the selected event.
-		StartTime = SelectedTimingEvent.StartTime;
-		EndTime = Viewport.RestrictEndTime(SelectedTimingEvent.EndTime);
+		StartTime = SelectedEvent->GetStartTime();
+		EndTime = Viewport.RestrictEndTime(SelectedEvent->GetEndTime());
 		if (EndTime == StartTime)
 		{
 			EndTime += 1.0 / Viewport.GetScaleX(); // +1px
@@ -2416,7 +2639,6 @@ void STimingView::FrameSelection()
 		if (Viewport.ZoomOnTimeInterval(StartTime - Duration * 0.1, Duration * 1.2))
 		{
 			UpdateHorizontalScrollBar();
-			bIsViewportDirty = true;
 		}
 	}
 }
@@ -2429,66 +2651,21 @@ TSharedRef<SWidget> STimingView::MakeTracksFilterMenu()
 
 	MenuBuilder.BeginSection("QuickFilter", LOCTEXT("TracksFilterHeading", "Quick Filter"));
 	{
-		const FTimingViewCommands& Commands = FTimingViewCommands::Get();
+		//const FTimingViewCommands& Commands = FTimingViewCommands::Get();
 
+		//TODO: MenuBuilder.AddMenuEntry(Commands.AutoHideGraphTrack);
 		MenuBuilder.AddMenuEntry(
 			LOCTEXT("ShowGraphTrack", "Graph Track - G"),
 			LOCTEXT("ShowGraphTrack_Tooltip", "Show/hide the Graph track"),
 			FSlateIcon(),
 			FUIAction(FExecuteAction::CreateSP(this, &STimingView::ShowHideGraphTrack_Execute),
-				FCanExecuteAction(),
-				FIsActionChecked::CreateSP(this, &STimingView::ShowHideGraphTrack_IsChecked)),
-			NAME_None,
-			EUserInterfaceActionType::ToggleButton
-		);
-
-		//TODO: MenuBuilder.AddMenuEntry(Commands.ShowAllGpuTracks);
-		MenuBuilder.AddMenuEntry(
-			LOCTEXT("ShowAllGpuTracks", "GPU Track - Y"),
-			LOCTEXT("ShowAllGpuTracks_Tooltip", "Show/hide the GPU track"),
-			FSlateIcon(),
-			FUIAction(FExecuteAction::CreateSP(this, &STimingView::ShowHideAllGpuTracks_Execute),
-					  FCanExecuteAction(),
-					  FIsActionChecked::CreateSP(this, &STimingView::ShowHideAllGpuTracks_IsChecked)),
-			NAME_None,
-			EUserInterfaceActionType::ToggleButton
-		);
-
-		//TODO: MenuBuilder.AddMenuEntry(Commands.ShowAllCpuTracks);
-		MenuBuilder.AddMenuEntry(
-			LOCTEXT("ShowAllCpuTracks", "CPU Thread Tracks - U"),
-			LOCTEXT("ShowAllCpuTracks_Tooltip", "Show/hide all CPU tracks (and all CPU thread groups)"),
-			FSlateIcon(),
-			FUIAction(FExecuteAction::CreateSP(this, &STimingView::ShowHideAllCpuTracks_Execute),
-					  FCanExecuteAction(),
-					  FIsActionChecked::CreateSP(this, &STimingView::ShowHideAllCpuTracks_IsChecked)),
-			NAME_None,
-			EUserInterfaceActionType::ToggleButton
-		);
-
-		//TODO: MenuBuilder.AddMenuEntry(Commands.ShowAllLoadingTracks);
-		MenuBuilder.AddMenuEntry(
-			LOCTEXT("ShowAllLoadingTracks", "Asset Loading Tracks - L"),
-			LOCTEXT("ShowAllLoadingTracks_Tooltip", "Show/hide the Asset Loading tracks"),
-			FSlateIcon(),
-			FUIAction(FExecuteAction::CreateSP(this, &STimingView::ShowHideAllLoadingTracks_Execute),
 					  FCanExecuteAction(),
-					  FIsActionChecked::CreateSP(this, &STimingView::ShowHideAllLoadingTracks_IsChecked)),
+					  FIsActionChecked::CreateSP(this, &STimingView::ShowHideGraphTrack_IsChecked)),
 			NAME_None,
 			EUserInterfaceActionType::ToggleButton
 		);
 
-		//TODO: MenuBuilder.AddMenuEntry(Commands.ShowAllIoTracks);
-		MenuBuilder.AddMenuEntry(
-			LOCTEXT("ShowAllIoTracks", "I/O Tracks - I"),
-			LOCTEXT("ShowAllIoTracks_Tooltip", "Show/hide the I/O (File Activity) tracks"),
-			FSlateIcon(),
-			FUIAction(FExecuteAction::CreateSP(this, &STimingView::ShowHideAllIoTracks_Execute),
-					  FCanExecuteAction(),
-					  FIsActionChecked::CreateSP(this, &STimingView::ShowHideAllIoTracks_IsChecked)),
-			NAME_None,
-			EUserInterfaceActionType::ToggleButton
-		);
+		MenuBuilder.AddMenuSeparator("QuickFilterSeparator");
 
 		//TODO: MenuBuilder.AddMenuEntry(Commands.AutoHideEmptyTracks);
 		MenuBuilder.AddMenuEntry(
@@ -2501,75 +2678,38 @@ TSharedRef<SWidget> STimingView::MakeTracksFilterMenu()
 			NAME_None,
 			EUserInterfaceActionType::ToggleButton
 		);
-
-		// Let any plugins extend the filter menu
-		for (Insights::ITimingViewExtender* Extender : GetExtenders())
-		{
-			Extender->ExtendFilterMenu(*this, MenuBuilder);
-		}
 	}
 	MenuBuilder.EndSection();
 
-	MenuBuilder.BeginSection("ThreadGroups", LOCTEXT("ThreadGroupsHeading", "CPU Thread Groups"));
-	CreateThreadGroupsMenu(MenuBuilder);
-	MenuBuilder.EndSection();
-
 	//MenuBuilder.BeginSection("Tracks", LOCTEXT("TracksHeading", "Tracks"));
 	//CreateTracksMenu(MenuBuilder);
 	//MenuBuilder.EndSection();
 
-	return MenuBuilder.MakeWidget();
-}
-
-////////////////////////////////////////////////////////////////////////////////////////////////////
-
-void STimingView::CreateThreadGroupsMenu(FMenuBuilder& MenuBuilder)
-{
-	// Sort the list of thread groups.
-	TArray<const FThreadGroup*> SortedThreadGroups;
-	SortedThreadGroups.Reserve(ThreadGroups.Num());
-	for (const auto& KV : ThreadGroups)
+	// Let any plugin extend the filter menu.
+	for (Insights::ITimingViewExtender* Extender : GetExtenders())
 	{
-		SortedThreadGroups.Add(&KV.Value);
+		Extender->ExtendFilterMenu(*this, MenuBuilder);
 	}
-	Algo::SortBy(SortedThreadGroups, &FThreadGroup::GetOrder);
 
-	for (const FThreadGroup* ThreadGroupPtr : SortedThreadGroups)
-	{
-		const FThreadGroup& ThreadGroup = *ThreadGroupPtr;
-		if (ThreadGroup.NumTimelines > 0)
-		{
-			MenuBuilder.AddMenuEntry(
-				//FText::FromString(ThreadGroup.Name),
-				FText::Format(LOCTEXT("ThreadGroupFmt", "{0} ({1})"), FText::FromString(ThreadGroup.Name), ThreadGroup.NumTimelines),
-				TAttribute<FText>(), // no tooltip
-				FSlateIcon(),
-				FUIAction(FExecuteAction::CreateSP(this, &STimingView::ToggleTrackVisibilityByGroup_Execute, ThreadGroup.Name),
-					FCanExecuteAction::CreateLambda([] { return true; }),
-					FIsActionChecked::CreateSP(this, &STimingView::ToggleTrackVisibilityByGroup_IsChecked, ThreadGroup.Name)),
-				NAME_None,
-				EUserInterfaceActionType::ToggleButton
-			);
-		}
-	}
+	return MenuBuilder.MakeWidget();
 }
 
 ////////////////////////////////////////////////////////////////////////////////////////////////////
 
 void STimingView::CreateTracksMenu(FMenuBuilder& MenuBuilder)
 {
-	for (int32 TrackIndex = 0; TrackIndex < TimingEventsTracks.Num(); ++TrackIndex)
+	for (const TSharedPtr<FBaseTimingTrack>& TrackPtr : ScrollableTracks)
 	{
-		const FTimingEventsTrack& Track = *TimingEventsTracks[TrackIndex];
-		if (Track.GetHeight() > 0.0f || Layout.TargetMinTimelineH > 0.0f)
+		const FBaseTimingTrack& Track = *TrackPtr;
+		if (Track.GetHeight() > 0.0f || Viewport.GetLayout().TargetMinTimelineH > 0.0f)
 		{
 			MenuBuilder.AddMenuEntry(
 				FText::FromString(Track.GetName()),
 				TAttribute<FText>(), // no tooltip
 				FSlateIcon(),
 				FUIAction(FExecuteAction::CreateSP(this, &STimingView::ToggleTrackVisibility_Execute, Track.GetId()),
-					FCanExecuteAction::CreateLambda([] { return true; }),
-					FIsActionChecked::CreateSP(this, &STimingView::ToggleTrackVisibility_IsChecked, Track.GetId())),
+						  FCanExecuteAction::CreateLambda([] { return true; }),
+						  FIsActionChecked::CreateSP(this, &STimingView::ToggleTrackVisibility_IsChecked, Track.GetId())),
 				NAME_None,
 				EUserInterfaceActionType::ToggleButton
 			);
@@ -2589,144 +2729,24 @@ bool STimingView::ShowHideGraphTrack_IsChecked() const
 void STimingView::ShowHideGraphTrack_Execute()
 {
 	GraphTrack->ToggleVisibility();
-	bTimingEventsNeedUpdate = true;
-}
-
-////////////////////////////////////////////////////////////////////////////////////////////////////
-
-bool STimingView::ShowHideAllCpuTracks_IsChecked() const
-{
-	return bShowHideAllCpuTracks;
-}
-
-////////////////////////////////////////////////////////////////////////////////////////////////////
-
-void STimingView::ShowHideAllCpuTracks_Execute()
-{
-	bShowHideAllCpuTracks = !bShowHideAllCpuTracks;
-
-	for (int32 TrackIndex = 0; TrackIndex < TimingEventsTracks.Num(); ++TrackIndex)
-	{
-		FTimingEventsTrack& Track = *TimingEventsTracks[TrackIndex];
-		if (Track.GetType() == FName(TEXT("Thread")) && Track.GetSubType() == FName(TEXT("CPU")))
-		{
-			Track.SetVisibilityFlag(bShowHideAllCpuTracks);
-		}
-	}
-
-	for (auto& KV: ThreadGroups)
-	{
-		KV.Value.bIsVisible = bShowHideAllCpuTracks;
-	}
-
-	OnTimingEventsTrackVisibilityChanged();
-}
-
-////////////////////////////////////////////////////////////////////////////////////////////////////
-
-bool STimingView::ShowHideAllGpuTracks_IsChecked() const
-{
-	return bShowHideAllGpuTracks;
-}
-
-////////////////////////////////////////////////////////////////////////////////////////////////////
-
-void STimingView::ShowHideAllGpuTracks_Execute()
-{
-	bShowHideAllGpuTracks = !bShowHideAllGpuTracks;
-
-	for (int32 TrackIndex = 0; TrackIndex < TimingEventsTracks.Num(); ++TrackIndex)
-	{
-		FTimingEventsTrack& Track = *TimingEventsTracks[TrackIndex];
-		if (Track.GetType() == FName(TEXT("Thread")) && Track.GetSubType() == FName(TEXT("GPU")))
-		{
-			Track.SetVisibilityFlag(bShowHideAllGpuTracks);
-		}
-	}
-
-	OnTimingEventsTrackVisibilityChanged();
-}
-
-////////////////////////////////////////////////////////////////////////////////////////////////////
-
-bool STimingView::ShowHideAllLoadingTracks_IsChecked() const
-{
-	return bShowHideAllLoadingTracks;
-}
-
-////////////////////////////////////////////////////////////////////////////////////////////////////
-
-void STimingView::ShowHideAllLoadingTracks_Execute()
-{
-	bShowHideAllLoadingTracks = !bShowHideAllLoadingTracks;
-
-	for (int32 TrackIndex = 0; TrackIndex < TimingEventsTracks.Num(); ++TrackIndex)
-	{
-		FTimingEventsTrack& Track = *TimingEventsTracks[TrackIndex];
-		if (Track.GetType() == FName(TEXT("Loading")))
-		{
-			Track.SetVisibilityFlag(bShowHideAllLoadingTracks);
-		}
-	}
-
-	OnTimingEventsTrackVisibilityChanged();
-}
-
-////////////////////////////////////////////////////////////////////////////////////////////////////
-
-bool STimingView::ShowHideAllIoTracks_IsChecked() const
-{
-	return bShowHideAllIoTracks;
-}
-
-////////////////////////////////////////////////////////////////////////////////////////////////////
-
-void STimingView::ShowHideAllIoTracks_Execute()
-{
-	bShowHideAllIoTracks = !bShowHideAllIoTracks;
-
-	for (int32 TrackIndex = 0; TrackIndex < TimingEventsTracks.Num(); ++TrackIndex)
-	{
-		FTimingEventsTrack& Track = *TimingEventsTracks[TrackIndex];
-		if (Track.GetType() == FName(TEXT("FileActivity")))
-		{
-			Track.SetVisibilityFlag(bShowHideAllIoTracks);
-		}
-	}
-
-	OnTimingEventsTrackVisibilityChanged();
-
-	if (bShowHideAllIoTracks)
-	{
-		FileActivitySharedState->RequestUpdate();
-	}
 }
 
 ////////////////////////////////////////////////////////////////////////////////////////////////////
 
 bool STimingView::IsAutoHideEmptyTracksEnabled() const
 {
-	return (Layout.TargetMinTimelineH == 0.0f);
+	return (Viewport.GetLayout().TargetMinTimelineH == 0.0f);
 }
 
 ////////////////////////////////////////////////////////////////////////////////////////////////////
 
 void STimingView::ToggleAutoHideEmptyTracks()
 {
-	if (Layout.TargetMinTimelineH == 0.0f)
-	{
-		Layout.TargetMinTimelineH = FTimingEventsTrackLayout::RealMinTimelineH;
-	}
-	else
-	{
-		Layout.TargetMinTimelineH = 0.0f;
-	}
-
-	Layout.MinTimelineH = Layout.TargetMinTimelineH; // no layout animation
+	Viewport.ToggleLayoutMinTrackHeight();
 
-	for (auto& CachedTimelineKV : CachedTimelines)
+	for (const TSharedPtr<FBaseTimingTrack>& TrackPtr : ScrollableTracks)
 	{
-		CachedTimelineKV.Value->SetHeight(0.0f);
+		TrackPtr->SetHeight(0.0f);
 	}
 
 	ScrollAtPosY(0.0f);
@@ -2736,10 +2756,10 @@ void STimingView::ToggleAutoHideEmptyTracks()
 
 bool STimingView::ToggleTrackVisibility_IsChecked(uint64 InTrackId) const
 {
-	if (CachedTimelines.Contains(InTrackId))
+	const TSharedPtr<FBaseTimingTrack>* const TrackPtrPtr = AllTracks.Find(InTrackId);
+	if (TrackPtrPtr)
 	{
-		const FTimingEventsTrack* Track = CachedTimelines[InTrackId];
-		return Track->IsVisible();
+		return (*TrackPtrPtr)->IsVisible();
 	}
 	return false;
 }
@@ -2748,76 +2768,55 @@ bool STimingView::ToggleTrackVisibility_IsChecked(uint64 InTrackId) const
 
 void STimingView::ToggleTrackVisibility_Execute(uint64 InTrackId)
 {
-	if (CachedTimelines.Contains(InTrackId))
+	const TSharedPtr<FBaseTimingTrack>* TrackPtrPtr = AllTracks.Find(InTrackId);
+	if (TrackPtrPtr)
 	{
-		FTimingEventsTrack* Track = CachedTimelines[InTrackId];
-		Track->ToggleVisibility();
-
-		OnTimingEventsTrackVisibilityChanged();
+		(*TrackPtrPtr)->ToggleVisibility();
+		OnTrackVisibilityChanged();
 	}
 }
 
 ////////////////////////////////////////////////////////////////////////////////////////////////////
 
-bool STimingView::ToggleTrackVisibilityByGroup_IsChecked(const TCHAR* InGroupName) const
+void STimingView::OnTrackVisibilityChanged()
 {
-	if (ThreadGroups.Contains(InGroupName))
+	if (HoveredTrack.IsValid())
 	{
-		const FThreadGroup& ThreadGroup = ThreadGroups[InGroupName];
-		return ThreadGroup.bIsVisible;
+		HoveredTrack.Reset();
+		OnHoveredTrackChangedDelegate.Broadcast(HoveredTrack);
 	}
-	return false;
-}
-
-////////////////////////////////////////////////////////////////////////////////////////////////////
-
-void STimingView::ToggleTrackVisibilityByGroup_Execute(const TCHAR* InGroupName)
-{
-	if (ThreadGroups.Contains(InGroupName))
+	if (HoveredEvent.IsValid())
 	{
-		FThreadGroup& ThreadGroup = ThreadGroups[InGroupName];
-		ThreadGroup.bIsVisible = !ThreadGroup.bIsVisible;
-
-		for (auto& KV : CachedTimelines)
-		{
-			FTimingEventsTrack& Track = *KV.Value;
-			if (Track.GetType() == FName(TEXT("Thread")) && Track.GetSubType() == FName(TEXT("CPU")) &&
-				((FCpuTimingTrack*)&Track)->GetGroupName() == InGroupName)
-			{
-				Track.SetVisibilityFlag(ThreadGroup.bIsVisible);
-			}
-		}
-
-		OnTimingEventsTrackVisibilityChanged();
+		HoveredEvent.Reset();
+		OnHoveredEventChangedDelegate.Broadcast(HoveredEvent);
+	}
+	if (SelectedTrack.IsValid())
+	{
+		SelectedTrack.Reset();
+		OnSelectedTrackChangedDelegate.Broadcast(SelectedTrack);
+	}
+	if (SelectedEvent.IsValid())
+	{
+		SelectedEvent.Reset();
+		OnSelectedEventChangedDelegate.Broadcast(SelectedEvent);
 	}
-}
-
-////////////////////////////////////////////////////////////////////////////////////////////////////
-
-void STimingView::OnTimingEventsTrackVisibilityChanged()
-{
-	HoveredTimingEvent.Reset();
-	SelectedTimingEvent.Reset();
 	Tooltip.SetDesiredOpacity(0.0f);
 
 	//TODO: ThreadFilterChangedEvent.Broadcast();
-
 	FTimingProfilerManager::Get()->OnThreadFilterChanged();
 
 	if (SelectionStartTime < SelectionEndTime)
 	{
 		UpdateAggregatedStats();
 	}
-
-	bTimingEventsNeedUpdate = true;
 }
 
 ////////////////////////////////////////////////////////////////////////////////////////////////////
 
-FTimingEventsTrack* STimingView::FindTimingEventsTrack(uint64 InTrackID)
+TSharedPtr<FBaseTimingTrack> STimingView::FindTrack(uint64 InTrackId)
 {
-	FTimingEventsTrack** TrackPtr = CachedTimelines.Find(InTrackID);
-	return TrackPtr != nullptr ? *TrackPtr : nullptr;
+	TSharedPtr<FBaseTimingTrack>* TrackPtrPtr = AllTracks.Find(InTrackId);
+	return TrackPtrPtr ? *TrackPtrPtr : nullptr;
 }
 
 ////////////////////////////////////////////////////////////////////////////////////////////////////
@@ -2827,4 +2826,6 @@ TArray<Insights::ITimingViewExtender*> STimingView::GetExtenders() const
 	return IModularFeatures::Get().GetModularFeatureImplementations<Insights::ITimingViewExtender>(Insights::TimingViewExtenderFeatureName);
 }
 
+////////////////////////////////////////////////////////////////////////////////////////////////////
+
 #undef LOCTEXT_NAMESPACE
