[Insights]

diff --git a/Engine/Source/Developer/TraceInsights/Private/Insights/Widgets/STimingView.h b/Engine/Source/Developer/TraceInsights/Private/Insights/Widgets/STimingView.h
--- a/Engine/Source/Developer/TraceInsights/Private/Insights/Widgets/STimingView.h
+++ b/Engine/Source/Developer/TraceInsights/Private/Insights/Widgets/STimingView.h
@@ -14,19 +14,20 @@
 
 // Insights
 #include "Insights/Common/FixedCircularBuffer.h"
-#include "Insights/ViewModels/MarkersTimingTrack.h"
+#include "Insights/ITimingViewSession.h"
 #include "Insights/ViewModels/TimerNode.h"
-#include "Insights/ViewModels/TimeRulerTrack.h"
 #include "Insights/ViewModels/TimingEvent.h"
 #include "Insights/ViewModels/TimingEventsTrack.h"
 #include "Insights/ViewModels/TimingTrackViewport.h"
 #include "Insights/ViewModels/TimingViewDrawHelper.h"
 #include "Insights/ViewModels/TooltipDrawState.h"
-#include "Insights/ITimingViewSession.h"
 
 class FFileActivitySharedState;
 class FLoadingSharedState;
+class FMarkersTimingTrack;
 class FMenuBuilder;
+class FThreadTimingSharedState;
+class FTimeRulerTrack; 
 class FTimingGraphTrack;
 class FTimingViewDrawHelper;
 class SScrollBar;
@@ -69,34 +70,19 @@ public:
 	void Construct(const FArguments& InArgs);
 
 	TSharedRef<SWidget> MakeTracksFilterMenu();
-	void CreateThreadGroupsMenu(FMenuBuilder& MenuBuilder);
 	void CreateTracksMenu(FMenuBuilder& MenuBuilder);
 
 	bool ShowHideGraphTrack_IsChecked() const;
 	void ShowHideGraphTrack_Execute();
 
-	bool ShowHideAllCpuTracks_IsChecked() const;
-	void ShowHideAllCpuTracks_Execute();
-
-	bool ShowHideAllGpuTracks_IsChecked() const;
-	void ShowHideAllGpuTracks_Execute();
-
-	bool ShowHideAllLoadingTracks_IsChecked() const;
-	void ShowHideAllLoadingTracks_Execute();
-
-	bool ShowHideAllIoTracks_IsChecked() const;
-	void ShowHideAllIoTracks_Execute();
-
 	bool IsAutoHideEmptyTracksEnabled() const;
 	void ToggleAutoHideEmptyTracks();
 
 	bool ToggleTrackVisibility_IsChecked(uint64 InTrackId) const;
 	void ToggleTrackVisibility_Execute(uint64 InTrackId);
 
-	bool ToggleTrackVisibilityByGroup_IsChecked(const TCHAR* InGroupName) const;
-	void ToggleTrackVisibilityByGroup_Execute(const TCHAR* InGroupName);
-
-	void EnableAssetLoadingMode();
+	bool IsAssetLoadingModeEnabled() const { return bAssetLoadingMode; }
+	void EnableAssetLoadingMode() { bAssetLoadingMode = true; }
 
 	/** Resets internal widget's data to the default one. */
 	void Reset();
@@ -234,20 +220,43 @@ public:
 	virtual FReply OnKeyUp(const FGeometry& MyGeometry, const FKeyEvent& InKeyEvent) override;
 
 	////////////////////////////////////////////////////////////////////////////////////////////////////
-
 	// ITimingViewSession interface
-	virtual FTimingEventsTrack* FindTimingEventsTrack(uint64 InTrackID) override;
-	virtual void AddTimingEventsTrack(FTimingEventsTrack* Track) override;
-	virtual Insights::FTimeMarkerChangedDelegate& OnTimeMarkerChanged() override { return OnTimeMarkerChangedDelegate; }
+
+	virtual TSharedPtr<FBaseTimingTrack> FindTrack(uint64 InTrackId) override;
+
+	virtual void AddTopDockedTrack(TSharedPtr<FBaseTimingTrack> Track) override;
+	virtual void AddBottomDockedTrack(TSharedPtr<FBaseTimingTrack> Track) override;
+	virtual void AddScrollableTrack(TSharedPtr<FBaseTimingTrack> Track) override;
+	virtual void AddForegroundTrack(TSharedPtr<FBaseTimingTrack> Track) override;
+
 	virtual Insights::FSelectionChangedDelegate& OnSelectionChanged() override { return OnSelectionChangedDelegate; }
+	virtual Insights::FTimeMarkerChangedDelegate& OnTimeMarkerChanged() override { return OnTimeMarkerChangedDelegate; }
+	virtual Insights::FHoveredTrackChangedDelegate& OnHoveredTrackChanged() override { return OnHoveredTrackChangedDelegate; }
 	virtual Insights::FHoveredEventChangedDelegate& OnHoveredEventChanged() override { return OnHoveredEventChangedDelegate; }
+	virtual Insights::FSelectedTrackChangedDelegate& OnSelectedTrackChanged() override { return OnSelectedTrackChangedDelegate; }
 	virtual Insights::FSelectedEventChangedDelegate& OnSelectedEventChanged() override { return OnSelectedEventChangedDelegate; }
 
-	const FTimingTrackViewport& GetViewport() { return Viewport; }
+	////////////////////////////////////////////////////////////////////////////////////////////////////
+
+	void UpdateScrollableTracksOrder();
+	void OnTrackVisibilityChanged();
+
+	bool IsGpuTrackVisible() const;
+	bool IsCpuTrackVisible(uint32 InThreadId) const;
+
+	////////////////////////////////////////////////////////////////////////////////////////////////////
+
+	const FTimingTrackViewport& GetViewport() const { return Viewport; }
+
+	const FVector2D& GetMousePosition() const { return MousePosition; }
 
 	double GetSelectionStartTime() const { return SelectionStartTime; }
 	double GetSelectionEndTime() const { return SelectionEndTime; }
 
+	bool IsPanning() const { return bIsPanning; }
+	bool IsSelecting() const { return bIsSelecting; }
+	bool IsScrubbing() const { return bIsScrubbing; }
+
 	bool IsTimeSelected(double Time) const { return Time >= SelectionStartTime && Time < SelectionEndTime; }
 	bool IsTimeSelectedInclusive(double Time) const { return Time >= SelectionStartTime && Time <= SelectionEndTime; }
 
@@ -259,22 +268,31 @@ public:
 	void SetAndCenterOnTimeMarker(double Time);
 	void SelectToTimeMarker(double Time);
 
-	bool AreTimeMarkersVisible() { return MarkersTrack.IsVisible(); }
+	//bool AreTimeMarkersVisible() { return MarkersTrack->IsVisible(); }
 	void SetTimeMarkersVisible(bool bOnOff);
-	bool IsDrawOnlyBookmarksEnabled() { return MarkersTrack.IsBookmarksTrack(); }
+	//bool IsDrawOnlyBookmarksEnabled() { return MarkersTrack->IsBookmarksTrack(); }
 	void SetDrawOnlyBookmarks(bool bOnOff);
 
-	bool IsGpuTrackVisible() const;
-	bool IsCpuTrackVisible(uint32 ThreadId) const;
-
 	TSharedPtr<FTimingGraphTrack> GetMainTimingGraphTrack() { return GraphTrack; }
 
+	uint64 GetHighlightedEventTypeId() const { return HighlightedEventTypeId; }
+	void SetHighlightedEventTypeId(uint64 InHighlightedEventTypeId) { HighlightedEventTypeId = InHighlightedEventTypeId; }
+
+	const TSharedPtr<FBaseTimingTrack> GetHoveredTrack() const { return HoveredTrack; }
+	const TSharedPtr<const ITimingEvent> GetHoveredEvent() const { return HoveredEvent; }
+	const TSharedPtr<FBaseTimingTrack> GetSelectedTrack() const { return SelectedTrack; }
+	const TSharedPtr<const ITimingEvent> GetSelectedEvent() const { return SelectedEvent; }
+
+	const TSharedPtr<FBaseTimingTrack> GetTrackAt(float InPosX, float InPosY) const;
+
 protected:
 	virtual FVector2D ComputeDesiredSize(float) const override
 	{
 		return FVector2D(16.0f, 16.0f);
 	}
 
+	void UpdateOtherViews();
+
 	void ShowContextMenu(const FPointerEvent& MouseEvent);
 
 	/** Binds our UI commands to delegates. */
@@ -308,10 +326,13 @@ protected:
 
 	void UpdateAggregatedStats();
 
-	void UpdateHoveredTimingEvent(float MX, float MY);
+	void UpdateHoveredTimingEvent(float InMousePosX, float InMousePosY);
 
 	void OnSelectedTimingEventChanged();
+
+	void SelectHoveredTimingTrack();
 	void SelectHoveredTimingEvent();
+
 	void SelectLeftTimingEvent();
 	void SelectRightTimingEvent();
 	void SelectUpTimingEvent();
@@ -319,8 +340,6 @@ protected:
 
 	void FrameSelection();
 
-	void OnTimingEventsTrackVisibilityChanged();
-
 	// Get all the plugin extenders we care about
 	TArray<Insights::ITimingViewExtender*> GetExtenders() const;
 
@@ -348,81 +367,40 @@ protected:
 	////////////////////////////////////////////////////////////////////////////////////////////////////
 
 protected:
-	bool bShowHideAllGpuTracks;
-	bool bShowHideAllCpuTracks;
-	bool bShowHideAllLoadingTracks;
-	bool bShowHideAllIoTracks;
-
-	////////////////////////////////////////////////////////////
-
 	/** The track's viewport. Encapsulates info about position and scale. */
 	FTimingTrackViewport Viewport;
 
-	bool bIsViewportDirty;
-	bool bIsVerticalViewportDirty;
-
 	////////////////////////////////////////////////////////////
 
 	/** All created tracks.
 	  * Maps track id to track pointer.
 	  */
-	TMap<uint64, FTimingEventsTrack*> CachedTimelines;
+	TMap<uint64, TSharedPtr<FBaseTimingTrack>> AllTracks;
 
-	//TODO: TArray<FBaseTimingTrack*> TopTracks; /**< tracks docked on top, in the order to be displayed (top to bottom) */
-	//TODO: TArray<FBaseTimingTrack*> ScrollableTracks; /**< tracks in scrollable area, in the order to be displayed (top to bottom) */
-	//TODO: TArray<FBaseTimingTrack*> BottomTracks; /**< tracks docked on bottom, in the order to be displayed (top to bottom) */
-	//TODO: TArray<FBaseTimingTrack*> ForegorundTracks; /**< tracks to draw over top/scrollable/bottom tracks (can use entire area), in the order to be displayed (back to front) */
+	TArray<TSharedPtr<FBaseTimingTrack>> TopDockedTracks; /**< tracks docked on top, in the order to be displayed (top to bottom) */
+	TArray<TSharedPtr<FBaseTimingTrack>> BottomDockedTracks; /**< tracks docked on bottom, in the order to be displayed (top to bottom) */
+	TArray<TSharedPtr<FBaseTimingTrack>> ScrollableTracks; /**< tracks in scrollable area, in the order to be displayed (top to bottom) */
+	TArray<TSharedPtr<FBaseTimingTrack>> ForegroundTracks; /**< tracks to draw over top/scrollable/bottom tracks (can use entire area), in the order to be displayed (back to front) */
 
 	////////////////////////////////////////////////////////////
 
-	TArray<FTimingEventsTrack*> TimingEventsTracks; /**< all timing events tracks, in the order to be displayed */
-
-	// Dirty flag that indicates the timing events need updating
-	bool bTimingEventsNeedUpdate;
-
-	// Dirty flag that indicates that the TimingEventsTracks list has changed
-	bool bTimingEventsTracksChanged;
-
-	////////////////////////////////////////////////////////////
-	// Cpu/Gpu
-
-	FTimingEventsTrack* GpuTrack;
-	TMap<uint32, FTimingEventsTrack*> CpuTracks; /**< maps thread id to track pointer */
-
-	struct FThreadGroup
-	{
-		const TCHAR* Name; /**< The thread group name; pointer to string owned by ThreadProvider. */
-		bool bIsVisible;  /**< Toggle to show/hide all thread timelines associated with this group at once. Used also as default for new thread timelines. */
-		uint32 NumTimelines; /**< Number of thread timelines associated with this group. */
-		int32 Order; //**< Order index used for sorting. Inherited from last thread timeline associated with this group. **/
-
-		int32 GetOrder() const { return Order; }
-	};
-
-	TMap<const TCHAR*, FThreadGroup> ThreadGroups; /**< maps thread group name to thread group info */
-
-	////////////////////////////////////////////////////////////
-	// Asset Loading
+	// Shared state for Cpu/Gpu Thread tracks
+	TSharedPtr<FThreadTimingSharedState> ThreadTimingSharedState;
 
+	// Shared state for Asset Loading tracks
 	TSharedPtr<FLoadingSharedState> LoadingSharedState;
-	TMap<uint32, FTimingEventsTrack*> LoadingTracks;
-
 	bool bAssetLoadingMode;
 
-	////////////////////////////////////////////////////////////
-	// File activity (I/O)
-
+	// Shared state for File Activity (I/O) tracks
 	TSharedPtr<FFileActivitySharedState> FileActivitySharedState;
-	FTimingEventsTrack* IoOverviewTrack;
-	FTimingEventsTrack* IoActivityTrack;
 
 	////////////////////////////////////////////////////////////
 
 	/** The time ruler track. */
-	FTimeRulerTrack TimeRulerTrack;
+	TSharedPtr<FTimeRulerTrack> TimeRulerTrack;
 
 	/** The time markers track. */
-	FMarkersTimingTrack MarkersTrack;
+	TSharedPtr<FMarkersTimingTrack> MarkersTrack;
 
 	/** A graph track for frame times. */
 	TSharedPtr<FTimingGraphTrack> GraphTrack;
@@ -457,11 +435,7 @@ protected:
 	////////////////////////////////////////////////////////////
 	// Panning
 
-	/**
-	 * True, if the user is currently interactively panning the view (horizontally and/or vertically),
-	 * either by holding the right mouse button and dragging
-	 * or by holding spacebar pressed and dragging with left mouse button.
-	 */
+	/** True, if the user is currently interactively panning the view (horizontally and/or vertically). */
 	bool bIsPanning;
 
 	/** How to pan. */
@@ -477,14 +451,16 @@ protected:
 	////////////////////////////////////////////////////////////
 	// Selection
 
-	/** True, if the user is currently changing the selection (by holding the left mouse button and dragging). */
+	/** True, if the user is currently changing the selection. */
 	bool bIsSelecting;
 
 	double SelectionStartTime;
 	double SelectionEndTime;
 
-	FTimingEvent HoveredTimingEvent;
-	FTimingEvent SelectedTimingEvent;
+	TSharedPtr<FBaseTimingTrack> HoveredTrack;
+	TSharedPtr<const ITimingEvent> HoveredEvent;
+	TSharedPtr<FBaseTimingTrack> SelectedTrack;
+	TSharedPtr<const ITimingEvent> SelectedEvent;
 
 	FTooltipDrawState Tooltip;
 
@@ -501,6 +477,8 @@ protected:
 	/** True of the user is currently dragging the time marker */
 	bool bIsScrubbing;
 
+	uint64 HighlightedEventTypeId;
+
 	////////////////////////////////////////////////////////////
 	// Misc
 
@@ -509,21 +487,26 @@ protected:
 	const FSlateBrush* WhiteBrush;
 	const FSlateFontInfo MainFont;
 
-	FTimingEventsTrackLayout Layout;
-
 	// Debug stats
 	int32 NumUpdatedEvents;
-	TFixedCircularBuffer<uint64, 32> TimelineCacheUpdateDurationHistory;
-	TFixedCircularBuffer<uint64, 32> TimeMarkerCacheUpdateDurationHistory;
-	mutable TFixedCircularBuffer<uint64, 32> DrawDurationHistory;
-	mutable TFixedCircularBuffer<uint64, 32> OnPaintDurationHistory;
+	TFixedCircularBuffer<uint64, 32> PreUpdateTracksDurationHistory;
+	TFixedCircularBuffer<uint64, 32> UpdateTracksDurationHistory;
+	TFixedCircularBuffer<uint64, 32> PostUpdateTracksDurationHistory;
+	TFixedCircularBuffer<uint64, 32> TickDurationHistory;
+	mutable TFixedCircularBuffer<uint64, 32> PreDrawTracksDurationHistory;
+	mutable TFixedCircularBuffer<uint64, 32> DrawTracksDurationHistory;
+	mutable TFixedCircularBuffer<uint64, 32> PostDrawTracksDurationHistory;
+	mutable TFixedCircularBuffer<uint64, 32> TotalDrawDurationHistory;
+	mutable TFixedCircularBuffer<uint64, 32> OnPaintDeltaTimeHistory;
 	mutable uint64 LastOnPaintTime;
 
 	////////////////////////////////////////////////////////////
 	// Delegates
 
 	Insights::FSelectionChangedDelegate OnSelectionChangedDelegate;
+	Insights::FTimeMarkerChangedDelegate OnTimeMarkerChangedDelegate;
+	Insights::FHoveredTrackChangedDelegate OnHoveredTrackChangedDelegate;
 	Insights::FHoveredEventChangedDelegate OnHoveredEventChangedDelegate;
+	Insights::FSelectedTrackChangedDelegate OnSelectedTrackChangedDelegate;
 	Insights::FSelectedEventChangedDelegate OnSelectedEventChangedDelegate;
-	Insights::FTimeMarkerChangedDelegate OnTimeMarkerChangedDelegate;
 };
diff --git a/Engine/Source/Developer/TraceInsights/Public/Insights/ITimingViewExtender.h b/Engine/Source/Developer/TraceInsights/Public/Insights/ITimingViewExtender.h
--- a/Engine/Source/Developer/TraceInsights/Public/Insights/ITimingViewExtender.h
+++ b/Engine/Source/Developer/TraceInsights/Public/Insights/ITimingViewExtender.h
@@ -14,7 +14,7 @@ namespace Insights
 class ITimingViewSession;
 extern TRACEINSIGHTS_API const FName TimingViewExtenderFeatureName;
 
-class ITimingViewExtender : public IModularFeature
+class TRACEINSIGHTS_API ITimingViewExtender : public IModularFeature
 {
 public:
 	virtual ~ITimingViewExtender() = default;
@@ -25,7 +25,7 @@ public:
 	/** Called to clear out any data at the end of the timing view session */
 	virtual void OnEndSession(ITimingViewSession& InSession) = 0;
 
-	/** Called each frame. If any new tracks are created they can be added via ITimingViewSession::AddTimingEventsTrack() */
+	/** Called each frame. If any new tracks are created they can be added via ITimingViewSession::Add*Track() */
 	virtual void Tick(ITimingViewSession& InSession, const Trace::IAnalysisSession& InAnalysisSession) = 0;
 
 	/** Extension hook for the 'quick filter' menu */
@@ -35,4 +35,4 @@ public:
 	virtual void OnTracksChanged(ITimingViewSession& InSession, int32& InOutOrder) = 0;
 };
 
-}
\ No newline at end of file
+} // namespace Insights
diff --git a/Engine/Source/Developer/TraceInsights/Public/Insights/ITimingViewSession.h b/Engine/Source/Developer/TraceInsights/Public/Insights/ITimingViewSession.h
--- a/Engine/Source/Developer/TraceInsights/Public/Insights/ITimingViewSession.h
+++ b/Engine/Source/Developer/TraceInsights/Public/Insights/ITimingViewSession.h
@@ -4,9 +4,9 @@
 
 #include "CoreMinimal.h"
 
+class FBaseTimingTrack;
 class FTimingEventsTrack;
-namespace Insights { class ITimingViewVisualizer; }
-struct FTimingEvent;
+class ITimingEvent;
 
 namespace Insights
 {
@@ -18,44 +18,61 @@ enum class ETimeChangedFlags : int32
 	// The event fired in response to an interactive change from the user. Will be followed by a non-interactive change finished.
 	Interactive = (1 << 0)
 };
-
 ENUM_CLASS_FLAGS(ETimeChangedFlags);
 
+/** The delegate to be invoked when the selection have been changed */
+DECLARE_MULTICAST_DELEGATE_ThreeParams(FSelectionChangedDelegate, ETimeChangedFlags /*InFlags*/, double /*StartTime*/, double /*EndTime*/);
+
 /** The delegate to be invoked when the time marker has changed */
 DECLARE_MULTICAST_DELEGATE_TwoParams(FTimeMarkerChangedDelegate, ETimeChangedFlags /*InFlags*/, double /*TimeMarker*/);
 
-/** The delegate to be invoked when the selection have been changed. */
-DECLARE_MULTICAST_DELEGATE_ThreeParams(FSelectionChangedDelegate, ETimeChangedFlags /*InFlags*/, double /*StartTime*/, double /*EndTime*/);
+/** The delegate to be invoked when the timing track being hovered by the mouse has changed */
+DECLARE_MULTICAST_DELEGATE_OneParam(FHoveredTrackChangedDelegate, const TSharedPtr<FBaseTimingTrack> /*InTrack*/);
+
+/** The delegate to be invoked when the timing event being hovered by the mouse has changed */
+DECLARE_MULTICAST_DELEGATE_OneParam(FHoveredEventChangedDelegate, const TSharedPtr<const ITimingEvent> /*InEvent*/);
 
-/** The delegate to be invoked when the timing event being hovered by the mouse has changed. */
-DECLARE_MULTICAST_DELEGATE_OneParam(FHoveredEventChangedDelegate, const FTimingEvent& /*InEvent*/);
+/** The delegate to be invoked when the selected timing track has changed */
+DECLARE_MULTICAST_DELEGATE_OneParam(FSelectedTrackChangedDelegate, const TSharedPtr<FBaseTimingTrack> /*InTrack*/);
 
-/** The delegate to be invoked when the selected timing event has changed. */
-DECLARE_MULTICAST_DELEGATE_OneParam(FSelectedEventChangedDelegate, const FTimingEvent& /*InEvent*/);
+/** The delegate to be invoked when the selected timing event has changed */
+DECLARE_MULTICAST_DELEGATE_OneParam(FSelectedEventChangedDelegate, const TSharedPtr<const ITimingEvent> /*InEvent*/);
 
-/** Hosts a number of timing view visualizers, represents a session of the timing view */
-class ITimingViewSession
+/** Hosts a number of timing view visualizers, represents a session of the timing view. */
+class TRACEINSIGHTS_API ITimingViewSession
 {
 public:
 	virtual ~ITimingViewSession() = default;
 
-	/** Adds a new timing events track. Transfers ownership of the track pointer to the session. */
-	virtual void AddTimingEventsTrack(FTimingEventsTrack* InTrack) = 0;
+	/** Adds a new top docked track. */
+	virtual void AddTopDockedTrack(TSharedPtr<FBaseTimingTrack> Track) = 0;
+	/** Adds a new bottom docked track. */
+	virtual void AddBottomDockedTrack(TSharedPtr<FBaseTimingTrack> Track) = 0;
+	/** Adds a new scrollable track. */
+	virtual void AddScrollableTrack(TSharedPtr<FBaseTimingTrack> Track) = 0;
+	/** Adds a new foreground track. */
+	virtual void AddForegroundTrack(TSharedPtr<FBaseTimingTrack> Track) = 0;
 
-	/** Find a timing events track has been added via AddTimingEventsTrack() */
-	virtual FTimingEventsTrack* FindTimingEventsTrack(uint64 InTrackID) = 0;
+	/** Finds a track has been added via Add*Track(). */
+	virtual TSharedPtr<FBaseTimingTrack> FindTrack(uint64 InTrackId) = 0;
 
-	/** Get the delegate to be invoked when the time marker has changed */
+	/** Gets the delegate to be invoked when the selection have been changed. */
+	virtual FSelectionChangedDelegate& OnSelectionChanged() = 0;
+
+	/** Gets the delegate to be invoked when the time marker has changed. */
 	virtual FTimeMarkerChangedDelegate& OnTimeMarkerChanged() = 0;
 
-	/** Get the delegate to be invoked when the selection have been changed. */
-	virtual FSelectionChangedDelegate& OnSelectionChanged() = 0;
+	/** Gets the delegate to be invoked when the timing track being hovered by the mouse has changed. */
+	virtual FHoveredTrackChangedDelegate& OnHoveredTrackChanged() = 0;
 
-	/** Get the delegate to be invoked when the timing event being hovered by the mouse has changed. */
+	/** Gets the delegate to be invoked when the timing event being hovered by the mouse has changed. */
 	virtual FHoveredEventChangedDelegate& OnHoveredEventChanged() = 0;
 
-	/** Get the delegate to be invoked when the selected timing event has changed. */
+	/** Gets the delegate to be invoked when the selected timing track has changed. */
+	virtual FSelectedTrackChangedDelegate& OnSelectedTrackChanged() = 0;
+
+	/** Gets the delegate to be invoked when the selected timing event has changed. */
 	virtual FSelectedEventChangedDelegate& OnSelectedEventChanged() = 0;
 };
 
-}
\ No newline at end of file
+} // namespace Insights
diff --git a/Engine/Source/Developer/TraceInsights/Public/Insights/ViewModels/BaseTimingTrack.h b/Engine/Source/Developer/TraceInsights/Public/Insights/ViewModels/BaseTimingTrack.h
--- a/Engine/Source/Developer/TraceInsights/Public/Insights/ViewModels/BaseTimingTrack.h
+++ b/Engine/Source/Developer/TraceInsights/Public/Insights/ViewModels/BaseTimingTrack.h
@@ -4,14 +4,46 @@
 
 #include "CoreMinimal.h"
 #include "Misc/EnumClassFlags.h"
-#include "Templates/SharedPointer.h"
 
 ////////////////////////////////////////////////////////////////////////////////////////////////////
 
-class FTimingTrackViewport;
 class FMenuBuilder;
 namespace Trace { class IAnalysisSession; };
 
+class FTimingEventSearchParameters;
+class FTimingTrackViewport;
+class FTimingViewDrawHelper;
+class FTooltipDrawState;
+class ITimingEvent;
+
+struct FDrawContext;
+
+////////////////////////////////////////////////////////////////////////////////////////////////////
+
+class TRACEINSIGHTS_API ITimingTrackUpdateContext
+{
+public:
+	virtual const FTimingTrackViewport& GetViewport() const = 0;
+	virtual const FVector2D& GetMousePosition() const = 0;
+	virtual const TSharedPtr<const ITimingEvent> GetHoveredEvent() const = 0;
+	virtual const TSharedPtr<const ITimingEvent> GetSelectedEvent() const = 0;
+	virtual double GetCurrentTime() const = 0;
+	virtual float GetDeltaTime() const = 0;
+};
+
+////////////////////////////////////////////////////////////////////////////////////////////////////
+
+class TRACEINSIGHTS_API ITimingTrackDrawContext
+{
+public:
+	virtual const FTimingTrackViewport& GetViewport() const = 0;
+	virtual const FVector2D& GetMousePosition() const = 0;
+	virtual const TSharedPtr<const ITimingEvent> GetHoveredEvent() const = 0;
+	virtual const TSharedPtr<const ITimingEvent> GetSelectedEvent() const = 0;
+	virtual FDrawContext& GetDrawContext() const = 0;
+	virtual const FTimingViewDrawHelper& GetHelper() const = 0;
+};
+
 ////////////////////////////////////////////////////////////////////////////////////////////////////
 
 enum class ETimingTrackFlags : uint32
@@ -27,13 +59,33 @@ ENUM_CLASS_FLAGS(ETimingTrackFlags);
 
 ////////////////////////////////////////////////////////////////////////////////////////////////////
 
+enum class EDrawEventMode : uint32
+{
+	None = 0,
+
+	/** Draw the content of the event. This flag can be omitted in order to draw only the hovered/selected highlights. */
+	Content = (1 << 0),
+
+	/** Draw the highlights for a hovered event. */
+	Hovered = (1 << 1),
+
+	/** Draw the highlights for a selected event. */
+	Selected = (1 << 2),
+
+	/** Draw the highlights for an event that is both selected and hovered. */
+	SelectedAndHovered = Hovered | Selected,
+};
+ENUM_CLASS_FLAGS(EDrawEventMode);
+
+////////////////////////////////////////////////////////////////////////////////////////////////////
+
 class TRACEINSIGHTS_API FBaseTimingTrack : public TSharedFromThis<FBaseTimingTrack>
 {
 	friend class FTimingViewDrawHelper;
 
 protected:
-	explicit FBaseTimingTrack(uint64 InId, const FName& InType = NAME_None, const FName& InSubType = NAME_None, const FString& InName = FString())
-		: Id(InId)
+	explicit FBaseTimingTrack(const FName& InType, const FName& InSubType = NAME_None, const FString& InName = FString())
+		: Id(GenerateId())
 		, Type(InType)
 		, SubType(InSubType)
 		, Name(InName)
@@ -49,7 +101,8 @@ protected:
 public:
 	virtual void Reset()
 	{
-		Id = 0;
+		//Name = FString();
+		//Order = 0;
 		PosY = 0.0f;
 		Height = 0.0f;
 		Flags = ETimingTrackFlags::IsVisible | ETimingTrackFlags::IsDirty;
@@ -61,9 +114,10 @@ public:
 	const FName& GetSubType() const { return SubType; }
 
 	const FString& GetName() const { return Name; }
+	void SetName(const FString& InName) { Name = InName; }
 
-	void SetOrder(int32 InOrder) { Order = InOrder; }
 	int32 GetOrder() const { return Order; }
+	void SetOrder(int32 InOrder) { Order = InOrder; }
 
 	float GetPosY() const { return PosY; }
 	virtual void SetPosY(float InPosY) { PosY = InPosY; }
@@ -76,7 +130,7 @@ public:
 	void Hide() { Flags &= ~ETimingTrackFlags::IsVisible; OnVisibilityChanged(); }
 	void ToggleVisibility() { Flags ^= ETimingTrackFlags::IsVisible; OnVisibilityChanged(); }
 	void SetVisibilityFlag(bool bIsVisible) { bIsVisible ? Show() : Hide(); }
-	virtual void OnVisibilityChanged() {}
+	virtual void OnVisibilityChanged() { if (IsVisible()) SetDirtyFlag(); }
 
 	bool IsDirty() const { return EnumHasAnyFlags(Flags, ETimingTrackFlags::IsDirty); }
 	void SetDirtyFlag() { Flags |= ETimingTrackFlags::IsDirty; OnDirtyFlagChanged(); }
@@ -94,19 +148,80 @@ public:
 	void SetHoveredState(bool bIsHovered) { bIsHovered ? Flags |= ETimingTrackFlags::IsHovered : Flags &= ~ETimingTrackFlags::IsHovered; }
 	bool IsHeaderHovered() const { return EnumHasAllFlags(Flags, ETimingTrackFlags::IsHovered | ETimingTrackFlags::IsHeaderHovered); }
 	void SetHeaderHoveredState(bool bIsHeaderHovered) { bIsHeaderHovered ? Flags |= ETimingTrackFlags::IsHeaderHovered : Flags &= ~ETimingTrackFlags::IsHeaderHovered; }
-	virtual void UpdateHoveredState(float MouseX, float MouseY, const FTimingTrackViewport& Viewport) = 0;
 
-	virtual void Tick(const double CurrentTime, const float DeltaTime) {}
-	virtual void Update(const FTimingTrackViewport& Viewport) {}
+	//////////////////////////////////////////////////
 
-	static uint64 GenerateId() { return IdGenerator++; }
+	// PreUpdate callback called each frame, but only if the track is visible.
+	// In this update, neither the postion nor the size of the track is yet computed.
+	// Track should update here its height.
+	virtual void PreUpdate(const ITimingTrackUpdateContext& Context) {}
+
+	// Update callback called each frame, but only if the track is visible.
+	// In this update, it is assumed the track position and the track size are valid.
+	virtual void Update(const ITimingTrackUpdateContext& Context) {}
+
+	// PostUpdate callback called each frame, but only if the track is visible.
+	// Track should update here its "hovered" state.
+	virtual void PostUpdate(const ITimingTrackUpdateContext& Context) {}
+
+	//////////////////////////////////////////////////
+
+	// PreDraw callback (called from OnPaint) to draw something in the background.
+	virtual void PreDraw(const ITimingTrackDrawContext& Context) const {}
+
+	// Draw callback (called from OnPaint) to draw the track's content.
+	virtual void Draw(const ITimingTrackDrawContext& Context) const {}
+
+	// Draw a single event (can be used to draw only the highlight for a selected and/or hovered event).
+	virtual void DrawEvent(const ITimingTrackDrawContext& Context, const ITimingEvent& InTimingEvent, EDrawEventMode InDrawMode) const {}
+
+	// PostDraw callback (called from OnPaint) to draw something in the foreground.
+	virtual void PostDraw(const ITimingTrackDrawContext& Context) const {}
+
+	//////////////////////////////////////////////////
 
 	virtual void BuildContextMenu(FMenuBuilder& MenuBuilder);
 
+	/**
+	 * Gets the event at a specified position. Leaves the InOutEvent unchanged if no event is found at specified position.
+	 * @param X The horizontal coordinate of the point tested; in Slate pixels (viewport coordinates).
+	 * @param Y The vertical coordinate of the point tested; in Slate pixels (viewport coordinates).
+	 * @param Viewport The timing viewport used to transform time in viewport coordinates.
+	 * @return The event located at (PosX, PosY) coordinates, if any; nullptr ottherwise.
+	 */
+	virtual const TSharedPtr<const ITimingEvent> GetEvent(float InPosX, float InPosY, const FTimingTrackViewport& Viewport) const { return nullptr; }
+
+	// Search for an event using custom parameters.
+	virtual const TSharedPtr<const ITimingEvent> SearchEvent(const FTimingEventSearchParameters& InSearchParameters) const { return nullptr; }
+
+	// Allows tracks to update event stats that are slower to compute (called at a lower frequency than GetEventAtPosition or Search or SearchTimingEvent).
+	virtual void UpdateEventStats(ITimingEvent& InOutEvent) const {}
+
+	// Called back from the timing view when an event is hovered by mouse.
+	virtual void OnEventHovered(const ITimingEvent& InHoveredEvent) const {}
+
+	// Called back from the timing view when an event is selected.
+	virtual void OnEventSelected(const ITimingEvent& InSelectedEvent) const {}
+
+	// Called to initialize the tooltip's content with info from a timing event.
+	virtual void InitTooltip(FTooltipDrawState& InOutTooltip, const ITimingEvent& InTooltipEvent) const {}
+
+	// Called back from the timing view when an event is copied to the clipboard with Ctrl+C.
+	virtual void OnClipboardCopyEvent(const ITimingEvent& InSelectedEvent) const {}
+
+	// Returns number of text lines needed to display the debug string.
+	//TODO: virtual int GetDebugStringLineCount() const { return 0; }
+
+	// Builds a string with debug text information.
+	//TODO: virtual void BuildDebugString(FString& OutStr) const {}
+
+protected:
+	static uint64 GenerateId() { return IdGenerator++; }
+
 private:
-	uint64 Id;
-	FName Type; // Thread, Loading, FileActivity
-	FName SubType; // Cpu, Gpu, MainThread, AsyncThread, Overview, Detailed
+	const uint64 Id;
+	const FName Type; // Thread, Loading, FileActivity
+	const FName SubType; // Cpu, Gpu, MainThread, AsyncThread, Overview, Detailed
 	FString Name;
 	int32 Order;
 	float PosY; // y position, in Slate units
diff --git a/Engine/Source/Developer/TraceInsights/Public/Insights/ViewModels/GraphSeries.h b/Engine/Source/Developer/TraceInsights/Public/Insights/ViewModels/GraphSeries.h
new file mode 100644
--- /dev/null
+++ b/Engine/Source/Developer/TraceInsights/Public/Insights/ViewModels/GraphSeries.h
@@ -0,0 +1,163 @@
+// Copyright 1998-2019 Epic Games, Inc. All Rights Reserved.
+
+#pragma once
+
+#include "CoreMinimal.h"
+
+// Insights
+#include "Insights/ViewModels/GraphTrackEvent.h"
+
+class FTimingTrackViewport;
+
+////////////////////////////////////////////////////////////////////////////////////////////////////
+
+class TRACEINSIGHTS_API FGraphSeries
+{
+	friend class FGraphTrack;
+	friend class FGraphTrackBuilder;
+
+public:
+	struct FBox
+	{
+		float X;
+		float W;
+		float Y;
+	};
+
+public:
+	FGraphSeries();
+	virtual ~FGraphSeries();
+
+	const FText& GetName() const { return Name; }
+	void SetName(const TCHAR* InName) { Name = FText::FromString(InName); }
+	void SetName(const FString& InName) { Name = FText::FromString(InName); }
+	void SetName(const FText& InName) { Name = InName; }
+
+	const FText& GetDescription() const { return Description; }
+	void SetDescription(const TCHAR* InDescription) { Description = FText::FromString(InDescription); }
+	void SetDescription(const FString& InDescription) { Description = FText::FromString(InDescription); }
+	void SetDescription(const FText& InDescription) { Description = InDescription; }
+
+	bool IsVisible() const { return bIsVisible; }
+	void SetVisibility(bool bOnOff) { bIsVisible = bOnOff; }
+
+	bool IsDirty() const { return bIsDirty; }
+	void SetDirtyFlag() { bIsDirty = true; }
+	void ClearDirtyFlag() { bIsDirty = false; }
+
+	const FLinearColor& GetColor() const { return Color; }
+	const FLinearColor& GetBorderColor() const { return BorderColor; }
+
+	void SetColor(FLinearColor InColor, FLinearColor InBorderColor)
+	{
+		Color = InColor;
+		BorderColor = InBorderColor;
+	}
+
+	/**
+	 * @return Y position (in viewport local space) of the baseline (with Value == 0); in pixels (Slate units).
+	 * Y == 0 at the top of the graph track, positive values are downward.
+	 */
+	double GetBaselineY() const { return BaselineY; }
+	void SetBaselineY(const double InBaselineY) { BaselineY = InBaselineY; }
+
+	/**
+	 * @return The scale between Value units and viewport units; in pixels (Slate units) / Value unit.
+	 */
+	double GetScaleY() const { return ScaleY; }
+	void SetScaleY(const double InScaleY) { ScaleY = InScaleY; }
+
+	/**
+	 * @param Value a value; in Value units
+	 * @return Y position (in viewport local space) for a Value; in pixels (Slate units).
+	 * Y == 0 at the top of the graph track, positive values are downward.
+	 */
+	float GetYForValue(double Value) const
+	{
+		return static_cast<float>(BaselineY - Value * ScaleY);
+	}
+
+	float GetRoundedYForValue(double Value) const
+	{
+		return FMath::RoundToFloat(FMath::Clamp<float>(GetYForValue(Value), -FLT_MAX, FLT_MAX));
+	}
+
+	/**
+	 * @param Y a Y position (in viewport local space); in pixels (Slate units).
+	 * @return Value for specified Y position.
+	 */
+	double GetValueForY(float Y) const
+	{
+		return (BaselineY - static_cast<double>(Y)) / ScaleY;
+	}
+
+	bool IsAutoZoomEnabled() const { return bAutoZoom; }
+	void EnableAutoZoom() { bAutoZoom = true; }
+
+	/** Target low value of auto zoom interval (corresponding to bottom of the track) */
+	double GetTargetAutoZoomLowValue() const { return TargetAutoZoomLowValue; }
+	/** Target high value of auto zoom interval (corresponding to top of the track) */
+	double GetTargetAutoZoomHighValue() const { return TargetAutoZoomHighValue; }
+	void SetTargetAutoZoomRange(double LowValue, double HighValue) { TargetAutoZoomLowValue = LowValue; TargetAutoZoomHighValue = HighValue; }
+
+	/** Current auto zoom low value */
+	double GetAutoZoomLowValue() const { return AutoZoomLowValue; }
+	/** Current auto zoom high value */
+	double GetAutoZoomHighValue() const { return AutoZoomHighValue; }
+	void SetAutoZoomRange(double LowValue, double HighValue) { AutoZoomLowValue = LowValue; AutoZoomHighValue = HighValue; }
+
+	/**
+	 * Compute BaselineY and ScaleY so the [Low, High] Value range will correspond to [Top, Bottom] Y position range.
+	 * GetYForValue(InHighValue) == InTopY
+	 * GetYForValue(InLowValue) == InBottomY
+	 */
+	void ComputeBaselineAndScale(const double InLowValue, const double InHighValue, const float InTopY, const float InBottomY, double& OutBaselineY, double& OutScaleY) const
+	{
+		ensure(InLowValue < InHighValue);
+		ensure(InTopY < InBottomY);
+		const double InvRange = 1.0 / (InHighValue - InLowValue);
+		OutScaleY = static_cast<double>(InBottomY - InTopY) * InvRange;
+		//OutBaselineY = (InHighValue * static_cast<double>(InBottomY) - InLowValue * static_cast<double>(InTopY)) * InvRange;
+		OutBaselineY = static_cast<double>(InTopY) + InHighValue * OutScaleY;
+	}
+
+	/**
+	 * @param X The horizontal coordinate of the point tested; in Slate pixels (local graph coordinates)
+	 * @param Y The vertical coordinate of the point tested; in Slate pixels (local graph coordinates)
+	 * @param Viewport The timing viewport used to transform time in local graph coordinates
+	 * @param bCheckLine If needs to check the bounding box of the horizontal line (determined by duration of event and value) or only the bounding box of the visual point
+	 * @param bCheckBox If needs to check the bounding box of the entire visual box (determined by duration of event, value and baseline)
+	 * @return A pointer to an Event located at (X, Y) coordinates, if any; nullptr if no event is located at respective coordinates
+	 * The returned pointer is valid only temporary until next Reset() or next usage of FGraphTrackBuilder for this series/track.
+	 */
+	const FGraphSeriesEvent* GetEvent(const float PosX, const float PosY, const FTimingTrackViewport& Viewport, bool bCheckLine, bool bCheckBox) const;
+
+	virtual FString FormatValue(double Value) const;
+
+private:
+	FText Name;
+	FText Description;
+
+	bool bIsVisible;
+	bool bIsDirty;
+
+	bool bAutoZoom;
+	double TargetAutoZoomLowValue; // target low value of auto zoom interval (corresponding to bottom of the track)
+	double TargetAutoZoomHighValue; // target high value of auto zoom interval (corresponding to top of the track)
+	double AutoZoomLowValue; // current auto zoom low value
+	double AutoZoomHighValue; // current auto zoom high value
+
+	double BaselineY; // Y position (in viewport local space) of the baseline (with Value == 0); in pixels (Slate units)
+	double ScaleY; // scale between Value units and viewport units; in pixels (Slate units) / Value unit
+
+	FLinearColor Color;
+	FLinearColor BorderColor;
+
+protected:
+	TArray<FGraphSeriesEvent> Events; // reduced list of events; used to identify an event at a certain screen position (ex.: the event hovered by mouse)
+	TArray<FVector2D> Points; // reduced list of points; for drawing points
+	TArray<FVector2D> LinePoints; // reduced list of points; for drawing the connected line and filled polygon
+	TArray<FBox> Boxes; // reduced list of boxes; for drawing boxes
+};
+
+////////////////////////////////////////////////////////////////////////////////////////////////////
diff --git a/Engine/Source/Developer/TraceInsights/Public/Insights/ViewModels/GraphTrack.h b/Engine/Source/Developer/TraceInsights/Public/Insights/ViewModels/GraphTrack.h
--- a/Engine/Source/Developer/TraceInsights/Public/Insights/ViewModels/GraphTrack.h
+++ b/Engine/Source/Developer/TraceInsights/Public/Insights/ViewModels/GraphTrack.h
@@ -7,169 +7,12 @@
 
 // Insights
 #include "Insights/ViewModels/BaseTimingTrack.h"
-#include "Insights/ViewModels/TooltipDrawState.h"
 
-struct FDrawContext;
-struct FSlateBrush;
 class FMenuBuilder;
-class FTimingTrackViewport;
-
-////////////////////////////////////////////////////////////////////////////////////////////////////
-
-class FGraphSeries
-{
-	friend class FGraphTrack;
-	friend class FGraphTrackBuilder;
-
-public:
-	struct FEvent
-	{
-		double Time;
-		double Duration;
-		double Value;
-	};
-
-	struct FBox
-	{
-		float X;
-		float W;
-		float Y;
-	};
-
-public:
-	FGraphSeries();
-	virtual ~FGraphSeries();
-
-	const FText& GetName() const { return Name; }
-	void SetName(const TCHAR* InName) { Name = FText::FromString(InName); }
-	void SetName(const FString& InName) { Name = FText::FromString(InName); }
-	void SetName(const FText& InName) { Name = InName; }
-
-	const FText& GetDescription() const { return Description; }
-	void SetDescription(const TCHAR* InDescription) { Description = FText::FromString(InDescription); }
-	void SetDescription(const FString& InDescription) { Description = FText::FromString(InDescription); }
-	void SetDescription(const FText& InDescription) { Description = InDescription; }
-
-	bool IsVisible() const { return bIsVisible; }
-	void SetVisibility(bool bOnOff) { bIsVisible = bOnOff; }
-
-	bool IsDirty() const { return bIsDirty; }
-	void SetDirtyFlag() { bIsDirty = true; }
-	void ClearDirtyFlag() { bIsDirty = false; }
-
-	const FLinearColor& GetColor() const { return Color; }
-	const FLinearColor& GetBorderColor() const { return BorderColor; }
-
-	void SetColor(FLinearColor InColor, FLinearColor InBorderColor)
-	{
-		Color = InColor;
-		BorderColor = InBorderColor;
-	}
-
-	/**
-	 * @return Y position (in viewport local space) of the baseline (with Value == 0); in pixels (Slate units).
-	 * Y == 0 at the top of the graph track, positive values are downward.
-	 */
-	double GetBaselineY() const { return BaselineY; }
-	void SetBaselineY(const double InBaselineY) { BaselineY = InBaselineY; }
-
-	/**
-	 * @return The scale between Value units and viewport units; in pixels (Slate units) / Value unit.
-	 */
-	double GetScaleY() const { return ScaleY; }
-	void SetScaleY(const double InScaleY) { ScaleY = InScaleY; }
-
-	/**
-	 * @param Value a value; in Value units
-	 * @return Y position (in viewport local space) for a Value; in pixels (Slate units).
-	 * Y == 0 at the top of the graph track, positive values are downward.
-	 */
-	float GetYForValue(double Value) const
-	{
-		return static_cast<float>(BaselineY - Value * ScaleY);
-	}
-
-	float GetRoundedYForValue(double Value) const
-	{
-		return FMath::RoundToFloat(FMath::Clamp<float>(GetYForValue(Value), -FLT_MAX, FLT_MAX));
-	}
-
-	/**
-	 * @param Y a Y position (in viewport local space); in pixels (Slate units).
-	 * @return Value for specified Y position.
-	 */
-	double GetValueForY(float Y) const
-	{
-		return (BaselineY - static_cast<double>(Y)) / ScaleY;
-	}
-
-	bool IsAutoZoomEnabled() const { return bAutoZoom; }
-	void EnableAutoZoom() { bAutoZoom = true; }
-
-	/** Target low value of auto zoom interval (corresponding to bottom of the track) */
-	double GetTargetAutoZoomLowValue() const { return TargetAutoZoomLowValue; }
-	/** Target high value of auto zoom interval (corresponding to top of the track) */
-	double GetTargetAutoZoomHighValue() const { return TargetAutoZoomHighValue; }
-	void SetTargetAutoZoomRange(double LowValue, double HighValue) { TargetAutoZoomLowValue = LowValue; TargetAutoZoomHighValue = HighValue; }
-
-	/** Current auto zoom low value */
-	double GetAutoZoomLowValue() const { return AutoZoomLowValue; }
-	/** Current auto zoom high value */
-	double GetAutoZoomHighValue() const { return AutoZoomHighValue; }
-	void SetAutoZoomRange(double LowValue, double HighValue) { AutoZoomLowValue = LowValue; AutoZoomHighValue = HighValue; }
-
-	/**
-	 * Compute BaselineY and ScaleY so the [Low, High] Value range will correspond to [Top, Bottom] Y position range.
-	 * GetYForValue(InHighValue) == InTopY
-	 * GetYForValue(InLowValue) == InBottomY
-	 */
-	void ComputeBaselineAndScale(const double InLowValue, const double InHighValue, const float InTopY, const float InBottomY, double& OutBaselineY, double& OutScaleY) const
-	{
-		ensure(InLowValue < InHighValue);
-		ensure(InTopY < InBottomY);
-		const double InvRange = 1.0 / (InHighValue - InLowValue);
-		OutScaleY = static_cast<double>(InBottomY - InTopY) * InvRange;
-		//OutBaselineY = (InHighValue * static_cast<double>(InBottomY) - InLowValue * static_cast<double>(InTopY)) * InvRange;
-		OutBaselineY = static_cast<double>(InTopY) + InHighValue * OutScaleY;
-	}
-
-	/**
-	 * @param X The horizontal coordinate of the point tested; in Slate pixels (local graph coordinates)
-	 * @param Y The vertical coordinate of the point tested; in Slate pixels (local graph coordinates)
-	 * @param Viewport The timing viewport used to transform time in local graph coordinates
-	 * @param bCheckLine If needs to check the bounding box of the horizontal line (determined by duration of event and value) or only the bounding box of the visual point
-	 * @param bCheckBox If needs to check the bounding box of the entire visual box (determined by duration of event, value and baseline)
-	 * @return A pointer to an Event located at (X, Y) coordinates, if any; nullptr if no event is located at respective coordinates
-	 */
-	const FGraphSeries::FEvent* GetEvent(const float PosX, const float PosY, const FTimingTrackViewport& Viewport, bool bCheckLine, bool bCheckBox) const;
-
-	virtual FString FormatValue(double Value) const;
-
-private:
-	FText Name;
-	FText Description;
-
-	bool bIsVisible;
-	bool bIsDirty;
-
-	bool bAutoZoom;
-	double TargetAutoZoomLowValue; // target low value of auto zoom interval (corresponding to bottom of the track)
-	double TargetAutoZoomHighValue; // target high value of auto zoom interval (corresponding to top of the track)
-	double AutoZoomLowValue; // current auto zoom low value
-	double AutoZoomHighValue; // current auto zoom high value
-
-	double BaselineY; // Y position (in viewport local space) of the baseline (with Value == 0); in pixels (Slate units)
-	double ScaleY; // scale between Value units and viewport units; in pixels (Slate units) / Value unit
-
-	FLinearColor Color;
-	FLinearColor BorderColor;
+struct FSlateBrush;
 
-protected:
-	TArray<FEvent> Events; // reduced list of events; used to identify an event at a certain screen position (ex.: the event hovered by mouse)
-	TArray<FVector2D> Points; // reduced list of points; for drawing points
-	TArray<FVector2D> LinePoints; // reduced list of points; for drawing the connected line and filled polygon
-	TArray<FBox> Boxes; // reduced list of boxes; for drawing boxes
-};
+class FGraphSeries;
+class FTimingTrackViewport;
 
 ////////////////////////////////////////////////////////////////////////////////////////////////////
 
@@ -177,16 +20,6 @@ class TRACEINSIGHTS_API FGraphTrack : public FBaseTimingTrack
 {
 	friend class FGraphTrackBuilder;
 
-public:
-	struct FEvent
-	{
-		TSharedPtr<FGraphSeries> Series;
-		FGraphSeries::FEvent SeriesEvent;
-
-		void Reset() { Series.Reset(); }
-		bool IsValid() const { return Series.IsValid(); }
-	};
-
 private:
 	// Visual size of points (in pixels).
 	static constexpr float PointVisualSize = 5.0f;
@@ -196,33 +29,31 @@ private:
 	static constexpr float PointSizeY = 3.0f;
 
 public:
-	explicit FGraphTrack(uint64 InTrackId);
+	explicit FGraphTrack(const FName& InSubType = NAME_None);
 	virtual ~FGraphTrack();
 
+	//////////////////////////////////////////////////
+	// FBaseTimingTrack
+
 	virtual void Reset() override;
 
-	virtual void UpdateHoveredState(float MouseX, float MouseY, const FTimingTrackViewport& Viewport) override;
+	virtual void PostUpdate(const ITimingTrackUpdateContext& Context) override;
 
-	virtual void Update(const FTimingTrackViewport& Viewport) override;
+	virtual void Draw(const ITimingTrackDrawContext& Context) const override;
+	virtual void DrawEvent(const ITimingTrackDrawContext& Context, const ITimingEvent& InTimingEvent, EDrawEventMode InDrawMode) const override;
 
-	void Draw(FDrawContext& DrawContext, const FTimingTrackViewport& Viewport, const FVector2D& MousePosition) const;
-	void PostDraw(FDrawContext& DrawContext, const FTimingTrackViewport& Viewport, const FVector2D& MousePosition) const;
+	virtual void InitTooltip(FTooltipDrawState& InOutTooltip, const ITimingEvent& InTooltipEvent) const override;
 
-	/**
-	 * @param X The horizontal coordinate of the point tested; in Slate pixels (local graph coordinates)
-	 * @param Y The vertical coordinate of the point tested; in Slate pixels (local graph coordinates)
-	 * @param Viewport The timing viewport used to transform time in local graph coordinates
-	 * @param OutEvent The returned event located at (PosX, PosY) coordinates, if any.
-	 * @return True if an event is found at (PosX, PosY) coordinates.
-	 */
-	const bool GetEvent(const float MouseX, const float MouseY, const FTimingTrackViewport& Viewport, FGraphTrack::FEvent& OutEvent) const;
+	virtual const TSharedPtr<const ITimingEvent> GetEvent(float InPosX, float InPosY, const FTimingTrackViewport& Viewport) const override;
 
 	virtual void BuildContextMenu(FMenuBuilder& MenuBuilder) override;
 
+	//////////////////////////////////////////////////
+
 	TArray<TSharedPtr<FGraphSeries>>& GetSeries() { return AllSeries; }
 
-	//virtual int GetDebugLineCount() override;
-	//virtual void BuildDebugLines(FString& OutStr) override;
+	//TODO: virtual int GetDebugStringLineCount() const override;
+	//TODO: virtual void BuildDebugString(FString& OutStr) const override;
 
 	int32 GetNumAddedEvents() const { return NumAddedEvents; }
 	int32 GetNumDrawPoints() const { return NumDrawPoints; }
@@ -232,11 +63,7 @@ public:
 protected:
 	void UpdateStats();
 
-	void DrawBackground(FDrawContext& DrawContext, const FTimingTrackViewport& Viewport) const;
 	void DrawSeries(const FGraphSeries& Series, FDrawContext& DrawContext, const FTimingTrackViewport& Viewport) const;
-	void DrawHighlightedEvent(FDrawContext& DrawContext, const FTimingTrackViewport& Viewport, const FGraphTrack::FEvent& GraphEvent) const;
-
-	void InitTooltip();
 
 private:
 	void ContextMenu_ShowPoints_Execute();
@@ -283,10 +110,6 @@ protected:
 	bool bUseEventDuration;
 	bool bDrawBoxes;
 
-	mutable FGraphTrack::FEvent HoveredGraphEvent;
-	FTooltipDrawState Tooltip;
-	FVector2D TooltipMousePosition;
-
 	// Stats
 	int32 NumAddedEvents; // total event count
 	int32 NumDrawPoints;
@@ -296,13 +119,13 @@ protected:
 
 ////////////////////////////////////////////////////////////////////////////////////////////////////
 
-class FRandomGraphTrack : public FGraphTrack
+class TRACEINSIGHTS_API FRandomGraphTrack : public FGraphTrack
 {
 public:
-	explicit FRandomGraphTrack(uint64 InTrackId);
+	FRandomGraphTrack();
 	virtual ~FRandomGraphTrack();
 
-	virtual void Update(const FTimingTrackViewport& Viewport) override;
+	virtual void Update(const ITimingTrackUpdateContext& Context) override;
 
 	void AddDefaultSeries();
 
@@ -311,70 +134,3 @@ protected:
 };
 
 ////////////////////////////////////////////////////////////////////////////////////////////////////
-
-class FGraphTrackBuilder
-{
-private:
-	struct FPointInfo
-	{
-		bool bValid;
-		float X;
-		float Y;
-
-		FPointInfo() : bValid(false) {}
-	};
-
-public:
-	explicit FGraphTrackBuilder(FGraphTrack& InTrack, FGraphSeries& InSeries, const FTimingTrackViewport& InViewport);
-	~FGraphTrackBuilder();
-
-	/**
-	 * Non-copyable
-	 */
-	FGraphTrackBuilder(const FGraphTrackBuilder&) = delete;
-	FGraphTrackBuilder& operator=(const FGraphTrackBuilder&) = delete;
-
-	FGraphTrack& GetTrack() const { return Track; }
-	FGraphSeries& GetSeries() const { return Series; }
-	const FTimingTrackViewport& GetViewport() const { return Viewport; }
-
-	void AddEvent(double Time, double Duration, double Value);
-
-private:
-	void BeginPoints();
-	bool AddPoint(double Time, double Value);
-	void FlushPoints();
-	void EndPoints();
-
-	void BeginConnectedLines();
-	void AddConnectedLine(double Time, double Value);
-	void FlushConnectedLine();
-	void EndConnectedLines();
-
-	void BeginBoxes();
-	void AddBox(double Time, double Duration, double Value);
-	void FlushBox();
-	void EndBoxes();
-
-private:
-	FGraphTrack& Track;
-	FGraphSeries& Series;
-	const FTimingTrackViewport& Viewport;
-
-	// Used by the point reduction algorithm.
-	double PointsCurrentX;
-	TArray<FPointInfo> PointsAtCurrentX;
-
-	// Used by the line reduction algorithm.
-	float LinesCurrentX;
-	float LinesMinY;
-	float LinesMaxY;
-	float LinesFirstY;
-	float LinesLastY;
-	bool bIsLastLineAdded;
-
-	// Used by the box reduction algorithm.
-	//...
-};
-
-////////////////////////////////////////////////////////////////////////////////////////////////////
diff --git a/Engine/Source/Developer/TraceInsights/Public/Insights/ViewModels/GraphTrackBuilder.h b/Engine/Source/Developer/TraceInsights/Public/Insights/ViewModels/GraphTrackBuilder.h
new file mode 100644
--- /dev/null
+++ b/Engine/Source/Developer/TraceInsights/Public/Insights/ViewModels/GraphTrackBuilder.h
@@ -0,0 +1,79 @@
+// Copyright 1998-2019 Epic Games, Inc. All Rights Reserved.
+
+#pragma once
+
+#include "CoreMinimal.h"
+
+// Insights
+class FGraphTrack;
+class FGraphSeries;
+class FTimingTrackViewport;
+
+////////////////////////////////////////////////////////////////////////////////////////////////////
+
+class TRACEINSIGHTS_API FGraphTrackBuilder
+{
+private:
+	struct FPointInfo
+	{
+		bool bValid;
+		float X;
+		float Y;
+
+		FPointInfo() : bValid(false) {}
+	};
+
+public:
+	explicit FGraphTrackBuilder(FGraphTrack& InTrack, FGraphSeries& InSeries, const FTimingTrackViewport& InViewport);
+	~FGraphTrackBuilder();
+
+	/**
+	 * Non-copyable
+	 */
+	FGraphTrackBuilder(const FGraphTrackBuilder&) = delete;
+	FGraphTrackBuilder& operator=(const FGraphTrackBuilder&) = delete;
+
+	FGraphTrack& GetTrack() const { return Track; }
+	FGraphSeries& GetSeries() const { return Series; }
+	const FTimingTrackViewport& GetViewport() const { return Viewport; }
+
+	void AddEvent(double Time, double Duration, double Value);
+
+private:
+	void BeginPoints();
+	bool AddPoint(double Time, double Value);
+	void FlushPoints();
+	void EndPoints();
+
+	void BeginConnectedLines();
+	void AddConnectedLine(double Time, double Value);
+	void FlushConnectedLine();
+	void EndConnectedLines();
+
+	void BeginBoxes();
+	void AddBox(double Time, double Duration, double Value);
+	void FlushBox();
+	void EndBoxes();
+
+private:
+	FGraphTrack& Track;
+	FGraphSeries& Series;
+	const FTimingTrackViewport& Viewport;
+
+	// Used by the point reduction algorithm.
+	double PointsCurrentX;
+	TArray<FPointInfo> PointsAtCurrentX;
+
+	// Used by the line reduction algorithm.
+	float LinesCurrentX;
+	float LinesMinY;
+	float LinesMaxY;
+	float LinesFirstY;
+	float LinesLastY;
+	bool bIsLastLineAdded;
+
+	// Used by the box reduction algorithm.
+	//...
+};
+
+////////////////////////////////////////////////////////////////////////////////////////////////////
diff --git a/Engine/Source/Developer/TraceInsights/Public/Insights/ViewModels/TimingEvent.h b/Engine/Source/Developer/TraceInsights/Public/Insights/ViewModels/TimingEvent.h
--- a/Engine/Source/Developer/TraceInsights/Public/Insights/ViewModels/TimingEvent.h
+++ b/Engine/Source/Developer/TraceInsights/Public/Insights/ViewModels/TimingEvent.h
@@ -3,41 +3,19 @@
 #pragma once
 
 #include "CoreMinimal.h"
-#include "TraceServices/Model/LoadTimeProfiler.h"
-#include "TimingEventSearch.h"
 
-class FTimingEventsTrack;
+// Insights
+#include "Insights/ViewModels/ITimingEvent.h"
+#include "Insights/ViewModels/TimingEventSearch.h"
 
-struct FTimingEvent
-{
-	// The track this timing event is contained within
-	const FTimingEventsTrack* Track;
+class FBaseTimingTrack;
 
-	// Handle to a previous search, used to accelerate access to underlying event data
-	mutable FTimingEventSearchHandle SearchHandle;
-
-	// The start time of the event
-	double StartTime;
-
-	// The end time of the event
-	double EndTime;
-
-	// For hierarchical events, the cached exclusive time
-	double ExclusiveTime;
+////////////////////////////////////////////////////////////////////////////////////////////////////
 
-	// The depth of the event
-	uint32 Depth;
-
-	FTimingEvent()
-		: Track(nullptr)
-		, StartTime(0.0)
-		, EndTime(-1.0)
-		, ExclusiveTime(0.0)
-		, Depth(0)
-	{
-	}
-
-	FTimingEvent(const FTimingEventsTrack* InTrack, double InStartTime, double InEndTime, uint32 InDepth)
+class TRACEINSIGHTS_API FTimingEvent : public ITimingEvent
+{
+public:
+	FTimingEvent(const TSharedRef<const FBaseTimingTrack> InTrack, double InStartTime, double InEndTime, uint32 InDepth)
 		: Track(InTrack)
 		, StartTime(InStartTime)
 		, EndTime(InEndTime)
@@ -45,37 +23,71 @@ struct FTimingEvent
 		, Depth(InDepth)
 	{}
 
+	virtual ~FTimingEvent() {}
+
 	FTimingEvent(const FTimingEvent&) = default;
 	FTimingEvent& operator=(const FTimingEvent&) = default;
 
 	FTimingEvent(FTimingEvent&&) = default;
 	FTimingEvent& operator=(FTimingEvent&&) = default;
 
-	void Reset()
-	{
-		Track = nullptr;
-		StartTime = 0.0;
-		EndTime = -1.0;
-		ExclusiveTime = 0.0;
-		Depth = 0;
-	}
+	//////////////////////////////////////////////////
+	// ITimingEvent interface
 
-	bool IsValidTrack() const { return Track != nullptr; }
+	virtual const FName& GetTypeName() const override { return FTimingEvent::TypeName; }
 
-	bool IsValid() const { return Track != nullptr && StartTime <= EndTime; }
+	virtual const TSharedRef<const FBaseTimingTrack> GetTrack() const override { return Track; }
 
-	double Duration() const { return EndTime - StartTime; }
+	virtual uint32 GetDepth() const override { return Depth; }
 
-	bool Equals(const FTimingEvent& Other) const
-	{
-		return Track == Other.Track
-			&& Depth == Other.Depth
-			&& StartTime == Other.StartTime
-			&& EndTime == Other.EndTime;
-	}
+	virtual double GetStartTime() const override { return StartTime; }
+	virtual double GetEndTime() const override { return EndTime; }
+	virtual double GetDuration() const override { return EndTime - StartTime; }
 
-	static bool AreEquals(const FTimingEvent& A, const FTimingEvent& B)
+	virtual bool Equals(const ITimingEvent& Other) const override
 	{
-		return A.Equals(B);
+		if (GetTypeName() != Other.GetTypeName())
+		{
+			return false;
+		}
+
+		const FTimingEvent& OtherTimingEvent = static_cast<const FTimingEvent&>(Other);
+		return Track == Other.GetTrack()
+			&& Depth == OtherTimingEvent.GetDepth()
+			&& StartTime == OtherTimingEvent.GetStartTime()
+			&& EndTime == OtherTimingEvent.GetEndTime();
 	}
+
+	//////////////////////////////////////////////////
+
+	FTimingEventSearchHandle& GetSearchHandle() const { return SearchHandle; }
+
+	double GetExclusiveTime() const { return ExclusiveTime; }
+	void SetExclusiveTime(double InExclusiveTime) { ExclusiveTime = InExclusiveTime; }
+
+	static const FName& GetStaticTypeName() { return FTimingEvent::TypeName; }
+	static bool CheckTypeName(const ITimingEvent& Event) { return Event.GetTypeName() == FTimingEvent::TypeName; }
+
+private:
+	// The track this timing event is contained within
+	TSharedRef<const FBaseTimingTrack> Track;
+
+	// Handle to a previous search, used to accelerate access to underlying event data
+	mutable FTimingEventSearchHandle SearchHandle;
+
+	// The start time of the event
+	double StartTime;
+
+	// The end time of the event
+	double EndTime;
+
+	// For hierarchical events, the cached exclusive time
+	double ExclusiveTime;
+
+	// The depth of the event
+	uint32 Depth;
+
+	static const FName TypeName;
 };
+
+////////////////////////////////////////////////////////////////////////////////////////////////////
diff --git a/Engine/Source/Developer/TraceInsights/Public/Insights/ViewModels/TimingEventsTrack.h b/Engine/Source/Developer/TraceInsights/Public/Insights/ViewModels/TimingEventsTrack.h
--- a/Engine/Source/Developer/TraceInsights/Public/Insights/ViewModels/TimingEventsTrack.h
+++ b/Engine/Source/Developer/TraceInsights/Public/Insights/ViewModels/TimingEventsTrack.h
@@ -7,88 +7,37 @@
 // Insights
 #include "Insights/ViewModels/BaseTimingTrack.h"
 
-struct FDrawContext;
-struct FSlateBrush;
-struct FSlateFontInfo;
-struct FTimingEvent;
-class FTimingTrackViewport;
-struct FTimingViewTooltip;
-class ITimingViewDrawHelper;
-class FTooltipDrawState;
-namespace Trace { class IAnalysisSession; }
-class FTimingEventSearchParameters;
-
-////////////////////////////////////////////////////////////////////////////////////////////////////
-
-struct TRACEINSIGHTS_API FTimingEventsTrackLayout
-{
-	static constexpr float RealMinTimelineH = 13.0f;
-
-	static constexpr float NormalLayoutEventH = 14.0f;
-	static constexpr float NormalLayoutEventDY = 2.0f;
-	static constexpr float NormalLayoutTimelineDY = 14.0f;
-	static constexpr float NormalLayoutMinTimelineH = 0.0f;
-
-	static constexpr float CompactLayoutEventH = 2.0f;
-	static constexpr float CompactLayoutEventDY = 1.0f;
-	static constexpr float CompactLayoutTimelineDY = 3.0f;
-	static constexpr float CompactLayoutMinTimelineH = 0.0f;
-
-	//////////////////////////////////////////////////
-
-	bool bIsCompactMode;
-
-	float EventH; // height of a timing event, in Slate units
-	float EventDY; // vertical distance between two timing event sub-tracks, in Slate units
-	float TimelineDY; // space at top and bottom of each timeline, in Slate units
-	float MinTimelineH;
-	float TargetMinTimelineH;
-
-	//////////////////////////////////////////////////
-
-	float GetLaneY(uint32 Depth) const { return 1.0f + TimelineDY + Depth * (EventDY + EventH); }
-
-	void ForceNormalMode();
-	void ForceCompactMode();
-	void Update();
-};
+struct FTimingEventsTrackDrawState;
 
 ////////////////////////////////////////////////////////////////////////////////////////////////////
 
 class TRACEINSIGHTS_API FTimingEventsTrack : public FBaseTimingTrack
 {
 public:
-	explicit FTimingEventsTrack(uint64 InTrackId, const FName& InType, const FName& InSubType, const FString& InName);
+	explicit FTimingEventsTrack(const FName& InType, const FName& InSubType, const FString& InName);
 	virtual ~FTimingEventsTrack();
 
-	int32 GetDepth() const { return NumLanes - 1; }
 	int32 GetNumLanes() const { return NumLanes; }
-	void SetDepth(int32 InDepth) { NumLanes = InDepth + 1; }
+	void SetNumLanes(int32 InNumLanes) { NumLanes = InNumLanes; }
+
+	//////////////////////////////////////////////////
+	// FBaseTimingTrack
 
 	virtual void Reset() override;
-	virtual void UpdateHoveredState(float MX, float MY, const FTimingTrackViewport& Viewport);
+	virtual void PostUpdate(const ITimingTrackUpdateContext& Context) override;
 
-	virtual void Draw(ITimingViewDrawHelper& Helper) const = 0;
-	virtual void DrawSelectedEventInfo(const FTimingEvent& SelectedTimingEvent, const FTimingTrackViewport& Viewport, const FDrawContext& DrawContext, const FSlateBrush* WhiteBrush, const FSlateFontInfo& Font) const {}
-	virtual void InitTooltip(FTooltipDrawState& Tooltip, const FTimingEvent& HoveredTimingEvent) const {}
+	virtual void DrawEvent(const ITimingTrackDrawContext& Context, const ITimingEvent& InTimingEvent, EDrawEventMode InDrawMode) const override;
 
-	virtual bool SearchTimingEvent(const FTimingEventSearchParameters& InSearchParameters, FTimingEvent& InOutTimingEvent) const { return false; }
-	virtual void ComputeTimingEventStats(FTimingEvent& InOutTimingEvent) const {}
+	virtual const TSharedPtr<const ITimingEvent> GetEvent(float InPosX, float InPosY, const FTimingTrackViewport& Viewport) const override;
 
-	// Called back from the timing view when an event is selected
-	virtual void OnEventSelected(const FTimingEvent& SelectedTimingEvent) const {}
+	//////////////////////////////////////////////////
 
-	// Called back from the timing view when a selected event is copied to the clipboard with Ctrl+C
-	virtual void OnClipboardCopyEvent(const FTimingEvent& SelectedTimingEvent) const {}
+protected:
+	void UpdateTrackHeight(const ITimingTrackUpdateContext& Context);
 
 private:
 	int32 NumLanes; // number of lanes (sub-tracks)
-
-	// TODO: Cached OnPaint state.
-	//TArray<FEventBoxInfo> Boxes;
-	//TArray<FEventBoxInfo> MergedBorderBoxes;
-	//TArray<FEventBoxInfo> Borders;
-	//TArray<FTextInfo> Texts;
+	TSharedPtr<FTimingEventsTrackDrawState> DrawState;
 
 public:
 	static bool bUseDownSampling;
diff --git a/Engine/Source/Developer/TraceInsights/Public/Insights/ViewModels/TimingTrackViewport.h b/Engine/Source/Developer/TraceInsights/Public/Insights/ViewModels/TimingTrackViewport.h
--- a/Engine/Source/Developer/TraceInsights/Public/Insights/ViewModels/TimingTrackViewport.h
+++ b/Engine/Source/Developer/TraceInsights/Public/Insights/ViewModels/TimingTrackViewport.h
@@ -4,10 +4,60 @@
 
 #include "CoreMinimal.h"
 
+// Insights
+#include "Insights/ViewModels/TimingViewLayout.h"
+
 class SScrollBar;
 
 ////////////////////////////////////////////////////////////////////////////////////////////////////
 
+enum class ETimingTrackViewportDirtyFlags : uint32
+{
+	None = 0,
+
+	//////////////////////////////////////////////////
+	// Horizontal viewport flags
+
+	/** Width has changed. */
+	HSizeChanged = (1 << 0),
+
+	/** StartTime has changed. */
+	HPositionChanged = (1 << 1),
+
+	/** ScaleX has changed. */
+	HScaleChanged = (1 << 2),
+
+	/** Intersection between valid session time ([StartSessionTime - EndSessionTime) interval) and viewport time ([StartTime - EndTime) interval) has changed. */
+	HClippedSessionTimeChanged = (1 << 3),
+
+	/** Misc horizontal viewport invalidation flag. */
+	HInvalidated = (1 << 4),
+
+	//////////////////////////////////////////////////
+	// Vertical viewport flags
+
+	/** Height has changed. */
+	VSizeChanged = (1 << 5),
+
+	/** ScrollPosY has changed. */
+	VPositionChanged = (1 << 6),
+
+	/** Layout has changed. */
+	VLayoutChanged = (1 << 7),
+
+	/** Misc vertical viewport invalidation flag. */
+	VInvalidated = (1 << 8),
+
+	//////////////////////////////////////////////////
+
+	AllHorizontal = HSizeChanged | HPositionChanged | HScaleChanged | HClippedSessionTimeChanged | HInvalidated,
+	AllVertical = VSizeChanged | VPositionChanged | VLayoutChanged | VInvalidated,
+	All = AllHorizontal | AllVertical
+};
+ENUM_CLASS_FLAGS(ETimingTrackViewportDirtyFlags);
+
+////////////////////////////////////////////////////////////////////////////////////////////////////
+
 class TRACEINSIGHTS_API FTimingTrackViewport
 {
 public:
@@ -32,8 +82,12 @@ public:
 		ScaleX = (5 * 20) / 5.0; // 5s between major tick marks
 
 		TopOffset = 0.0f;
+		BottomOffset = 0.0f;
 		ScrollHeight = 1.0f;
 		ScrollPosY = 0.0f;
+		Layout.ForceNormalMode();
+
+		DirtyFlags = ETimingTrackViewportDirtyFlags::All;
 	}
 
 	/** Width of viewport, in pixels. [px] */
@@ -127,12 +181,17 @@ public:
 	//////////////////////////////////////////////////
 	// Vertical axis
 
-	/** Top offset (to allow the time ruller to be visible), in pixels. [px] */
+	/** Top offset (total height of the visible top docked tracks), in pixels. [px] */
 	float GetTopOffset() const { return TopOffset; }
-
 	void SetTopOffset(float InTopOffset) { TopOffset = InTopOffset; }
+	
+	/** Bottom offset (total height of the visible bottom docked tracks), in pixels. [px] */
+	float GetBottomOffset() const { return BottomOffset; }
+	void SetBottomOffset(float InBottomOffset) { BottomOffset = InBottomOffset; }
 
-	/** Height of the vertical scrollable area, in pixels. [px] */
+	float GetScrollableAreaHeight() const { return Height - TopOffset - BottomOffset; }
+
+	/** Height of the vertical (virtual) scrollable area, in pixels. [px] */
 	float GetScrollHeight() const { return ScrollHeight; }
 	void SetScrollHeight(const float InScrollHeight) { ScrollHeight = InScrollHeight; }
 
@@ -142,6 +201,30 @@ public:
 
 	float GetViewportY(const float Y) const { return TopOffset + Y - ScrollPosY; }
 
+	const FTimingViewLayout& GetLayout() const { return Layout; }
+	void UpdateLayout();
+
+	void SwitchLayoutCompactMode()
+	{
+		Layout.bIsCompactMode = !Layout.bIsCompactMode;
+		AddDirtyFlags(ETimingTrackViewportDirtyFlags::VLayoutChanged);
+	}
+
+	void ToggleLayoutMinTrackHeight()
+	{
+		if (Layout.TargetMinTimelineH == 0.0f)
+		{
+			Layout.TargetMinTimelineH = FTimingViewLayout::RealMinTimelineH;
+		}
+		else
+		{
+			Layout.TargetMinTimelineH = 0.0f;
+		}
+
+		Layout.MinTimelineH = Layout.TargetMinTimelineH; // no layout animation
+		AddDirtyFlags(ETimingTrackViewportDirtyFlags::VLayoutChanged);
+	}
+
 	//////////////////////////////////////////////////
 
 	bool OnUserScrolled(TSharedPtr<SScrollBar> ScrollBar, float ScrollOffset);
@@ -152,6 +235,14 @@ public:
 
 	//////////////////////////////////////////////////
 
+	ETimingTrackViewportDirtyFlags GetDirtyFlags() const { return DirtyFlags; }
+	bool IsDirty() const { return DirtyFlags != ETimingTrackViewportDirtyFlags::None; }
+	bool IsDirty(ETimingTrackViewportDirtyFlags InDirtyFlags) const { return EnumHasAnyFlags(DirtyFlags, InDirtyFlags); }
+	bool IsHorizontalViewportDirty() const { return EnumHasAnyFlags(DirtyFlags, ETimingTrackViewportDirtyFlags::AllHorizontal); }
+	bool IsVerticalViewportDirty() const { return EnumHasAnyFlags(DirtyFlags, ETimingTrackViewportDirtyFlags::AllVertical); }
+	void AddDirtyFlags(ETimingTrackViewportDirtyFlags InFlags) { DirtyFlags |= InFlags; }
+	void ResetDirtyFlags() { DirtyFlags = ETimingTrackViewportDirtyFlags::None; }
+
 private:
 	float Width; // width of viewport, in pixels; [px]
 	float Height; // height of viewport, in pixels; [px]
@@ -166,9 +257,13 @@ private:
 	double MaxScaleX; // max scale factor; [px/s]
 	double ScaleX; // scale factor between seconds and pixels; [px/s]
 
-	float TopOffset; // top offset (to allow the time ruller to be visible), in pixels; [px]
+	float TopOffset; // top offset, in pixels; [px]
+	float BottomOffset; // bottom offset, in pixels; [px]
 	float ScrollHeight; // height of the vertical scrollable area, in pixels; [px]
 	float ScrollPosY; // current vertical scroll position, in pixels; [px]
+	FTimingViewLayout Layout;
+
+	ETimingTrackViewportDirtyFlags DirtyFlags;
 };
 
 ////////////////////////////////////////////////////////////////////////////////////////////////////
