[Insights]
  - Added  ITimingEventsTrackDrawStateBuilder public API. Also FTimingEventsTrack now exposes BuildDrawState() -- to be implemented by derived classes.

diff --git a/Engine/Source/Developer/TraceInsights/Private/Insights/ViewModels/FileActivityTimingTrack.cpp b/Engine/Source/Developer/TraceInsights/Private/Insights/ViewModels/FileActivityTimingTrack.cpp
--- a/Engine/Source/Developer/TraceInsights/Private/Insights/ViewModels/FileActivityTimingTrack.cpp
+++ b/Engine/Source/Developer/TraceInsights/Private/Insights/ViewModels/FileActivityTimingTrack.cpp
@@ -17,7 +17,6 @@
 #include "Insights/TimingProfilerCommon.h"
 #include "Insights/ViewModels/TimingEvent.h"
 #include "Insights/ViewModels/TimingTrackViewport.h"
-#include "Insights/ViewModels/TimingViewDrawHelper.h"
 #include "Insights/ViewModels/TooltipDrawState.h"
 #include "Insights/ViewModels/TimingEventSearch.h"
 #include "Insights/Widgets/STimingView.h"
@@ -628,67 +627,48 @@ bool FFileActivityTimingTrack::FindIoTimingEvent(const FTimingEventSearchParamet
 // FOverviewFileActivityTimingTrack
 ////////////////////////////////////////////////////////////////////////////////////////////////////
 
-void FOverviewFileActivityTimingTrack::PreUpdate(const ITimingTrackUpdateContext& Context)
+void FOverviewFileActivityTimingTrack::BuildDrawState(ITimingEventsTrackDrawStateBuilder& Builder, const ITimingTrackUpdateContext& Context)
 {
 	const FTimingTrackViewport& Viewport = Context.GetViewport();
-	if (IsDirty() || Viewport.IsHorizontalViewportDirty())
-	{
-		ClearDirtyFlag();
 
-		FTimingEventsTrackDrawStateBuilder Builder(DrawState, Viewport);
+	for (const FFileActivitySharedState::FIoTimingEvent& Event : SharedState.AllIoEvents)
+	{
+		const Trace::EFileActivityType ActivityType = static_cast<Trace::EFileActivityType>(Event.Type & 0x0F);
+		const uint64 EventType = static_cast<uint64>(ActivityType);
 
-		for (const FFileActivitySharedState::FIoTimingEvent& Event : SharedState.AllIoEvents)
+		if (ActivityType >= Trace::FileActivityType_Count)
 		{
-			const Trace::EFileActivityType ActivityType = static_cast<Trace::EFileActivityType>(Event.Type & 0x0F);
-			const uint64 EventType = static_cast<uint64>(ActivityType);
-
-			if (ActivityType >= Trace::FileActivityType_Count)
-			{
-				// Ignore "Idle" and "NotClosed" events.
-				continue;
-			}
+			// Ignore "Idle" and "NotClosed" events.
+			continue;
+		}
 
-			//const double EventEndTime = Event.EndTime; // keep duration of events
-			const double EventEndTime = Event.StartTime; // make all 0 duration events
+		//const double EventEndTime = Event.EndTime; // keep duration of events
+		const double EventEndTime = Event.StartTime; // make all 0 duration events
 
-			if (EventEndTime <= Viewport.GetStartTime())
-			{
-				continue;
-			}
-			if (Event.StartTime >= Viewport.GetEndTime())
-			{
-				break;
-			}
+		if (EventEndTime <= Viewport.GetStartTime())
+		{
+			continue;
+		}
+		if (Event.StartTime >= Viewport.GetEndTime())
+		{
+			break;
+		}
 
-			const bool bHasFailed = ((Event.Type & 0xF0) != 0);
+		const bool bHasFailed = ((Event.Type & 0xF0) != 0);
 
-			if (bHasFailed)
-			{
-				FString EventName = TEXT("Failed ");
-				EventName += GetFileActivityTypeName(ActivityType);
-				const uint32 Color = 0xFFAA0000;
-				Builder.AddEvent(Event.StartTime, EventEndTime, 0, *EventName, EventType, Color);
-			}
-			else
-			{
-				const uint32 Color = GetFileActivityTypeColor(ActivityType);
-				Builder.AddEvent(Event.StartTime, EventEndTime, 0, GetFileActivityTypeName(ActivityType), EventType, Color);
-			}
+		if (bHasFailed)
+		{
+			FString EventName = TEXT("Failed ");
+			EventName += GetFileActivityTypeName(ActivityType);
+			const uint32 Color = 0xFFAA0000;
+			Builder.AddEvent(Event.StartTime, EventEndTime, 0, *EventName, EventType, Color);
+		}
+		else
+		{
+			const uint32 Color = GetFileActivityTypeColor(ActivityType);
+			Builder.AddEvent(Event.StartTime, EventEndTime, 0, GetFileActivityTypeName(ActivityType), EventType, Color);
 		}
-
-		Builder.Flush();
-		SetNumLanes(1);
 	}
-
-	UpdateTrackHeight(Context);
-}
-
-////////////////////////////////////////////////////////////////////////////////////////////////////
-
-void FOverviewFileActivityTimingTrack::Draw(const ITimingTrackDrawContext& Context) const
-{
-	const FTimingViewDrawHelper& Helper = Context.GetHelper();
-	Helper.DrawEventsTrack(DrawState, *this);
 }
 
 ////////////////////////////////////////////////////////////////////////////////////////////////////
@@ -710,95 +690,76 @@ const TSharedPtr<const ITimingEvent> FOverviewFileActivityTimingTrack::SearchEve
 // FDetailedFileActivityTimingTrack
 ////////////////////////////////////////////////////////////////////////////////////////////////////
 
-void FDetailedFileActivityTimingTrack::PreUpdate(const ITimingTrackUpdateContext& Context)
+void FDetailedFileActivityTimingTrack::BuildDrawState(ITimingEventsTrackDrawStateBuilder& Builder, const ITimingTrackUpdateContext& Context)
 {
 	const FTimingTrackViewport& Viewport = Context.GetViewport();
-	if (IsDirty() || Viewport.IsHorizontalViewportDirty())
-	{
-		ClearDirtyFlag();
-
-		FTimingEventsTrackDrawStateBuilder Builder(DrawState, Viewport);
 
-		// Add IO file activity background events.
-		if (SharedState.bShowFileActivityBackgroundEvents)
+	// Add IO file activity background events.
+	if (SharedState.bShowFileActivityBackgroundEvents)
+	{
+		for (const TSharedPtr<FFileActivitySharedState::FIoFileActivity> Activity : SharedState.FileActivities)
 		{
-			for (const TSharedPtr<FFileActivitySharedState::FIoFileActivity> Activity : SharedState.FileActivities)
+			if (Activity->EndTime <= Viewport.GetStartTime())
 			{
-				if (Activity->EndTime <= Viewport.GetStartTime())
-				{
-					continue;
-				}
-				if (Activity->StartTime >= Viewport.GetEndTime())
-				{
-					break;
-				}
+				continue;
+			}
+			if (Activity->StartTime >= Viewport.GetEndTime())
+			{
+				break;
+			}
 
-				ensure(Activity->Depth <= 10000);
+			ensure(Activity->Depth <= 10000);
 
-				const uint64 EventType = uint64(Trace::EFileActivityType::FileActivityType_Count);
+			const uint64 EventType = uint64(Trace::EFileActivityType::FileActivityType_Count);
 
-				Builder.AddEvent(Activity->StartTime, Activity->EndTime, Activity->Depth, Activity->Path, EventType, 0x55333333);
-			}
+			Builder.AddEvent(Activity->StartTime, Activity->EndTime, Activity->Depth, Activity->Path, EventType, 0x55333333);
 		}
+	}
 
-		// Add IO file activity foreground events.
-		for (const FFileActivitySharedState::FIoTimingEvent& Event : SharedState.AllIoEvents)
+	// Add IO file activity foreground events.
+	for (const FFileActivitySharedState::FIoTimingEvent& Event : SharedState.AllIoEvents)
+	{
+		if (Event.EndTime <= Viewport.GetStartTime())
 		{
-			if (Event.EndTime <= Viewport.GetStartTime())
-			{
-				continue;
-			}
-			if (Event.StartTime >= Viewport.GetEndTime())
-			{
-				break;
-			}
+			continue;
+		}
+		if (Event.StartTime >= Viewport.GetEndTime())
+		{
+			break;
+		}
 
-			ensure(Event.Depth <= 10000);
+		ensure(Event.Depth <= 10000);
 
-			const Trace::EFileActivityType ActivityType = static_cast<Trace::EFileActivityType>(Event.Type & 0x0F);
-			const uint64 EventType = static_cast<uint64>(ActivityType);
+		const Trace::EFileActivityType ActivityType = static_cast<Trace::EFileActivityType>(Event.Type & 0x0F);
+		const uint64 EventType = static_cast<uint64>(ActivityType);
 
-			if (ActivityType < Trace::FileActivityType_Count)
-			{
-				const bool bHasFailed = ((Event.Type & 0xF0) != 0);
+		if (ActivityType < Trace::FileActivityType_Count)
+		{
+			const bool bHasFailed = ((Event.Type & 0xF0) != 0);
 
-				if (bHasFailed)
-				{
-					FString EventName = TEXT("Failed ");
-					EventName += GetFileActivityTypeName(ActivityType);
-					const uint32 Color = 0xFFAA0000;
-					Builder.AddEvent(Event.StartTime, Event.EndTime, Event.Depth, *EventName, EventType, Color);
-				}
-				else
-				{
-					const uint32 Color = GetFileActivityTypeColor(ActivityType);
-					Builder.AddEvent(Event.StartTime, Event.EndTime, Event.Depth, GetFileActivityTypeName(ActivityType), EventType, Color);
-				}
+			if (bHasFailed)
+			{
+				FString EventName = TEXT("Failed ");
+				EventName += GetFileActivityTypeName(ActivityType);
+				const uint32 Color = 0xFFAA0000;
+				Builder.AddEvent(Event.StartTime, Event.EndTime, Event.Depth, *EventName, EventType, Color);
 			}
 			else
 			{
-				FString EventName = GetFileActivityTypeName(ActivityType);
-				EventName += " [";
-				EventName += Event.FileActivity->Path;
-				EventName += "]";
 				const uint32 Color = GetFileActivityTypeColor(ActivityType);
-				Builder.AddEvent(Event.StartTime, Event.EndTime, Event.Depth, *EventName, EventType, Color);
+				Builder.AddEvent(Event.StartTime, Event.EndTime, Event.Depth, GetFileActivityTypeName(ActivityType), EventType, Color);
 			}
 		}
-
-		Builder.Flush();
-		SetNumLanes(Builder.GetMaxDepth() + 1);
+		else
+		{
+			FString EventName = GetFileActivityTypeName(ActivityType);
+			EventName += " [";
+			EventName += Event.FileActivity->Path;
+			EventName += "]";
+			const uint32 Color = GetFileActivityTypeColor(ActivityType);
+			Builder.AddEvent(Event.StartTime, Event.EndTime, Event.Depth, *EventName, EventType, Color);
+		}
 	}
-
-	UpdateTrackHeight(Context);
-}
-
-////////////////////////////////////////////////////////////////////////////////////////////////////
-
-void FDetailedFileActivityTimingTrack::Draw(const ITimingTrackDrawContext& Context) const
-{
-	const FTimingViewDrawHelper& Helper = Context.GetHelper();
-	Helper.DrawEventsTrack(DrawState, *this);
 }
 
 ////////////////////////////////////////////////////////////////////////////////////////////////////
diff --git a/Engine/Source/Developer/TraceInsights/Private/Insights/ViewModels/FileActivityTimingTrack.h b/Engine/Source/Developer/TraceInsights/Private/Insights/ViewModels/FileActivityTimingTrack.h
--- a/Engine/Source/Developer/TraceInsights/Private/Insights/ViewModels/FileActivityTimingTrack.h
+++ b/Engine/Source/Developer/TraceInsights/Private/Insights/ViewModels/FileActivityTimingTrack.h
@@ -7,7 +7,6 @@
 // Insights
 #include "Insights/ITimingViewExtender.h"
 #include "Insights/ViewModels/TimingEventsTrack.h"
-#include "Insights/ViewModels/TimingViewDrawHelper.h"
 
 class FTimingEventSearchParameters;
 class STimingView;
@@ -101,7 +100,6 @@ protected:
 
 protected:
 	FFileActivitySharedState& SharedState;
-	FTimingEventsTrackDrawState DrawState;
 };
 
 ////////////////////////////////////////////////////////////////////////////////////////////////////
@@ -114,8 +112,7 @@ public:
 	{
 	}
 
-	virtual void PreUpdate(const ITimingTrackUpdateContext& Context) override;
-	virtual void Draw(const ITimingTrackDrawContext& Context) const override;
+	virtual void BuildDrawState(ITimingEventsTrackDrawStateBuilder& Builder, const ITimingTrackUpdateContext& Context) override;
 	virtual const TSharedPtr<const ITimingEvent> SearchEvent(const FTimingEventSearchParameters& InSearchParameters) const override;
 };
 
@@ -129,8 +126,7 @@ public:
 	{
 	}
 
-	virtual void PreUpdate(const ITimingTrackUpdateContext& Context) override;
-	virtual void Draw(const ITimingTrackDrawContext& Context) const override;
+	virtual void BuildDrawState(ITimingEventsTrackDrawStateBuilder& Builder, const ITimingTrackUpdateContext& Context) override;
 	virtual const TSharedPtr<const ITimingEvent> SearchEvent(const FTimingEventSearchParameters& InSearchParameters) const override;
 };
 
diff --git a/Engine/Source/Developer/TraceInsights/Private/Insights/ViewModels/LoadingTimingTrack.cpp b/Engine/Source/Developer/TraceInsights/Private/Insights/ViewModels/LoadingTimingTrack.cpp
--- a/Engine/Source/Developer/TraceInsights/Private/Insights/ViewModels/LoadingTimingTrack.cpp
+++ b/Engine/Source/Developer/TraceInsights/Private/Insights/ViewModels/LoadingTimingTrack.cpp
@@ -15,7 +15,6 @@
 #include "Insights/ITimingViewSession.h"
 #include "Insights/ViewModels/TimingEvent.h"
 #include "Insights/ViewModels/TimingTrackViewport.h"
-#include "Insights/ViewModels/TimingViewDrawHelper.h"
 #include "Insights/ViewModels/TooltipDrawState.h"
 #include "Insights/ViewModels/TimingEventSearch.h"
 #include "Insights/Widgets/STimingView.h"
@@ -226,6 +225,46 @@ void FLoadingSharedState::SetColorSchema(int32 Schema)
 // FLoadingTimingTrack
 ////////////////////////////////////////////////////////////////////////////////////////////////////
 
+void FLoadingTimingTrack::BuildDrawState(ITimingEventsTrackDrawStateBuilder& Builder, const ITimingTrackUpdateContext& Context)
+{
+	TSharedPtr<const Trace::IAnalysisSession> Session = FInsightsManager::Get()->GetSession();
+	if (Session.IsValid() && Trace::ReadLoadTimeProfilerProvider(*Session.Get()))
+	{
+		Trace::FAnalysisSessionReadScope SessionReadScope(*Session.Get());
+
+		const Trace::ILoadTimeProfilerProvider& LoadTimeProfilerProvider = *Trace::ReadLoadTimeProfilerProvider(*Session.Get());
+
+		const FTimingTrackViewport& Viewport = Context.GetViewport();
+
+		LoadTimeProfilerProvider.ReadTimeline(TimelineIndex, [this, &Builder, &Viewport](const Trace::ILoadTimeProfilerProvider::CpuTimeline& Timeline)
+		{
+			if (FTimingEventsTrack::bUseDownSampling)
+			{
+				const double SecondsPerPixel = 1.0 / Viewport.GetScaleX();
+				Timeline.EnumerateEventsDownSampled(Viewport.GetStartTime(), Viewport.GetEndTime(), SecondsPerPixel, [this, &Builder](double StartTime, double EndTime, uint32 Depth, const Trace::FLoadTimeProfilerCpuEvent& Event)
+				{
+					const TCHAR* Name = SharedState.GetEventName(Depth, Event);
+					const uint64 Type = static_cast<uint64>(Event.EventType);
+					const uint32 Color = 0;
+					Builder.AddEvent(StartTime, EndTime, Depth, Name, Type, Color);
+				});
+			}
+			else
+			{
+				Timeline.EnumerateEvents(Viewport.GetStartTime(), Viewport.GetEndTime(), [this, &Builder](double StartTime, double EndTime, uint32 Depth, const Trace::FLoadTimeProfilerCpuEvent& Event)
+				{
+					const TCHAR* Name = SharedState.GetEventName(Depth, Event);
+					const uint64 Type = static_cast<uint64>(Event.EventType);
+					const uint32 Color = 0;
+					Builder.AddEvent(StartTime, EndTime, Depth, Name, Type, Color);
+				});
+			}
+		});
+	}
+}
+
+////////////////////////////////////////////////////////////////////////////////////////////////////
+
 void FLoadingTimingTrack::InitTooltip(FTooltipDrawState& InOutTooltip, const ITimingEvent& InTooltipEvent) const
 {
 	if (InTooltipEvent.CheckTrack(this) && FTimingEvent::CheckTypeName(InTooltipEvent))
@@ -274,65 +313,6 @@ void FLoadingTimingTrack::InitTooltip(FTooltipDrawState& InOutTooltip, const ITi
 
 ////////////////////////////////////////////////////////////////////////////////////////////////////
 
-void FLoadingTimingTrack::PreUpdate(const ITimingTrackUpdateContext& Context)
-{
-	if (IsDirty() || Context.GetViewport().IsHorizontalViewportDirty())
-	{
-		ClearDirtyFlag();
-
-		FTimingEventsTrackDrawStateBuilder Builder(DrawState, Context.GetViewport());
-
-		TSharedPtr<const Trace::IAnalysisSession> Session = FInsightsManager::Get()->GetSession();
-		if (Trace::ReadLoadTimeProfilerProvider(*Session.Get()))
-		{
-			Trace::FAnalysisSessionReadScope SessionReadScope(*Session.Get());
-
-			const Trace::ILoadTimeProfilerProvider& LoadTimeProfilerProvider = *Trace::ReadLoadTimeProfilerProvider(*Session.Get());
-
-			LoadTimeProfilerProvider.ReadTimeline(TimelineIndex, [this, &Builder](const Trace::ILoadTimeProfilerProvider::CpuTimeline& Timeline)
-			{
-				const FTimingTrackViewport& Viewport = Builder.GetViewport();
-				if (FTimingEventsTrack::bUseDownSampling)
-				{
-					const double SecondsPerPixel = 1.0 / Viewport.GetScaleX();
-					Timeline.EnumerateEventsDownSampled(Viewport.GetStartTime(), Viewport.GetEndTime(), SecondsPerPixel, [this, &Builder](double StartTime, double EndTime, uint32 Depth, const Trace::FLoadTimeProfilerCpuEvent& Event)
-					{
-						const TCHAR* Name = SharedState.GetEventName(Depth, Event);
-						const uint64 Type = static_cast<uint64>(Event.EventType);
-						const uint32 Color = 0;
-						Builder.AddEvent(StartTime, EndTime, Depth, Name, Type, Color);
-					});
-				}
-				else
-				{
-					Timeline.EnumerateEvents(Viewport.GetStartTime(), Viewport.GetEndTime(), [this, &Builder](double StartTime, double EndTime, uint32 Depth, const Trace::FLoadTimeProfilerCpuEvent& Event)
-					{
-						const TCHAR* Name = SharedState.GetEventName(Depth, Event);
-						const uint64 Type = static_cast<uint64>(Event.EventType);
-						const uint32 Color = 0;
-						Builder.AddEvent(StartTime, EndTime, Depth, Name, Type, Color);
-					});
-				}
-			});
-		}
-
-		Builder.Flush();
-		SetNumLanes(Builder.GetMaxDepth() + 1);
-	}
-
-	UpdateTrackHeight(Context);
-}
-
-////////////////////////////////////////////////////////////////////////////////////////////////////
-
-void FLoadingTimingTrack::Draw(const ITimingTrackDrawContext& Context) const
-{
-	const FTimingViewDrawHelper& Helper = Context.GetHelper();
-	Helper.DrawEventsTrack(DrawState, *this);
-}
-
-////////////////////////////////////////////////////////////////////////////////////////////////////
-
 const TSharedPtr<const ITimingEvent> FLoadingTimingTrack::SearchEvent(const FTimingEventSearchParameters& InSearchParameters) const
 {
 	TSharedPtr<const ITimingEvent> FoundEvent;
diff --git a/Engine/Source/Developer/TraceInsights/Private/Insights/ViewModels/LoadingTimingTrack.h b/Engine/Source/Developer/TraceInsights/Private/Insights/ViewModels/LoadingTimingTrack.h
--- a/Engine/Source/Developer/TraceInsights/Private/Insights/ViewModels/LoadingTimingTrack.h
+++ b/Engine/Source/Developer/TraceInsights/Private/Insights/ViewModels/LoadingTimingTrack.h
@@ -8,7 +8,6 @@
 // Insights
 #include "Insights/ITimingViewExtender.h"
 #include "Insights/ViewModels/TimingEventsTrack.h"
-#include "Insights/ViewModels/TimingViewDrawHelper.h"
 
 class FTimingEventSearchParameters;
 class STimingView;
@@ -76,9 +75,10 @@ public:
 	}
 	virtual ~FLoadingTimingTrack() {}
 
+	virtual void BuildDrawState(ITimingEventsTrackDrawStateBuilder& Builder, const ITimingTrackUpdateContext& Context) override;
+
 	virtual void InitTooltip(FTooltipDrawState& InOutTooltip, const ITimingEvent& InTooltipEvent) const override;
-	virtual void PreUpdate(const ITimingTrackUpdateContext& Context) override;
-	virtual void Draw(const ITimingTrackDrawContext& Context) const override;
+
 	virtual const TSharedPtr<const ITimingEvent> SearchEvent(const FTimingEventSearchParameters& InSearchParameters) const override;
 
 protected:
@@ -88,7 +88,6 @@ protected:
 protected:
 	FLoadingSharedState& SharedState;
 	uint32 TimelineIndex;
-	FTimingEventsTrackDrawState DrawState;
 };
 
 ////////////////////////////////////////////////////////////////////////////////////////////////////
diff --git a/Engine/Source/Developer/TraceInsights/Private/Insights/ViewModels/ThreadTimingTrack.cpp b/Engine/Source/Developer/TraceInsights/Private/Insights/ViewModels/ThreadTimingTrack.cpp
--- a/Engine/Source/Developer/TraceInsights/Private/Insights/ViewModels/ThreadTimingTrack.cpp
+++ b/Engine/Source/Developer/TraceInsights/Private/Insights/ViewModels/ThreadTimingTrack.cpp
@@ -384,64 +384,45 @@ void FThreadTimingSharedState::OnTracksChanged(Insights::ITimingViewSession& InS
 // FThreadTimingTrack
 ////////////////////////////////////////////////////////////////////////////////////////////////////
 
-void FThreadTimingTrack::PreUpdate(const ITimingTrackUpdateContext& Context)
+void FThreadTimingTrack::BuildDrawState(ITimingEventsTrackDrawStateBuilder& Builder, const ITimingTrackUpdateContext& Context)
 {
-	if (IsDirty() || Context.GetViewport().IsHorizontalViewportDirty())
+	TSharedPtr<const Trace::IAnalysisSession> Session = FInsightsManager::Get()->GetSession();
+	if (Session.IsValid() && Trace::ReadTimingProfilerProvider(*Session.Get()))
 	{
-		ClearDirtyFlag();
+		Trace::FAnalysisSessionReadScope SessionReadScope(*Session.Get());
 
-		FTimingEventsTrackDrawStateBuilder Builder(GetDrawState(), Context.GetViewport());
+		const Trace::ITimingProfilerProvider& TimingProfilerProvider = *Trace::ReadTimingProfilerProvider(*Session.Get());
 
-		TSharedPtr<const Trace::IAnalysisSession> Session = FInsightsManager::Get()->GetSession();
-		if (Session.IsValid() && Trace::ReadTimingProfilerProvider(*Session.Get()))
-		{
-			Trace::FAnalysisSessionReadScope SessionReadScope(*Session.Get());
-
-			const Trace::ITimingProfilerProvider& TimingProfilerProvider = *Trace::ReadTimingProfilerProvider(*Session.Get());
+		const FTimingTrackViewport& Viewport = Context.GetViewport();
 
-			TimingProfilerProvider.ReadTimers([this, &Builder, &TimingProfilerProvider](const Trace::FTimingProfilerTimer* Timers, uint64 TimersCount)
+		TimingProfilerProvider.ReadTimers([this, &Builder, &Viewport, &TimingProfilerProvider](const Trace::FTimingProfilerTimer* Timers, uint64 TimersCount)
+		{
+			TimingProfilerProvider.ReadTimeline(TimelineIndex, [this, &Builder, &Viewport, Timers](const Trace::ITimingProfilerProvider::Timeline& Timeline)
 			{
-				TimingProfilerProvider.ReadTimeline(TimelineIndex, [this, &Builder, Timers](const Trace::ITimingProfilerProvider::Timeline& Timeline)
+				if (FTimingEventsTrack::bUseDownSampling)
 				{
-					const FTimingTrackViewport& Viewport = Builder.GetViewport();
-					if (FTimingEventsTrack::bUseDownSampling)
-					{
-						const double SecondsPerPixel = 1.0 / Viewport.GetScaleX();
-						Timeline.EnumerateEventsDownSampled(Viewport.GetStartTime(), Viewport.GetEndTime(), SecondsPerPixel,
-							[this, &Builder, Timers](double StartTime, double EndTime, uint32 Depth, const Trace::FTimingProfilerEvent& Event)
-							{
-								const uint64 Type = Timers[Event.TimerIndex].Id;
-								const uint32 Color = 0;
-								Builder.AddEvent(StartTime, EndTime, Depth, Timers[Event.TimerIndex].Name, Type, Color);
-							});
-					}
-					else
-					{
-						Timeline.EnumerateEvents(Viewport.GetStartTime(), Viewport.GetEndTime(),
-							[this, &Builder, Timers](double StartTime, double EndTime, uint32 Depth, const Trace::FTimingProfilerEvent& Event)
-							{
-								const uint64 Type = Timers[Event.TimerIndex].Id;
-								const uint32 Color = 0;
-								Builder.AddEvent(StartTime, EndTime, Depth, Timers[Event.TimerIndex].Name, Type, Color);
-							});
-					}
-				});
+					const double SecondsPerPixel = 1.0 / Viewport.GetScaleX();
+					Timeline.EnumerateEventsDownSampled(Viewport.GetStartTime(), Viewport.GetEndTime(), SecondsPerPixel,
+						[this, &Builder, Timers](double StartTime, double EndTime, uint32 Depth, const Trace::FTimingProfilerEvent& Event)
+						{
+							const uint64 Type = Timers[Event.TimerIndex].Id;
+							const uint32 Color = 0;
+							Builder.AddEvent(StartTime, EndTime, Depth, Timers[Event.TimerIndex].Name, Type, Color);
+						});
+				}
+				else
+				{
+					Timeline.EnumerateEvents(Viewport.GetStartTime(), Viewport.GetEndTime(),
+						[this, &Builder, Timers](double StartTime, double EndTime, uint32 Depth, const Trace::FTimingProfilerEvent& Event)
+						{
+							const uint64 Type = Timers[Event.TimerIndex].Id;
+							const uint32 Color = 0;
+							Builder.AddEvent(StartTime, EndTime, Depth, Timers[Event.TimerIndex].Name, Type, Color);
+						});
+				}
 			});
-		}
-
-		Builder.Flush();
-		SetNumLanes(Builder.GetMaxDepth() + 1);
+		});
 	}
-
-	UpdateTrackHeight(Context);
-}
-
-////////////////////////////////////////////////////////////////////////////////////////////////////
-
-void FThreadTimingTrack::Draw(const ITimingTrackDrawContext& Context) const
-{
-	const FTimingViewDrawHelper& Helper = Context.GetHelper();
-	Helper.DrawEventsTrack(DrawState, *this);
 }
 
 ////////////////////////////////////////////////////////////////////////////////////////////////////
diff --git a/Engine/Source/Developer/TraceInsights/Private/Insights/ViewModels/ThreadTimingTrack.h b/Engine/Source/Developer/TraceInsights/Private/Insights/ViewModels/ThreadTimingTrack.h
--- a/Engine/Source/Developer/TraceInsights/Private/Insights/ViewModels/ThreadTimingTrack.h
+++ b/Engine/Source/Developer/TraceInsights/Private/Insights/ViewModels/ThreadTimingTrack.h
@@ -3,19 +3,20 @@
 #pragma once
 
 #include "CoreMinimal.h"
+#include "Fonts/SlateFontInfo.h"
 #include "TraceServices/Model/TimingProfiler.h"
 
 // Insights
 #include "Insights/ITimingViewExtender.h"
 #include "Insights/ViewModels/TimingEventSearch.h" // for TTimingEventSearchCache
 #include "Insights/ViewModels/TimingEventsTrack.h"
-#include "Insights/ViewModels/TimingViewDrawHelper.h" // for FTimingEventsTrackDrawState
 
 class FTimingEvent;
 class FTimingEventSearchParameters;
 class FGpuTimingTrack;
 class FCpuTimingTrack;
 class STimingView;
+struct FSlateBrush;
 
 ////////////////////////////////////////////////////////////////////////////////////////////////////
 
@@ -107,9 +108,8 @@ public:
 	uint32 GetThreadId() const { return ThreadId; }
 	//void SetThreadId(uint32 InThreadId) { ThreadId = InThreadId; }
 
-	virtual void PreUpdate(const ITimingTrackUpdateContext& Context) override;
+	virtual void BuildDrawState(ITimingEventsTrackDrawStateBuilder& Builder, const ITimingTrackUpdateContext& Context) override;
 
-	virtual void Draw(const ITimingTrackDrawContext& Context) const override;
 	virtual void PostDraw(const ITimingTrackDrawContext& Context) const override;
 
 	virtual void InitTooltip(FTooltipDrawState& InOutTooltip, const ITimingEvent& InTooltipEvent) const override;
@@ -121,9 +121,6 @@ public:
 	virtual void OnClipboardCopyEvent(const ITimingEvent& InSelectedEvent) const override;
 	virtual void BuildContextMenu(FMenuBuilder& MenuBuilder) override;
 
-protected:
-	FTimingEventsTrackDrawState& GetDrawState() { return DrawState; }
-
 private:
 	void DrawSelectedEventInfo(const FTimingEvent& SelectedEvent, const FTimingTrackViewport& Viewport, const FDrawContext& DrawContext, const FSlateBrush* WhiteBrush, const FSlateFontInfo& Font) const;
 
@@ -143,8 +140,6 @@ private:
 	uint32 TimelineIndex;
 	uint32 ThreadId;
 
-	FTimingEventsTrackDrawState DrawState;
-
 	// Search cache
 	mutable TTimingEventSearchCache<Trace::FTimingProfilerEvent> SearchCache;
 };
diff --git a/Engine/Source/Developer/TraceInsights/Private/Insights/ViewModels/TimingEventsTrack.cpp b/Engine/Source/Developer/TraceInsights/Private/Insights/ViewModels/TimingEventsTrack.cpp
--- a/Engine/Source/Developer/TraceInsights/Private/Insights/ViewModels/TimingEventsTrack.cpp
+++ b/Engine/Source/Developer/TraceInsights/Private/Insights/ViewModels/TimingEventsTrack.cpp
@@ -19,7 +19,7 @@ bool FTimingEventsTrack::bUseDownSampling = true;
 FTimingEventsTrack::FTimingEventsTrack(const FName& InType, const FName& InSubType, const FString& InName)
 	: FBaseTimingTrack(InType, InSubType, InName)
 	, NumLanes(0)
-	, DrawState(MakeShareable(new FTimingEventsTrackDrawState()))
+	, DrawState(MakeShared<FTimingEventsTrackDrawState>())
 {
 }
 
@@ -41,6 +41,63 @@ void FTimingEventsTrack::Reset()
 
 ////////////////////////////////////////////////////////////////////////////////////////////////////
 
+void FTimingEventsTrack::PreUpdate(const ITimingTrackUpdateContext& Context)
+{
+	if (IsDirty() || Context.GetViewport().IsHorizontalViewportDirty())
+	{
+		ClearDirtyFlag();
+
+		FTimingEventsTrackDrawStateBuilder Builder(*DrawState, Context.GetViewport());
+
+		BuildDrawState(Builder, Context);
+
+		Builder.Flush();
+
+		SetNumLanes(Builder.GetMaxDepth() + 1);
+	}
+
+	UpdateTrackHeight(Context);
+}
+
+////////////////////////////////////////////////////////////////////////////////////////////////////
+
+void FTimingEventsTrack::UpdateTrackHeight(const ITimingTrackUpdateContext& Context)
+{
+	const FTimingTrackViewport& Viewport = Context.GetViewport();
+
+	const float CurrentTrackHeight = GetHeight();
+	const float DesiredTrackHeight = Viewport.GetLayout().ComputeTrackHeight(NumLanes);
+
+	if (CurrentTrackHeight < DesiredTrackHeight)
+	{
+		float NewTrackHeight;
+		if (Viewport.IsDirty(ETimingTrackViewportDirtyFlags::VLayoutChanged))
+		{
+			NewTrackHeight = DesiredTrackHeight;
+		}
+		else
+		{
+			NewTrackHeight = FMath::CeilToFloat(CurrentTrackHeight * 0.9f + DesiredTrackHeight * 0.1f);
+		}
+		SetHeight(NewTrackHeight);
+	}
+	else if (CurrentTrackHeight > DesiredTrackHeight)
+	{
+		float NewTrackHeight;
+		if (Viewport.IsDirty(ETimingTrackViewportDirtyFlags::VLayoutChanged))
+		{
+			NewTrackHeight = DesiredTrackHeight;
+		}
+		else
+		{
+			NewTrackHeight = FMath::FloorToFloat(CurrentTrackHeight * 0.9f + DesiredTrackHeight * 0.1f);
+		}
+		SetHeight(NewTrackHeight);
+	}
+}
+
+////////////////////////////////////////////////////////////////////////////////////////////////////
+
 void FTimingEventsTrack::PostUpdate(const ITimingTrackUpdateContext& Context)
 {
 	FBaseTimingTrack::PostUpdate(Context);
@@ -64,6 +121,14 @@ void FTimingEventsTrack::PostUpdate(const ITimingTrackUpdateContext& Context)
 
 ////////////////////////////////////////////////////////////////////////////////////////////////////
 
+void FTimingEventsTrack::Draw(const ITimingTrackDrawContext& Context) const
+{
+	const FTimingViewDrawHelper& Helper = Context.GetHelper();
+	Helper.DrawEventsTrack(*DrawState, *this);
+}
+
+////////////////////////////////////////////////////////////////////////////////////////////////////
+
 void FTimingEventsTrack::DrawEvent(const ITimingTrackDrawContext& Context, const ITimingEvent& InTimingEvent, EDrawEventMode InDrawMode) const
 {
 	if (InTimingEvent.CheckTrack(this) && FTimingEvent::CheckTypeName(InTimingEvent))
@@ -105,41 +170,4 @@ const TSharedPtr<const ITimingEvent> FTimingEventsTrack::GetEvent(float InPosX,
 
 ////////////////////////////////////////////////////////////////////////////////////////////////////
 
-void FTimingEventsTrack::UpdateTrackHeight(const ITimingTrackUpdateContext& Context)
-{
-	const FTimingTrackViewport& Viewport = Context.GetViewport();
-
-	const float CurrentTrackHeight = GetHeight();
-	const float DesiredTrackHeight = Viewport.GetLayout().ComputeTrackHeight(NumLanes);
-
-	if (CurrentTrackHeight < DesiredTrackHeight)
-	{
-		float NewTrackHeight;
-		if (Viewport.IsDirty(ETimingTrackViewportDirtyFlags::VLayoutChanged))
-		{
-			NewTrackHeight = DesiredTrackHeight;
-		}
-		else
-		{
-			NewTrackHeight = FMath::CeilToFloat(CurrentTrackHeight * 0.9f + DesiredTrackHeight * 0.1f);
-		}
-		SetHeight(NewTrackHeight);
-	}
-	else if (CurrentTrackHeight > DesiredTrackHeight)
-	{
-		float NewTrackHeight;
-		if (Viewport.IsDirty(ETimingTrackViewportDirtyFlags::VLayoutChanged))
-		{
-			NewTrackHeight = DesiredTrackHeight;
-		}
-		else
-		{
-			NewTrackHeight = FMath::FloorToFloat(CurrentTrackHeight * 0.9f + DesiredTrackHeight * 0.1f);
-		}
-		SetHeight(NewTrackHeight);
-	}
-}
-
-////////////////////////////////////////////////////////////////////////////////////////////////////
-
 #undef LOCTEXT_NAMESPACE
diff --git a/Engine/Source/Developer/TraceInsights/Private/Insights/ViewModels/TimingViewDrawHelper.h b/Engine/Source/Developer/TraceInsights/Private/Insights/ViewModels/TimingViewDrawHelper.h
--- a/Engine/Source/Developer/TraceInsights/Private/Insights/ViewModels/TimingViewDrawHelper.h
+++ b/Engine/Source/Developer/TraceInsights/Private/Insights/ViewModels/TimingViewDrawHelper.h
@@ -6,6 +6,9 @@
 #include "Fonts/SlateFontInfo.h"
 #include "Math/Color.h"
 
+// Insights
+#include "Insights/ViewModels/TimingEventsTrack.h"
+
 ////////////////////////////////////////////////////////////////////////////////////////////////////
 
 struct FSlateBrush;
@@ -81,7 +84,7 @@ struct FTimingEventsTrackDrawState
 
 ////////////////////////////////////////////////////////////////////////////////////////////////////
 
-class FTimingEventsTrackDrawStateBuilder
+class FTimingEventsTrackDrawStateBuilder : public ITimingEventsTrackDrawStateBuilder
 {
 private:
 	struct FBoxData
@@ -97,6 +100,7 @@ private:
 
 public:
 	explicit FTimingEventsTrackDrawStateBuilder(FTimingEventsTrackDrawState& InState, const FTimingTrackViewport& InViewport);
+	virtual ~FTimingEventsTrackDrawStateBuilder() {}
 
 	/**
 	 * Non-copyable
@@ -106,7 +110,7 @@ public:
 
 	const FTimingTrackViewport& GetViewport() const { return Viewport; }
 
-	void AddEvent(double InEventStartTime, double InEventEndTime, uint32 InEventDepth, const TCHAR* InEventName, uint64 InEventType, uint32 InEventColor);
+	virtual void AddEvent(double InEventStartTime, double InEventEndTime, uint32 InEventDepth, const TCHAR* InEventName, uint64 InEventType, uint32 InEventColor) override;
 	void Flush();
 
 	int32 GetMaxDepth() const { return MaxDepth; }
diff --git a/Engine/Source/Developer/TraceInsights/Public/Insights/ViewModels/TimingEventsTrack.h b/Engine/Source/Developer/TraceInsights/Public/Insights/ViewModels/TimingEventsTrack.h
--- a/Engine/Source/Developer/TraceInsights/Public/Insights/ViewModels/TimingEventsTrack.h
+++ b/Engine/Source/Developer/TraceInsights/Public/Insights/ViewModels/TimingEventsTrack.h
@@ -7,7 +7,15 @@
 // Insights
 #include "Insights/ViewModels/BaseTimingTrack.h"
 
-struct FTimingEventsTrackDrawState;
+////////////////////////////////////////////////////////////////////////////////////////////////////
+
+class TRACEINSIGHTS_API ITimingEventsTrackDrawStateBuilder
+{
+public:
+	virtual ~ITimingEventsTrackDrawStateBuilder() = default;
+
+	virtual void AddEvent(double InEventStartTime, double InEventEndTime, uint32 InEventDepth, const TCHAR* InEventName, uint64 InEventType, uint32 InEventColor) = 0;
+};
 
 ////////////////////////////////////////////////////////////////////////////////////////////////////
 
@@ -17,27 +25,32 @@ public:
 	explicit FTimingEventsTrack(const FName& InType, const FName& InSubType, const FString& InName);
 	virtual ~FTimingEventsTrack();
 
-	int32 GetNumLanes() const { return NumLanes; }
-	void SetNumLanes(int32 InNumLanes) { NumLanes = InNumLanes; }
-
 	//////////////////////////////////////////////////
 	// FBaseTimingTrack
 
 	virtual void Reset() override;
+
+	virtual void PreUpdate(const ITimingTrackUpdateContext& Context) override;
 	virtual void PostUpdate(const ITimingTrackUpdateContext& Context) override;
 
+	virtual void Draw(const ITimingTrackDrawContext& Context) const override;
 	virtual void DrawEvent(const ITimingTrackDrawContext& Context, const ITimingEvent& InTimingEvent, EDrawEventMode InDrawMode) const override;
 
 	virtual const TSharedPtr<const ITimingEvent> GetEvent(float InPosX, float InPosY, const FTimingTrackViewport& Viewport) const override;
 
 	//////////////////////////////////////////////////
 
+	virtual void BuildDrawState(ITimingEventsTrackDrawStateBuilder& Builder, const ITimingTrackUpdateContext& Context) = 0;
+
 protected:
+	int32 GetNumLanes() const { return NumLanes; }
+	void SetNumLanes(int32 InNumLanes) { NumLanes = InNumLanes; }
+
 	void UpdateTrackHeight(const ITimingTrackUpdateContext& Context);
 
 private:
 	int32 NumLanes; // number of lanes (sub-tracks)
-	TSharedPtr<FTimingEventsTrackDrawState> DrawState;
+	TSharedRef<struct FTimingEventsTrackDrawState> DrawState;
 
 public:
 	static bool bUseDownSampling;
