[Insights] Added DrawEvents, DrawTracks protected api in TimingEventsTrack.

diff --git a/Engine/Source/Developer/TraceInsights/Private/Insights/ViewModels/TimingEventsTrack.cpp b/Engine/Source/Developer/TraceInsights/Private/Insights/ViewModels/TimingEventsTrack.cpp
--- a/Engine/Source/Developer/TraceInsights/Private/Insights/ViewModels/TimingEventsTrack.cpp
+++ b/Engine/Source/Developer/TraceInsights/Private/Insights/ViewModels/TimingEventsTrack.cpp
@@ -109,7 +109,7 @@ void FTimingEventsTrack::PostUpdate(const ITimingTrackUpdateContext& Context)
 	if (MouseY >= GetPosY() && MouseY < GetPosY() + GetHeight())
 	{
 		SetHoveredState(true);
-		
+
 		const float MouseX = Context.GetMousePosition().X;
 		SetHeaderHoveredState(MouseX < HeaderWidth && MouseY < GetPosY() + HeaderHeight);
 	}
@@ -124,7 +124,40 @@ void FTimingEventsTrack::PostUpdate(const ITimingTrackUpdateContext& Context)
 void FTimingEventsTrack::Draw(const ITimingTrackDrawContext& Context) const
 {
 	const FTimingViewDrawHelper& Helper = Context.GetHelper();
-	Helper.DrawEventsTrack(*DrawState, *this);
+	Helper.DrawEvents(*DrawState, *this);
+	Helper.DrawTrackHeader(*this);
+}
+
+////////////////////////////////////////////////////////////////////////////////////////////////////
+
+void FTimingEventsTrack::DrawEvents(const ITimingTrackDrawContext& Context, const float OffsetY) const
+{
+	const FTimingViewDrawHelper& Helper = Context.GetHelper();
+	Helper.DrawEvents(*DrawState, *this, OffsetY);
+}
+
+////////////////////////////////////////////////////////////////////////////////////////////////////
+
+int32 FTimingEventsTrack::GetHeaderBackgroundLayerId(const ITimingTrackDrawContext& Context) const
+{
+	const FTimingViewDrawHelper& Helper = Context.GetHelper();
+	return Helper.GetHeaderBackgroundLayerId();
+}
+
+////////////////////////////////////////////////////////////////////////////////////////////////////
+
+int32 FTimingEventsTrack::GetHeaderTextLayerId(const ITimingTrackDrawContext& Context) const
+{
+	const FTimingViewDrawHelper& Helper = Context.GetHelper();
+	return Helper.GetHeaderTextLayerId();
+}
+
+////////////////////////////////////////////////////////////////////////////////////////////////////
+
+void FTimingEventsTrack::DrawHeader(const ITimingTrackDrawContext& Context) const
+{
+	const FTimingViewDrawHelper& Helper = Context.GetHelper();
+	Helper.DrawTrackHeader(*this);
 }
 
 ////////////////////////////////////////////////////////////////////////////////////////////////////
diff --git a/Engine/Source/Developer/TraceInsights/Private/Insights/ViewModels/TimingViewDrawHelper.cpp b/Engine/Source/Developer/TraceInsights/Private/Insights/ViewModels/TimingViewDrawHelper.cpp
--- a/Engine/Source/Developer/TraceInsights/Private/Insights/ViewModels/TimingViewDrawHelper.cpp
+++ b/Engine/Source/Developer/TraceInsights/Private/Insights/ViewModels/TimingViewDrawHelper.cpp
@@ -349,12 +349,13 @@ void FTimingViewDrawHelper::DrawBackground() const
 
 ////////////////////////////////////////////////////////////////////////////////////////////////////
 
-void FTimingViewDrawHelper::DrawEventsTrack(const FTimingEventsTrackDrawState& DrawState, const FTimingEventsTrack& Track) const
+void FTimingViewDrawHelper::DrawEvents(const FTimingEventsTrackDrawState& DrawState, const FTimingEventsTrack& Track, const float OffsetY) const
 {
 	const float TrackY = Track.GetPosY();
+	const float TrackH = Track.GetHeight();
 
-	if (Track.GetHeight() > 0.0f &&
-		TrackY + Track.GetHeight() > Viewport.GetTopOffset() &&
+	if (TrackH > 0.0f &&
+		TrackY + TrackH > Viewport.GetTopOffset() &&
 		TrackY < Viewport.GetHeight() - Viewport.GetBottomOffset())
 	{
 		const FTimingViewLayout& Layout = Viewport.GetLayout();
@@ -362,10 +363,10 @@ void FTimingViewDrawHelper::DrawEventsTrack(const FTimingEventsTrackDrawState& D
 		NumEvents += DrawState.GetNumEvents();
 		NumMergedBoxes += DrawState.GetNumMergedBoxes();
 
-		const float TopLaneY = TrackY + 1.0f + Layout.TimelineDY; // +1.0f is for horizontal line between timelines
+		const float TopLaneY = TrackY + OffsetY + Layout.TimelineDY;
 
-		// Draw filled boxes.
-		//if (Layout.EventH > 2.0f)
+		// Draw filled boxes (merged borders).
+		//if (Layout.EventH > 0.0f)
 		{
 			const int32 EventFillLayerId = ReservedLayerId + ToInt32(EDrawLayer::EventFill);
 			const float EventFillH = Layout.EventH;
@@ -377,7 +378,7 @@ void FTimingViewDrawHelper::DrawEventsTrack(const FTimingEventsTrackDrawState& D
 			NumDrawBoxes += DrawState.Boxes.Num();
 		}
 
-		// Draw merged borders.
+		// Draw filled boxes (event inside area).
 		if (Layout.EventH > 2.0f)
 		{
 			const int32 EventFillLayerId = ReservedLayerId + ToInt32(EDrawLayer::EventFill);
@@ -414,13 +415,26 @@ void FTimingViewDrawHelper::DrawEventsTrack(const FTimingEventsTrackDrawState& D
 			}
 			NumDrawTexts += DrawState.Texts.Num();
 		}
+	}
+}
+
+////////////////////////////////////////////////////////////////////////////////////////////////////
 
+void FTimingViewDrawHelper::DrawTrackHeader(const FTimingEventsTrack& Track) const
+{
+	const float TrackY = Track.GetPosY();
+	const float TrackH = Track.GetHeight();
+
+	if (TrackH > 0.0f &&
+		TrackY + TrackH > Viewport.GetTopOffset() &&
+		TrackY < Viewport.GetHeight() - Viewport.GetBottomOffset())
+	{
 		// Draw a horizontal line between timelines (top line of a track).
 		const int32 HeaderLayerId = ReservedLayerId + ToInt32(EDrawLayer::HeaderBackground);
 		DrawContext.DrawBox(HeaderLayerId, 0.0f, TrackY, Viewport.GetWidth(), 1.0f, WhiteBrush, EdgeColor);
 
 		// Draw header with name of timeline.
-		if (Track.GetHeight() > 4.0f)
+		if (TrackH > 4.0f)
 		{
 			const TSharedRef<FSlateFontMeasure> FontMeasureService = FSlateApplication::Get().GetRenderer()->GetFontMeasureService();
 			float NameWidth = FontMeasureService->Measure(Track.GetName(), EventFont).X;
diff --git a/Engine/Source/Developer/TraceInsights/Private/Insights/ViewModels/TimingViewDrawHelper.h b/Engine/Source/Developer/TraceInsights/Private/Insights/ViewModels/TimingViewDrawHelper.h
--- a/Engine/Source/Developer/TraceInsights/Private/Insights/ViewModels/TimingViewDrawHelper.h
+++ b/Engine/Source/Developer/TraceInsights/Private/Insights/ViewModels/TimingViewDrawHelper.h
@@ -168,7 +168,11 @@ public:
 	void DrawBackground() const;
 
 	void BeginDrawTracks() const;
-	void DrawEventsTrack(const FTimingEventsTrackDrawState& DrawState, const FTimingEventsTrack& Track) const;
+	// OffsetY = 1.0f is for the top horizontal line (which separates the timelines) added by DrawTrackHeader.
+	void DrawEvents(const FTimingEventsTrackDrawState& DrawState, const FTimingEventsTrack& Track, const float OffsetY = 1.0f) const;
+	int32 GetHeaderBackgroundLayerId() const { return ReservedLayerId + ToInt32(EDrawLayer::HeaderBackground); }
+	int32 GetHeaderTextLayerId() const { return ReservedLayerId + ToInt32(EDrawLayer::HeaderText); }
+	void DrawTrackHeader(const FTimingEventsTrack& Track) const;
 	void EndDrawTracks() const;
 
 	void DrawTimingEventHighlight(double StartTime, double EndTime, float Y, EDrawEventMode Mode) const;
diff --git a/Engine/Source/Developer/TraceInsights/Public/Insights/ViewModels/TimingEventsTrack.h b/Engine/Source/Developer/TraceInsights/Public/Insights/ViewModels/TimingEventsTrack.h
--- a/Engine/Source/Developer/TraceInsights/Public/Insights/ViewModels/TimingEventsTrack.h
+++ b/Engine/Source/Developer/TraceInsights/Public/Insights/ViewModels/TimingEventsTrack.h
@@ -48,6 +48,12 @@ protected:
 
 	void UpdateTrackHeight(const ITimingTrackUpdateContext& Context);
 
+	void DrawEvents(const ITimingTrackDrawContext& Context, const float OffsetY = 1.0f) const;
+	void DrawHeader(const ITimingTrackDrawContext& Context) const;
+
+	int32 GetHeaderBackgroundLayerId(const ITimingTrackDrawContext& Context) const;
+	int32 GetHeaderTextLayerId(const ITimingTrackDrawContext& Context) const;
+
 private:
 	int32 NumLanes; // number of lanes (sub-tracks)
 	TSharedRef<struct FTimingEventsTrackDrawState> DrawState;
