[Insights]
  - Changed several "MakeShareable(new FMyClass())" usage patterns with "MakeShared<FMyClass>()".
  - Increased version to 0.11

diff --git a/Engine/Source/Developer/TraceInsights/Private/Insights/InsightsStyle.cpp b/Engine/Source/Developer/TraceInsights/Private/Insights/InsightsStyle.cpp
--- a/Engine/Source/Developer/TraceInsights/Private/Insights/InsightsStyle.cpp
+++ b/Engine/Source/Developer/TraceInsights/Private/Insights/InsightsStyle.cpp
@@ -33,7 +33,7 @@ FName FInsightsStyle::GetStyleSetName()
 
 TSharedRef<FSlateStyleSet> FInsightsStyle::Create()
 {
-	TSharedRef<FSlateStyleSet> StyleRef = MakeShareable(new FSlateStyleSet(FInsightsStyle::GetStyleSetName()));
+	TSharedRef<FSlateStyleSet> StyleRef = MakeShared<FSlateStyleSet>(FInsightsStyle::GetStyleSetName());
 	StyleRef->SetContentRoot(FPaths::EngineContentDir() / TEXT("Editor/Slate"));
 	StyleRef->SetCoreContentRoot(FPaths::EngineContentDir() / TEXT("Slate"));
 
diff --git a/Engine/Source/Developer/TraceInsights/Private/Insights/LoadingProfiler/Widgets/SLoadingProfilerWindow.cpp b/Engine/Source/Developer/TraceInsights/Private/Insights/LoadingProfiler/Widgets/SLoadingProfilerWindow.cpp
--- a/Engine/Source/Developer/TraceInsights/Private/Insights/LoadingProfiler/Widgets/SLoadingProfilerWindow.cpp
+++ b/Engine/Source/Developer/TraceInsights/Private/Insights/LoadingProfiler/Widgets/SLoadingProfilerWindow.cpp
@@ -196,7 +196,8 @@ TSharedRef<SDockTab> SLoadingProfilerWindow::SpawnTab_EventAggregationTreeView(c
 {
 	FLoadingProfilerManager::Get()->SetEventAggregationTreeViewVisible(true);
 
-	TSharedPtr<Insights::FTable> TablePtr = MakeShareable(new Insights::FTable());
+	TSharedRef<Insights::FTable> Table = MakeShared<Insights::FTable>();
+
 	TSharedPtr<const Trace::IAnalysisSession> Session = FInsightsManager::Get()->GetSession();
 	if (Session.IsValid() && Trace::ReadLoadTimeProfilerProvider(*Session.Get()))
 	{
@@ -205,14 +206,14 @@ TSharedRef<SDockTab> SLoadingProfilerWindow::SpawnTab_EventAggregationTreeView(c
 		const double SelectionStartTime = TimingView ? TimingView->GetSelectionStartTime() : 0.0;
 		const double SelectionEndTime = TimingView ? TimingView->GetSelectionEndTime() : 0.0;
 		TSharedPtr<Trace::ITable<Trace::FLoadTimeProfilerAggregatedStats>> SourceTable = MakeShareable(LoadTimeProfilerProvider.CreateEventAggregation(SelectionStartTime, SelectionEndTime));
-		TablePtr->Init(SourceTable);
+		Table->Init(SourceTable);
 	}
 
 	const TSharedRef<SDockTab> DockTab = SNew(SDockTab)
 		.ShouldAutosize(false)
 		.TabRole(ETabRole::PanelTab)
 		[
-			SAssignNew(EventAggregationTreeView, Insights::STableTreeView, TablePtr)
+			SAssignNew(EventAggregationTreeView, Insights::STableTreeView, Table)
 		];
 
 	EventAggregationTreeView->Reset();
@@ -236,7 +237,8 @@ TSharedRef<SDockTab> SLoadingProfilerWindow::SpawnTab_ObjectTypeAggregationTreeV
 {
 	FLoadingProfilerManager::Get()->SetObjectTypeAggregationTreeViewVisible(true);
 
-	TSharedPtr<Insights::FTable> TablePtr = MakeShareable(new Insights::FTable());
+	TSharedRef<Insights::FTable> Table = MakeShared<Insights::FTable>();
+
 	TSharedPtr<const Trace::IAnalysisSession> Session = FInsightsManager::Get()->GetSession();
 	if (Session.IsValid() && Trace::ReadLoadTimeProfilerProvider(*Session.Get()))
 	{
@@ -245,14 +247,14 @@ TSharedRef<SDockTab> SLoadingProfilerWindow::SpawnTab_ObjectTypeAggregationTreeV
 		const double SelectionStartTime = TimingView ? TimingView->GetSelectionStartTime() : 0.0;
 		const double SelectionEndTime = TimingView ? TimingView->GetSelectionEndTime() : 0.0;
 		TSharedPtr<Trace::ITable<Trace::FLoadTimeProfilerAggregatedStats>> SourceTable = MakeShareable(LoadTimeProfilerProvider.CreateObjectTypeAggregation(SelectionStartTime, SelectionEndTime));
-		TablePtr->Init(SourceTable);
+		Table->Init(SourceTable);
 	}
 
 	const TSharedRef<SDockTab> DockTab = SNew(SDockTab)
 		.ShouldAutosize(false)
 		.TabRole(ETabRole::PanelTab)
 		[
-			SAssignNew(ObjectTypeAggregationTreeView, Insights::STableTreeView, TablePtr)
+			SAssignNew(ObjectTypeAggregationTreeView, Insights::STableTreeView, Table)
 		];
 
 	ObjectTypeAggregationTreeView->Reset();
@@ -276,7 +278,8 @@ TSharedRef<SDockTab> SLoadingProfilerWindow::SpawnTab_PackageDetailsTreeView(con
 {
 	FLoadingProfilerManager::Get()->SetPackageDetailsTreeViewVisible(true);
 
-	TSharedPtr<Insights::FTable> TablePtr = MakeShareable(new Insights::FTable());
+	TSharedRef<Insights::FTable> Table = MakeShared<Insights::FTable>();
+
 	TSharedPtr<const Trace::IAnalysisSession> Session = FInsightsManager::Get()->GetSession();
 	if (Session.IsValid() && Trace::ReadLoadTimeProfilerProvider(*Session.Get()))
 	{
@@ -285,14 +288,14 @@ TSharedRef<SDockTab> SLoadingProfilerWindow::SpawnTab_PackageDetailsTreeView(con
 		const double SelectionStartTime = TimingView ? TimingView->GetSelectionStartTime() : 0.0;
 		const double SelectionEndTime = TimingView ? TimingView->GetSelectionEndTime() : 0.0;
 		TSharedPtr<Trace::ITable<Trace::FPackagesTableRow>> SourceTable = MakeShareable(LoadTimeProfilerProvider.CreatePackageDetailsTable(SelectionStartTime, SelectionEndTime));
-		TablePtr->Init(SourceTable);
+		Table->Init(SourceTable);
 	}
 
 	const TSharedRef<SDockTab> DockTab = SNew(SDockTab)
 		.ShouldAutosize(false)
 		.TabRole(ETabRole::PanelTab)
 		[
-			SAssignNew(PackageDetailsTreeView, Insights::STableTreeView, TablePtr)
+			SAssignNew(PackageDetailsTreeView, Insights::STableTreeView, Table)
 		];
 
 	PackageDetailsTreeView->Reset();
@@ -316,7 +319,8 @@ TSharedRef<SDockTab> SLoadingProfilerWindow::SpawnTab_ExportDetailsTreeView(cons
 {
 	FLoadingProfilerManager::Get()->SetExportDetailsTreeViewVisible(true);
 
-	TSharedPtr<Insights::FTable> TablePtr = MakeShareable(new Insights::FTable());
+	TSharedRef<Insights::FTable> Table = MakeShared<Insights::FTable>();
+
 	TSharedPtr<const Trace::IAnalysisSession> Session = FInsightsManager::Get()->GetSession();
 	if (Session.IsValid() && Trace::ReadLoadTimeProfilerProvider(*Session.Get()))
 	{
@@ -325,14 +329,14 @@ TSharedRef<SDockTab> SLoadingProfilerWindow::SpawnTab_ExportDetailsTreeView(cons
 		const double SelectionStartTime = TimingView ? TimingView->GetSelectionStartTime() : 0.0;
 		const double SelectionEndTime = TimingView ? TimingView->GetSelectionEndTime() : 0.0;
 		TSharedPtr<Trace::ITable<Trace::FExportsTableRow>> SourceTable = MakeShareable(LoadTimeProfilerProvider.CreateExportDetailsTable(SelectionStartTime, SelectionEndTime));
-		TablePtr->Init(SourceTable);
+		Table->Init(SourceTable);
 	}
 
 	const TSharedRef<SDockTab> DockTab = SNew(SDockTab)
 		.ShouldAutosize(false)
 		.TabRole(ETabRole::PanelTab)
 		[
-			SAssignNew(ExportDetailsTreeView, Insights::STableTreeView, TablePtr)
+			SAssignNew(ExportDetailsTreeView, Insights::STableTreeView, Table)
 		];
 
 	ExportDetailsTreeView->Reset();
diff --git a/Engine/Source/Developer/TraceInsights/Private/Insights/NetworkingProfiler/ViewModels/NetStatsViewColumnFactory.cpp b/Engine/Source/Developer/TraceInsights/Private/Insights/NetworkingProfiler/ViewModels/NetStatsViewColumnFactory.cpp
--- a/Engine/Source/Developer/TraceInsights/Private/Insights/NetworkingProfiler/ViewModels/NetStatsViewColumnFactory.cpp
+++ b/Engine/Source/Developer/TraceInsights/Private/Insights/NetworkingProfiler/ViewModels/NetStatsViewColumnFactory.cpp
@@ -55,8 +55,8 @@ TSharedRef<Insights::FTableColumn> FNetStatsViewColumnFactory::CreateNameColumn(
 {
 	using namespace Insights;
 
-	TSharedPtr<FTableColumn> ColumnPtr = MakeShareable(new FTableColumn(FNetStatsViewColumns::NameColumnID));
-	FTableColumn& Column = *ColumnPtr;
+	TSharedRef<FTableColumn> ColumnRef = MakeShared<FTableColumn>(FNetStatsViewColumns::NameColumnID);
+	FTableColumn& Column = *ColumnRef;
 
 	Column.SetShortName(LOCTEXT("Name_ColumnName", "Name"));
 	Column.SetTitleName(LOCTEXT("Name_ColumnTitle", "NetEvent or Group Name"));
@@ -72,16 +72,16 @@ TSharedRef<Insights::FTableColumn> FNetStatsViewColumnFactory::CreateNameColumn(
 
 	Column.SetDataType(ETableCellDataType::Text);
 
-	TSharedPtr<ITableCellValueGetter> GetterPtr = MakeShareable(new FDisplayNameValueGetter());
-	Column.SetValueGetter(GetterPtr.ToSharedRef());
+	TSharedRef<ITableCellValueGetter> Getter = MakeShared<FDisplayNameValueGetter>();
+	Column.SetValueGetter(Getter);
 
-	TSharedPtr<ITableCellValueFormatter> FormatterPtr = MakeShareable(new FTextValueFormatter());
-	Column.SetValueFormatter(FormatterPtr.ToSharedRef());
+	TSharedRef<ITableCellValueFormatter> Formatter = MakeShared<FTextValueFormatter>();
+	Column.SetValueFormatter(Formatter);
 
-	TSharedPtr<ITableCellValueSorter> SorterPtr = MakeShareable(new FSorterByName(ColumnPtr.ToSharedRef()));
-	Column.SetValueSorter(SorterPtr);
+	TSharedRef<ITableCellValueSorter> Sorter = MakeShared<FSorterByName>(ColumnRef);
+	Column.SetValueSorter(Sorter);
 
-	return ColumnPtr.ToSharedRef();
+	return ColumnRef;
 }
 
 ////////////////////////////////////////////////////////////////////////////////////////////////////
@@ -90,8 +90,8 @@ TSharedRef<Insights::FTableColumn> FNetStatsViewColumnFactory::CreateTypeColumn(
 {
 	using namespace Insights;
 
-	TSharedPtr<FTableColumn> ColumnPtr = MakeShareable(new FTableColumn(FNetStatsViewColumns::TypeColumnID));
-	FTableColumn& Column = *ColumnPtr;
+	TSharedRef<FTableColumn> ColumnRef = MakeShared<FTableColumn>(FNetStatsViewColumns::TypeColumnID);
+	FTableColumn& Column = *ColumnRef;
 
 	Column.SetShortName(LOCTEXT("Type_ColumnName", "Type"));
 	Column.SetTitleName(LOCTEXT("Type_ColumnTitle", "Type"));
@@ -117,17 +117,17 @@ TSharedRef<Insights::FTableColumn> FNetStatsViewColumnFactory::CreateTypeColumn(
 		}
 	};
 
-	TSharedPtr<FTableCellValueGetter> GetterPtr = MakeShareable(new FNetEventTypeValueGetter());
-	Column.SetValueGetter(GetterPtr.ToSharedRef());
+	TSharedRef<FTableCellValueGetter> Getter = MakeShared<FNetEventTypeValueGetter>();
+	Column.SetValueGetter(Getter);
 
-	TSharedPtr<FTableCellValueFormatter> FormatterPtr = MakeShareable(new FTextValueFormatter());
-	Column.SetValueFormatter(FormatterPtr.ToSharedRef());
+	TSharedRef<FTableCellValueFormatter> Formatter = MakeShared<FTextValueFormatter>();
+	Column.SetValueFormatter(Formatter);
 
-	//TSharedPtr<ITableCellValueSorter> SorterPtr = MakeShareable(new FSorterByTextValue(ColumnPtr.ToSharedRef()));
-	TSharedPtr<ITableCellValueSorter> SorterPtr = MakeShareable(new FNetEventNodeSortingByEventType(ColumnPtr.ToSharedRef()));
-	Column.SetValueSorter(SorterPtr);
+	//TSharedRef<ITableCellValueSorter> Sorter = MakeShared<FSorterByTextValue>(ColumnRef);
+	TSharedRef<ITableCellValueSorter> Sorter = MakeShared<FNetEventNodeSortingByEventType>(ColumnRef);
+	Column.SetValueSorter(Sorter);
 
-	return ColumnPtr.ToSharedRef();
+	return ColumnRef;
 }
 
 ////////////////////////////////////////////////////////////////////////////////////////////////////
@@ -136,8 +136,8 @@ TSharedRef<Insights::FTableColumn> FNetStatsViewColumnFactory::CreateLevelColumn
 {
 	using namespace Insights;
 
-	TSharedPtr<FTableColumn> ColumnPtr = MakeShareable(new FTableColumn(FNetStatsViewColumns::LevelColumnID));
-	FTableColumn& Column = *ColumnPtr;
+	TSharedRef<FTableColumn> ColumnRef = MakeShared<FTableColumn>(FNetStatsViewColumns::LevelColumnID);
+	FTableColumn& Column = *ColumnRef;
 
 	Column.SetShortName(LOCTEXT("Level_ColumnName", "Level"));
 	Column.SetTitleName(LOCTEXT("Level_ColumnTitle", "Level"));
@@ -163,16 +163,16 @@ TSharedRef<Insights::FTableColumn> FNetStatsViewColumnFactory::CreateLevelColumn
 		}
 	};
 
-	TSharedPtr<FTableCellValueGetter> GetterPtr = MakeShareable(new FLevelValueGetter());
-	Column.SetValueGetter(GetterPtr.ToSharedRef());
+	TSharedRef<FTableCellValueGetter> Getter = MakeShared<FLevelValueGetter>();
+	Column.SetValueGetter(Getter);
 
-	TSharedPtr<FTableCellValueFormatter> FormatterPtr = MakeShareable(new FInt64ValueFormatterAsNumber());
-	Column.SetValueFormatter(FormatterPtr.ToSharedRef());
+	TSharedRef<FTableCellValueFormatter> Formatter = MakeShared<FInt64ValueFormatterAsNumber>();
+	Column.SetValueFormatter(Formatter);
 
-	TSharedPtr<ITableCellValueSorter> SorterPtr = MakeShareable(new FSorterByInt64Value(ColumnPtr.ToSharedRef()));
-	Column.SetValueSorter(SorterPtr);
+	TSharedRef<ITableCellValueSorter> Sorter = MakeShared<FSorterByInt64Value>(ColumnRef);
+	Column.SetValueSorter(Sorter);
 
-	return ColumnPtr.ToSharedRef();
+	return ColumnRef;
 }
 
 ////////////////////////////////////////////////////////////////////////////////////////////////////
@@ -181,8 +181,8 @@ TSharedRef<Insights::FTableColumn> FNetStatsViewColumnFactory::CreateInstanceCou
 {
 	using namespace Insights;
 
-	TSharedPtr<FTableColumn> ColumnPtr = MakeShareable(new FTableColumn(FNetStatsViewColumns::InstanceCountColumnID));
-	FTableColumn& Column = *ColumnPtr;
+	TSharedRef<FTableColumn> ColumnRef = MakeShared<FTableColumn>(FNetStatsViewColumns::InstanceCountColumnID);
+	FTableColumn& Column = *ColumnRef;
 
 	Column.SetShortName(LOCTEXT("InstanceCount_ColumnName", "Count"));
 	Column.SetTitleName(LOCTEXT("InstanceCount_ColumnTitle", "Instance Count"));
@@ -208,17 +208,17 @@ TSharedRef<Insights::FTableColumn> FNetStatsViewColumnFactory::CreateInstanceCou
 		}
 	};
 
-	TSharedPtr<FTableCellValueGetter> GetterPtr = MakeShareable(new FInstanceCountValueGetter());
-	Column.SetValueGetter(GetterPtr.ToSharedRef());
+	TSharedRef<FTableCellValueGetter> Getter = MakeShared<FInstanceCountValueGetter>();
+	Column.SetValueGetter(Getter);
 
-	TSharedPtr<FTableCellValueFormatter> FormatterPtr = MakeShareable(new FInt64ValueFormatterAsNumber());
-	Column.SetValueFormatter(FormatterPtr.ToSharedRef());
+	TSharedRef<FTableCellValueFormatter> Formatter = MakeShared<FInt64ValueFormatterAsNumber>();
+	Column.SetValueFormatter(Formatter);
 
-	//TSharedPtr<ITableCellValueSorter> SorterPtr = MakeShareable(new FSorterByInt64Value(ColumnPtr.ToSharedRef()));
-	TSharedPtr<ITableCellValueSorter> SorterPtr = MakeShareable(new FNetEventNodeSortingByInstanceCount(ColumnPtr.ToSharedRef()));
-	Column.SetValueSorter(SorterPtr);
+	//TSharedRef<ITableCellValueSorter> Sorter = MakeShared<FSorterByInt64Value>(ColumnRef);
+	TSharedRef<ITableCellValueSorter> Sorter = MakeShared<FNetEventNodeSortingByInstanceCount>(ColumnRef);
+	Column.SetValueSorter(Sorter);
 
-	return ColumnPtr.ToSharedRef();
+	return ColumnRef;
 }
 
 ////////////////////////////////////////////////////////////////////////////////////////////////////
@@ -229,8 +229,8 @@ TSharedRef<Insights::FTableColumn> FNetStatsViewColumnFactory::CreateTotalInclus
 {
 	using namespace Insights;
 
-	TSharedPtr<FTableColumn> ColumnPtr = MakeShareable(new FTableColumn(FNetStatsViewColumns::TotalInclusiveSizeColumnID));
-	FTableColumn& Column = *ColumnPtr;
+	TSharedRef<FTableColumn> ColumnRef = MakeShared<FTableColumn>(FNetStatsViewColumns::TotalInclusiveSizeColumnID);
+	FTableColumn& Column = *ColumnRef;
 
 	Column.SetShortName(LOCTEXT("TotalInclusive_ColumnName", "Incl"));
 	Column.SetTitleName(LOCTEXT("TotalInclusive_ColumnTitle", "Total Inclusive Size"));
@@ -257,17 +257,17 @@ TSharedRef<Insights::FTableColumn> FNetStatsViewColumnFactory::CreateTotalInclus
 		}
 	};
 
-	TSharedPtr<FTableCellValueGetter> GetterPtr = MakeShareable(new FTotalInclusiveValueGetter());
-	Column.SetValueGetter(GetterPtr.ToSharedRef());
+	TSharedRef<FTableCellValueGetter> Getter = MakeShared<FTotalInclusiveValueGetter>();
+	Column.SetValueGetter(Getter);
 
-	TSharedPtr<FTableCellValueFormatter> FormatterPtr = MakeShareable(new FInt64ValueFormatterAsNumber());
-	Column.SetValueFormatter(FormatterPtr.ToSharedRef());
+	TSharedRef<FTableCellValueFormatter> Formatter = MakeShared<FInt64ValueFormatterAsNumber>();
+	Column.SetValueFormatter(Formatter);
 
-	//TSharedPtr<ITableCellValueSorter> SorterPtr = MakeShareable(new FSorterByInt64Value(ColumnPtr.ToSharedRef()));
-	TSharedPtr<ITableCellValueSorter> SorterPtr = MakeShareable(new FNetEventNodeSortingByTotalInclusiveSize(ColumnPtr.ToSharedRef()));
-	Column.SetValueSorter(SorterPtr);
+	//TSharedRef<ITableCellValueSorter> Sorter = MakeShared<FSorterByInt64Value>(ColumnRef);
+	TSharedRef<ITableCellValueSorter> Sorter = MakeShared<FNetEventNodeSortingByTotalInclusiveSize>(ColumnRef);
+	Column.SetValueSorter(Sorter);
 
-	return ColumnPtr.ToSharedRef();
+	return ColumnRef;
 }
 
 ////////////////////////////////////////////////////////////////////////////////////////////////////
@@ -276,8 +276,8 @@ TSharedRef<Insights::FTableColumn> FNetStatsViewColumnFactory::CreateMaxInclusiv
 {
 	using namespace Insights;
 
-	TSharedPtr<FTableColumn> ColumnPtr = MakeShareable(new FTableColumn(FNetStatsViewColumns::MaxInclusiveSizeColumnID));
-	FTableColumn& Column = *ColumnPtr;
+	TSharedRef<FTableColumn> ColumnRef = MakeShared<FTableColumn>(FNetStatsViewColumns::MaxInclusiveSizeColumnID);
+	FTableColumn& Column = *ColumnRef;
 
 	Column.SetShortName(LOCTEXT("MaxInclusive_ColumnName", "I.Max"));
 	Column.SetTitleName(LOCTEXT("MaxInclusive_ColumnTitle", "Max Inclusive Size"));
@@ -303,16 +303,16 @@ TSharedRef<Insights::FTableColumn> FNetStatsViewColumnFactory::CreateMaxInclusiv
 		}
 	};
 
-	TSharedPtr<FTableCellValueGetter> GetterPtr = MakeShareable(new FMaxInclusiveValueGetter());
-	Column.SetValueGetter(GetterPtr.ToSharedRef());
+	TSharedRef<FTableCellValueGetter> Getter = MakeShared<FMaxInclusiveValueGetter>();
+	Column.SetValueGetter(Getter);
 
-	TSharedPtr<FTableCellValueFormatter> FormatterPtr = MakeShareable(new FInt64ValueFormatterAsNumber());
-	Column.SetValueFormatter(FormatterPtr.ToSharedRef());
+	TSharedRef<FTableCellValueFormatter> Formatter = MakeShared<FInt64ValueFormatterAsNumber>();
+	Column.SetValueFormatter(Formatter);
 
-	TSharedPtr<ITableCellValueSorter> SorterPtr = MakeShareable(new FSorterByInt64Value(ColumnPtr.ToSharedRef()));
-	Column.SetValueSorter(SorterPtr);
+	TSharedRef<ITableCellValueSorter> Sorter = MakeShared<FSorterByInt64Value>(ColumnRef);
+	Column.SetValueSorter(Sorter);
 
-	return ColumnPtr.ToSharedRef();
+	return ColumnRef;
 }
 
 ////////////////////////////////////////////////////////////////////////////////////////////////////
@@ -321,8 +321,8 @@ TSharedRef<Insights::FTableColumn> FNetStatsViewColumnFactory::CreateAverageIncl
 {
 	using namespace Insights;
 
-	TSharedPtr<FTableColumn> ColumnPtr = MakeShareable(new FTableColumn(FNetStatsViewColumns::AverageInclusiveSizeColumnID));
-	FTableColumn& Column = *ColumnPtr;
+	TSharedRef<FTableColumn> ColumnRef = MakeShared<FTableColumn>(FNetStatsViewColumns::AverageInclusiveSizeColumnID);
+	FTableColumn& Column = *ColumnRef;
 
 	Column.SetShortName(LOCTEXT("AvgInclusive_ColumnName", "I.Avg"));
 	Column.SetTitleName(LOCTEXT("AvgInclusive_ColumnTitle", "Average Inclusive Size"));
@@ -348,16 +348,16 @@ TSharedRef<Insights::FTableColumn> FNetStatsViewColumnFactory::CreateAverageIncl
 		}
 	};
 
-	TSharedPtr<FTableCellValueGetter> GetterPtr = MakeShareable(new FAverageInclusiveValueGetter());
-	Column.SetValueGetter(GetterPtr.ToSharedRef());
+	TSharedRef<FTableCellValueGetter> Getter = MakeShared<FAverageInclusiveValueGetter>();
+	Column.SetValueGetter(Getter);
 
-	TSharedPtr<FTableCellValueFormatter> FormatterPtr = MakeShareable(new FInt64ValueFormatterAsNumber());
-	Column.SetValueFormatter(FormatterPtr.ToSharedRef());
+	TSharedRef<FTableCellValueFormatter> Formatter = MakeShared<FInt64ValueFormatterAsNumber>();
+	Column.SetValueFormatter(Formatter);
 
-	TSharedPtr<ITableCellValueSorter> SorterPtr = MakeShareable(new FSorterByInt64Value(ColumnPtr.ToSharedRef()));
-	Column.SetValueSorter(SorterPtr);
+	TSharedRef<ITableCellValueSorter> Sorter = MakeShared<FSorterByInt64Value>(ColumnRef);
+	Column.SetValueSorter(Sorter);
 
-	return ColumnPtr.ToSharedRef();
+	return ColumnRef;
 }
 
 ////////////////////////////////////////////////////////////////////////////////////////////////////
@@ -368,8 +368,8 @@ TSharedRef<Insights::FTableColumn> FNetStatsViewColumnFactory::CreateTotalExclus
 {
 	using namespace Insights;
 
-	TSharedPtr<FTableColumn> ColumnPtr = MakeShareable(new FTableColumn(FNetStatsViewColumns::TotalExclusiveSizeColumnID));
-	FTableColumn& Column = *ColumnPtr;
+	TSharedRef<FTableColumn> ColumnRef = MakeShared<FTableColumn>(FNetStatsViewColumns::TotalExclusiveSizeColumnID);
+	FTableColumn& Column = *ColumnRef;
 
 	Column.SetShortName(LOCTEXT("TotalExclusive_ColumnName", "Excl"));
 	Column.SetTitleName(LOCTEXT("TotalExclusive_ColumnTitle", "Total Exclusive Size"));
@@ -396,17 +396,17 @@ TSharedRef<Insights::FTableColumn> FNetStatsViewColumnFactory::CreateTotalExclus
 		}
 	};
 
-	TSharedPtr<FTableCellValueGetter> GetterPtr = MakeShareable(new FTotalExclusiveValueGetter());
-	Column.SetValueGetter(GetterPtr.ToSharedRef());
+	TSharedRef<FTableCellValueGetter> Getter = MakeShared<FTotalExclusiveValueGetter>();
+	Column.SetValueGetter(Getter);
 
-	TSharedPtr<FTableCellValueFormatter> FormatterPtr = MakeShareable(new FInt64ValueFormatterAsNumber());
-	Column.SetValueFormatter(FormatterPtr.ToSharedRef());
+	TSharedRef<FTableCellValueFormatter> Formatter = MakeShared<FInt64ValueFormatterAsNumber>();
+	Column.SetValueFormatter(Formatter);
 
-	//TSharedPtr<ITableCellValueSorter> SorterPtr = MakeShareable(new FSorterByInt64Value(ColumnPtr.ToSharedRef()));
-	TSharedPtr<ITableCellValueSorter> SorterPtr = MakeShareable(new FNetEventNodeSortingByTotalExclusiveSize(ColumnPtr.ToSharedRef()));
-	Column.SetValueSorter(SorterPtr);
+	//TSharedRef<ITableCellValueSorter> Sorter = MakeShared<FSorterByInt64Value(ColumnRef));
+	TSharedRef<ITableCellValueSorter> Sorter = MakeShared<FNetEventNodeSortingByTotalExclusiveSize>(ColumnRef);
+	Column.SetValueSorter(Sorter);
 
-	return ColumnPtr.ToSharedRef();
+	return ColumnRef;
 }
 
 ////////////////////////////////////////////////////////////////////////////////////////////////////
@@ -415,8 +415,8 @@ TSharedRef<Insights::FTableColumn> FNetStatsViewColumnFactory::CreateMaxExclusiv
 {
 	using namespace Insights;
 
-	TSharedPtr<FTableColumn> ColumnPtr = MakeShareable(new FTableColumn(FNetStatsViewColumns::MaxExclusiveSizeColumnID));
-	FTableColumn& Column = *ColumnPtr;
+	TSharedRef<FTableColumn> ColumnRef = MakeShared<FTableColumn>(FNetStatsViewColumns::MaxExclusiveSizeColumnID);
+	FTableColumn& Column = *ColumnRef;
 
 	Column.SetShortName(LOCTEXT("MaxExclusive_ColumnName", "E.Max"));
 	Column.SetTitleName(LOCTEXT("MaxExclusive_ColumnTitle", "Max Exclusive Size"));
@@ -442,16 +442,16 @@ TSharedRef<Insights::FTableColumn> FNetStatsViewColumnFactory::CreateMaxExclusiv
 		}
 	};
 
-	TSharedPtr<FTableCellValueGetter> GetterPtr = MakeShareable(new FMaxExclusiveValueGetter());
-	Column.SetValueGetter(GetterPtr.ToSharedRef());
+	TSharedRef<FTableCellValueGetter> Getter = MakeShared<FMaxExclusiveValueGetter>();
+	Column.SetValueGetter(Getter);
 
-	TSharedPtr<FTableCellValueFormatter> FormatterPtr = MakeShareable(new FInt64ValueFormatterAsNumber());
-	Column.SetValueFormatter(FormatterPtr.ToSharedRef());
+	TSharedRef<FTableCellValueFormatter> Formatter = MakeShared<FInt64ValueFormatterAsNumber>();
+	Column.SetValueFormatter(Formatter);
 
-	TSharedPtr<ITableCellValueSorter> SorterPtr = MakeShareable(new FSorterByInt64Value(ColumnPtr.ToSharedRef()));
-	Column.SetValueSorter(SorterPtr);
+	TSharedRef<ITableCellValueSorter> Sorter = MakeShared<FSorterByInt64Value>(ColumnRef);
+	Column.SetValueSorter(Sorter);
 
-	return ColumnPtr.ToSharedRef();
+	return ColumnRef;
 }
 
 ////////////////////////////////////////////////////////////////////////////////////////////////////
@@ -460,8 +460,8 @@ TSharedRef<Insights::FTableColumn> FNetStatsViewColumnFactory::CreateAverageExcl
 {
 	using namespace Insights;
 
-	TSharedPtr<FTableColumn> ColumnPtr = MakeShareable(new FTableColumn(FNetStatsViewColumns::AverageExclusiveSizeColumnID));
-	FTableColumn& Column = *ColumnPtr;
+	TSharedRef<FTableColumn> ColumnRef = MakeShared<FTableColumn>(FNetStatsViewColumns::AverageExclusiveSizeColumnID);
+	FTableColumn& Column = *ColumnRef;
 
 	Column.SetShortName(LOCTEXT("AvgExclusiveName", "E.Avg"));
 	Column.SetTitleName(LOCTEXT("AvgExclusiveTitle", "Average Exclusive  (ms)"));
@@ -488,16 +488,16 @@ TSharedRef<Insights::FTableColumn> FNetStatsViewColumnFactory::CreateAverageExcl
 		}
 	};
 
-	TSharedPtr<FTableCellValueGetter> GetterPtr = MakeShareable(new FAverageExclusiveValueGetter());
-	Column.SetValueGetter(GetterPtr.ToSharedRef());
+	TSharedRef<FTableCellValueGetter> Getter = MakeShared<FAverageExclusiveValueGetter>();
+	Column.SetValueGetter(Getter);
 
-	TSharedPtr<FTableCellValueFormatter> FormatterPtr = MakeShareable(new FInt64ValueFormatterAsNumber());
-	Column.SetValueFormatter(FormatterPtr.ToSharedRef());
+	TSharedRef<FTableCellValueFormatter> Formatter = MakeShared<FInt64ValueFormatterAsNumber>();
+	Column.SetValueFormatter(Formatter);
 
-	TSharedPtr<ITableCellValueSorter> SorterPtr = MakeShareable(new FSorterByInt64Value(ColumnPtr.ToSharedRef()));
-	Column.SetValueSorter(SorterPtr);
+	TSharedRef<ITableCellValueSorter> Sorter = MakeShared<FSorterByInt64Value>(ColumnRef);
+	Column.SetValueSorter(Sorter);
 
-	return ColumnPtr.ToSharedRef();
+	return ColumnRef;
 }
 
 ////////////////////////////////////////////////////////////////////////////////////////////////////
diff --git a/Engine/Source/Developer/TraceInsights/Private/Insights/NetworkingProfiler/Widgets/SPacketContentView.cpp b/Engine/Source/Developer/TraceInsights/Private/Insights/NetworkingProfiler/Widgets/SPacketContentView.cpp
--- a/Engine/Source/Developer/TraceInsights/Private/Insights/NetworkingProfiler/Widgets/SPacketContentView.cpp
+++ b/Engine/Source/Developer/TraceInsights/Private/Insights/NetworkingProfiler/Widgets/SPacketContentView.cpp
@@ -26,7 +26,7 @@
 
 SPacketContentView::SPacketContentView()
 	: ProfilerWindow()
-	, DrawState(MakeShareable(new FPacketContentViewDrawState()))
+	, DrawState(MakeShared<FPacketContentViewDrawState>())
 {
 	Reset();
 }
diff --git a/Engine/Source/Developer/TraceInsights/Private/Insights/NetworkingProfiler/Widgets/SPacketContentView.h b/Engine/Source/Developer/TraceInsights/Private/Insights/NetworkingProfiler/Widgets/SPacketContentView.h
--- a/Engine/Source/Developer/TraceInsights/Private/Insights/NetworkingProfiler/Widgets/SPacketContentView.h
+++ b/Engine/Source/Developer/TraceInsights/Private/Insights/NetworkingProfiler/Widgets/SPacketContentView.h
@@ -169,7 +169,7 @@ private:
 	int64 PacketBitSize; // total number of bits; [bit]
 
 	/** Cached draw state of the packet content (i.e. all it needs to render). */
-	TSharedPtr<FPacketContentViewDrawState> DrawState;
+	TSharedRef<FPacketContentViewDrawState> DrawState;
 	bool bIsStateDirty;
 	//////////////////////////////////////////////////
 
diff --git a/Engine/Source/Developer/TraceInsights/Private/Insights/ViewModels/TimersViewColumnFactory.cpp b/Engine/Source/Developer/TraceInsights/Private/Insights/ViewModels/TimersViewColumnFactory.cpp
--- a/Engine/Source/Developer/TraceInsights/Private/Insights/ViewModels/TimersViewColumnFactory.cpp
+++ b/Engine/Source/Developer/TraceInsights/Private/Insights/ViewModels/TimersViewColumnFactory.cpp
@@ -81,8 +81,8 @@ TSharedRef<Insights::FTableColumn> FTimersViewColumnFactory::CreateNameColumn()
 {
 	using namespace Insights;
 
-	TSharedPtr<FTableColumn> ColumnPtr = MakeShareable(new FTableColumn(FTimersViewColumns::NameColumnID));
-	FTableColumn& Column = *ColumnPtr;
+	TSharedRef<FTableColumn> ColumnRef = MakeShared<FTableColumn>(FTimersViewColumns::NameColumnID);
+	FTableColumn& Column = *ColumnRef;
 
 	Column.SetShortName(LOCTEXT("TimerNameColumnName", "Name"));
 	Column.SetTitleName(LOCTEXT("TimerNameColumnTitle", "Timer or Group Name"));
@@ -98,16 +98,16 @@ TSharedRef<Insights::FTableColumn> FTimersViewColumnFactory::CreateNameColumn()
 
 	Column.SetDataType(ETableCellDataType::Text);
 
-	TSharedPtr<ITableCellValueGetter> GetterPtr = MakeShareable(new FDisplayNameValueGetter());
-	Column.SetValueGetter(GetterPtr.ToSharedRef());
+	TSharedRef<ITableCellValueGetter> Getter = MakeShared<FDisplayNameValueGetter>();
+	Column.SetValueGetter(Getter);
 
-	TSharedPtr<ITableCellValueFormatter> FormatterPtr = MakeShareable(new FTextValueFormatter());
-	Column.SetValueFormatter(FormatterPtr.ToSharedRef());
+	TSharedRef<ITableCellValueFormatter> Formatter = MakeShared<FTextValueFormatter>();
+	Column.SetValueFormatter(Formatter);
 
-	TSharedPtr<ITableCellValueSorter> SorterPtr = MakeShareable(new FSorterByName(ColumnPtr.ToSharedRef()));
-	Column.SetValueSorter(SorterPtr);
+	TSharedRef<ITableCellValueSorter> Sorter = MakeShared<FSorterByName>(ColumnRef);
+	Column.SetValueSorter(Sorter);
 
-	return ColumnPtr.ToSharedRef();
+	return ColumnRef;
 }
 
 ////////////////////////////////////////////////////////////////////////////////////////////////////
@@ -116,8 +116,8 @@ TSharedRef<Insights::FTableColumn> FTimersViewColumnFactory::CreateMetaGroupName
 {
 	using namespace Insights;
 
-	TSharedPtr<FTableColumn> ColumnPtr = MakeShareable(new FTableColumn(FTimersViewColumns::MetaGroupNameColumnID));
-	FTableColumn& Column = *ColumnPtr;
+	TSharedRef<FTableColumn> ColumnRef = MakeShared<FTableColumn>(FTimersViewColumns::MetaGroupNameColumnID);
+	FTableColumn& Column = *ColumnRef;
 
 	Column.SetShortName(LOCTEXT("TimerMetaGroupNameColumnName", "Meta Group"));
 	Column.SetTitleName(LOCTEXT("TimerMetaGroupNameColumnTitle", "Meta Group Name"));
@@ -143,16 +143,16 @@ TSharedRef<Insights::FTableColumn> FTimersViewColumnFactory::CreateMetaGroupName
 		}
 	};
 
-	TSharedPtr<FTableCellValueGetter> GetterPtr = MakeShareable(new FMetaGroupNameValueGetter());
-	Column.SetValueGetter(GetterPtr.ToSharedRef());
+	TSharedRef<FTableCellValueGetter> Getter = MakeShared<FMetaGroupNameValueGetter>();
+	Column.SetValueGetter(Getter);
 
-	TSharedPtr<FTableCellValueFormatter> FormatterPtr = MakeShareable(new FTextValueFormatter());
-	Column.SetValueFormatter(FormatterPtr.ToSharedRef());
+	TSharedRef<FTableCellValueFormatter> Formatter = MakeShared<FTextValueFormatter>();
+	Column.SetValueFormatter(Formatter);
 
-	TSharedPtr<ITableCellValueSorter> SorterPtr = MakeShareable(new FSorterByTextValue(ColumnPtr.ToSharedRef()));
-	Column.SetValueSorter(SorterPtr);
+	TSharedRef<ITableCellValueSorter> Sorter = MakeShared<FSorterByTextValue>(ColumnRef);
+	Column.SetValueSorter(Sorter);
 
-	return ColumnPtr.ToSharedRef();
+	return ColumnRef;
 }
 
 ////////////////////////////////////////////////////////////////////////////////////////////////////
@@ -161,8 +161,8 @@ TSharedRef<Insights::FTableColumn> FTimersViewColumnFactory::CreateTypeColumn()
 {
 	using namespace Insights;
 
-	TSharedPtr<FTableColumn> ColumnPtr = MakeShareable(new FTableColumn(FTimersViewColumns::TypeColumnID));
-	FTableColumn& Column = *ColumnPtr;
+	TSharedRef<FTableColumn> ColumnRef = MakeShared<FTableColumn>(FTimersViewColumns::TypeColumnID);
+	FTableColumn& Column = *ColumnRef;
 
 	Column.SetShortName(LOCTEXT("TimerTypeColumnName", "Type"));
 	Column.SetTitleName(LOCTEXT("TimerTypeColumnTitle", "Type"));
@@ -188,17 +188,17 @@ TSharedRef<Insights::FTableColumn> FTimersViewColumnFactory::CreateTypeColumn()
 		}
 	};
 
-	TSharedPtr<FTableCellValueGetter> GetterPtr = MakeShareable(new FTimerTypeValueGetter());
-	Column.SetValueGetter(GetterPtr.ToSharedRef());
+	TSharedRef<FTableCellValueGetter> Getter = MakeShared<FTimerTypeValueGetter>();
+	Column.SetValueGetter(Getter);
 
-	TSharedPtr<FTableCellValueFormatter> FormatterPtr = MakeShareable(new FTextValueFormatter());
-	Column.SetValueFormatter(FormatterPtr.ToSharedRef());
+	TSharedRef<FTableCellValueFormatter> Formatter = MakeShared<FTextValueFormatter>();
+	Column.SetValueFormatter(Formatter);
 
-	//TSharedPtr<ITableCellValueSorter> SorterPtr = MakeShareable(new FSorterByTextValue(ColumnPtr.ToSharedRef()));
-	TSharedPtr<ITableCellValueSorter> SorterPtr = MakeShareable(new FTimerNodeSortingByTimerType(ColumnPtr.ToSharedRef()));
-	Column.SetValueSorter(SorterPtr);
+	//TSharedRef<ITableCellValueSorter> Sorter = MakeShared<FSorterByTextValue>(ColumnRef);
+	TSharedRef<ITableCellValueSorter> Sorter = MakeShared<FTimerNodeSortingByTimerType>(ColumnRef);
+	Column.SetValueSorter(Sorter);
 
-	return ColumnPtr.ToSharedRef();
+	return ColumnRef;
 }
 
 ////////////////////////////////////////////////////////////////////////////////////////////////////
@@ -207,8 +207,8 @@ TSharedRef<Insights::FTableColumn> FTimersViewColumnFactory::CreateInstanceCount
 {
 	using namespace Insights;
 
-	TSharedPtr<FTableColumn> ColumnPtr = MakeShareable(new FTableColumn(FTimersViewColumns::InstanceCountColumnID));
-	FTableColumn& Column = *ColumnPtr;
+	TSharedRef<FTableColumn> ColumnRef = MakeShared<FTableColumn>(FTimersViewColumns::InstanceCountColumnID);
+	FTableColumn& Column = *ColumnRef;
 
 	Column.SetShortName(LOCTEXT("InstanceCountName", "Count"));
 	Column.SetTitleName(LOCTEXT("InstanceCountTitle", "Instance Count"));
@@ -234,17 +234,17 @@ TSharedRef<Insights::FTableColumn> FTimersViewColumnFactory::CreateInstanceCount
 		}
 	};
 
-	TSharedPtr<FTableCellValueGetter> GetterPtr = MakeShareable(new FInstanceCountValueGetter());
-	Column.SetValueGetter(GetterPtr.ToSharedRef());
+	TSharedRef<FTableCellValueGetter> Getter = MakeShared<FInstanceCountValueGetter>();
+	Column.SetValueGetter(Getter);
 
-	TSharedPtr<FTableCellValueFormatter> FormatterPtr = MakeShareable(new FInt64ValueFormatterAsNumber());
-	Column.SetValueFormatter(FormatterPtr.ToSharedRef());
+	TSharedRef<FTableCellValueFormatter> Formatter = MakeShared<FInt64ValueFormatterAsNumber>();
+	Column.SetValueFormatter(Formatter);
 
-	//TSharedPtr<ITableCellValueSorter> SorterPtr = MakeShareable(new FSorterByInt64Value(ColumnPtr.ToSharedRef()));
-	TSharedPtr<ITableCellValueSorter> SorterPtr = MakeShareable(new FTimerNodeSortingByInstanceCount(ColumnPtr.ToSharedRef()));
-	Column.SetValueSorter(SorterPtr);
+	//TSharedRef<ITableCellValueSorter> Sorter = MakeShared<FSorterByInt64Value>(ColumnRef);
+	TSharedRef<ITableCellValueSorter> Sorter = MakeShared<FTimerNodeSortingByInstanceCount>(ColumnRef);
+	Column.SetValueSorter(Sorter);
 
-	return ColumnPtr.ToSharedRef();
+	return ColumnRef;
 }
 
 ////////////////////////////////////////////////////////////////////////////////////////////////////
@@ -255,8 +255,8 @@ TSharedRef<Insights::FTableColumn> FTimersViewColumnFactory::CreateTotalInclusiv
 {
 	using namespace Insights;
 
-	TSharedPtr<FTableColumn> ColumnPtr = MakeShareable(new FTableColumn(FTimersViewColumns::TotalInclusiveTimeColumnID));
-	FTableColumn& Column = *ColumnPtr;
+	TSharedRef<FTableColumn> ColumnRef = MakeShared<FTableColumn>(FTimersViewColumns::TotalInclusiveTimeColumnID);
+	FTableColumn& Column = *ColumnRef;
 
 	Column.SetShortName(LOCTEXT("TotalInclusiveTimeName", "Incl"));
 	Column.SetTitleName(LOCTEXT("TotalInclusiveTimeTitle", "Total Inclusive Time"));
@@ -283,17 +283,17 @@ TSharedRef<Insights::FTableColumn> FTimersViewColumnFactory::CreateTotalInclusiv
 		}
 	};
 
-	TSharedPtr<FTableCellValueGetter> GetterPtr = MakeShareable(new FTotalInclusiveTimeValueGetter());
-	Column.SetValueGetter(GetterPtr.ToSharedRef());
+	TSharedRef<FTableCellValueGetter> Getter = MakeShared<FTotalInclusiveTimeValueGetter>();
+	Column.SetValueGetter(Getter);
 
-	TSharedPtr<FTableCellValueFormatter> FormatterPtr = MakeShareable(new FDoubleValueFormatterAsTimeAuto());
-	Column.SetValueFormatter(FormatterPtr.ToSharedRef());
+	TSharedRef<FTableCellValueFormatter> Formatter = MakeShared<FDoubleValueFormatterAsTimeAuto>();
+	Column.SetValueFormatter(Formatter);
 
-	//TSharedPtr<ITableCellValueSorter> SorterPtr = MakeShareable(new FSorterByDoubleValue(ColumnPtr.ToSharedRef()));
-	TSharedPtr<ITableCellValueSorter> SorterPtr = MakeShareable(new FTimerNodeSortingByTotalInclusiveTime(ColumnPtr.ToSharedRef()));
-	Column.SetValueSorter(SorterPtr);
+	//TSharedRef<ITableCellValueSorter> Sorter = MakeShared<FSorterByDoubleValue>(ColumnRef);
+	TSharedRef<ITableCellValueSorter> Sorter = MakeShared<FTimerNodeSortingByTotalInclusiveTime>(ColumnRef);
+	Column.SetValueSorter(Sorter);
 
-	return ColumnPtr.ToSharedRef();
+	return ColumnRef;
 }
 
 ////////////////////////////////////////////////////////////////////////////////////////////////////
@@ -302,8 +302,8 @@ TSharedRef<Insights::FTableColumn> FTimersViewColumnFactory::CreateMaxInclusiveT
 {
 	using namespace Insights;
 
-	TSharedPtr<FTableColumn> ColumnPtr = MakeShareable(new FTableColumn(FTimersViewColumns::MaxInclusiveTimeColumnID));
-	FTableColumn& Column = *ColumnPtr;
+	TSharedRef<FTableColumn> ColumnRef = MakeShared<FTableColumn>(FTimersViewColumns::MaxInclusiveTimeColumnID);
+	FTableColumn& Column = *ColumnRef;
 
 	Column.SetShortName(LOCTEXT("MaxInclusiveTimeName", "I.Max"));
 	Column.SetTitleName(LOCTEXT("MaxInclusiveTimeTitle", "Max Inclusive Time (ms)"));
@@ -329,16 +329,16 @@ TSharedRef<Insights::FTableColumn> FTimersViewColumnFactory::CreateMaxInclusiveT
 		}
 	};
 
-	TSharedPtr<FTableCellValueGetter> GetterPtr = MakeShareable(new FMaxInclusiveTimeValueGetter());
-	Column.SetValueGetter(GetterPtr.ToSharedRef());
+	TSharedRef<FTableCellValueGetter> Getter = MakeShared<FMaxInclusiveTimeValueGetter>();
+	Column.SetValueGetter(Getter);
 
-	TSharedPtr<FTableCellValueFormatter> FormatterPtr = MakeShareable(new FDoubleValueFormatterAsTimeMs());
-	Column.SetValueFormatter(FormatterPtr.ToSharedRef());
+	TSharedRef<FTableCellValueFormatter> Formatter = MakeShared<FDoubleValueFormatterAsTimeMs>();
+	Column.SetValueFormatter(Formatter);
 
-	TSharedPtr<ITableCellValueSorter> SorterPtr = MakeShareable(new FSorterByDoubleValue(ColumnPtr.ToSharedRef()));
-	Column.SetValueSorter(SorterPtr);
+	TSharedRef<ITableCellValueSorter> Sorter = MakeShared<FSorterByDoubleValue>(ColumnRef);
+	Column.SetValueSorter(Sorter);
 
-	return ColumnPtr.ToSharedRef();
+	return ColumnRef;
 }
 
 ////////////////////////////////////////////////////////////////////////////////////////////////////
@@ -347,8 +347,8 @@ TSharedRef<Insights::FTableColumn> FTimersViewColumnFactory::CreateAverageInclus
 {
 	using namespace Insights;
 
-	TSharedPtr<FTableColumn> ColumnPtr = MakeShareable(new FTableColumn(FTimersViewColumns::AverageInclusiveTimeColumnID));
-	FTableColumn& Column = *ColumnPtr;
+	TSharedRef<FTableColumn> ColumnRef = MakeShared<FTableColumn>(FTimersViewColumns::AverageInclusiveTimeColumnID);
+	FTableColumn& Column = *ColumnRef;
 
 	Column.SetShortName(LOCTEXT("AvgInclusiveTimeName", "I.Avg"));
 	Column.SetTitleName(LOCTEXT("AvgInclusiveTimeTitle", "Average Inclusive Time (ms)"));
@@ -374,16 +374,16 @@ TSharedRef<Insights::FTableColumn> FTimersViewColumnFactory::CreateAverageInclus
 		}
 	};
 
-	TSharedPtr<FTableCellValueGetter> GetterPtr = MakeShareable(new FAverageInclusiveTimeValueGetter());
-	Column.SetValueGetter(GetterPtr.ToSharedRef());
+	TSharedRef<FTableCellValueGetter> Getter = MakeShared<FAverageInclusiveTimeValueGetter>();
+	Column.SetValueGetter(Getter);
 
-	TSharedPtr<FTableCellValueFormatter> FormatterPtr = MakeShareable(new FDoubleValueFormatterAsTimeMs());
-	Column.SetValueFormatter(FormatterPtr.ToSharedRef());
+	TSharedRef<FTableCellValueFormatter> Formatter = MakeShared<FDoubleValueFormatterAsTimeMs>();
+	Column.SetValueFormatter(Formatter);
 
-	TSharedPtr<ITableCellValueSorter> SorterPtr = MakeShareable(new FSorterByDoubleValue(ColumnPtr.ToSharedRef()));
-	Column.SetValueSorter(SorterPtr);
+	TSharedRef<ITableCellValueSorter> Sorter = MakeShared<FSorterByDoubleValue>(ColumnRef);
+	Column.SetValueSorter(Sorter);
 
-	return ColumnPtr.ToSharedRef();
+	return ColumnRef;
 }
 
 ////////////////////////////////////////////////////////////////////////////////////////////////////
@@ -392,8 +392,8 @@ TSharedRef<Insights::FTableColumn> FTimersViewColumnFactory::CreateMedianInclusi
 {
 	using namespace Insights;
 
-	TSharedPtr<FTableColumn> ColumnPtr = MakeShareable(new FTableColumn(FTimersViewColumns::MedianInclusiveTimeColumnID));
-	FTableColumn& Column = *ColumnPtr;
+	TSharedRef<FTableColumn> ColumnRef = MakeShared<FTableColumn>(FTimersViewColumns::MedianInclusiveTimeColumnID);
+	FTableColumn& Column = *ColumnRef;
 
 	Column.SetShortName(LOCTEXT("MedInclusiveTimeName", "I.Med"));
 	Column.SetTitleName(LOCTEXT("MedInclusiveTimeTitle", "Median Inclusive Time (ms)"));
@@ -419,16 +419,16 @@ TSharedRef<Insights::FTableColumn> FTimersViewColumnFactory::CreateMedianInclusi
 		}
 	};
 
-	TSharedPtr<FTableCellValueGetter> GetterPtr = MakeShareable(new FMedianInclusiveTimeValueGetter());
-	Column.SetValueGetter(GetterPtr.ToSharedRef());
+	TSharedRef<FTableCellValueGetter> Getter = MakeShared<FMedianInclusiveTimeValueGetter>();
+	Column.SetValueGetter(Getter);
 
-	TSharedPtr<FTableCellValueFormatter> FormatterPtr = MakeShareable(new FDoubleValueFormatterAsTimeMs());
-	Column.SetValueFormatter(FormatterPtr.ToSharedRef());
+	TSharedRef<FTableCellValueFormatter> Formatter = MakeShared<FDoubleValueFormatterAsTimeMs>();
+	Column.SetValueFormatter(Formatter);
 
-	TSharedPtr<ITableCellValueSorter> SorterPtr = MakeShareable(new FSorterByDoubleValue(ColumnPtr.ToSharedRef()));
-	Column.SetValueSorter(SorterPtr);
+	TSharedRef<ITableCellValueSorter> Sorter = MakeShared<FSorterByDoubleValue>(ColumnRef);
+	Column.SetValueSorter(Sorter);
 
-	return ColumnPtr.ToSharedRef();
+	return ColumnRef;
 }
 
 ////////////////////////////////////////////////////////////////////////////////////////////////////
@@ -437,8 +437,8 @@ TSharedRef<Insights::FTableColumn> FTimersViewColumnFactory::CreateMinInclusiveT
 {
 	using namespace Insights;
 
-	TSharedPtr<FTableColumn> ColumnPtr = MakeShareable(new FTableColumn(FTimersViewColumns::MinInclusiveTimeColumnID));
-	FTableColumn& Column = *ColumnPtr;
+	TSharedRef<FTableColumn> ColumnRef = MakeShared<FTableColumn>(FTimersViewColumns::MinInclusiveTimeColumnID);
+	FTableColumn& Column = *ColumnRef;
 
 	Column.SetShortName(LOCTEXT("MinInclusiveTimeName", "I.Min"));
 	Column.SetTitleName(LOCTEXT("MinInclusiveTimeTitle", "Min Inclusive Time (ms)"));
@@ -464,16 +464,16 @@ TSharedRef<Insights::FTableColumn> FTimersViewColumnFactory::CreateMinInclusiveT
 		}
 	};
 
-	TSharedPtr<FTableCellValueGetter> GetterPtr = MakeShareable(new FMinInclusiveTimeValueGetter());
-	Column.SetValueGetter(GetterPtr.ToSharedRef());
+	TSharedRef<FTableCellValueGetter> Getter = MakeShared<FMinInclusiveTimeValueGetter>();
+	Column.SetValueGetter(Getter);
 
-	TSharedPtr<FTableCellValueFormatter> FormatterPtr = MakeShareable(new FDoubleValueFormatterAsTimeMs());
-	Column.SetValueFormatter(FormatterPtr.ToSharedRef());
+	TSharedRef<FTableCellValueFormatter> Formatter = MakeShared<FDoubleValueFormatterAsTimeMs>();
+	Column.SetValueFormatter(Formatter);
 
-	TSharedPtr<ITableCellValueSorter> SorterPtr = MakeShareable(new FSorterByDoubleValue(ColumnPtr.ToSharedRef()));
-	Column.SetValueSorter(SorterPtr);
+	TSharedRef<ITableCellValueSorter> Sorter = MakeShared<FSorterByDoubleValue>(ColumnRef);
+	Column.SetValueSorter(Sorter);
 
-	return ColumnPtr.ToSharedRef();
+	return ColumnRef;
 }
 
 ////////////////////////////////////////////////////////////////////////////////////////////////////
@@ -484,8 +484,8 @@ TSharedRef<Insights::FTableColumn> FTimersViewColumnFactory::CreateTotalExclusiv
 {
 	using namespace Insights;
 
-	TSharedPtr<FTableColumn> ColumnPtr = MakeShareable(new FTableColumn(FTimersViewColumns::TotalExclusiveTimeColumnID));
-	FTableColumn& Column = *ColumnPtr;
+	TSharedRef<FTableColumn> ColumnRef = MakeShared<FTableColumn>(FTimersViewColumns::TotalExclusiveTimeColumnID);
+	FTableColumn& Column = *ColumnRef;
 
 	Column.SetShortName(LOCTEXT("TotalExclusiveTimeName", "Excl"));
 	Column.SetTitleName(LOCTEXT("TotalExclusiveTimeTitle", "Total Exclusive Time"));
@@ -512,17 +512,17 @@ TSharedRef<Insights::FTableColumn> FTimersViewColumnFactory::CreateTotalExclusiv
 		}
 	};
 
-	TSharedPtr<FTableCellValueGetter> GetterPtr = MakeShareable(new FTotalExclusiveTimeValueGetter());
-	Column.SetValueGetter(GetterPtr.ToSharedRef());
+	TSharedRef<FTableCellValueGetter> Getter = MakeShared<FTotalExclusiveTimeValueGetter>();
+	Column.SetValueGetter(Getter);
 
-	TSharedPtr<FTableCellValueFormatter> FormatterPtr = MakeShareable(new FDoubleValueFormatterAsTimeAuto());
-	Column.SetValueFormatter(FormatterPtr.ToSharedRef());
+	TSharedRef<FTableCellValueFormatter> Formatter = MakeShared<FDoubleValueFormatterAsTimeAuto>();
+	Column.SetValueFormatter(Formatter);
 
-	//TSharedPtr<ITableCellValueSorter> SorterPtr = MakeShareable(new FSorterByDoubleValue(ColumnPtr.ToSharedRef()));
-	TSharedPtr<ITableCellValueSorter> SorterPtr = MakeShareable(new FTimerNodeSortingByTotalExclusiveTime(ColumnPtr.ToSharedRef()));
-	Column.SetValueSorter(SorterPtr);
+	//TSharedRef<ITableCellValueSorter> Sorter = MakeShared<FSorterByDoubleValue>(ColumnRef);
+	TSharedRef<ITableCellValueSorter> Sorter = MakeShared<FTimerNodeSortingByTotalExclusiveTime>(ColumnRef);
+	Column.SetValueSorter(Sorter);
 
-	return ColumnPtr.ToSharedRef();
+	return ColumnRef;
 }
 
 ////////////////////////////////////////////////////////////////////////////////////////////////////
@@ -531,8 +531,8 @@ TSharedRef<Insights::FTableColumn> FTimersViewColumnFactory::CreateMaxExclusiveT
 {
 	using namespace Insights;
 
-	TSharedPtr<FTableColumn> ColumnPtr = MakeShareable(new FTableColumn(FTimersViewColumns::MaxExclusiveTimeColumnID));
-	FTableColumn& Column = *ColumnPtr;
+	TSharedRef<FTableColumn> ColumnRef = MakeShared<FTableColumn>(FTimersViewColumns::MaxExclusiveTimeColumnID);
+	FTableColumn& Column = *ColumnRef;
 
 	Column.SetShortName(LOCTEXT("MaxExclusiveTimeName", "E.Max"));
 	Column.SetTitleName(LOCTEXT("MaxExclusiveTimeTitle", "Max Exclusive Time (ms)"));
@@ -558,16 +558,16 @@ TSharedRef<Insights::FTableColumn> FTimersViewColumnFactory::CreateMaxExclusiveT
 		}
 	};
 
-	TSharedPtr<FTableCellValueGetter> GetterPtr = MakeShareable(new FMaxExclusiveTimeValueGetter());
-	Column.SetValueGetter(GetterPtr.ToSharedRef());
+	TSharedRef<FTableCellValueGetter> Getter = MakeShared<FMaxExclusiveTimeValueGetter>();
+	Column.SetValueGetter(Getter);
 
-	TSharedPtr<FTableCellValueFormatter> FormatterPtr = MakeShareable(new FDoubleValueFormatterAsTimeMs());
-	Column.SetValueFormatter(FormatterPtr.ToSharedRef());
+	TSharedRef<FTableCellValueFormatter> Formatter = MakeShared<FDoubleValueFormatterAsTimeMs>();
+	Column.SetValueFormatter(Formatter);
 
-	TSharedPtr<ITableCellValueSorter> SorterPtr = MakeShareable(new FSorterByDoubleValue(ColumnPtr.ToSharedRef()));
-	Column.SetValueSorter(SorterPtr);
+	TSharedRef<ITableCellValueSorter> Sorter = MakeShared<FSorterByDoubleValue>(ColumnRef);
+	Column.SetValueSorter(Sorter);
 
-	return ColumnPtr.ToSharedRef();
+	return ColumnRef;
 }
 
 ////////////////////////////////////////////////////////////////////////////////////////////////////
@@ -576,8 +576,8 @@ TSharedRef<Insights::FTableColumn> FTimersViewColumnFactory::CreateAverageExclus
 {
 	using namespace Insights;
 
-	TSharedPtr<FTableColumn> ColumnPtr = MakeShareable(new FTableColumn(FTimersViewColumns::AverageExclusiveTimeColumnID));
-	FTableColumn& Column = *ColumnPtr;
+	TSharedRef<FTableColumn> ColumnRef = MakeShared<FTableColumn>(FTimersViewColumns::AverageExclusiveTimeColumnID);
+	FTableColumn& Column = *ColumnRef;
 
 	Column.SetShortName(LOCTEXT("AvgExclusiveTimeName", "E.Avg"));
 	Column.SetTitleName(LOCTEXT("AvgExclusiveTimeTitle", "Average Exclusive Time (ms)"));
@@ -603,16 +603,16 @@ TSharedRef<Insights::FTableColumn> FTimersViewColumnFactory::CreateAverageExclus
 		}
 	};
 
-	TSharedPtr<FTableCellValueGetter> GetterPtr = MakeShareable(new FAverageExclusiveTimeValueGetter());
-	Column.SetValueGetter(GetterPtr.ToSharedRef());
+	TSharedRef<FTableCellValueGetter> Getter = MakeShared<FAverageExclusiveTimeValueGetter>();
+	Column.SetValueGetter(Getter);
 
-	TSharedPtr<FTableCellValueFormatter> FormatterPtr = MakeShareable(new FDoubleValueFormatterAsTimeMs());
-	Column.SetValueFormatter(FormatterPtr.ToSharedRef());
+	TSharedRef<FTableCellValueFormatter> Formatter = MakeShared<FDoubleValueFormatterAsTimeMs>();
+	Column.SetValueFormatter(Formatter);
 
-	TSharedPtr<ITableCellValueSorter> SorterPtr = MakeShareable(new FSorterByDoubleValue(ColumnPtr.ToSharedRef()));
-	Column.SetValueSorter(SorterPtr);
+	TSharedRef<ITableCellValueSorter> Sorter = MakeShared<FSorterByDoubleValue>(ColumnRef);
+	Column.SetValueSorter(Sorter);
 
-	return ColumnPtr.ToSharedRef();
+	return ColumnRef;
 }
 
 ////////////////////////////////////////////////////////////////////////////////////////////////////
@@ -621,8 +621,8 @@ TSharedRef<Insights::FTableColumn> FTimersViewColumnFactory::CreateMedianExclusi
 {
 	using namespace Insights;
 
-	TSharedPtr<FTableColumn> ColumnPtr = MakeShareable(new FTableColumn(FTimersViewColumns::MedianExclusiveTimeColumnID));
-	FTableColumn& Column = *ColumnPtr;
+	TSharedRef<FTableColumn> ColumnRef = MakeShared<FTableColumn>(FTimersViewColumns::MedianExclusiveTimeColumnID);
+	FTableColumn& Column = *ColumnRef;
 
 	Column.SetShortName(LOCTEXT("MedExclusiveTimeName", "E.Med"));
 	Column.SetTitleName(LOCTEXT("MedExclusiveTimeTitle", "Median Exclusive Time (ms)"));
@@ -648,16 +648,16 @@ TSharedRef<Insights::FTableColumn> FTimersViewColumnFactory::CreateMedianExclusi
 		}
 	};
 
-	TSharedPtr<FTableCellValueGetter> GetterPtr = MakeShareable(new FMedianExclusiveTimeValueGetter());
-	Column.SetValueGetter(GetterPtr.ToSharedRef());
+	TSharedRef<FTableCellValueGetter> Getter = MakeShared<FMedianExclusiveTimeValueGetter>();
+	Column.SetValueGetter(Getter);
 
-	TSharedPtr<FTableCellValueFormatter> FormatterPtr = MakeShareable(new FDoubleValueFormatterAsTimeMs());
-	Column.SetValueFormatter(FormatterPtr.ToSharedRef());
+	TSharedRef<FTableCellValueFormatter> Formatter = MakeShared<FDoubleValueFormatterAsTimeMs>();
+	Column.SetValueFormatter(Formatter);
 
-	TSharedPtr<ITableCellValueSorter> SorterPtr = MakeShareable(new FSorterByDoubleValue(ColumnPtr.ToSharedRef()));
-	Column.SetValueSorter(SorterPtr);
+	TSharedRef<ITableCellValueSorter> Sorter = MakeShared<FSorterByDoubleValue>(ColumnRef);
+	Column.SetValueSorter(Sorter);
 
-	return ColumnPtr.ToSharedRef();
+	return ColumnRef;
 }
 
 ////////////////////////////////////////////////////////////////////////////////////////////////////
@@ -666,8 +666,8 @@ TSharedRef<Insights::FTableColumn> FTimersViewColumnFactory::CreateMinExclusiveT
 {
 	using namespace Insights;
 
-	TSharedPtr<FTableColumn> ColumnPtr = MakeShareable(new FTableColumn(FTimersViewColumns::MinExclusiveTimeColumnID));
-	FTableColumn& Column = *ColumnPtr;
+	TSharedRef<FTableColumn> ColumnRef = MakeShared<FTableColumn>(FTimersViewColumns::MinExclusiveTimeColumnID);
+	FTableColumn& Column = *ColumnRef;
 
 	Column.SetShortName(LOCTEXT("MinExclusiveTimeName", "E.Min"));
 	Column.SetTitleName(LOCTEXT("MinExclusiveTimeTitle", "Min Exclusive Time (ms)"));
@@ -693,16 +693,16 @@ TSharedRef<Insights::FTableColumn> FTimersViewColumnFactory::CreateMinExclusiveT
 		}
 	};
 
-	TSharedPtr<FTableCellValueGetter> GetterPtr = MakeShareable(new FMinExclusiveTimeValueGetter());
-	Column.SetValueGetter(GetterPtr.ToSharedRef());
+	TSharedRef<FTableCellValueGetter> Getter = MakeShared<FMinExclusiveTimeValueGetter>();
+	Column.SetValueGetter(Getter);
 
-	TSharedPtr<FTableCellValueFormatter> FormatterPtr = MakeShareable(new FDoubleValueFormatterAsTimeMs());
-	Column.SetValueFormatter(FormatterPtr.ToSharedRef());
+	TSharedRef<FTableCellValueFormatter> Formatter = MakeShared<FDoubleValueFormatterAsTimeMs>();
+	Column.SetValueFormatter(Formatter);
 
-	TSharedPtr<ITableCellValueSorter> SorterPtr = MakeShareable(new FSorterByDoubleValue(ColumnPtr.ToSharedRef()));
-	Column.SetValueSorter(SorterPtr);
+	TSharedRef<ITableCellValueSorter> Sorter = MakeShared<FSorterByDoubleValue>(ColumnRef);
+	Column.SetValueSorter(Sorter);
 
-	return ColumnPtr.ToSharedRef();
+	return ColumnRef;
 }
 
 ////////////////////////////////////////////////////////////////////////////////////////////////////
diff --git a/Engine/Source/Developer/TraceInsights/Public/Insights/Version.h b/Engine/Source/Developer/TraceInsights/Public/Insights/Version.h
--- a/Engine/Source/Developer/TraceInsights/Public/Insights/Version.h
+++ b/Engine/Source/Developer/TraceInsights/Public/Insights/Version.h
@@ -5,10 +5,10 @@
 // See https://semver.org/
 
 #define UNREAL_INSIGHTS_VERSION_MAJOR 0
-#define UNREAL_INSIGHTS_VERSION_MINOR 10
+#define UNREAL_INSIGHTS_VERSION_MINOR 11
 #define UNREAL_INSIGHTS_VERSION_PATCH 0
 
-#define UNREAL_INSIGHTS_VERSION_STRING "0.10"
+#define UNREAL_INSIGHTS_VERSION_STRING "0.11"
 
 /** Extra identifier string (like: "alpha", "beta.1", "beta.2", "test", "featureX", etc.). */
 //#define UNREAL_INSIGHTS_VERSION_ID_STRING "beta"
