[Insights]
  - Changed almost all usages of MakeShareable(new FMyClass()) to MakeShared<FMyClass>().
  - Changed several TSharedPtr to TSharedRef.
  - Asset Loading: Fixed the invalidation of tracks when pressing 1/2/3/4 key shortcut to change coloring schema.

diff --git a/Engine/Source/Developer/TraceInsights/Private/Insights/NetworkingProfiler/ViewModels/NetStatsViewColumnFactory.cpp b/Engine/Source/Developer/TraceInsights/Private/Insights/NetworkingProfiler/ViewModels/NetStatsViewColumnFactory.cpp
--- a/Engine/Source/Developer/TraceInsights/Private/Insights/NetworkingProfiler/ViewModels/NetStatsViewColumnFactory.cpp
+++ b/Engine/Source/Developer/TraceInsights/Private/Insights/NetworkingProfiler/ViewModels/NetStatsViewColumnFactory.cpp
@@ -31,7 +31,7 @@ const FName FNetStatsViewColumns::AverageExclusiveSizeColumnID(TEXT("AverageExcl
 
 ////////////////////////////////////////////////////////////////////////////////////////////////////
 
-void FNetStatsViewColumnFactory::CreateNetStatsViewColumns(TArray<TSharedPtr<Insights::FTableColumn>>& Columns)
+void FNetStatsViewColumnFactory::CreateNetStatsViewColumns(TArray<TSharedRef<Insights::FTableColumn>>& Columns)
 {
 	Columns.Reset();
 
diff --git a/Engine/Source/Developer/TraceInsights/Private/Insights/NetworkingProfiler/ViewModels/NetStatsViewColumnFactory.h b/Engine/Source/Developer/TraceInsights/Private/Insights/NetworkingProfiler/ViewModels/NetStatsViewColumnFactory.h
--- a/Engine/Source/Developer/TraceInsights/Private/Insights/NetworkingProfiler/ViewModels/NetStatsViewColumnFactory.h
+++ b/Engine/Source/Developer/TraceInsights/Private/Insights/NetworkingProfiler/ViewModels/NetStatsViewColumnFactory.h
@@ -39,7 +39,7 @@ namespace Insights
 struct FNetStatsViewColumnFactory
 {
 public:
-	static void CreateNetStatsViewColumns(TArray<TSharedPtr<Insights::FTableColumn>>& Columns);
+	static void CreateNetStatsViewColumns(TArray<TSharedRef<Insights::FTableColumn>>& Columns);
 
 	static TSharedRef<Insights::FTableColumn> CreateNameColumn();
 	static TSharedRef<Insights::FTableColumn> CreateTypeColumn();
diff --git a/Engine/Source/Developer/TraceInsights/Private/Insights/NetworkingProfiler/Widgets/SNetStatsTableRow.cpp b/Engine/Source/Developer/TraceInsights/Private/Insights/NetworkingProfiler/Widgets/SNetStatsTableRow.cpp
--- a/Engine/Source/Developer/TraceInsights/Private/Insights/NetworkingProfiler/Widgets/SNetStatsTableRow.cpp
+++ b/Engine/Source/Developer/TraceInsights/Private/Insights/NetworkingProfiler/Widgets/SNetStatsTableRow.cpp
@@ -33,7 +33,7 @@ void SNetStatsTableRow::Construct(const FArguments& InArgs, const TSharedRef<STa
 	TablePtr = InArgs._TablePtr;
 	NetEventNodePtr = InArgs._NetEventNodePtr;
 
-	RowToolTip = MakeShareable(new SNetEventTableRowToolTip(NetEventNodePtr));
+	RowToolTip = MakeShared<SNetEventTableRowToolTip>(NetEventNodePtr);
 
 	SetEnabled(TAttribute<bool>(this, &SNetStatsTableRow::HandleShouldBeEnabled));
 
diff --git a/Engine/Source/Developer/TraceInsights/Private/Insights/NetworkingProfiler/Widgets/SNetStatsView.cpp b/Engine/Source/Developer/TraceInsights/Private/Insights/NetworkingProfiler/Widgets/SNetStatsView.cpp
--- a/Engine/Source/Developer/TraceInsights/Private/Insights/NetworkingProfiler/Widgets/SNetStatsView.cpp
+++ b/Engine/Source/Developer/TraceInsights/Private/Insights/NetworkingProfiler/Widgets/SNetStatsView.cpp
@@ -30,7 +30,7 @@
 ////////////////////////////////////////////////////////////////////////////////////////////////////
 
 SNetStatsView::SNetStatsView()
-	: Table(MakeShareable(new Insights::FTable()))
+	: Table(MakeShared<Insights::FTable>())
 	, bExpansionSaved(false)
 	, bFilterOutZeroCountEvents(false)
 	, GroupingMode(ENetEventGroupingMode::Flat)
@@ -248,8 +248,8 @@ void SNetStatsView::Construct(const FArguments& InArgs)
 	//BindCommands();
 
 	// Create the search filters: text based, type based etc.
-	TextFilter = MakeShareable(new FNetEventNodeTextFilter(FNetEventNodeTextFilter::FItemToStringArray::CreateSP(this, &SNetStatsView::HandleItemToStringArray)));
-	Filters = MakeShareable(new FNetEventNodeFilterCollection());
+	TextFilter = MakeShared<FNetEventNodeTextFilter>(FNetEventNodeTextFilter::FItemToStringArray::CreateSP(this, &SNetStatsView::HandleItemToStringArray));
+	Filters = MakeShared<FNetEventNodeFilterCollection>();
 	Filters->Add(TextFilter);
 
 	CreateGroupByOptionsSources();
@@ -424,9 +424,9 @@ void SNetStatsView::TreeView_BuildSortByMenu(FMenuBuilder& MenuBuilder)
 
 	MenuBuilder.BeginSection("ColumnName", LOCTEXT("ContextMenu_Header_Misc_ColumnName", "Column Name"));
 
-	for (const TSharedPtr<Insights::FTableColumn>& ColumnPtr : Table->GetColumns())
+	for (const TSharedRef<Insights::FTableColumn>& ColumnRef : Table->GetColumns())
 	{
-		const Insights::FTableColumn& Column = *ColumnPtr;
+		const Insights::FTableColumn& Column = *ColumnRef;
 
 		if (Column.IsVisible() && Column.CanBeSorted())
 		{
@@ -486,9 +486,9 @@ void SNetStatsView::TreeView_BuildViewColumnMenu(FMenuBuilder& MenuBuilder)
 {
 	MenuBuilder.BeginSection("ViewColumn", LOCTEXT("ContextMenu_Header_Columns_View", "View Column"));
 
-	for (const TSharedPtr<Insights::FTableColumn>& ColumnPtr : Table->GetColumns())
+	for (const TSharedRef<Insights::FTableColumn>& ColumnRef : Table->GetColumns())
 	{
-		const Insights::FTableColumn& Column = *ColumnPtr;
+		const Insights::FTableColumn& Column = *ColumnRef;
 
 		FUIAction Action_ToggleColumn
 		(
@@ -512,16 +512,16 @@ void SNetStatsView::TreeView_BuildViewColumnMenu(FMenuBuilder& MenuBuilder)
 void SNetStatsView::InitializeAndShowHeaderColumns()
 {
 	// Create columns.
-	TArray<TSharedPtr<Insights::FTableColumn>> Columns;
+	TArray<TSharedRef<Insights::FTableColumn>> Columns;
 	FNetStatsViewColumnFactory::CreateNetStatsViewColumns(Columns);
 	Table->SetColumns(Columns);
 
 	// Show columns.
-	for (const TSharedPtr<Insights::FTableColumn>& ColumnPtr : Table->GetColumns())
+	for (const TSharedRef<Insights::FTableColumn>& ColumnRef : Table->GetColumns())
 	{
-		if (ColumnPtr->ShouldBeVisible())
+		if (ColumnRef->ShouldBeVisible())
 		{
-			ShowColumn(ColumnPtr->GetId());
+			ShowColumn(ColumnRef->GetId());
 		}
 	}
 }
@@ -950,7 +950,7 @@ void SNetStatsView::CreateGroups()
 		FNetEventNodePtr* GroupPtr = GroupNodeSet.Find(GroupName);
 		if (!GroupPtr)
 		{
-			GroupPtr = &GroupNodeSet.Add(GroupName, MakeShareable(new FNetEventNode(GroupName)));
+			GroupPtr = &GroupNodeSet.Add(GroupName, MakeShared<FNetEventNode>(GroupName));
 		}
 
 		for (const FNetEventNodePtr& NetEventNodePtr : NetEventNodes)
@@ -970,7 +970,7 @@ void SNetStatsView::CreateGroups()
 			FNetEventNodePtr* GroupPtr = GroupNodeSet.Find(GroupName);
 			if (!GroupPtr)
 			{
-				GroupPtr = &GroupNodeSet.Add(GroupName, MakeShareable(new FNetEventNode(GroupName)));
+				GroupPtr = &GroupNodeSet.Add(GroupName, MakeShared<FNetEventNode>(GroupName));
 			}
 
 			(*GroupPtr)->AddChildAndSetGroupPtr(NetEventNodePtr);
@@ -986,7 +986,7 @@ void SNetStatsView::CreateGroups()
 			FNetEventNodePtr* GroupPtr = GroupNodeSet.Find(GroupName);
 			if (!GroupPtr)
 			{
-				GroupPtr = &GroupNodeSet.Add(GroupName, MakeShareable(new FNetEventNode(GroupName)));
+				GroupPtr = &GroupNodeSet.Add(GroupName, MakeShared<FNetEventNode>(GroupName));
 			}
 
 			(*GroupPtr)->AddChildAndSetGroupPtr(NetEventNodePtr);
@@ -1003,7 +1003,7 @@ void SNetStatsView::CreateGroups()
 			FNetEventNodePtr* GroupPtr = GroupNodeSet.Find(GroupName);
 			if (!GroupPtr)
 			{
-				GroupPtr = &GroupNodeSet.Add(GroupName, MakeShareable(new FNetEventNode(GroupName)));
+				GroupPtr = &GroupNodeSet.Add(GroupName, MakeShared<FNetEventNode>(GroupName));
 			}
 
 			(*GroupPtr)->AddChildAndSetGroupPtr(NetEventNodePtr);
@@ -1021,10 +1021,10 @@ void SNetStatsView::CreateGroupByOptionsSources()
 	GroupByOptionsSource.Reset(3);
 
 	// Must be added in order of elements in the ENetEventGroupingMode.
-	GroupByOptionsSource.Add(MakeShareable(new ENetEventGroupingMode(ENetEventGroupingMode::Flat)));
-	GroupByOptionsSource.Add(MakeShareable(new ENetEventGroupingMode(ENetEventGroupingMode::ByName)));
-	//GroupByOptionsSource.Add(MakeShareable(new ENetEventGroupingMode(ENetEventGroupingMode::ByType)));
-	GroupByOptionsSource.Add(MakeShareable(new ENetEventGroupingMode(ENetEventGroupingMode::ByLevel)));
+	GroupByOptionsSource.Add(MakeShared<ENetEventGroupingMode>(ENetEventGroupingMode::Flat));
+	GroupByOptionsSource.Add(MakeShared<ENetEventGroupingMode>(ENetEventGroupingMode::ByName));
+	//GroupByOptionsSource.Add(MakeShared<ENetEventGroupingMode>(ENetEventGroupingMode::ByType));
+	GroupByOptionsSource.Add(MakeShared<ENetEventGroupingMode>(ENetEventGroupingMode::ByLevel));
 
 	ENetEventGroupingModePtr* GroupingModePtrPtr = GroupByOptionsSource.FindByPredicate([&](const ENetEventGroupingModePtr InGroupingModePtr) { return *InGroupingModePtr == GroupingMode; });
 	if (GroupingModePtrPtr != nullptr)
@@ -1095,11 +1095,11 @@ void SNetStatsView::CreateSortings()
 	AvailableSorters.Reset();
 	CurrentSorter = nullptr;
 
-	for (const TSharedPtr<Insights::FTableColumn> ColumnPtr : Table->GetColumns())
+	for (const TSharedRef<Insights::FTableColumn> ColumnRef : Table->GetColumns())
 	{
-		if (ColumnPtr->CanBeSorted())
+		if (ColumnRef->CanBeSorted())
 		{
-			TSharedPtr<Insights::ITableCellValueSorter> SorterPtr = ColumnPtr->GetValueSorter();
+			TSharedPtr<Insights::ITableCellValueSorter> SorterPtr = ColumnRef->GetValueSorter();
 			if (ensure(SorterPtr.IsValid()))
 			{
 				AvailableSorters.Add(SorterPtr);
@@ -1389,9 +1389,9 @@ void SNetStatsView::ContextMenu_ShowAllColumns_Execute()
 	ColumnSortMode = GetDefaultColumnSortMode();
 	ColumnBeingSorted = GetDefaultColumnBeingSorted();
 
-	for (const TSharedPtr<Insights::FTableColumn>& ColumnPtr : Table->GetColumns())
+	for (const TSharedRef<Insights::FTableColumn>& ColumnRef : Table->GetColumns())
 	{
-		const Insights::FTableColumn& Column = *ColumnPtr;
+		const Insights::FTableColumn& Column = *ColumnRef;
 
 		if (!Column.IsVisible())
 		{
@@ -1433,9 +1433,9 @@ void SNetStatsView::ContextMenu_ShowMinMaxMedColumns_Execute()
 	ColumnSortMode = EColumnSortMode::Descending;
 	UpdateCurrentSortingByColumn();
 
-	for (const TSharedPtr<Insights::FTableColumn>& ColumnPtr : Table->GetColumns())
+	for (const TSharedRef<Insights::FTableColumn>& ColumnRef : Table->GetColumns())
 	{
-		const Insights::FTableColumn& Column = *ColumnPtr;
+		const Insights::FTableColumn& Column = *ColumnRef;
 
 		const bool bShouldBeVisible = Preset.Contains(Column.GetId());
 
@@ -1467,9 +1467,9 @@ void SNetStatsView::ContextMenu_ResetColumns_Execute()
 	ColumnSortMode = GetDefaultColumnSortMode();
 	UpdateCurrentSortingByColumn();
 
-	for (const TSharedPtr<Insights::FTableColumn>& ColumnPtr : Table->GetColumns())
+	for (const TSharedRef<Insights::FTableColumn>& ColumnRef : Table->GetColumns())
 	{
-		const Insights::FTableColumn& Column = *ColumnPtr;
+		const Insights::FTableColumn& Column = *ColumnRef;
 
 		if (Column.ShouldBeVisible() && !Column.IsVisible())
 		{
@@ -1549,7 +1549,7 @@ void SNetStatsView::RebuildTree(bool bResync)
 					});
 					const FName Name(NamePtr);
 					const ENetEventNodeType Type = ENetEventNodeType::NetEvent;
-					FNetEventNodePtr NetEventNodePtr = MakeShareable(new FNetEventNode(Id, Name, Type, NetEvent.Level));
+					FNetEventNodePtr NetEventNodePtr = MakeShared<FNetEventNode>(Id, Name, Type, NetEvent.Level);
 					NetEventNodes.Add(NetEventNodePtr);
 					NetEventNodesIdMap.Add(Id, NetEventNodePtr);
 				}
diff --git a/Engine/Source/Developer/TraceInsights/Private/Insights/NetworkingProfiler/Widgets/SNetworkingProfilerWindow.cpp b/Engine/Source/Developer/TraceInsights/Private/Insights/NetworkingProfiler/Widgets/SNetworkingProfilerWindow.cpp
--- a/Engine/Source/Developer/TraceInsights/Private/Insights/NetworkingProfiler/Widgets/SNetworkingProfilerWindow.cpp
+++ b/Engine/Source/Developer/TraceInsights/Private/Insights/NetworkingProfiler/Widgets/SNetworkingProfilerWindow.cpp
@@ -510,7 +510,7 @@ void SNetworkingProfilerWindow::UpdateAvailableGameInstances()
 		const Trace::INetProfilerProvider& NetProfilerProvider = Trace::ReadNetProfilerProvider(*Session.Get());
 		NetProfilerProvider.ReadGameInstances([this](const Trace::FNetProfilerGameInstance& GameInstance)
 		{
-			AvailableGameInstances.Add(MakeShareable(new FGameInstanceItem(GameInstance)));
+			AvailableGameInstances.Add(MakeShared<FGameInstanceItem>(GameInstance));
 		});
 	}
 
@@ -535,7 +535,7 @@ void SNetworkingProfilerWindow::UpdateAvailableConnections()
 		const Trace::INetProfilerProvider& NetProfilerProvider = Trace::ReadNetProfilerProvider(*Session.Get());
 		NetProfilerProvider.ReadConnections(SelectedGameInstance->GetIndex(), [this](const Trace::FNetProfilerConnection& Connection)
 		{
-			AvailableConnections.Add(MakeShareable(new FConnectionItem(Connection)));
+			AvailableConnections.Add(MakeShared<FConnectionItem>(Connection));
 		});
 	}
 
@@ -557,11 +557,11 @@ void SNetworkingProfilerWindow::UpdateAvailableConnectionModes()
 	{
 		if (SelectedConnection->Connection.bHasIncomingData)
 		{
-			AvailableConnectionModes.Add(MakeShareable(new FConnectionModeItem(Trace::ENetProfilerConnectionMode::Incoming)));
+			AvailableConnectionModes.Add(MakeShared<FConnectionModeItem>(Trace::ENetProfilerConnectionMode::Incoming));
 		}
 		if (SelectedConnection->Connection.bHasOutgoingData)
 		{
-			AvailableConnectionModes.Add(MakeShareable(new FConnectionModeItem(Trace::ENetProfilerConnectionMode::Outgoing)));
+			AvailableConnectionModes.Add(MakeShared<FConnectionModeItem>(Trace::ENetProfilerConnectionMode::Outgoing));
 		}
 	}
 
diff --git a/Engine/Source/Developer/TraceInsights/Private/Insights/NetworkingProfiler/Widgets/SPacketView.cpp b/Engine/Source/Developer/TraceInsights/Private/Insights/NetworkingProfiler/Widgets/SPacketView.cpp
--- a/Engine/Source/Developer/TraceInsights/Private/Insights/NetworkingProfiler/Widgets/SPacketView.cpp
+++ b/Engine/Source/Developer/TraceInsights/Private/Insights/NetworkingProfiler/Widgets/SPacketView.cpp
@@ -32,7 +32,7 @@
 
 SPacketView::SPacketView()
 	: ProfilerWindow()
-	, PacketSeries(MakeShareable(new FNetworkPacketSeries()))
+	, PacketSeries(MakeShared<FNetworkPacketSeries>())
 {
 	Reset();
 }
@@ -354,7 +354,7 @@ FNetworkPacketSampleRef SPacketView::GetSampleAtMousePosition(float X, float Y)
 
 					if (Y >= TopY && Y < BottomY)
 					{
-						return FNetworkPacketSampleRef(PacketSeries, MakeShareable(new FNetworkPacketAggregatedSample(Sample)));
+						return FNetworkPacketSampleRef(PacketSeries, MakeShared<FNetworkPacketAggregatedSample>(Sample));
 					}
 				}
 			}
diff --git a/Engine/Source/Developer/TraceInsights/Private/Insights/Table/ViewModels/Table.cpp b/Engine/Source/Developer/TraceInsights/Private/Insights/Table/ViewModels/Table.cpp
--- a/Engine/Source/Developer/TraceInsights/Private/Insights/Table/ViewModels/Table.cpp
+++ b/Engine/Source/Developer/TraceInsights/Private/Insights/Table/ViewModels/Table.cpp
@@ -107,7 +107,7 @@ void FTable::Reset()
 
 int32 FTable::GetColumnPositionIndex(const FName& ColumnId) const
 {
-	return Columns.IndexOfByPredicate([&ColumnId](const TSharedPtr<FTableColumn>& ColumnPtr) -> bool { return ColumnPtr->GetId() == ColumnId; });
+	return Columns.IndexOfByPredicate([&ColumnId](const TSharedRef<FTableColumn>& ColumnRef) -> bool { return ColumnRef->GetId() == ColumnId; });
 }
 
 ////////////////////////////////////////////////////////////////////////////////////////////////////
@@ -162,23 +162,23 @@ void FTable::UpdateSourceTable(TSharedPtr<Trace::IUntypedTable> InSourceTable)
 
 ////////////////////////////////////////////////////////////////////////////////////////////////////
 
-void FTable::SetColumns(const TArray<TSharedPtr<Insights::FTableColumn>>& InColumns)
+void FTable::SetColumns(const TArray<TSharedRef<Insights::FTableColumn>>& InColumns)
 {
 	Columns.Reset(InColumns.Num());
 	ColumnIdToPtrMapping.Reset();
-	for (TSharedPtr<Insights::FTableColumn> ColumnPtr : InColumns)
+	for (TSharedRef<Insights::FTableColumn> ColumnRef : InColumns)
 	{
-		AddColumn(ColumnPtr);
+		AddColumn(ColumnRef);
 	}
 }
 
 ////////////////////////////////////////////////////////////////////////////////////////////////////
 
-void FTable::AddColumn(TSharedPtr<FTableColumn> ColumnPtr)
+void FTable::AddColumn(TSharedRef<FTableColumn> ColumnRef)
 {
-	ColumnPtr->SetParentTable(SharedThis(this));
-	Columns.Add(ColumnPtr);
-	ColumnIdToPtrMapping.Add(ColumnPtr->GetId(), ColumnPtr);
+	ColumnRef->SetParentTable(SharedThis(this));
+	Columns.Add(ColumnRef);
+	ColumnIdToPtrMapping.Add(ColumnRef->GetId(), ColumnRef);
 }
 
 ////////////////////////////////////////////////////////////////////////////////////////////////////
@@ -187,9 +187,9 @@ void FTable::CreateHierarchyColumn(int32 ColumnIndex, const TCHAR* ColumnName)
 {
 	const FName HierarchyColumnId(TEXT("_Hierarchy"));
 
-	TSharedPtr<FTableColumn> ColumnPtr = MakeShareable(new FTableColumn(HierarchyColumnId));
-	FTableColumn& Column = *ColumnPtr;
-	
+	TSharedRef<FTableColumn> ColumnRef = MakeShared<FTableColumn>(HierarchyColumnId);
+	FTableColumn& Column = *ColumnRef;
+
 	Column.SetIndex(ColumnIndex);
 
 	const FString ColumnNameStr = ColumnName ? FString::Printf(TEXT("Hierarchy (%s)"), ColumnName) : TEXT("Hierarchy");
@@ -206,16 +206,16 @@ void FTable::CreateHierarchyColumn(int32 ColumnIndex, const TCHAR* ColumnName)
 
 	Column.SetDataType(ETableCellDataType::CString);
 
-	TSharedPtr<ITableCellValueGetter> GetterPtr = MakeShareable(new FDisplayNameValueGetter());
-	Column.SetValueGetter(GetterPtr.ToSharedRef());
+	TSharedRef<ITableCellValueGetter> Getter = MakeShared<FDisplayNameValueGetter>();
+	Column.SetValueGetter(Getter);
 
-	TSharedPtr<ITableCellValueFormatter> FormatterPtr = MakeShareable(new FTextValueFormatter());
-	Column.SetValueFormatter(FormatterPtr.ToSharedRef());
+	TSharedRef<ITableCellValueFormatter> Formatter = MakeShared<FTextValueFormatter>();
+	Column.SetValueFormatter(Formatter);
 
-	TSharedPtr<ITableCellValueSorter> SorterPtr = MakeShareable(new FSorterByName(ColumnPtr.ToSharedRef()));
-	Column.SetValueSorter(SorterPtr);
+	TSharedRef<ITableCellValueSorter> Sorter = MakeShared<FSorterByName>(ColumnRef);
+	Column.SetValueSorter(Sorter);
 
-	AddColumn(ColumnPtr);
+	AddColumn(ColumnRef);
 }
 
 ////////////////////////////////////////////////////////////////////////////////////////////////////
@@ -255,8 +255,8 @@ void FTable::CreateColumnsFromTableLayout()
 		Trace::ETableColumnType ColumnType = TableLayout.GetColumnType(ColumnIndex);
 		const TCHAR* ColumnName = TableLayout.GetColumnName(ColumnIndex);
 
-		TSharedPtr<FTableColumn> ColumnPtr = MakeShareable(new FTableColumn(FName(ColumnName)));
-		FTableColumn& Column = *ColumnPtr;
+		TSharedRef<FTableColumn> ColumnRef = MakeShared<FTableColumn>(FName(ColumnName));
+		FTableColumn& Column = *ColumnRef;
 
 		const FString ColumnNameStr(ColumnName);
 		const FText ColumnNameText = FText::FromString(ColumnNameStr);
@@ -283,8 +283,8 @@ void FTable::CreateColumnsFromTableLayout()
 			InitialColumnWidth = 40.0f;
 			//TODO: if (Hint == AsOnOff)
 			//else // if (Hint == AsTrueFalse)
-			FormatterPtr = MakeShareable(new FBoolValueFormatterAsTrueFalse());
-			SorterPtr = MakeShareable(new FSorterByBoolValue(ColumnPtr.ToSharedRef()));
+			FormatterPtr = MakeShared<FBoolValueFormatterAsTrueFalse>();
+			SorterPtr = MakeShared<FSorterByBoolValue>(ColumnRef);
 			break;
 
 		case Trace::TableColumnType_Int:
@@ -294,11 +294,11 @@ void FTable::CreateColumnsFromTableLayout()
 			InitialColumnWidth = 60.0f;
 			//TODO: if (Hint == AsMemory)
 			//{
-			//	FormatterPtr = MakeShareable(new FInt64ValueFormatterAsMemory());
+			//	FormatterPtr = MakeShared<FInt64ValueFormatterAsMemory>();
 			//}
 			//else // AsNumber
-			FormatterPtr = MakeShareable(new FInt64ValueFormatterAsNumber());
-			SorterPtr = MakeShareable(new FSorterByInt64Value(ColumnPtr.ToSharedRef()));
+			FormatterPtr = MakeShared<FInt64ValueFormatterAsNumber>();
+			SorterPtr = MakeShared<FSorterByInt64Value>(ColumnRef);
 			break;
 
 		case Trace::TableColumnType_Float:
@@ -308,8 +308,8 @@ void FTable::CreateColumnsFromTableLayout()
 			InitialColumnWidth = 60.0f;
 			//TODO: if (Hint == AsTimeMs)
 			//else // if (Hint == AsTimeAuto)
-			FormatterPtr = MakeShareable(new FFloatValueFormatterAsTimeAuto());
-			SorterPtr = MakeShareable(new FSorterByFloatValue(ColumnPtr.ToSharedRef()));
+			FormatterPtr = MakeShared<FFloatValueFormatterAsTimeAuto>();
+			SorterPtr = MakeShared<FSorterByFloatValue>(ColumnRef);
 			break;
 
 		case Trace::TableColumnType_Double:
@@ -319,16 +319,16 @@ void FTable::CreateColumnsFromTableLayout()
 			InitialColumnWidth = 80.0f;
 			//TODO: if (Hint == AsTimeMs)
 			//else // if (Hint == AsTimeAuto)
-			FormatterPtr = MakeShareable(new FDoubleValueFormatterAsTimeAuto());
-			SorterPtr = MakeShareable(new FSorterByDoubleValue(ColumnPtr.ToSharedRef()));
+			FormatterPtr = MakeShared<FDoubleValueFormatterAsTimeAuto>();
+			SorterPtr = MakeShared<FSorterByDoubleValue>(ColumnRef);
 			break;
 
 		case Trace::TableColumnType_CString:
 			Column.SetDataType(ETableCellDataType::CString);
 			HorizontalAlignment = HAlign_Left;
 			InitialColumnWidth = FMath::Max(120.0f, 6.0f * ColumnNameStr.Len());
-			FormatterPtr = MakeShareable(new FCStringValueFormatterAsText());
-			SorterPtr = MakeShareable(new FSorterByCStringValue(ColumnPtr.ToSharedRef()));
+			FormatterPtr = MakeShared<FCStringValueFormatterAsText>();
+			SorterPtr = MakeShared<FSorterByCStringValue>(ColumnRef);
 			break;
 		}
 
@@ -346,7 +346,7 @@ void FTable::CreateColumnsFromTableLayout()
 
 		Column.SetAggregation(Aggregation);
 
-		Column.SetValueGetter(MakeShareable(new FTableTreeNodeValueGetter(Column.GetDataType())));
+		Column.SetValueGetter(MakeShared<FTableTreeNodeValueGetter>(Column.GetDataType()));
 
 		if (FormatterPtr.IsValid())
 		{
@@ -355,7 +355,7 @@ void FTable::CreateColumnsFromTableLayout()
 
 		Column.SetValueSorter(SorterPtr);
 
-		AddColumn(ColumnPtr);
+		AddColumn(ColumnRef);
 	}
 }
 
diff --git a/Engine/Source/Developer/TraceInsights/Private/Insights/Table/ViewModels/Table.h b/Engine/Source/Developer/TraceInsights/Private/Insights/Table/ViewModels/Table.h
--- a/Engine/Source/Developer/TraceInsights/Private/Insights/Table/ViewModels/Table.h
+++ b/Engine/Source/Developer/TraceInsights/Private/Insights/Table/ViewModels/Table.h
@@ -56,18 +56,22 @@ public:
 
 	bool IsValid() const { return Columns.Num() > 0; }
 
-	const TArray<TSharedPtr<FTableColumn>>& GetColumns() const { return Columns; }
-	void SetColumns(const TArray<TSharedPtr<Insights::FTableColumn>>& InColumns);
+	const TArray<TSharedRef<FTableColumn>>& GetColumns() const { return Columns; }
+	void SetColumns(const TArray<TSharedRef<Insights::FTableColumn>>& InColumns);
 
-	TSharedPtr<FTableColumn> FindColumnChecked(const FName& ColumnId) const
+	TSharedRef<FTableColumn> FindColumnChecked(const FName& ColumnId) const
 	{
 		return ColumnIdToPtrMapping.FindChecked(ColumnId);
 	}
 
 	TSharedPtr<FTableColumn> FindColumn(const FName& ColumnId) const
 	{
-		const TSharedPtr<FTableColumn>* const ColumnPtrPtr = ColumnIdToPtrMapping.Find(ColumnId);
-		return (ColumnPtrPtr != nullptr) ? *ColumnPtrPtr : nullptr;
+		const TSharedRef<FTableColumn>* const ColumnRefPtr = ColumnIdToPtrMapping.Find(ColumnId);
+		if (ColumnRefPtr != nullptr)
+		{
+			return *ColumnRefPtr;
+		}
+		return nullptr;
 	}
 
 	int32 GetColumnPositionIndex(const FName& ColumnId) const;
@@ -80,7 +84,7 @@ public:
 	void UpdateSourceTable(TSharedPtr<Trace::IUntypedTable> InSourceTable);
 
 private:
-	void AddColumn(TSharedPtr<FTableColumn> ColumnPtr);
+	void AddColumn(TSharedRef<FTableColumn> Column);
 	void CreateHierarchyColumn(int32 ColumnIndex, const TCHAR* ColumnName);
 	void CreateColumnsFromTableLayout();
 
@@ -89,10 +93,10 @@ private:
 	FText Description;
 
 	/** All available columns. */
-	TArray<TSharedPtr<FTableColumn>> Columns;
+	TArray<TSharedRef<FTableColumn>> Columns;
 
-	/** Mapping between column Ids and FTableColumn shared pointers. */
-	TMap<FName, TSharedPtr<FTableColumn>> ColumnIdToPtrMapping;
+	/** Mapping between column Ids and FTableColumn shared refs. */
+	TMap<FName, TSharedRef<FTableColumn>> ColumnIdToPtrMapping;
 
 	TSharedPtr<Trace::IUntypedTable> SourceTable;
 	TSharedPtr<Trace::IUntypedTableReader> TableReader;
diff --git a/Engine/Source/Developer/TraceInsights/Private/Insights/Table/ViewModels/TableCellValue.h b/Engine/Source/Developer/TraceInsights/Private/Insights/Table/ViewModels/TableCellValue.h
--- a/Engine/Source/Developer/TraceInsights/Private/Insights/Table/ViewModels/TableCellValue.h
+++ b/Engine/Source/Developer/TraceInsights/Private/Insights/Table/ViewModels/TableCellValue.h
@@ -33,7 +33,7 @@ public:
 	FTableCellValue(float Value) : DataType(ETableCellDataType::Float), Float(Value) {}
 	FTableCellValue(double Value) : DataType(ETableCellDataType::Double), Double(Value) {}
 	FTableCellValue(const TCHAR* Value) : DataType(ETableCellDataType::CString), CString(Value) {}
-	FTableCellValue(const FText& Value) : DataType(ETableCellDataType::Text), TextPtr(MakeShareable(new FText(Value))) {}
+	FTableCellValue(const FText& Value) : DataType(ETableCellDataType::Text), TextPtr(MakeShared<FText>(Value)) {}
 
 	ETableCellDataType DataType;
 
diff --git a/Engine/Source/Developer/TraceInsights/Private/Insights/Table/ViewModels/TableColumn.cpp b/Engine/Source/Developer/TraceInsights/Private/Insights/Table/ViewModels/TableColumn.cpp
--- a/Engine/Source/Developer/TraceInsights/Private/Insights/Table/ViewModels/TableColumn.cpp
+++ b/Engine/Source/Developer/TraceInsights/Private/Insights/Table/ViewModels/TableColumn.cpp
@@ -15,7 +15,7 @@ namespace Insights
 
 TSharedRef<ITableCellValueGetter> FTableColumn::GetDefaultValueGetter()
 {
-	return MakeShareable(new FTableCellValueGetter());
+	return MakeShared<FTableCellValueGetter>();
 }
 
 ////////////////////////////////////////////////////////////////////////////////////////////////////
@@ -29,7 +29,7 @@ const TOptional<FTableCellValue> FTableColumn::GetValue(const FBaseTreeNode& InN
 
 TSharedRef<ITableCellValueFormatter> FTableColumn::GetDefaultValueFormatter()
 {
-	return MakeShareable(new FTableCellValueFormatter());
+	return MakeShared<FTableCellValueFormatter>();
 }
 
 ////////////////////////////////////////////////////////////////////////////////////////////////////
diff --git a/Engine/Source/Developer/TraceInsights/Private/Insights/Table/Widgets/STableTreeView.cpp b/Engine/Source/Developer/TraceInsights/Private/Insights/Table/Widgets/STableTreeView.cpp
--- a/Engine/Source/Developer/TraceInsights/Private/Insights/Table/Widgets/STableTreeView.cpp
+++ b/Engine/Source/Developer/TraceInsights/Private/Insights/Table/Widgets/STableTreeView.cpp
@@ -45,7 +45,7 @@ STableTreeView::STableTreeView()
 	, HoveredColumnId()
 	, HoveredNodePtr(nullptr)
 	, HighlightedNodeName()
-	, Root(MakeShareable(new FTableTreeNode(RootNodeName, Table)))
+	, Root(MakeShared<FTableTreeNode>(RootNodeName, Table))
 	, TableTreeNodes()
 	, FilteredGroupNodes()
 	, TableTreeNodesIdMap()
@@ -206,8 +206,8 @@ void STableTreeView::Construct(const FArguments& InArgs, TSharedPtr<FTable> InTa
 	//BindCommands();
 
 	// Create the search filters: text based, type based etc.
-	TextFilter = MakeShareable(new FTableTreeNodeTextFilter(FTableTreeNodeTextFilter::FItemToStringArray::CreateSP(this, &STableTreeView::HandleItemToStringArray)));
-	Filters = MakeShareable(new FTableTreeNodeFilterCollection());
+	TextFilter = MakeShared<FTableTreeNodeTextFilter>(FTableTreeNodeTextFilter::FItemToStringArray::CreateSP(this, &STableTreeView::HandleItemToStringArray));
+	Filters = MakeShared<FTableTreeNodeFilterCollection>();
 	Filters->Add(TextFilter);
 
 	CreateGroupings();
@@ -351,9 +351,9 @@ void STableTreeView::TreeView_BuildSortByMenu(FMenuBuilder& MenuBuilder)
 	MenuBuilder.BeginSection("ColumnName", LOCTEXT("ContextMenu_Header_Misc_ColumnName", "Column Name"));
 
 	//TODO: for (Sorting : AvailableSortings)
-	for (const TSharedPtr<FTableColumn>& ColumnPtr : Table->GetColumns())
+	for (const TSharedRef<FTableColumn>& ColumnRef : Table->GetColumns())
 	{
-		const FTableColumn& Column = *ColumnPtr;
+		const FTableColumn& Column = *ColumnRef;
 
 		if (Column.IsVisible() && Column.CanBeSorted())
 		{
@@ -413,9 +413,9 @@ void STableTreeView::TreeView_BuildViewColumnMenu(FMenuBuilder& MenuBuilder)
 {
 	MenuBuilder.BeginSection("ViewColumn", LOCTEXT("ContextMenu_Header_Columns_View", "View Column"));
 
-	for (const TSharedPtr<FTableColumn>& ColumnPtr : Table->GetColumns())
+	for (const TSharedRef<FTableColumn>& ColumnRef : Table->GetColumns())
 	{
-		const FTableColumn& Column = *ColumnPtr;
+		const FTableColumn& Column = *ColumnRef;
 
 		FUIAction Action_ToggleColumn
 		(
@@ -438,11 +438,11 @@ void STableTreeView::TreeView_BuildViewColumnMenu(FMenuBuilder& MenuBuilder)
 
 void STableTreeView::InitializeAndShowHeaderColumns()
 {
-	for (const TSharedPtr<FTableColumn>& ColumnPtr : Table->GetColumns())
+	for (const TSharedRef<FTableColumn>& ColumnRef : Table->GetColumns())
 	{
-		if (ColumnPtr->ShouldBeVisible())
+		if (ColumnRef->ShouldBeVisible())
 		{
-			ShowColumn(ColumnPtr->GetId());
+			ShowColumn(ColumnRef->GetId());
 		}
 	}
 }
@@ -806,9 +806,9 @@ void STableTreeView::CreateGroups()
 	GroupNodesRec(TableTreeNodes, *Root, 0);
 
 	ResetAggregatedValuesRec(*Root);
-	for (const TSharedPtr<FTableColumn> ColumnPtr : Table->GetColumns())
+	for (const TSharedRef<FTableColumn>& ColumnRef : Table->GetColumns())
 	{
-		FTableColumn& Column = *ColumnPtr;
+		FTableColumn& Column = *ColumnRef;
 		if (Column.GetAggregation() == ETableColumnAggregation::Sum)
 		{
 			switch (Column.GetDataType())
@@ -851,7 +851,7 @@ void STableTreeView::GroupNodesRec(const TArray<FTableTreeNodePtr>& Nodes, FTabl
 		FTableTreeNodePtr* GroupPtrPtr = GroupMap.Find(GroupInfo.Name);
 		if (!GroupPtrPtr)
 		{
-			GroupPtr = MakeShareable(new FTableTreeNode(GroupInfo.Name, Table));
+			GroupPtr = MakeShared<FTableTreeNode>(GroupInfo.Name, Table);
 
 			GroupMap.Add(GroupInfo.Name, GroupPtr);
 			ParentGroup.AddChildAndSetGroupPtr(GroupPtr);
@@ -986,15 +986,15 @@ void STableTreeView::CreateGroupings()
 {
 	AvailableGroupings.Reset(3);
 
-	AvailableGroupings.Add(MakeShareable(new FTreeNodeGroupingFlat()));
-	//AvailableGroupings.Add(MakeShareable(new FTreeNodeGroupingByNameFirstLetter()));
-	//AvailableGroupings.Add(MakeShareable(new FTreeNodeGroupingByType()));
+	AvailableGroupings.Add(MakeShared<FTreeNodeGroupingFlat>());
+	//AvailableGroupings.Add(MakeShared<FTreeNodeGroupingByNameFirstLetter>());
+	//AvailableGroupings.Add(MakeShared<FTreeNodeGroupingByType>());
 
-	for (const TSharedPtr<FTableColumn> ColumnPtr : Table->GetColumns())
+	for (const TSharedRef<FTableColumn>& ColumnRef : Table->GetColumns())
 	{
-		if (!ColumnPtr->IsHierarchy())
+		if (!ColumnRef->IsHierarchy())
 		{
-			AvailableGroupings.Add(MakeShareable(new FTreeNodeGroupingByUniqueValue(ColumnPtr.ToSharedRef())));
+			AvailableGroupings.Add(MakeShared<FTreeNodeGroupingByUniqueValue>(ColumnRef));
 		}
 	}
 
@@ -1476,11 +1476,11 @@ void STableTreeView::CreateSortings()
 	AvailableSorters.Reset();
 	CurrentSorter = nullptr;
 
-	for (const TSharedPtr<FTableColumn> ColumnPtr : Table->GetColumns())
+	for (const TSharedRef<FTableColumn>& ColumnRef : Table->GetColumns())
 	{
-		if (ColumnPtr->CanBeSorted())
+		if (ColumnRef->CanBeSorted())
 		{
-			TSharedPtr<Insights::ITableCellValueSorter> SorterPtr = ColumnPtr->GetValueSorter();
+			TSharedPtr<Insights::ITableCellValueSorter> SorterPtr = ColumnRef->GetValueSorter();
 			if (ensure(SorterPtr.IsValid()))
 			{
 				AvailableSorters.Add(SorterPtr);
@@ -1767,9 +1767,9 @@ void STableTreeView::ContextMenu_ShowAllColumns_Execute()
 	ColumnSortMode = GetDefaultColumnSortMode();
 	UpdateCurrentSortingByColumn();
 
-	for (const TSharedPtr<FTableColumn>& ColumnPtr : Table->GetColumns())
+	for (const TSharedRef<FTableColumn>& ColumnRef : Table->GetColumns())
 	{
-		const FTableColumn& Column = *ColumnPtr;
+		const FTableColumn& Column = *ColumnRef;
 
 		if (!Column.IsVisible())
 		{
@@ -1795,9 +1795,9 @@ void STableTreeView::ContextMenu_ResetColumns_Execute()
 	ColumnSortMode = GetDefaultColumnSortMode();
 	UpdateCurrentSortingByColumn();
 
-	for (const TSharedPtr<FTableColumn>& ColumnPtr : Table->GetColumns())
+	for (const TSharedRef<FTableColumn>& ColumnRef : Table->GetColumns())
 	{
-		const FTableColumn& Column = *ColumnPtr;
+		const FTableColumn& Column = *ColumnRef;
 
 		if (Column.ShouldBeVisible() && !Column.IsVisible())
 		{
@@ -1858,7 +1858,7 @@ void STableTreeView::RebuildTree(bool bResync)
 				TableReader->SetRowIndex(RowIndex);
 				uint64 NodeId = static_cast<uint64>(RowIndex);
 				FName NodeName(*FString::Printf(TEXT("row %d"), RowIndex));
-				FTableTreeNodePtr NodePtr = MakeShareable(new FTableTreeNode(NodeId, NodeName, Table, RowIndex));
+				FTableTreeNodePtr NodePtr = MakeShared<FTableTreeNode>(NodeId, NodeName, Table, RowIndex);
 				TableTreeNodes.Add(NodePtr);
 				TableTreeNodesIdMap.Add(NodeId, NodePtr);
 			}
diff --git a/Engine/Source/Developer/TraceInsights/Private/Insights/Table/Widgets/STableTreeViewRow.cpp b/Engine/Source/Developer/TraceInsights/Private/Insights/Table/Widgets/STableTreeViewRow.cpp
--- a/Engine/Source/Developer/TraceInsights/Private/Insights/Table/Widgets/STableTreeViewRow.cpp
+++ b/Engine/Source/Developer/TraceInsights/Private/Insights/Table/Widgets/STableTreeViewRow.cpp
@@ -36,7 +36,7 @@ void STableTreeViewRow::Construct(const FArguments& InArgs, const TSharedRef<STa
 	TablePtr = InArgs._TablePtr;
 	TableTreeNodePtr = InArgs._TableTreeNodePtr;
 
-	RowToolTip = MakeShareable(new STableTreeRowToolTip(TableTreeNodePtr));
+	RowToolTip = MakeShared<STableTreeRowToolTip>(TableTreeNodePtr);
 
 	SetEnabled(TAttribute<bool>(this, &STableTreeViewRow::HandleShouldBeEnabled));
 
diff --git a/Engine/Source/Developer/TraceInsights/Private/Insights/ViewModels/FileActivityTimingTrack.cpp b/Engine/Source/Developer/TraceInsights/Private/Insights/ViewModels/FileActivityTimingTrack.cpp
--- a/Engine/Source/Developer/TraceInsights/Private/Insights/ViewModels/FileActivityTimingTrack.cpp
+++ b/Engine/Source/Developer/TraceInsights/Private/Insights/ViewModels/FileActivityTimingTrack.cpp
@@ -127,7 +127,7 @@ void FFileActivitySharedState::Tick(Insights::ITimingViewSession& InSession, con
 
 	if (!IoOverviewTrack.IsValid())
 	{
-		IoOverviewTrack = MakeShareable(new FOverviewFileActivityTimingTrack(*this));
+		IoOverviewTrack = MakeShared<FOverviewFileActivityTimingTrack>(*this);
 		IoOverviewTrack->SetOrder(-999999); // first track
 		IoOverviewTrack->SetVisibilityFlag(bShowHideAllIoTracks);
 		InSession.AddScrollableTrack(IoOverviewTrack);
@@ -135,7 +135,7 @@ void FFileActivitySharedState::Tick(Insights::ITimingViewSession& InSession, con
 
 	if (!IoActivityTrack.IsValid())
 	{
-		IoActivityTrack = MakeShareable(new FDetailedFileActivityTimingTrack(*this));
+		IoActivityTrack = MakeShared<FDetailedFileActivityTimingTrack>(*this);
 		IoActivityTrack->SetOrder(+999999); // last track
 		IoActivityTrack->SetVisibilityFlag(bShowHideAllIoTracks);
 		InSession.AddScrollableTrack(IoActivityTrack);
@@ -170,7 +170,7 @@ void FFileActivitySharedState::Tick(Insights::ITimingViewSession& InSession, con
 
 					if (!FileActivityMap.Contains(FileInfo.Id))
 					{
-						Activity = MakeShareable(new FIoFileActivity());
+						Activity = MakeShared<FIoFileActivity>();
 
 						Activity->Id = FileInfo.Id;
 						Activity->Path = FileInfo.Path;
@@ -680,7 +680,7 @@ const TSharedPtr<const ITimingEvent> FOverviewFileActivityTimingTrack::SearchEve
 	constexpr bool bIgnoreEventDepth = true;
 	FindIoTimingEvent(InSearchParameters, bIgnoreEventDepth, [this, &FoundEvent](double InFoundStartTime, double InFoundEndTime, uint32 InFoundDepth, const FFileActivitySharedState::FIoTimingEvent& InEvent)
 	{
-		FoundEvent = MakeShareable(new FTimingEvent(SharedThis(this), InFoundStartTime, InFoundEndTime, InFoundDepth));
+		FoundEvent = MakeShared<FTimingEvent>(SharedThis(this), InFoundStartTime, InFoundEndTime, InFoundDepth);
 	});
 
 	return FoundEvent;
@@ -771,7 +771,7 @@ const TSharedPtr<const ITimingEvent> FDetailedFileActivityTimingTrack::SearchEve
 	constexpr bool bIgnoreEventDepth = false;
 	FindIoTimingEvent(InSearchParameters, bIgnoreEventDepth, [this, &FoundEvent](double InFoundStartTime, double InFoundEndTime, uint32 InFoundDepth, const FFileActivitySharedState::FIoTimingEvent& InEvent)
 	{
-		FoundEvent = MakeShareable(new FTimingEvent(SharedThis(this), InFoundStartTime, InFoundEndTime, InFoundDepth));
+		FoundEvent = MakeShared<FTimingEvent>(SharedThis(this), InFoundStartTime, InFoundEndTime, InFoundDepth);
 	});
 
 	return FoundEvent;
diff --git a/Engine/Source/Developer/TraceInsights/Private/Insights/ViewModels/GraphTrack.cpp b/Engine/Source/Developer/TraceInsights/Private/Insights/ViewModels/GraphTrack.cpp
--- a/Engine/Source/Developer/TraceInsights/Private/Insights/ViewModels/GraphTrack.cpp
+++ b/Engine/Source/Developer/TraceInsights/Private/Insights/ViewModels/GraphTrack.cpp
@@ -475,7 +475,7 @@ const TSharedPtr<const ITimingEvent> FGraphTrack::GetEvent(float InPosX, float I
 			const FGraphSeriesEvent* Event = Series->GetEvent(LocalPosX, LocalPosY, Viewport, bCheckLine, bDrawBoxes);
 			if (Event != nullptr)
 			{
-				return MakeShareable(new FGraphTrackEvent(SharedThis(this), Series.ToSharedRef(), *Event));
+				return MakeShared<FGraphTrackEvent>(SharedThis(this), Series.ToSharedRef(), *Event);
 			}
 		}
 	}
@@ -784,21 +784,21 @@ FRandomGraphTrack::FRandomGraphTrack()
 
 void FRandomGraphTrack::AddDefaultSeries()
 {
-	TSharedPtr<FGraphSeries> Series0 = MakeShareable(new FGraphSeries());
+	TSharedRef<FGraphSeries> Series0 = MakeShared<FGraphSeries>();
 	Series0->SetName(TEXT("Random Blue"));
 	Series0->SetDescription(TEXT("Random series; for debuging purposes"));
 	Series0->SetColor(FLinearColor(0.1f, 0.5f, 1.0f, 1.0f), FLinearColor(0.4f, 0.8f, 1.0f, 1.0f));
 	Series0->SetVisibility(true);
 	AllSeries.Add(Series0);
 
-	TSharedPtr<FGraphSeries> Series1 = MakeShareable(new FGraphSeries());
+	TSharedRef<FGraphSeries> Series1 = MakeShared<FGraphSeries>();
 	Series1->SetName(TEXT("Random Yellow"));
 	Series1->SetDescription(TEXT("Random series; for debuging purposes"));
 	Series1->SetColor(FLinearColor(0.9f, 0.9f, 0.1f, 1.0f), FLinearColor(1.0f, 1.0f, 0.4f, 1.0f));
 	Series1->SetVisibility(false);
 	AllSeries.Add(Series1);
 
-	TSharedPtr<FGraphSeries> Series2 = MakeShareable(new FGraphSeries());
+	TSharedRef<FGraphSeries> Series2 = MakeShared<FGraphSeries>();
 	Series2->SetName(TEXT("Random Red"));
 	Series2->SetDescription(TEXT("Random series; for debuging purposes"));
 	Series2->SetColor(FLinearColor(1.0f, 0.1f, 0.2f, 1.0f), FLinearColor(1.0f, 0.4f, 0.5f, 1.0f));
diff --git a/Engine/Source/Developer/TraceInsights/Private/Insights/ViewModels/LoadingTimingTrack.cpp b/Engine/Source/Developer/TraceInsights/Private/Insights/ViewModels/LoadingTimingTrack.cpp
--- a/Engine/Source/Developer/TraceInsights/Private/Insights/ViewModels/LoadingTimingTrack.cpp
+++ b/Engine/Source/Developer/TraceInsights/Private/Insights/ViewModels/LoadingTimingTrack.cpp
@@ -89,7 +89,7 @@ void FLoadingSharedState::Tick(Insights::ITimingViewSession& InSession, const Tr
 					const TCHAR* const GroupName = ThreadInfo.GroupName ? ThreadInfo.GroupName : ThreadInfo.Name;
 					const FString TrackName(ThreadInfo.Name && *ThreadInfo.Name ? FString::Printf(TEXT("Loading - %s"), ThreadInfo.Name) : FString::Printf(TEXT("Loading - Thread %u"), ThreadInfo.Id));
 
-					TSharedPtr<FLoadingTimingTrack> LoadingThreadTrack = MakeShareable(new FLoadingTimingTrack(*this, LoadingTimelineIndex, GroupName, TrackName));
+					TSharedRef<FLoadingTimingTrack> LoadingThreadTrack = MakeShared<FLoadingTimingTrack>(*this, LoadingTimelineIndex, GroupName, TrackName);
 					LoadingThreadTrack->SetOrder(-100 + LoadingTracks.Num());
 					LoadingThreadTrack->SetVisibilityFlag(bShowHideAllLoadingTracks);
 					InSession.AddScrollableTrack(LoadingThreadTrack);
@@ -219,6 +219,12 @@ void FLoadingSharedState::SetColorSchema(int32 Schema)
 		case 2: GetEventNameDelegate = FLoadingTrackGetEventNameDelegate::CreateRaw(this, &FLoadingSharedState::GetEventNameByExportClassName); break;
 		case 3: GetEventNameDelegate = FLoadingTrackGetEventNameDelegate::CreateRaw(this, &FLoadingSharedState::GetEventNameByPackageAndExportClassName); break;
 	};
+
+	for (const auto& KV : LoadingTracks)
+	{
+		FLoadingTimingTrack& Track = *KV.Value;
+		Track.SetDirtyFlag();
+	}
 }
 
 ////////////////////////////////////////////////////////////////////////////////////////////////////
@@ -319,7 +325,7 @@ const TSharedPtr<const ITimingEvent> FLoadingTimingTrack::SearchEvent(const FTim
 
 	FindLoadTimeProfilerCpuEvent(InSearchParameters, [this, &FoundEvent](double InFoundStartTime, double InFoundEndTime, uint32 InFoundDepth, const Trace::FLoadTimeProfilerCpuEvent& InFoundEvent)
 	{
-		FoundEvent = MakeShareable(new FTimingEvent(SharedThis(this), InFoundStartTime, InFoundEndTime, InFoundDepth));
+		FoundEvent = MakeShared<FTimingEvent>(SharedThis(this), InFoundStartTime, InFoundEndTime, InFoundDepth);
 	});
 
 	return FoundEvent;
diff --git a/Engine/Source/Developer/TraceInsights/Private/Insights/ViewModels/LogMessage.cpp b/Engine/Source/Developer/TraceInsights/Private/Insights/ViewModels/LogMessage.cpp
--- a/Engine/Source/Developer/TraceInsights/Private/Insights/ViewModels/LogMessage.cpp
+++ b/Engine/Source/Developer/TraceInsights/Private/Insights/ViewModels/LogMessage.cpp
@@ -190,7 +190,7 @@ TSharedPtr<FLogMessageRecord> FLogMessageCache::GetUncached(uint64 Index) const
 		const Trace::ILogProvider& LogProvider = Trace::ReadLogProvider(*Session.Get());
 		LogProvider.ReadMessage(Index, [&EntryPtr](const Trace::FLogMessage& Message)
 		{
-			EntryPtr = MakeShareable(new FLogMessageRecord(Message));
+			EntryPtr = MakeShared<FLogMessageRecord>(Message);
 		});
 	}
 
diff --git a/Engine/Source/Developer/TraceInsights/Private/Insights/ViewModels/ThreadTimingTrack.cpp b/Engine/Source/Developer/TraceInsights/Private/Insights/ViewModels/ThreadTimingTrack.cpp
--- a/Engine/Source/Developer/TraceInsights/Private/Insights/ViewModels/ThreadTimingTrack.cpp
+++ b/Engine/Source/Developer/TraceInsights/Private/Insights/ViewModels/ThreadTimingTrack.cpp
@@ -553,12 +553,12 @@ void FThreadTimingTrack::GetParentAndRoot(const FTimingEvent& TimingEvent, TShar
 						if (EventDepth == 0)
 						{
 							OutRootEvent = Event;
-							OutRootTimingEvent = MakeShareable(new FTimingEvent(TimingEvent.GetTrack(), EventStartTime, EventEndTime, EventDepth));
+							OutRootTimingEvent = MakeShared<FTimingEvent>(TimingEvent.GetTrack(), EventStartTime, EventEndTime, EventDepth);
 						}
 						if (EventDepth == TimingEvent.GetDepth() - 1)
 						{
 							OutParentEvent = Event;
-							OutParentTimingEvent = MakeShareable(new FTimingEvent(TimingEvent.GetTrack(), EventStartTime, EventEndTime, EventDepth));
+							OutParentTimingEvent = MakeShared<FTimingEvent>(TimingEvent.GetTrack(), EventStartTime, EventEndTime, EventDepth);
 						}
 					});
 				});
@@ -575,7 +575,7 @@ const TSharedPtr<const ITimingEvent> FThreadTimingTrack::SearchEvent(const FTimi
 
 	FindTimingProfilerEvent(InSearchParameters, [this, &FoundEvent](double InFoundStartTime, double InFoundEndTime, uint32 InFoundDepth, const Trace::FTimingProfilerEvent& InFoundEvent)
 	{
-		FoundEvent = MakeShareable(new FTimingEvent(SharedThis(this), InFoundStartTime, InFoundEndTime, InFoundDepth));
+		FoundEvent = MakeShared<FTimingEvent>(SharedThis(this), InFoundStartTime, InFoundEndTime, InFoundDepth);
 	});
 
 	return FoundEvent;
diff --git a/Engine/Source/Developer/TraceInsights/Private/Insights/ViewModels/TimersViewColumnFactory.cpp b/Engine/Source/Developer/TraceInsights/Private/Insights/ViewModels/TimersViewColumnFactory.cpp
--- a/Engine/Source/Developer/TraceInsights/Private/Insights/ViewModels/TimersViewColumnFactory.cpp
+++ b/Engine/Source/Developer/TraceInsights/Private/Insights/ViewModels/TimersViewColumnFactory.cpp
@@ -40,7 +40,7 @@ const FName FTimersViewColumns::MinExclusiveTimeColumnID(TEXT("MinExclTime"));
 
 ////////////////////////////////////////////////////////////////////////////////////////////////////
 
-void FTimersViewColumnFactory::CreateTimersViewColumns(TArray<TSharedPtr<Insights::FTableColumn>>& Columns)
+void FTimersViewColumnFactory::CreateTimersViewColumns(TArray<TSharedRef<Insights::FTableColumn>>& Columns)
 {
 	Columns.Reset();
 
@@ -64,7 +64,7 @@ void FTimersViewColumnFactory::CreateTimersViewColumns(TArray<TSharedPtr<Insight
 
 ////////////////////////////////////////////////////////////////////////////////////////////////////
 
-void FTimersViewColumnFactory::CreateTimerTreeViewColumns(TArray<TSharedPtr<Insights::FTableColumn>>& Columns)
+void FTimersViewColumnFactory::CreateTimerTreeViewColumns(TArray<TSharedRef<Insights::FTableColumn>>& Columns)
 {
 	Columns.Reset();
 
diff --git a/Engine/Source/Developer/TraceInsights/Private/Insights/ViewModels/TimersViewColumnFactory.h b/Engine/Source/Developer/TraceInsights/Private/Insights/ViewModels/TimersViewColumnFactory.h
--- a/Engine/Source/Developer/TraceInsights/Private/Insights/ViewModels/TimersViewColumnFactory.h
+++ b/Engine/Source/Developer/TraceInsights/Private/Insights/ViewModels/TimersViewColumnFactory.h
@@ -47,8 +47,8 @@ namespace Insights
 struct FTimersViewColumnFactory
 {
 public:
-	static void CreateTimersViewColumns(TArray<TSharedPtr<Insights::FTableColumn>>& Columns);
-	static void CreateTimerTreeViewColumns(TArray<TSharedPtr<Insights::FTableColumn>>& Columns);
+	static void CreateTimersViewColumns(TArray<TSharedRef<Insights::FTableColumn>>& Columns);
+	static void CreateTimerTreeViewColumns(TArray<TSharedRef<Insights::FTableColumn>>& Columns);
 
 	static TSharedRef<Insights::FTableColumn> CreateNameColumn();
 	static TSharedRef<Insights::FTableColumn> CreateMetaGroupNameColumn();
diff --git a/Engine/Source/Developer/TraceInsights/Private/Insights/ViewModels/TimingGraphTrack.cpp b/Engine/Source/Developer/TraceInsights/Private/Insights/ViewModels/TimingGraphTrack.cpp
--- a/Engine/Source/Developer/TraceInsights/Private/Insights/ViewModels/TimingGraphTrack.cpp
+++ b/Engine/Source/Developer/TraceInsights/Private/Insights/ViewModels/TimingGraphTrack.cpp
@@ -73,7 +73,7 @@ void FTimingGraphTrack::AddDefaultFrameSeries()
 	const float BaselineY = GetHeight() - 1.0f;
 	const double ScaleY = GetHeight() / 0.1; // 100ms
 
-	TSharedPtr<FTimingGraphSeries> GameFramesSeries = MakeShareable(new FTimingGraphSeries());
+	TSharedRef<FTimingGraphSeries> GameFramesSeries = MakeShared<FTimingGraphSeries>();
 	GameFramesSeries->SetName(TEXT("Game Frames"));
 	GameFramesSeries->SetDescription(TEXT("Duration of Game frames"));
 	GameFramesSeries->SetColor(FLinearColor(0.3f, 0.3f, 1.0f, 1.0f), FLinearColor(1.0f, 1.0f, 1.0f, 1.0f));
@@ -83,7 +83,7 @@ void FTimingGraphTrack::AddDefaultFrameSeries()
 	GameFramesSeries->SetScaleY(ScaleY);
 	AllSeries.Add(GameFramesSeries);
 
-	TSharedPtr<FTimingGraphSeries> RenderingFramesSeries = MakeShareable(new FTimingGraphSeries());
+	TSharedRef<FTimingGraphSeries> RenderingFramesSeries = MakeShared<FTimingGraphSeries>();
 	RenderingFramesSeries->SetName(TEXT("Rendering Frames"));
 	RenderingFramesSeries->SetDescription(TEXT("Duration of Rendering frames"));
 	RenderingFramesSeries->SetColor(FLinearColor(1.0f, 0.3f, 0.3f, 1.0f), FLinearColor(1.0f, 1.0f, 1.0f, 1.0f));
@@ -110,7 +110,7 @@ TSharedPtr<FTimingGraphSeries> FTimingGraphTrack::GetStatsCounterSeries(uint32 C
 
 TSharedPtr<FTimingGraphSeries> FTimingGraphTrack::AddStatsCounterSeries(uint32 CounterId, FLinearColor Color)
 {
-	TSharedPtr<FTimingGraphSeries> Series = MakeShareable(new FTimingGraphSeries());
+	TSharedRef<FTimingGraphSeries> Series = MakeShared<FTimingGraphSeries>();
 
 	const TCHAR* CounterName = nullptr;
 	bool bIsFloatingPoint = false;
diff --git a/Engine/Source/Developer/TraceInsights/Private/Insights/Widgets/SFrameTrack.cpp b/Engine/Source/Developer/TraceInsights/Private/Insights/Widgets/SFrameTrack.cpp
--- a/Engine/Source/Developer/TraceInsights/Private/Insights/Widgets/SFrameTrack.cpp
+++ b/Engine/Source/Developer/TraceInsights/Private/Insights/Widgets/SFrameTrack.cpp
@@ -203,19 +203,19 @@ void SFrameTrack::Tick(const FGeometry& AllottedGeometry, const double InCurrent
 
 ////////////////////////////////////////////////////////////////////////////////////////////////////
 
-TSharedPtr<FFrameTrackSeries> SFrameTrack::FindOrAddSeries(int32 FrameType)
+TSharedRef<FFrameTrackSeries> SFrameTrack::FindOrAddSeries(int32 FrameType)
 {
 	TSharedPtr<FFrameTrackSeries>* SeriesPtrPtr = SeriesMap.Find(FrameType);
 	if (SeriesPtrPtr)
 	{
 		ensure((**SeriesPtrPtr).FrameType == FrameType);
-		return *SeriesPtrPtr;
+		return (*SeriesPtrPtr).ToSharedRef();
 	}
 	else
 	{
-		TSharedPtr<FFrameTrackSeries> SeriesPtr = MakeShareable(new FFrameTrackSeries(FrameType));
-		SeriesMap.Add(FrameType, SeriesPtr);
-		return SeriesPtr;
+		TSharedRef<FFrameTrackSeries> SeriesRef = MakeShared<FFrameTrackSeries>(FrameType);
+		SeriesMap.Add(FrameType, SeriesRef);
+		return SeriesRef;
 	}
 }
 
@@ -333,7 +333,7 @@ FFrameTrackSampleRef SFrameTrack::GetSampleAtMousePosition(float X, float Y)
 
 							if (Y >= TopY && Y < BottomY)
 							{
-								return FFrameTrackSampleRef(SeriesPtr, MakeShareable(new FFrameTrackSample(Sample)));
+								return FFrameTrackSampleRef(SeriesPtr, MakeShared<FFrameTrackSample>(Sample));
 							}
 						}
 					}
diff --git a/Engine/Source/Developer/TraceInsights/Private/Insights/Widgets/SFrameTrack.h b/Engine/Source/Developer/TraceInsights/Private/Insights/Widgets/SFrameTrack.h
--- a/Engine/Source/Developer/TraceInsights/Private/Insights/Widgets/SFrameTrack.h
+++ b/Engine/Source/Developer/TraceInsights/Private/Insights/Widgets/SFrameTrack.h
@@ -129,7 +129,7 @@ public:
 	virtual FCursorReply OnCursorQuery(const FGeometry& MyGeometry, const FPointerEvent& CursorEvent) const override;
 
 protected:
-	TSharedPtr<FFrameTrackSeries> FindOrAddSeries(int32 FrameType);
+	TSharedRef<FFrameTrackSeries> FindOrAddSeries(int32 FrameType);
 	TSharedPtr<FFrameTrackSeries> FindSeries(int32 FrameType) const;
 	void UpdateState();
 
diff --git a/Engine/Source/Developer/TraceInsights/Private/Insights/Widgets/SStartPageWindow.cpp b/Engine/Source/Developer/TraceInsights/Private/Insights/Widgets/SStartPageWindow.cpp
--- a/Engine/Source/Developer/TraceInsights/Private/Insights/Widgets/SStartPageWindow.cpp
+++ b/Engine/Source/Developer/TraceInsights/Private/Insights/Widgets/SStartPageWindow.cpp
@@ -1263,15 +1263,15 @@ TSharedRef<SWidget> SStartPageWindow::MakeSessionListMenu()
 			Trace::FSessionInfo  SessionInfo;
 		};
 
-		TArray<TSharedPtr<FSessionInfoEx>> SortedAvailableSessions;
+		TArray<TSharedRef<FSessionInfoEx>> SortedAvailableSessions;
 		SortedAvailableSessions.Reserve(AvailableSessions.Num());
 
 		for (Trace::FSessionHandle SessionHandle : AvailableSessions)
 		{
-			FSessionInfoEx* SessionInfoExPtr = new FSessionInfoEx;
-			SessionInfoExPtr->SessionHandle = SessionHandle;
-			SessionService->GetSessionInfo(SessionHandle, SessionInfoExPtr->SessionInfo);
-			SortedAvailableSessions.Add(MakeShareable(SessionInfoExPtr));
+			TSharedRef<FSessionInfoEx> SessionInfoEx = MakeShared<FSessionInfoEx>();
+			SessionInfoEx->SessionHandle = SessionHandle;
+			SessionService->GetSessionInfo(SessionHandle, SessionInfoEx->SessionInfo);
+			SortedAvailableSessions.Add(SessionInfoEx);
 		}
 
 		Algo::SortBy(SortedAvailableSessions, &FSessionInfoEx::GetTimeStamp);
diff --git a/Engine/Source/Developer/TraceInsights/Private/Insights/Widgets/SStatsTableRow.cpp b/Engine/Source/Developer/TraceInsights/Private/Insights/Widgets/SStatsTableRow.cpp
--- a/Engine/Source/Developer/TraceInsights/Private/Insights/Widgets/SStatsTableRow.cpp
+++ b/Engine/Source/Developer/TraceInsights/Private/Insights/Widgets/SStatsTableRow.cpp
@@ -31,7 +31,7 @@ void SStatsTableRow::Construct(const FArguments& InArgs, const TSharedRef<STable
 
 	StatsNodePtr = InArgs._StatsNodePtr;
 
-	RowToolTip = MakeShareable(new SStatsCounterTableRowToolTip(StatsNodePtr));
+	RowToolTip = MakeShared<SStatsCounterTableRowToolTip>(StatsNodePtr);
 
 	SetEnabled(TAttribute<bool>(this, &SStatsTableRow::HandleShouldBeEnabled));
 
diff --git a/Engine/Source/Developer/TraceInsights/Private/Insights/Widgets/SStatsView.cpp b/Engine/Source/Developer/TraceInsights/Private/Insights/Widgets/SStatsView.cpp
--- a/Engine/Source/Developer/TraceInsights/Private/Insights/Widgets/SStatsView.cpp
+++ b/Engine/Source/Developer/TraceInsights/Private/Insights/Widgets/SStatsView.cpp
@@ -242,8 +242,8 @@ void SStatsView::Construct(const FArguments& InArgs)
 	//BindCommands();
 
 	// Create the search filters: text based, type based etc.
-	TextFilter = MakeShareable(new FStatsNodeTextFilter(FStatsNodeTextFilter::FItemToStringArray::CreateSP(this, &SStatsView::HandleItemToStringArray)));
-	Filters = MakeShareable(new FStatsNodeFilterCollection());
+	TextFilter = MakeShared<FStatsNodeTextFilter>(FStatsNodeTextFilter::FItemToStringArray::CreateSP(this, &SStatsView::HandleItemToStringArray));
+	Filters = MakeShared<FStatsNodeFilterCollection>();
 	Filters->Add(TextFilter);
 
 	CreateGroupByOptionsSources();
@@ -1064,7 +1064,7 @@ void SStatsView::CreateGroups()
 		FStatsNodePtr* GroupPtr = GroupNodeSet.Find(GroupName);
 		if (!GroupPtr)
 		{
-			GroupPtr = &GroupNodeSet.Add(GroupName, MakeShareable(new FStatsNode(GroupName)));
+			GroupPtr = &GroupNodeSet.Add(GroupName, MakeShared<FStatsNode>(GroupName));
 		}
 
 		for (const FStatsNodePtr& StatsNodePtr : StatsNodes)
@@ -1084,7 +1084,7 @@ void SStatsView::CreateGroups()
 			FStatsNodePtr* GroupPtr = GroupNodeSet.Find(GroupName);
 			if (!GroupPtr)
 			{
-				GroupPtr = &GroupNodeSet.Add(GroupName, MakeShareable(new FStatsNode(GroupName)));
+				GroupPtr = &GroupNodeSet.Add(GroupName, MakeShared<FStatsNode>(GroupName));
 			}
 
 			(*GroupPtr)->AddChildAndSetGroupPtr(StatsNodePtr);
@@ -1101,7 +1101,7 @@ void SStatsView::CreateGroups()
 			FStatsNodePtr* GroupPtr = GroupNodeSet.Find(GroupName);
 			if (!GroupPtr)
 			{
-				GroupPtr = &GroupNodeSet.Add(GroupName, MakeShareable(new FStatsNode(GroupName)));
+				GroupPtr = &GroupNodeSet.Add(GroupName, MakeShared<FStatsNode>(GroupName));
 			}
 
 			(*GroupPtr)->AddChildAndSetGroupPtr(StatsNodePtr);
@@ -1118,7 +1118,7 @@ void SStatsView::CreateGroups()
 			FStatsNodePtr* GroupPtr = GroupNodeSet.Find(GroupName);
 			if (!GroupPtr)
 			{
-				GroupPtr = &GroupNodeSet.Add(GroupName, MakeShareable(new FStatsNode(GroupName)));
+				GroupPtr = &GroupNodeSet.Add(GroupName, MakeShared<FStatsNode>(GroupName));
 			}
 
 			(*GroupPtr)->AddChildAndSetGroupPtr(StatsNodePtr);
@@ -1138,10 +1138,10 @@ void SStatsView::CreateGroupByOptionsSources()
 	GroupByOptionsSource.Reset(3);
 
 	// Must be added in order of elements in the EStatsGroupingMode.
-	GroupByOptionsSource.Add(MakeShareable(new EStatsGroupingMode(EStatsGroupingMode::Flat)));
-	GroupByOptionsSource.Add(MakeShareable(new EStatsGroupingMode(EStatsGroupingMode::ByName)));
-	GroupByOptionsSource.Add(MakeShareable(new EStatsGroupingMode(EStatsGroupingMode::ByMetaGroupName)));
-	GroupByOptionsSource.Add(MakeShareable(new EStatsGroupingMode(EStatsGroupingMode::ByType)));
+	GroupByOptionsSource.Add(MakeShared<EStatsGroupingMode>(EStatsGroupingMode::Flat));
+	GroupByOptionsSource.Add(MakeShared<EStatsGroupingMode>(EStatsGroupingMode::ByName));
+	GroupByOptionsSource.Add(MakeShared<EStatsGroupingMode>(EStatsGroupingMode::ByMetaGroupName));
+	GroupByOptionsSource.Add(MakeShared<EStatsGroupingMode>(EStatsGroupingMode::ByType));
 
 	EStatsGroupingModePtr* GroupingModePtrPtr = GroupByOptionsSource.FindByPredicate([&](const EStatsGroupingModePtr InGroupingModePtr) { return *InGroupingModePtr == GroupingMode; });
 	if (GroupingModePtrPtr != nullptr)
@@ -1553,7 +1553,7 @@ void SStatsView::RebuildTree(bool bResync)
 				FName Group(Counter.GetDisplayHint() == Trace::CounterDisplayHint_Memory ? TEXT("Memory") :
 							Counter.IsFloatingPoint() ? TEXT("float") : TEXT("int64"));
 				EStatsNodeType Type = Counter.IsFloatingPoint() ? EStatsNodeType::Float : EStatsNodeType::Int64;
-				FStatsNodePtr StatsNodePtr = MakeShareable(new FStatsNode(CounterId, Name, Group, Type));
+				FStatsNodePtr StatsNodePtr = MakeShared<FStatsNode>(CounterId, Name, Group, Type);
 				UpdateStatsNode(StatsNodePtr);
 				StatsNodes.Add(StatsNodePtr);
 				//StatsNodesMap.Add(Name, StatsNodePtr);
diff --git a/Engine/Source/Developer/TraceInsights/Private/Insights/Widgets/STimerTableRow.cpp b/Engine/Source/Developer/TraceInsights/Private/Insights/Widgets/STimerTableRow.cpp
--- a/Engine/Source/Developer/TraceInsights/Private/Insights/Widgets/STimerTableRow.cpp
+++ b/Engine/Source/Developer/TraceInsights/Private/Insights/Widgets/STimerTableRow.cpp
@@ -34,7 +34,7 @@ void STimerTableRow::Construct(const FArguments& InArgs, const TSharedRef<STable
 	TablePtr = InArgs._TablePtr;
 	TimerNodePtr = InArgs._TimerNodePtr;
 
-	RowToolTip = MakeShareable(new STimerTableRowToolTip(TimerNodePtr));
+	RowToolTip = MakeShared<STimerTableRowToolTip>(TimerNodePtr);
 
 	SetEnabled(TAttribute<bool>(this, &STimerTableRow::HandleShouldBeEnabled));
 
diff --git a/Engine/Source/Developer/TraceInsights/Private/Insights/Widgets/STimerTreeView.cpp b/Engine/Source/Developer/TraceInsights/Private/Insights/Widgets/STimerTreeView.cpp
--- a/Engine/Source/Developer/TraceInsights/Private/Insights/Widgets/STimerTreeView.cpp
+++ b/Engine/Source/Developer/TraceInsights/Private/Insights/Widgets/STimerTreeView.cpp
@@ -22,7 +22,7 @@
 ////////////////////////////////////////////////////////////////////////////////////////////////////
 
 STimerTreeView::STimerTreeView()
-	: Table(MakeShareable(new Insights::FTable()))
+	: Table(MakeShared<Insights::FTable>())
 	, ViewName()
 	, TreeView(nullptr)
 	, TreeViewHeaderRow(nullptr)
@@ -227,9 +227,9 @@ void STimerTreeView::TreeView_BuildSortByMenu(FMenuBuilder& MenuBuilder)
 {
 	MenuBuilder.BeginSection("ColumnName", LOCTEXT("ContextMenu_Header_Misc_ColumnName", "Column Name"));
 
-	for (const TSharedPtr<Insights::FTableColumn>& ColumnPtr : Table->GetColumns())
+	for (const TSharedRef<Insights::FTableColumn>& ColumnRef : Table->GetColumns())
 	{
-		const Insights::FTableColumn& Column = *ColumnPtr;
+		const Insights::FTableColumn& Column = *ColumnRef;
 
 		if (Column.IsVisible() && Column.CanBeSorted())
 		{
@@ -289,9 +289,9 @@ void STimerTreeView::TreeView_BuildViewColumnMenu(FMenuBuilder& MenuBuilder)
 {
 	MenuBuilder.BeginSection("ViewColumn", LOCTEXT("ContextMenu_Header_Columns_View", "View Column"));
 
-	for (const TSharedPtr<Insights::FTableColumn>& ColumnPtr : Table->GetColumns())
+	for (const TSharedRef<Insights::FTableColumn>& ColumnRef : Table->GetColumns())
 	{
-		const Insights::FTableColumn& Column = *ColumnPtr;
+		const Insights::FTableColumn& Column = *ColumnRef;
 
 		FUIAction Action_ToggleColumn
 		(
@@ -315,7 +315,7 @@ void STimerTreeView::TreeView_BuildViewColumnMenu(FMenuBuilder& MenuBuilder)
 void STimerTreeView::InitializeAndShowHeaderColumns()
 {
 	// Create columns.
-	TArray<TSharedPtr<Insights::FTableColumn>> Columns;
+	TArray<TSharedRef<Insights::FTableColumn>> Columns;
 	FTimersViewColumnFactory::CreateTimerTreeViewColumns(Columns);
 	ensure(Columns.Num() > 0 && Columns[0]->IsHierarchy());
 	Columns[0]->SetShortName(ViewName);
@@ -323,11 +323,11 @@ void STimerTreeView::InitializeAndShowHeaderColumns()
 	Table->SetColumns(Columns);
 
 	// Show columns.
-	for (const TSharedPtr<Insights::FTableColumn>& ColumnPtr : Table->GetColumns())
+	for (const TSharedRef<Insights::FTableColumn>& ColumnRef : Table->GetColumns())
 	{
-		if (ColumnPtr->ShouldBeVisible())
+		if (ColumnRef->ShouldBeVisible())
 		{
-			ShowColumn(ColumnPtr->GetId());
+			ShowColumn(ColumnRef->GetId());
 		}
 	}
 }
@@ -563,11 +563,11 @@ void STimerTreeView::CreateSortings()
 	AvailableSorters.Reset();
 	CurrentSorter = nullptr;
 
-	for (const TSharedPtr<Insights::FTableColumn> ColumnPtr : Table->GetColumns())
+	for (const TSharedRef<Insights::FTableColumn>& ColumnRef : Table->GetColumns())
 	{
-		if (ColumnPtr->CanBeSorted())
+		if (ColumnRef->CanBeSorted())
 		{
-			TSharedPtr<Insights::ITableCellValueSorter> SorterPtr = ColumnPtr->GetValueSorter();
+			TSharedPtr<Insights::ITableCellValueSorter> SorterPtr = ColumnRef->GetValueSorter();
 			if (ensure(SorterPtr.IsValid()))
 			{
 				AvailableSorters.Add(SorterPtr);
@@ -858,9 +858,9 @@ void STimerTreeView::ContextMenu_ShowAllColumns_Execute()
 	ColumnSortMode = GetDefaultColumnSortMode();
 	UpdateCurrentSortingByColumn();
 
-	for (const TSharedPtr<Insights::FTableColumn>& ColumnPtr : Table->GetColumns())
+	for (const TSharedRef<Insights::FTableColumn>& ColumnRef : Table->GetColumns())
 	{
-		const Insights::FTableColumn& Column = *ColumnPtr;
+		const Insights::FTableColumn& Column = *ColumnRef;
 
 		if (!Column.IsVisible())
 		{
@@ -886,9 +886,9 @@ void STimerTreeView::ContextMenu_ResetColumns_Execute()
 	ColumnSortMode = GetDefaultColumnSortMode();
 	UpdateCurrentSortingByColumn();
 
-	for (const TSharedPtr<Insights::FTableColumn>& ColumnPtr : Table->GetColumns())
+	for (const TSharedRef<Insights::FTableColumn>& ColumnRef : Table->GetColumns())
 	{
-		const Insights::FTableColumn& Column = *ColumnPtr;
+		const Insights::FTableColumn& Column = *ColumnRef;
 
 		if (Column.ShouldBeVisible() && !Column.IsVisible())
 		{
@@ -947,7 +947,7 @@ FTimerNodePtr STimerTreeView::CreateTimerNodeRec(const Trace::FTimingProfilerBut
 	if (Node.Timer == nullptr)
 	{
 		return nullptr;
-		//return MakeShareable(new FTimerNode(0, FName(TEXT("!!!!!")), FName(TEXT("[Group]")), ETimerNodeType::InvalidOrMax));
+		//return MakeShared<FTimerNode>(0, FName(TEXT("!!!!!")), FName(TEXT("[Group]")), ETimerNodeType::InvalidOrMax);
 	}
 
 	const uint64 Id = Node.Timer->Id;
@@ -955,7 +955,7 @@ FTimerNodePtr STimerTreeView::CreateTimerNodeRec(const Trace::FTimingProfilerBut
 	const FName Group(Node.Timer->IsGpuTimer ? TEXT("GPU") : TEXT("CPU"));
 	const ETimerNodeType Type = Node.Timer->IsGpuTimer ? ETimerNodeType::GpuScope : ETimerNodeType::CpuScope;
 
-	FTimerNodePtr TimerNodePtr = MakeShareable(new FTimerNode(Id, Name, Group, Type));
+	FTimerNodePtr TimerNodePtr = MakeShared<FTimerNode>(Id, Name, Group, Type);
 
 	Trace::FTimingProfilerAggregatedStats AggregatedStats;
 	AggregatedStats.InstanceCount = Node.Count;
diff --git a/Engine/Source/Developer/TraceInsights/Private/Insights/Widgets/STimersView.cpp b/Engine/Source/Developer/TraceInsights/Private/Insights/Widgets/STimersView.cpp
--- a/Engine/Source/Developer/TraceInsights/Private/Insights/Widgets/STimersView.cpp
+++ b/Engine/Source/Developer/TraceInsights/Private/Insights/Widgets/STimersView.cpp
@@ -33,7 +33,7 @@
 ////////////////////////////////////////////////////////////////////////////////////////////////////
 
 STimersView::STimersView()
-	: Table(MakeShareable(new Insights::FTable()))
+	: Table(MakeShared<Insights::FTable>())
 	, bExpansionSaved(false)
 	, bFilterOutZeroCountTimers(false)
 	, GroupingMode(ETimerGroupingMode::ByType)
@@ -242,8 +242,8 @@ void STimersView::Construct(const FArguments& InArgs)
 	//BindCommands();
 
 	// Create the search filters: text based, type based etc.
-	TextFilter = MakeShareable(new FTimerNodeTextFilter(FTimerNodeTextFilter::FItemToStringArray::CreateSP(this, &STimersView::HandleItemToStringArray)));
-	Filters = MakeShareable(new FTimerNodeFilterCollection());
+	TextFilter = MakeShared<FTimerNodeTextFilter>(FTimerNodeTextFilter::FItemToStringArray::CreateSP(this, &STimersView::HandleItemToStringArray));
+	Filters = MakeShared<FTimerNodeFilterCollection>();
 	Filters->Add(TextFilter);
 
 	CreateGroupByOptionsSources();
@@ -401,9 +401,9 @@ void STimersView::TreeView_BuildSortByMenu(FMenuBuilder& MenuBuilder)
 
 	MenuBuilder.BeginSection("ColumnName", LOCTEXT("ContextMenu_Header_Misc_ColumnName", "Column Name"));
 
-	for (const TSharedPtr<Insights::FTableColumn>& ColumnPtr : Table->GetColumns())
+	for (const TSharedRef<Insights::FTableColumn>& ColumnRef : Table->GetColumns())
 	{
-		const Insights::FTableColumn& Column = *ColumnPtr;
+		const Insights::FTableColumn& Column = *ColumnRef;
 
 		if (Column.IsVisible() && Column.CanBeSorted())
 		{
@@ -463,9 +463,9 @@ void STimersView::TreeView_BuildViewColumnMenu(FMenuBuilder& MenuBuilder)
 {
 	MenuBuilder.BeginSection("ViewColumn", LOCTEXT("ContextMenu_Header_Columns_View", "View Column"));
 
-	for (const TSharedPtr<Insights::FTableColumn>& ColumnPtr : Table->GetColumns())
+	for (const TSharedRef<Insights::FTableColumn>& ColumnRef : Table->GetColumns())
 	{
-		const Insights::FTableColumn& Column = *ColumnPtr;
+		const Insights::FTableColumn& Column = *ColumnRef;
 
 		FUIAction Action_ToggleColumn
 		(
@@ -489,16 +489,16 @@ void STimersView::TreeView_BuildViewColumnMenu(FMenuBuilder& MenuBuilder)
 void STimersView::InitializeAndShowHeaderColumns()
 {
 	// Create columns.
-	TArray<TSharedPtr<Insights::FTableColumn>> Columns;
+	TArray<TSharedRef<Insights::FTableColumn>> Columns;
 	FTimersViewColumnFactory::CreateTimersViewColumns(Columns);
 	Table->SetColumns(Columns);
 
 	// Show columns.
-	for (const TSharedPtr<Insights::FTableColumn>& ColumnPtr : Table->GetColumns())
+	for (const TSharedRef<Insights::FTableColumn>& ColumnRef : Table->GetColumns())
 	{
-		if (ColumnPtr->ShouldBeVisible())
+		if (ColumnRef->ShouldBeVisible())
 		{
-			ShowColumn(ColumnPtr->GetId());
+			ShowColumn(ColumnRef->GetId());
 		}
 	}
 }
@@ -936,7 +936,7 @@ void STimersView::CreateGroups()
 		FTimerNodePtr* GroupPtr = GroupNodeSet.Find(GroupName);
 		if (!GroupPtr)
 		{
-			GroupPtr = &GroupNodeSet.Add(GroupName, MakeShareable(new FTimerNode(GroupName)));
+			GroupPtr = &GroupNodeSet.Add(GroupName, MakeShared<FTimerNode>(GroupName));
 		}
 
 		for (const FTimerNodePtr& TimerNodePtr : TimerNodes)
@@ -956,7 +956,7 @@ void STimersView::CreateGroups()
 			FTimerNodePtr* GroupPtr = GroupNodeSet.Find(GroupName);
 			if (!GroupPtr)
 			{
-				GroupPtr = &GroupNodeSet.Add(GroupName, MakeShareable(new FTimerNode(GroupName)));
+				GroupPtr = &GroupNodeSet.Add(GroupName, MakeShared<FTimerNode>(GroupName));
 			}
 
 			(*GroupPtr)->AddChildAndSetGroupPtr(TimerNodePtr);
@@ -973,7 +973,7 @@ void STimersView::CreateGroups()
 			FTimerNodePtr* GroupPtr = GroupNodeSet.Find(GroupName);
 			if (!GroupPtr)
 			{
-				GroupPtr = &GroupNodeSet.Add(GroupName, MakeShareable(new FTimerNode(GroupName)));
+				GroupPtr = &GroupNodeSet.Add(GroupName, MakeShared<FTimerNode>(GroupName));
 			}
 
 			(*GroupPtr)->AddChildAndSetGroupPtr(TimerNodePtr);
@@ -990,7 +990,7 @@ void STimersView::CreateGroups()
 			FTimerNodePtr* GroupPtr = GroupNodeSet.Find(GroupName);
 			if (!GroupPtr)
 			{
-				GroupPtr = &GroupNodeSet.Add(GroupName, MakeShareable(new FTimerNode(GroupName)));
+				GroupPtr = &GroupNodeSet.Add(GroupName, MakeShared<FTimerNode>(GroupName));
 			}
 
 			(*GroupPtr)->AddChildAndSetGroupPtr(TimerNodePtr);
@@ -1022,10 +1022,10 @@ void STimersView::CreateGroupByOptionsSources()
 	GroupByOptionsSource.Reset(3);
 
 	// Must be added in order of elements in the ETimerGroupingMode.
-	GroupByOptionsSource.Add(MakeShareable(new ETimerGroupingMode(ETimerGroupingMode::Flat)));
-	GroupByOptionsSource.Add(MakeShareable(new ETimerGroupingMode(ETimerGroupingMode::ByName)));
-	//GroupByOptionsSource.Add(MakeShareable(new ETimerGroupingMode(ETimerGroupingMode::ByMetaGroupName)));
-	GroupByOptionsSource.Add(MakeShareable(new ETimerGroupingMode(ETimerGroupingMode::ByType)));
+	GroupByOptionsSource.Add(MakeShared<ETimerGroupingMode>(ETimerGroupingMode::Flat));
+	GroupByOptionsSource.Add(MakeShared<ETimerGroupingMode>(ETimerGroupingMode::ByName));
+	//GroupByOptionsSource.Add(MakeShared<ETimerGroupingMode>(ETimerGroupingMode::ByMetaGroupName));
+	GroupByOptionsSource.Add(MakeShared<ETimerGroupingMode>(ETimerGroupingMode::ByType));
 
 	ETimerGroupingModePtr* GroupingModePtrPtr = GroupByOptionsSource.FindByPredicate([&](const ETimerGroupingModePtr InGroupingModePtr) { return *InGroupingModePtr == GroupingMode; });
 	if (GroupingModePtrPtr != nullptr)
@@ -1096,11 +1096,11 @@ void STimersView::CreateSortings()
 	AvailableSorters.Reset();
 	CurrentSorter = nullptr;
 
-	for (const TSharedPtr<Insights::FTableColumn> ColumnPtr : Table->GetColumns())
+	for (const TSharedRef<Insights::FTableColumn>& ColumnRef : Table->GetColumns())
 	{
-		if (ColumnPtr->CanBeSorted())
+		if (ColumnRef->CanBeSorted())
 		{
-			TSharedPtr<Insights::ITableCellValueSorter> SorterPtr = ColumnPtr->GetValueSorter();
+			TSharedPtr<Insights::ITableCellValueSorter> SorterPtr = ColumnRef->GetValueSorter();
 			if (ensure(SorterPtr.IsValid()))
 			{
 				AvailableSorters.Add(SorterPtr);
@@ -1391,9 +1391,9 @@ void STimersView::ContextMenu_ShowAllColumns_Execute()
 	ColumnSortMode = GetDefaultColumnSortMode();
 	UpdateCurrentSortingByColumn();
 
-	for (const TSharedPtr<Insights::FTableColumn>& ColumnPtr : Table->GetColumns())
+	for (const TSharedRef<Insights::FTableColumn>& ColumnRef : Table->GetColumns())
 	{
-		const Insights::FTableColumn& Column = *ColumnPtr;
+		const Insights::FTableColumn& Column = *ColumnRef;
 
 		if (!Column.IsVisible())
 		{
@@ -1437,9 +1437,9 @@ void STimersView::ContextMenu_ShowMinMaxMedColumns_Execute()
 	ColumnSortMode = EColumnSortMode::Descending;
 	UpdateCurrentSortingByColumn();
 
-	for (const TSharedPtr<Insights::FTableColumn>& ColumnPtr : Table->GetColumns())
+	for (const TSharedRef<Insights::FTableColumn>& ColumnRef : Table->GetColumns())
 	{
-		const Insights::FTableColumn& Column = *ColumnPtr;
+		const Insights::FTableColumn& Column = *ColumnRef;
 
 		const bool bShouldBeVisible = Preset.Contains(Column.GetId());
 
@@ -1471,9 +1471,9 @@ void STimersView::ContextMenu_ResetColumns_Execute()
 	ColumnSortMode = GetDefaultColumnSortMode();
 	UpdateCurrentSortingByColumn();
 
-	for (const TSharedPtr<Insights::FTableColumn>& ColumnPtr : Table->GetColumns())
+	for (const TSharedRef<Insights::FTableColumn>& ColumnRef : Table->GetColumns())
 	{
-		const Insights::FTableColumn& Column = *ColumnPtr;
+		const Insights::FTableColumn& Column = *ColumnRef;
 
 		if (Column.ShouldBeVisible() && !Column.IsVisible())
 		{
@@ -1540,7 +1540,7 @@ void STimersView::RebuildTree(bool bResync)
 					FName Name(Timer.Name);// +TEXT(" [GPU]")));
 					FName Group(Timer.IsGpuTimer ? TEXT("GPU") : TEXT("CPU"));
 					ETimerNodeType Type = Timer.IsGpuTimer ? ETimerNodeType::GpuScope : ETimerNodeType::CpuScope;
-					FTimerNodePtr TimerNodePtr = MakeShareable(new FTimerNode(Timer.Id, Name, Group, Type));
+					FTimerNodePtr TimerNodePtr = MakeShared<FTimerNode>(Timer.Id, Name, Group, Type);
 					TimerNodes.Add(TimerNodePtr);
 					TimerNodesIdMap.Add(Timer.Id, TimerNodePtr);
 				}
diff --git a/Engine/Source/Developer/TraceInsights/Private/Insights/Widgets/STimingView.cpp b/Engine/Source/Developer/TraceInsights/Private/Insights/Widgets/STimingView.cpp
--- a/Engine/Source/Developer/TraceInsights/Private/Insights/Widgets/STimingView.cpp
+++ b/Engine/Source/Developer/TraceInsights/Private/Insights/Widgets/STimingView.cpp
@@ -72,13 +72,13 @@ namespace Insights { const FName TimingViewExtenderFeatureName(TEXT("TimingViewE
 ////////////////////////////////////////////////////////////////////////////////////////////////////
 
 STimingView::STimingView()
-	: ThreadTimingSharedState(MakeShareable(new FThreadTimingSharedState(this)))
-	, LoadingSharedState(MakeShareable(new FLoadingSharedState(this)))
+	: ThreadTimingSharedState(MakeShared<FThreadTimingSharedState>(this))
+	, LoadingSharedState(MakeShared<FLoadingSharedState>(this))
 	, bAssetLoadingMode(false)
-	, FileActivitySharedState(MakeShareable(new FFileActivitySharedState(this)))
-	, TimeRulerTrack(MakeShareable(new FTimeRulerTrack()))
-	, MarkersTrack(MakeShareable(new FMarkersTimingTrack()))
-	, GraphTrack(MakeShareable(new FTimingGraphTrack()))
+	, FileActivitySharedState(MakeShared<FFileActivitySharedState>(this))
+	, TimeRulerTrack(MakeShared<FTimeRulerTrack>())
+	, MarkersTrack(MakeShared<FMarkersTimingTrack>())
+	, GraphTrack(MakeShared<FTimingGraphTrack>())
 	, WhiteBrush(FInsightsStyle::Get().GetBrush("WhiteBrush"))
 	, MainFont(FCoreStyle::GetDefaultFontStyle("Regular", 8))
 {
