Implemented global track ordering calbacks in existing track types

Exposed the ability to flag for a re-sort of scrollable tracks when ordering changes

diff --git a/Engine/Source/Developer/TraceInsights/Private/Insights/ViewModels/FileActivityTimingTrack.cpp b/Engine/Source/Developer/TraceInsights/Private/Insights/ViewModels/FileActivityTimingTrack.cpp
--- a/Engine/Source/Developer/TraceInsights/Private/Insights/ViewModels/FileActivityTimingTrack.cpp
+++ b/Engine/Source/Developer/TraceInsights/Private/Insights/ViewModels/FileActivityTimingTrack.cpp
@@ -471,16 +471,6 @@ void FFileActivitySharedState::ExtendFilterMenu(Insights::ITimingViewSession& In
 
 ////////////////////////////////////////////////////////////////////////////////////////////////////
 
-void FFileActivitySharedState::OnTracksChanged(Insights::ITimingViewSession& InSession, int32& InOutOrder)
-{
-	if (&InSession != TimingView)
-	{
-		return;
-	}
-}
-
-////////////////////////////////////////////////////////////////////////////////////////////////////
-
 bool FFileActivitySharedState::ShowHideAllIoTracks_IsChecked() const
 {
 	return bShowHideAllIoTracks;
diff --git a/Engine/Source/Developer/TraceInsights/Private/Insights/ViewModels/FileActivityTimingTrack.h b/Engine/Source/Developer/TraceInsights/Private/Insights/ViewModels/FileActivityTimingTrack.h
--- a/Engine/Source/Developer/TraceInsights/Private/Insights/ViewModels/FileActivityTimingTrack.h
+++ b/Engine/Source/Developer/TraceInsights/Private/Insights/ViewModels/FileActivityTimingTrack.h
@@ -49,7 +49,6 @@ public:
 	virtual void OnEndSession(Insights::ITimingViewSession& InSession) override;
 	virtual void Tick(Insights::ITimingViewSession& InSession, const Trace::IAnalysisSession& InAnalysisSession) override;
 	virtual void ExtendFilterMenu(Insights::ITimingViewSession& InSession, FMenuBuilder& InMenuBuilder) override;
-	virtual void OnTracksChanged(Insights::ITimingViewSession& InSession, int32& InOutOrder) override;
 
 	const TArray<FIoTimingEvent>& GetAllEvents() const { return AllIoEvents; }
 
diff --git a/Engine/Source/Developer/TraceInsights/Private/Insights/ViewModels/LoadingTimingTrack.cpp b/Engine/Source/Developer/TraceInsights/Private/Insights/ViewModels/LoadingTimingTrack.cpp
--- a/Engine/Source/Developer/TraceInsights/Private/Insights/ViewModels/LoadingTimingTrack.cpp
+++ b/Engine/Source/Developer/TraceInsights/Private/Insights/ViewModels/LoadingTimingTrack.cpp
@@ -152,16 +152,6 @@ void FLoadingSharedState::ShowHideAllLoadingTracks_Execute()
 
 ////////////////////////////////////////////////////////////////////////////////////////////////////
 
-void FLoadingSharedState::OnTracksChanged(Insights::ITimingViewSession& InSession, int32& InOutOrder)
-{
-	if (&InSession != TimingView)
-	{
-		return;
-	}
-}
-
-////////////////////////////////////////////////////////////////////////////////////////////////////
-
 const TCHAR* FLoadingSharedState::GetEventNameByEventType(uint32 Depth, const Trace::FLoadTimeProfilerCpuEvent& Event) const
 {
 	return Trace::GetLoadTimeProfilerObjectEventTypeString(Event.EventType);
diff --git a/Engine/Source/Developer/TraceInsights/Private/Insights/ViewModels/LoadingTimingTrack.h b/Engine/Source/Developer/TraceInsights/Private/Insights/ViewModels/LoadingTimingTrack.h
--- a/Engine/Source/Developer/TraceInsights/Private/Insights/ViewModels/LoadingTimingTrack.h
+++ b/Engine/Source/Developer/TraceInsights/Private/Insights/ViewModels/LoadingTimingTrack.h
@@ -29,7 +29,6 @@ public:
 	virtual void OnEndSession(Insights::ITimingViewSession& InSession) override;
 	virtual void Tick(Insights::ITimingViewSession& InSession, const Trace::IAnalysisSession& InAnalysisSession) override;
 	virtual void ExtendFilterMenu(Insights::ITimingViewSession& InSession, FMenuBuilder& InMenuBuilder) override;
-	virtual void OnTracksChanged(Insights::ITimingViewSession& InSession, int32& InOutOrder) override;
 
 	const TCHAR* GetEventName(uint32 Depth, const Trace::FLoadTimeProfilerCpuEvent& Event) const;
 	void SetColorSchema(int32 Schema);
diff --git a/Engine/Source/Developer/TraceInsights/Private/Insights/ViewModels/ThreadTimingTrack.cpp b/Engine/Source/Developer/TraceInsights/Private/Insights/ViewModels/ThreadTimingTrack.cpp
--- a/Engine/Source/Developer/TraceInsights/Private/Insights/ViewModels/ThreadTimingTrack.cpp
+++ b/Engine/Source/Developer/TraceInsights/Private/Insights/ViewModels/ThreadTimingTrack.cpp
@@ -197,9 +197,9 @@ void FThreadTimingSharedState::Tick(Insights::ITimingViewSession& InSession, con
 			Order += 100;
 		});
 
-		if (bTracksOrderChanged && TimingView)
+		if (bTracksOrderChanged)
 		{
-			TimingView->UpdateScrollableTracksOrder();
+			InSession.InvalidateScrollableTracksOrder();
 		}
 	}
 }
@@ -370,16 +370,6 @@ void FThreadTimingSharedState::ToggleTrackVisibilityByGroup_Execute(const TCHAR*
 	}
 }
 
-////////////////////////////////////////////////////////////////////////////////////////////////////
-
-void FThreadTimingSharedState::OnTracksChanged(Insights::ITimingViewSession& InSession, int32& InOutOrder)
-{
-	if (&InSession != TimingView)
-	{
-		return;
-	}
-}
-
 ////////////////////////////////////////////////////////////////////////////////////////////////////
 // FThreadTimingTrack
 ////////////////////////////////////////////////////////////////////////////////////////////////////
diff --git a/Engine/Source/Developer/TraceInsights/Private/Insights/ViewModels/ThreadTimingTrack.h b/Engine/Source/Developer/TraceInsights/Private/Insights/ViewModels/ThreadTimingTrack.h
--- a/Engine/Source/Developer/TraceInsights/Private/Insights/ViewModels/ThreadTimingTrack.h
+++ b/Engine/Source/Developer/TraceInsights/Private/Insights/ViewModels/ThreadTimingTrack.h
@@ -50,7 +50,6 @@ public:
 	virtual void OnEndSession(Insights::ITimingViewSession& InSession) override;
 	virtual void Tick(Insights::ITimingViewSession& InSession, const Trace::IAnalysisSession& InAnalysisSession) override;
 	virtual void ExtendFilterMenu(Insights::ITimingViewSession& InSession, FMenuBuilder& InMenuBuilder) override;
-	virtual void OnTracksChanged(Insights::ITimingViewSession& InSession, int32& InOutOrder) override;
 
 	//////////////////////////////////////////////////
 
diff --git a/Engine/Source/Developer/TraceInsights/Private/Insights/ViewModels/TimingViewDrawHelper.h b/Engine/Source/Developer/TraceInsights/Private/Insights/ViewModels/TimingViewDrawHelper.h
--- a/Engine/Source/Developer/TraceInsights/Private/Insights/ViewModels/TimingViewDrawHelper.h
+++ b/Engine/Source/Developer/TraceInsights/Private/Insights/ViewModels/TimingViewDrawHelper.h
@@ -110,7 +110,7 @@ public:
 
 	const FTimingTrackViewport& GetViewport() const { return Viewport; }
 
-	virtual void AddEvent(double InEventStartTime, double InEventEndTime, uint32 InEventDepth, const TCHAR* InEventName, uint64 InEventType, uint32 InEventColor) override;
+	virtual void AddEvent(double InEventStartTime, double InEventEndTime, uint32 InEventDepth, const TCHAR* InEventName, uint64 InEventType = 0, uint32 InEventColor = 0) override;
 	void Flush();
 
 	int32 GetMaxDepth() const { return MaxDepth; }
diff --git a/Engine/Source/Developer/TraceInsights/Private/Insights/Widgets/STimingView.cpp b/Engine/Source/Developer/TraceInsights/Private/Insights/Widgets/STimingView.cpp
--- a/Engine/Source/Developer/TraceInsights/Private/Insights/Widgets/STimingView.cpp
+++ b/Engine/Source/Developer/TraceInsights/Private/Insights/Widgets/STimingView.cpp
@@ -72,7 +72,8 @@ namespace Insights { const FName TimingViewExtenderFeatureName(TEXT("TimingViewE
 ////////////////////////////////////////////////////////////////////////////////////////////////////
 
 STimingView::STimingView()
-	: ThreadTimingSharedState(MakeShared<FThreadTimingSharedState>(this))
+	: bScrollableTracksDirty(false)
+	, ThreadTimingSharedState(MakeShared<FThreadTimingSharedState>(this))
 	, LoadingSharedState(MakeShared<FLoadingSharedState>(this))
 	, bAssetLoadingMode(false)
 	, FileActivitySharedState(MakeShared<FFileActivitySharedState>(this))
@@ -395,6 +396,9 @@ void STimingView::Tick(const FGeometry& AllottedGeometry, const double InCurrent
 		{
 			Extender->Tick(*this, *Session.Get());
 		}
+
+		// Re-sort now if we need to
+		UpdateScrollableTracksOrder();
 	}
 
 	////////////////////////////////////////////////////////////////////////////////////////////////////
@@ -1233,14 +1237,25 @@ void STimingView::AddScrollableTrack(TSharedPtr<FBaseTimingTrack> Track)
 	check(!AllTracks.Contains(Track->GetId()));
 	AllTracks.Add(Track->GetId(), Track);
 	ScrollableTracks.Add(Track);
-	UpdateScrollableTracksOrder();
+	InvalidateScrollableTracksOrder();
+}
+
+////////////////////////////////////////////////////////////////////////////////////////////////////
+
+void STimingView::InvalidateScrollableTracksOrder()
+{
+	bScrollableTracksDirty = true;
 }
 
 ////////////////////////////////////////////////////////////////////////////////////////////////////
 
 void STimingView::UpdateScrollableTracksOrder()
 {
-	Algo::SortBy(ScrollableTracks, &FBaseTimingTrack::GetOrder);
+	if(bScrollableTracksDirty)
+	{
+		Algo::SortBy(ScrollableTracks, &FBaseTimingTrack::GetOrder);
+		bScrollableTracksDirty = false;
+	}
 }
 
 ////////////////////////////////////////////////////////////////////////////////////////////////////
diff --git a/Engine/Source/Developer/TraceInsights/Private/Insights/Widgets/STimingView.h b/Engine/Source/Developer/TraceInsights/Private/Insights/Widgets/STimingView.h
--- a/Engine/Source/Developer/TraceInsights/Private/Insights/Widgets/STimingView.h
+++ b/Engine/Source/Developer/TraceInsights/Private/Insights/Widgets/STimingView.h
@@ -228,6 +228,8 @@ public:
 	virtual void AddScrollableTrack(TSharedPtr<FBaseTimingTrack> Track) override;
 	virtual void AddForegroundTrack(TSharedPtr<FBaseTimingTrack> Track) override;
 
+	virtual void InvalidateScrollableTracksOrder() override;
+
 	virtual Insights::FSelectionChangedDelegate& OnSelectionChanged() override { return OnSelectionChangedDelegate; }
 	virtual Insights::FTimeMarkerChangedDelegate& OnTimeMarkerChanged() override { return OnTimeMarkerChangedDelegate; }
 	virtual Insights::FHoveredTrackChangedDelegate& OnHoveredTrackChanged() override { return OnHoveredTrackChangedDelegate; }
@@ -381,6 +383,9 @@ protected:
 	TArray<TSharedPtr<FBaseTimingTrack>> ScrollableTracks; /**< tracks in scrollable area, in the order to be displayed (top to bottom) */
 	TArray<TSharedPtr<FBaseTimingTrack>> ForegroundTracks; /**< tracks to draw over top/scrollable/bottom tracks (can use entire area), in the order to be displayed (back to front) */
 
+	/** Whether the scrollable tracks are dirty and need to be re-sorted */
+	bool bScrollableTracksDirty;
+
 	////////////////////////////////////////////////////////////
 
 	// Shared state for Cpu/Gpu Thread tracks
diff --git a/Engine/Source/Developer/TraceInsights/Public/Insights/ITimingViewExtender.h b/Engine/Source/Developer/TraceInsights/Public/Insights/ITimingViewExtender.h
--- a/Engine/Source/Developer/TraceInsights/Public/Insights/ITimingViewExtender.h
+++ b/Engine/Source/Developer/TraceInsights/Public/Insights/ITimingViewExtender.h
@@ -30,9 +30,6 @@ public:
 
 	/** Extension hook for the 'quick filter' menu */
 	virtual void ExtendFilterMenu(ITimingViewSession& InSession, FMenuBuilder& InMenuBuilder) = 0;
-
-	/** Called when tracks have changed and need to be re-ordered */
-	virtual void OnTracksChanged(ITimingViewSession& InSession, int32& InOutOrder) = 0;
 };
 
 } // namespace Insights
diff --git a/Engine/Source/Developer/TraceInsights/Public/Insights/ITimingViewSession.h b/Engine/Source/Developer/TraceInsights/Public/Insights/ITimingViewSession.h
--- a/Engine/Source/Developer/TraceInsights/Public/Insights/ITimingViewSession.h
+++ b/Engine/Source/Developer/TraceInsights/Public/Insights/ITimingViewSession.h
@@ -53,6 +53,9 @@ public:
 	/** Adds a new foreground track. */
 	virtual void AddForegroundTrack(TSharedPtr<FBaseTimingTrack> Track) = 0;
 
+	/** Marks the scrollable tracks as dirty so they will be re-sorted */
+	virtual void InvalidateScrollableTracksOrder() = 0;
+
 	/** Finds a track has been added via Add*Track(). */
 	virtual TSharedPtr<FBaseTimingTrack> FindTrack(uint64 InTrackId) = 0;
 
diff --git a/Engine/Source/Developer/TraceInsights/Public/Insights/ViewModels/TimingEventsTrack.h b/Engine/Source/Developer/TraceInsights/Public/Insights/ViewModels/TimingEventsTrack.h
--- a/Engine/Source/Developer/TraceInsights/Public/Insights/ViewModels/TimingEventsTrack.h
+++ b/Engine/Source/Developer/TraceInsights/Public/Insights/ViewModels/TimingEventsTrack.h
@@ -14,7 +14,7 @@ class TRACEINSIGHTS_API ITimingEventsTrackDrawStateBuilder
 public:
 	virtual ~ITimingEventsTrackDrawStateBuilder() = default;
 
-	virtual void AddEvent(double InEventStartTime, double InEventEndTime, uint32 InEventDepth, const TCHAR* InEventName, uint64 InEventType, uint32 InEventColor) = 0;
+	virtual void AddEvent(double InEventStartTime, double InEventEndTime, uint32 InEventDepth, const TCHAR* InEventName, uint64 InEventType = 0, uint32 InEventColor = 0) = 0;
 };
 
 ////////////////////////////////////////////////////////////////////////////////////////////////////
