Track title area drawing changes

Exposed ITimingViewDrawHelper publicly again, but only exposed accessors for usful constants (and track selection/hover colors)

diff --git a/Engine/Source/Developer/TraceInsights/Private/Insights/ViewModels/ThreadTimingTrack.cpp b/Engine/Source/Developer/TraceInsights/Private/Insights/ViewModels/ThreadTimingTrack.cpp
--- a/Engine/Source/Developer/TraceInsights/Private/Insights/ViewModels/ThreadTimingTrack.cpp
+++ b/Engine/Source/Developer/TraceInsights/Private/Insights/ViewModels/ThreadTimingTrack.cpp
@@ -425,7 +425,7 @@ void FThreadTimingTrack::PostDraw(const ITimingTrackDrawContext& Context) const
 		FTimingEvent::CheckTypeName(*SelectedEventPtr))
 	{
 		const FTimingEvent& SelectedEvent = static_cast<const FTimingEvent&>(*SelectedEventPtr);
-		const FTimingViewDrawHelper& Helper = Context.GetHelper();
+		const ITimingViewDrawHelper& Helper = Context.GetHelper();
 		DrawSelectedEventInfo(SelectedEvent, Context.GetViewport(), Context.GetDrawContext(), Helper.GetWhiteBrush(), Helper.GetEventFont());
 	}
 }
diff --git a/Engine/Source/Developer/TraceInsights/Private/Insights/ViewModels/TimingEventsTrack.cpp b/Engine/Source/Developer/TraceInsights/Private/Insights/ViewModels/TimingEventsTrack.cpp
--- a/Engine/Source/Developer/TraceInsights/Private/Insights/ViewModels/TimingEventsTrack.cpp
+++ b/Engine/Source/Developer/TraceInsights/Private/Insights/ViewModels/TimingEventsTrack.cpp
@@ -123,7 +123,7 @@ void FTimingEventsTrack::PostUpdate(const ITimingTrackUpdateContext& Context)
 
 void FTimingEventsTrack::Draw(const ITimingTrackDrawContext& Context) const
 {
-	const FTimingViewDrawHelper& Helper = Context.GetHelper();
+	const FTimingViewDrawHelper& Helper = *static_cast<const FTimingViewDrawHelper*>(&Context.GetHelper());
 	Helper.DrawEvents(*DrawState, *this);
 	Helper.DrawTrackHeader(*this);
 }
@@ -132,7 +132,7 @@ void FTimingEventsTrack::Draw(const ITimingTrackDrawContext& Context) const
 
 void FTimingEventsTrack::DrawEvents(const ITimingTrackDrawContext& Context, const float OffsetY) const
 {
-	const FTimingViewDrawHelper& Helper = Context.GetHelper();
+	const FTimingViewDrawHelper& Helper = *static_cast<const FTimingViewDrawHelper*>(&Context.GetHelper());
 	Helper.DrawEvents(*DrawState, *this, OffsetY);
 }
 
@@ -140,7 +140,7 @@ void FTimingEventsTrack::DrawEvents(const ITimingTrackDrawContext& Context, cons
 
 int32 FTimingEventsTrack::GetHeaderBackgroundLayerId(const ITimingTrackDrawContext& Context) const
 {
-	const FTimingViewDrawHelper& Helper = Context.GetHelper();
+	const FTimingViewDrawHelper& Helper = *static_cast<const FTimingViewDrawHelper*>(&Context.GetHelper());
 	return Helper.GetHeaderBackgroundLayerId();
 }
 
@@ -148,7 +148,7 @@ int32 FTimingEventsTrack::GetHeaderBackgroundLayerId(const ITimingTrackDrawConte
 
 int32 FTimingEventsTrack::GetHeaderTextLayerId(const ITimingTrackDrawContext& Context) const
 {
-	const FTimingViewDrawHelper& Helper = Context.GetHelper();
+	const FTimingViewDrawHelper& Helper = *static_cast<const FTimingViewDrawHelper*>(&Context.GetHelper());
 	return Helper.GetHeaderTextLayerId();
 }
 
@@ -156,7 +156,7 @@ int32 FTimingEventsTrack::GetHeaderTextLayerId(const ITimingTrackDrawContext& Co
 
 void FTimingEventsTrack::DrawHeader(const ITimingTrackDrawContext& Context) const
 {
-	const FTimingViewDrawHelper& Helper = Context.GetHelper();
+	const FTimingViewDrawHelper& Helper = *static_cast<const FTimingViewDrawHelper*>(&Context.GetHelper());
 	Helper.DrawTrackHeader(*this);
 }
 
@@ -169,7 +169,9 @@ void FTimingEventsTrack::DrawEvent(const ITimingTrackDrawContext& Context, const
 		const FTimingEvent& TrackEvent = static_cast<const FTimingEvent&>(InTimingEvent);
 		const FTimingViewLayout& Layout = Context.GetViewport().GetLayout();
 		const float Y = TrackEvent.GetTrack()->GetPosY() + Layout.GetLaneY(TrackEvent.GetDepth());
-		Context.GetHelper().DrawTimingEventHighlight(TrackEvent.GetStartTime(), TrackEvent.GetEndTime(), Y, InDrawMode);
+
+		const FTimingViewDrawHelper& Helper = *static_cast<const FTimingViewDrawHelper*>(&Context.GetHelper());
+		Helper.DrawTimingEventHighlight(TrackEvent.GetStartTime(), TrackEvent.GetEndTime(), Y, InDrawMode);
 	}
 }
 
diff --git a/Engine/Source/Developer/TraceInsights/Private/Insights/ViewModels/TimingViewDrawHelper.cpp b/Engine/Source/Developer/TraceInsights/Private/Insights/ViewModels/TimingViewDrawHelper.cpp
--- a/Engine/Source/Developer/TraceInsights/Private/Insights/ViewModels/TimingViewDrawHelper.cpp
+++ b/Engine/Source/Developer/TraceInsights/Private/Insights/ViewModels/TimingViewDrawHelper.cpp
@@ -442,9 +442,7 @@ void FTimingViewDrawHelper::DrawTrackHeader(const FTimingEventsTrack& Track) con
 			const float HeaderH = FMath::Min(12.0f, Track.GetHeight() - 1.0f);
 			DrawContext.DrawBox(HeaderLayerId, 0.0f, TrackY + 1.0f, NameWidth + 4.0f, HeaderH, WhiteBrush, EdgeColor);
 
-			const FLinearColor TextColor = Track.IsHovered() ?  FLinearColor(1.0f, 1.0f, 0.0f, 1.0f) :
-			                               Track.IsSelected() ? FLinearColor(1.0f, 1.0f, 0.5f, 1.0f) :
-			                                                    FLinearColor(1.0f, 1.0f, 1.0f, 1.0f);
+			const FLinearColor TextColor = GetTrackNameTextColor(Track);
 
 			const float TextY = TrackY + FMath::Min(0.0f, (Track.GetHeight() - 13.0f) / 2.0f);
 			const int32 HeaderTextLayerId = ReservedLayerId + ToInt32(EDrawLayer::HeaderText);
@@ -565,3 +563,12 @@ void FTimingViewDrawHelper::DrawTimingEventHighlight(double StartTime, double En
 }
 
 ////////////////////////////////////////////////////////////////////////////////////////////////////
+
+FLinearColor FTimingViewDrawHelper::GetTrackNameTextColor(const FTimingEventsTrack& Track) const
+{
+	return  Track.IsHovered() ?  FLinearColor(1.0f, 1.0f, 0.0f, 1.0f) :
+			Track.IsSelected() ? FLinearColor(1.0f, 1.0f, 0.5f, 1.0f) :
+			                     FLinearColor(1.0f, 1.0f, 1.0f, 1.0f);
+}
+
+////////////////////////////////////////////////////////////////////////////////////////////////////
diff --git a/Engine/Source/Developer/TraceInsights/Private/Insights/ViewModels/TimingViewDrawHelper.h b/Engine/Source/Developer/TraceInsights/Private/Insights/ViewModels/TimingViewDrawHelper.h
--- a/Engine/Source/Developer/TraceInsights/Private/Insights/ViewModels/TimingViewDrawHelper.h
+++ b/Engine/Source/Developer/TraceInsights/Private/Insights/ViewModels/TimingViewDrawHelper.h
@@ -8,6 +8,7 @@
 
 // Insights
 #include "Insights/ViewModels/TimingEventsTrack.h"
+#include "Insights/ViewModels/ITimingViewDrawHelper.h"
 
 ////////////////////////////////////////////////////////////////////////////////////////////////////
 
@@ -133,7 +134,7 @@ private:
 
 ////////////////////////////////////////////////////////////////////////////////////////////////////
 
-class FTimingViewDrawHelper final
+class FTimingViewDrawHelper final : public ITimingViewDrawHelper
 {
 private:
 	enum class EDrawLayer : int32
@@ -159,19 +160,22 @@ public:
 	FTimingViewDrawHelper(const FTimingViewDrawHelper&) = delete;
 	FTimingViewDrawHelper& operator=(const FTimingViewDrawHelper&) = delete;
 
+	// ITimingViewDrawHelper interface
+	virtual const FSlateBrush* GetWhiteBrush() const override { return WhiteBrush; }
+	virtual const FSlateFontInfo& GetEventFont() const override { return EventFont; }
+	virtual FLinearColor GetEdgeColor() const override { return EdgeColor; }
+	virtual FLinearColor GetTrackNameTextColor(const FTimingEventsTrack& Track) const override;
+	virtual int32 GetHeaderBackgroundLayerId() const override { return ReservedLayerId + ToInt32(EDrawLayer::HeaderBackground); }
+	virtual int32 GetHeaderTextLayerId() const override { return ReservedLayerId + ToInt32(EDrawLayer::HeaderText); }
+
 	const FDrawContext& GetDrawContext() const { return DrawContext; }
 	const FTimingTrackViewport& GetViewport() const { return Viewport; }
 
-	const FSlateBrush* GetWhiteBrush() const { return WhiteBrush; }
-	const FSlateFontInfo& GetEventFont() const { return EventFont; }
-
 	void DrawBackground() const;
 
 	void BeginDrawTracks() const;
 	// OffsetY = 1.0f is for the top horizontal line (which separates the timelines) added by DrawTrackHeader.
 	void DrawEvents(const FTimingEventsTrackDrawState& DrawState, const FTimingEventsTrack& Track, const float OffsetY = 1.0f) const;
-	int32 GetHeaderBackgroundLayerId() const { return ReservedLayerId + ToInt32(EDrawLayer::HeaderBackground); }
-	int32 GetHeaderTextLayerId() const { return ReservedLayerId + ToInt32(EDrawLayer::HeaderText); }
 	void DrawTrackHeader(const FTimingEventsTrack& Track) const;
 	void EndDrawTracks() const;
 
diff --git a/Engine/Source/Developer/TraceInsights/Private/Insights/Widgets/STimingView.cpp b/Engine/Source/Developer/TraceInsights/Private/Insights/Widgets/STimingView.cpp
--- a/Engine/Source/Developer/TraceInsights/Private/Insights/Widgets/STimingView.cpp
+++ b/Engine/Source/Developer/TraceInsights/Private/Insights/Widgets/STimingView.cpp
@@ -811,7 +811,7 @@ int32 STimingView::OnPaint(const FPaintArgs& Args, const FGeometry& AllottedGeom
 		virtual const TSharedPtr<const ITimingEvent> GetHoveredEvent() const override { return TimingView->GetHoveredEvent(); }
 		virtual const TSharedPtr<const ITimingEvent> GetSelectedEvent() const override { return TimingView->GetSelectedEvent(); }
 		virtual FDrawContext& GetDrawContext() const override { return DrawContext; }
-		virtual const FTimingViewDrawHelper& GetHelper() const override { return Helper; }
+		virtual const ITimingViewDrawHelper& GetHelper() const override { return Helper; }
 
 	public:
 		const STimingView* TimingView;
diff --git a/Engine/Source/Developer/TraceInsights/Public/Insights/ViewModels/BaseTimingTrack.h b/Engine/Source/Developer/TraceInsights/Public/Insights/ViewModels/BaseTimingTrack.h
--- a/Engine/Source/Developer/TraceInsights/Public/Insights/ViewModels/BaseTimingTrack.h
+++ b/Engine/Source/Developer/TraceInsights/Public/Insights/ViewModels/BaseTimingTrack.h
@@ -12,7 +12,7 @@ namespace Trace { class IAnalysisSession; };
 
 class FTimingEventSearchParameters;
 class FTimingTrackViewport;
-class FTimingViewDrawHelper;
+class ITimingViewDrawHelper;
 class FTooltipDrawState;
 class ITimingEvent;
 
@@ -41,7 +41,7 @@ public:
 	virtual const TSharedPtr<const ITimingEvent> GetHoveredEvent() const = 0;
 	virtual const TSharedPtr<const ITimingEvent> GetSelectedEvent() const = 0;
 	virtual FDrawContext& GetDrawContext() const = 0;
-	virtual const FTimingViewDrawHelper& GetHelper() const = 0;
+	virtual const ITimingViewDrawHelper& GetHelper() const = 0;
 };
 
 ////////////////////////////////////////////////////////////////////////////////////////////////////
diff --git a/Engine/Source/Developer/TraceInsights/Public/Insights/ViewModels/ITimingViewDrawHelper.h b/Engine/Source/Developer/TraceInsights/Public/Insights/ViewModels/ITimingViewDrawHelper.h
new file mode 100644
--- /dev/null
+++ b/Engine/Source/Developer/TraceInsights/Public/Insights/ViewModels/ITimingViewDrawHelper.h
@@ -0,0 +1,21 @@
+// Copyright 1998-2019 Epic Games, Inc. All Rights Reserved.
+
+#pragma once
+
+#include "CoreMinimal.h"
+
+class FTimingEventsTrack;
+struct FSlateBrush;
+struct FSlateFontInfo;
+
+/** Helper allowing access to common drawing elements for tracks */
+class ITimingViewDrawHelper
+{
+public:
+	virtual const FSlateBrush* GetWhiteBrush() const = 0;
+	virtual const FSlateFontInfo& GetEventFont() const = 0;
+	virtual FLinearColor GetEdgeColor() const = 0;
+	virtual FLinearColor GetTrackNameTextColor(const FTimingEventsTrack& Track) const = 0;
+	virtual int32 GetHeaderBackgroundLayerId() const = 0;
+	virtual int32 GetHeaderTextLayerId() const = 0;
+};
\ No newline at end of file
