[Insights]
  - Timing view: Improved navigation by using a pinned track for vertical position reference when scrolling or zooming.
  - Added ">" in front of name of selected (pinned) track.
  - Minor code refactoring in tree node and log message files.

diff --git a/Engine/Source/Developer/TraceInsights/Private/Insights/Table/ViewModels/TreeNodeGrouping.cpp b/Engine/Source/Developer/TraceInsights/Private/Insights/Table/ViewModels/TreeNodeGrouping.cpp
--- a/Engine/Source/Developer/TraceInsights/Private/Insights/Table/ViewModels/TreeNodeGrouping.cpp
+++ b/Engine/Source/Developer/TraceInsights/Private/Insights/Table/ViewModels/TreeNodeGrouping.cpp
@@ -11,7 +11,7 @@ namespace Insights
 
 ////////////////////////////////////////////////////////////////////////////////////////////////////
 
-FTreeNodeGrouping::FTreeNodeGrouping(FText InShortName, FText InTitleName, FText InDescription, FName InBrushName,  FSlateBrush* InIcon)
+FTreeNodeGrouping::FTreeNodeGrouping(const FText& InShortName, const FText& InTitleName, const FText& InDescription, const FName InBrushName, const FSlateBrush* InIcon)
 	: ShortName(InShortName)
 	, TitleName(InTitleName)
 	, Description(InDescription)
diff --git a/Engine/Source/Developer/TraceInsights/Private/Insights/Table/ViewModels/TreeNodeGrouping.h b/Engine/Source/Developer/TraceInsights/Private/Insights/Table/ViewModels/TreeNodeGrouping.h
--- a/Engine/Source/Developer/TraceInsights/Private/Insights/Table/ViewModels/TreeNodeGrouping.h
+++ b/Engine/Source/Developer/TraceInsights/Private/Insights/Table/ViewModels/TreeNodeGrouping.h
@@ -43,7 +43,7 @@ struct FTreeNodeGroupInfo
 class FTreeNodeGrouping : public ITreeNodeGrouping
 {
 public:
-	FTreeNodeGrouping(FText InShortName, FText InTitleName, FText InDescription, FName InBrushName,  FSlateBrush* InIcon);
+	FTreeNodeGrouping(const FText& InShortName, const FText& InTitleName, const FText& InDescription, const FName InBrushName, const FSlateBrush* InIcon);
 	virtual ~FTreeNodeGrouping() {}
 
 	virtual FText GetShortName() const override { return ShortName; }
diff --git a/Engine/Source/Developer/TraceInsights/Private/Insights/ViewModels/LogMessage.cpp b/Engine/Source/Developer/TraceInsights/Private/Insights/ViewModels/LogMessage.cpp
--- a/Engine/Source/Developer/TraceInsights/Private/Insights/ViewModels/LogMessage.cpp
+++ b/Engine/Source/Developer/TraceInsights/Private/Insights/ViewModels/LogMessage.cpp
@@ -163,7 +163,7 @@ FLogMessageRecord& FLogMessageCache::Get(uint64 Index)
 		{
 			FScopeLock Lock(&CriticalSection);
 			FLogMessageRecord Entry(Message);
-			Map.Add(Index, Entry);
+			Map.Add(Index, MoveTemp(Entry));
 		});
 	}
 
diff --git a/Engine/Source/Developer/TraceInsights/Private/Insights/ViewModels/TimingViewDrawHelper.cpp b/Engine/Source/Developer/TraceInsights/Private/Insights/ViewModels/TimingViewDrawHelper.cpp
--- a/Engine/Source/Developer/TraceInsights/Private/Insights/ViewModels/TimingViewDrawHelper.cpp
+++ b/Engine/Source/Developer/TraceInsights/Private/Insights/ViewModels/TimingViewDrawHelper.cpp
@@ -439,6 +439,13 @@ void FTimingViewDrawHelper::DrawTrackHeader(const FTimingEventsTrack& Track) con
 			const TSharedRef<FSlateFontMeasure> FontMeasureService = FSlateApplication::Get().GetRenderer()->GetFontMeasureService();
 			float NameWidth = FontMeasureService->Measure(Track.GetName(), EventFont).X;
 
+			float TextX = 2.0f;
+			if (Track.IsSelected())
+			{
+				NameWidth += 8.0f;
+				TextX += 8.0f;
+			}
+
 			const float HeaderH = FMath::Min(12.0f, Track.GetHeight() - 1.0f);
 			DrawContext.DrawBox(HeaderLayerId, 0.0f, TrackY + 1.0f, NameWidth + 4.0f, HeaderH, WhiteBrush, EdgeColor);
 
@@ -446,7 +453,13 @@ void FTimingViewDrawHelper::DrawTrackHeader(const FTimingEventsTrack& Track) con
 
 			const float TextY = TrackY + FMath::Min(0.0f, (Track.GetHeight() - 13.0f) / 2.0f);
 			const int32 HeaderTextLayerId = ReservedLayerId + ToInt32(EDrawLayer::HeaderText);
-			DrawContext.DrawText(HeaderTextLayerId, 2.0f, TextY, Track.GetName(), EventFont, TextColor);
+			DrawContext.DrawText(HeaderTextLayerId, TextX, TextY, Track.GetName(), EventFont, TextColor);
+
+			if (Track.IsSelected())
+			{
+				// TODO: Use a "pin" image brush instead.
+				DrawContext.DrawText(HeaderTextLayerId, 2.0f, TextY, TEXT(">"), EventFont, TextColor);
+			}
 		}
 	}
 }
diff --git a/Engine/Source/Developer/TraceInsights/Private/Insights/Widgets/SLogView.cpp b/Engine/Source/Developer/TraceInsights/Private/Insights/Widgets/SLogView.cpp
--- a/Engine/Source/Developer/TraceInsights/Private/Insights/Widgets/SLogView.cpp
+++ b/Engine/Source/Developer/TraceInsights/Private/Insights/Widgets/SLogView.cpp
@@ -811,8 +811,7 @@ void SLogView::Tick(const FGeometry& AllottedGeometry, const double InCurrentTim
 
 				for (int32 Index = TotalNumMessages; Index < NewMessageCount; Index++)
 				{
-					FLogMessage LogMessage(Index);
-					Messages.Add(MakeShared<FLogMessage>(MoveTemp(LogMessage)));
+					Messages.Add(MakeShared<FLogMessage>(Index));
 				}
 
 				const int32 NumAddedMessages = NewMessageCount - TotalNumMessages;
@@ -880,8 +879,7 @@ void SLogView::Tick(const FGeometry& AllottedGeometry, const double InCurrentTim
 			const int32 NumFilteredMessages = FilteredMessages.Num();
 			for (int32 Index = 0; Index < NumFilteredMessages; Index++)
 			{
-				FLogMessage LogMessage(FilteredMessages[Index]);
-				Messages.Add(MakeShared<FLogMessage>(MoveTemp(LogMessage)));
+				Messages.Add(MakeShared<FLogMessage>(FilteredMessages[Index]));
 			}
 
 			TotalNumMessages = Task.GetEndIndex();
diff --git a/Engine/Source/Developer/TraceInsights/Private/Insights/Widgets/STimingView.cpp b/Engine/Source/Developer/TraceInsights/Private/Insights/Widgets/STimingView.cpp
--- a/Engine/Source/Developer/TraceInsights/Private/Insights/Widgets/STimingView.cpp
+++ b/Engine/Source/Developer/TraceInsights/Private/Insights/Widgets/STimingView.cpp
@@ -72,7 +72,7 @@ namespace Insights { const FName TimingViewExtenderFeatureName(TEXT("TimingViewE
 ////////////////////////////////////////////////////////////////////////////////////////////////////
 
 STimingView::STimingView()
-	: bScrollableTracksDirty(false)
+	: bScrollableTracksOrderIsDirty(false)
 	, ThreadTimingSharedState(MakeShared<FThreadTimingSharedState>(this))
 	, LoadingSharedState(MakeShared<FLoadingSharedState>(this))
 	, bAssetLoadingMode(false)
@@ -211,6 +211,8 @@ void STimingView::Reset()
 	ScrollableTracks.Reset();
 	ForegroundTracks.Reset();
 
+	bScrollableTracksOrderIsDirty = false;
+
 	FTimingEventsTrack::bUseDownSampling = true;
 
 	//////////////////////////////////////////////////
@@ -240,6 +242,8 @@ void STimingView::Reset()
 
 	MousePositionOnButtonUp = FVector2D::ZeroVector;
 
+	LastScrollPosY = 0.0f;
+
 	bIsLMB_Pressed = false;
 	bIsRMB_Pressed = false;
 
@@ -534,8 +538,29 @@ void STimingView::Tick(const FGeometry& AllottedGeometry, const double InCurrent
 	}
 
 	////////////////////////////////////////////////////////////////////////////////////////////////////
-	// Elastic snap to vertical scroll limits.
 
+	const float InitialScrollPosY = Viewport.GetScrollPosY();
+	const float InitialPinnedTrackPosY = SelectedTrack.IsValid() ? SelectedTrack->GetPosY() : 0.0f;
+
+	// Update the Y position for visible scrollable tracks.
+	UpdatePositionForScrollableTracks();
+
+	// The selected track will be pinned (keeps Y pos fixed unless user scrolls vertically).
+	if (SelectedTrack.IsValid())
+	{
+		const float ScrollingDY = LastScrollPosY - InitialScrollPosY;
+		const float PinnedTrackPosY = SelectedTrack->GetPosY();
+		const float AdjustmentDY = InitialPinnedTrackPosY - PinnedTrackPosY + ScrollingDY;
+
+		if (!FMath::IsNearlyZero(AdjustmentDY, 0.5f))
+		{
+			ViewportScrollPosYOnButtonDown -= AdjustmentDY;
+			ScrollAtPosY(InitialScrollPosY - AdjustmentDY);
+			UpdatePositionForScrollableTracks();
+		}
+	}
+
+	// Elastic snap to vertical scroll limits.
 	if (!bIsPanning)
 	{
 		const float DY = Viewport.GetScrollHeight() - Viewport.GetScrollableAreaHeight() + 7.0f; // +7 is to allow some space for the horizontal scrollbar
@@ -581,22 +606,15 @@ void STimingView::Tick(const FGeometry& AllottedGeometry, const double InCurrent
 			}
 		}
 
-		ScrollAtPosY(ScrollPosY);
-	}
-
-	////////////////////////////////////////////////////////////////////////////////////////////////////
-	// Update Y position for visible scrollable tracks.
-
-	float ScrollableTrackPosY = Viewport.GetTopOffset() - Viewport.GetScrollPosY();
-	for (TSharedPtr<FBaseTimingTrack>& TrackPtr : ScrollableTracks)
-	{
-		TrackPtr->SetPosY(ScrollableTrackPosY); // set pos y also for the hidden tracks
-		if (TrackPtr->IsVisible())
+		if (ScrollPosY != Viewport.GetScrollPosY())
 		{
-			ScrollableTrackPosY += TrackPtr->GetHeight();
+			ScrollAtPosY(ScrollPosY);
+			UpdatePositionForScrollableTracks();
 		}
 	}
 
+	LastScrollPosY = Viewport.GetScrollPosY();
+
 	////////////////////////////////////////////////////////////////////////////////////////////////////
 	// At this point it is assumed all tracks have proper position and size.
 	// Update.
@@ -755,6 +773,22 @@ void STimingView::UpdateOtherViews()
 
 ////////////////////////////////////////////////////////////////////////////////////////////////////
 
+void STimingView::UpdatePositionForScrollableTracks()
+{
+	// Update the Y postion for visible scrollable tracks.
+	float ScrollableTrackPosY = Viewport.GetTopOffset() - Viewport.GetScrollPosY();
+	for (TSharedPtr<FBaseTimingTrack>& TrackPtr : ScrollableTracks)
+	{
+		TrackPtr->SetPosY(ScrollableTrackPosY); // set pos y also for the hidden tracks
+		if (TrackPtr->IsVisible())
+		{
+			ScrollableTrackPosY += TrackPtr->GetHeight();
+		}
+	}
+}
+
+////////////////////////////////////////////////////////////////////////////////////////////////////
+
 int32 STimingView::OnPaint(const FPaintArgs& Args, const FGeometry& AllottedGeometry, const FSlateRect& MyCullingRect, FSlateWindowElementList& OutDrawElements, int32 LayerId, const FWidgetStyle& InWidgetStyle, bool bParentEnabled) const
 {
 	FStopwatch Stopwatch;
@@ -1244,17 +1278,17 @@ void STimingView::AddScrollableTrack(TSharedPtr<FBaseTimingTrack> Track)
 
 void STimingView::InvalidateScrollableTracksOrder()
 {
-	bScrollableTracksDirty = true;
+	bScrollableTracksOrderIsDirty = true;
 }
 
 ////////////////////////////////////////////////////////////////////////////////////////////////////
 
 void STimingView::UpdateScrollableTracksOrder()
 {
-	if(bScrollableTracksDirty)
+	if (bScrollableTracksOrderIsDirty)
 	{
 		Algo::SortBy(ScrollableTracks, &FBaseTimingTrack::GetOrder);
-		bScrollableTracksDirty = false;
+		bScrollableTracksOrderIsDirty = false;
 	}
 }
 
diff --git a/Engine/Source/Developer/TraceInsights/Private/Insights/Widgets/STimingView.h b/Engine/Source/Developer/TraceInsights/Private/Insights/Widgets/STimingView.h
--- a/Engine/Source/Developer/TraceInsights/Private/Insights/Widgets/STimingView.h
+++ b/Engine/Source/Developer/TraceInsights/Private/Insights/Widgets/STimingView.h
@@ -229,6 +229,7 @@ public:
 	virtual void AddForegroundTrack(TSharedPtr<FBaseTimingTrack> Track) override;
 
 	virtual void InvalidateScrollableTracksOrder() override;
+	//TODO: virtual void InvalidateScrollableTracksVisibility() override;
 
 	virtual Insights::FSelectionChangedDelegate& OnSelectionChanged() override { return OnSelectionChangedDelegate; }
 	virtual Insights::FTimeMarkerChangedDelegate& OnTimeMarkerChanged() override { return OnTimeMarkerChangedDelegate; }
@@ -301,6 +302,8 @@ protected:
 
 	////////////////////////////////////////////////////////////////////////////////////////////////////
 
+	void UpdatePositionForScrollableTracks();
+
 	/**
 	 * Called when the user scrolls the horizontal scrollbar.
 	 *
@@ -383,8 +386,8 @@ protected:
 	TArray<TSharedPtr<FBaseTimingTrack>> ScrollableTracks; /**< tracks in scrollable area, in the order to be displayed (top to bottom) */
 	TArray<TSharedPtr<FBaseTimingTrack>> ForegroundTracks; /**< tracks to draw over top/scrollable/bottom tracks (can use entire area), in the order to be displayed (back to front) */
 
-	/** Whether the scrollable tracks are dirty and need to be re-sorted */
-	bool bScrollableTracksDirty;
+	/** Whether the order of scrollable tracks is dirty and list need to be re-sorted */
+	bool bScrollableTracksOrderIsDirty;
 
 	////////////////////////////////////////////////////////////
 
@@ -430,6 +433,8 @@ protected:
 	/** Mouse position during the call on mouse button up. */
 	FVector2D MousePositionOnButtonUp;
 
+	float LastScrollPosY;
+
 	bool bIsLMB_Pressed;
 	bool bIsRMB_Pressed;
 
diff --git a/Engine/Source/Developer/TraceInsights/Public/Insights/ITimingViewSession.h b/Engine/Source/Developer/TraceInsights/Public/Insights/ITimingViewSession.h
--- a/Engine/Source/Developer/TraceInsights/Public/Insights/ITimingViewSession.h
+++ b/Engine/Source/Developer/TraceInsights/Public/Insights/ITimingViewSession.h
@@ -53,7 +53,7 @@ public:
 	/** Adds a new foreground track. */
 	virtual void AddForegroundTrack(TSharedPtr<FBaseTimingTrack> Track) = 0;
 
-	/** Marks the scrollable tracks as dirty so they will be re-sorted */
+	/** Marks the scrollable tracks as not being in the correct order so they will be re-sorted */
 	virtual void InvalidateScrollableTracksOrder() = 0;
 
 	/** Finds a track has been added via Add*Track(). */
