[Insights] Timers view: Fixed Ctrl + Double Click in Timing view not keeping the selected timer visible in Timers view (i.e. when selection is updated it updates the aggregation stats which triggers the list of timers to be resorted).

diff --git a/Engine/Source/Developer/TraceInsights/Private/Insights/Widgets/STimersView.cpp b/Engine/Source/Developer/TraceInsights/Private/Insights/Widgets/STimersView.cpp
--- a/Engine/Source/Developer/TraceInsights/Private/Insights/Widgets/STimersView.cpp
+++ b/Engine/Source/Developer/TraceInsights/Private/Insights/Widgets/STimersView.cpp
@@ -1640,6 +1640,12 @@ void STimersView::UpdateStats(double StartTime, double EndTime)
 	}
 
 	UpdateTree();
+
+	const TArray<FTimerNodePtr> SelectedNodes = TreeView->GetSelectedItems();
+	if (SelectedNodes.Num() > 0)
+	{
+		TreeView->RequestScrollIntoView(SelectedNodes[0]);
+	}
 }
 
 ////////////////////////////////////////////////////////////////////////////////////////////////////
