[Insights] Fixed tooltips not appearing on Linux when trace session is first opened.

diff --git a/Engine/Source/Developer/TraceInsights/Private/Insights/ViewModels/TooltipDrawState.cpp b/Engine/Source/Developer/TraceInsights/Private/Insights/ViewModels/TooltipDrawState.cpp
--- a/Engine/Source/Developer/TraceInsights/Private/Insights/ViewModels/TooltipDrawState.cpp
+++ b/Engine/Source/Developer/TraceInsights/Private/Insights/ViewModels/TooltipDrawState.cpp
@@ -186,7 +186,25 @@ void FTooltipDrawState::Update()
 		}
 	}
 
-	const float RealDesiredOpacity = FMath::Min(DesiredOpacity, 1.0f - FMath::Abs(Size.X - DesiredSize.X) / DesiredSize.X);
+	float RealDesiredOpacity;
+	if (DesiredSize.X > 1.0f)
+	{
+		const float DesiredOpacityByTooltipWidth = 1.0f - FMath::Abs(Size.X - DesiredSize.X) / DesiredSize.X;
+
+		if (FMath::IsNearlyEqual(DesiredOpacity, DesiredOpacityByTooltipWidth, 0.001f))
+		{
+			RealDesiredOpacity = DesiredOpacity;
+		}
+		else
+		{
+			RealDesiredOpacity = FMath::Min(DesiredOpacity, DesiredOpacityByTooltipWidth);
+		}
+	}
+	else
+	{
+		RealDesiredOpacity = 0.0f;
+	}
+
 	if (Opacity != RealDesiredOpacity)
 	{
 		if (Opacity < RealDesiredOpacity)
@@ -200,7 +218,7 @@ void FTooltipDrawState::Update()
 			Opacity = Opacity * 0.75f + RealDesiredOpacity * 0.25f;
 		}
 
-		if (FMath::IsNearlyEqual(Opacity, RealDesiredOpacity))
+		if (FMath::IsNearlyEqual(Opacity, RealDesiredOpacity, 0.001f))
 		{
 			Opacity = RealDesiredOpacity;
 		}
