[Insights]
  - Timing view: Added highlight/filtered mode for a selected timer. Double clicking a timing event (in the Timing view) or a timer (in the Timers list) will highlight all timing events of same type. The other timing events are displayed with a faded color. Double clicking an empty space (in the Timing view) or same filtered timer (in the Timers list) will reset the filtered timer.  - Timing view: Fixed the jittery vertical scrolling (i.e. when user scralls vertically using the vertical scrollbar after clicking the TimeRuler track or Markers (Bookmarks) track).

diff --git a/Engine/Source/Developer/TraceInsights/Private/Insights/ViewModels/ThreadTimingTrack.cpp b/Engine/Source/Developer/TraceInsights/Private/Insights/ViewModels/ThreadTimingTrack.cpp
--- a/Engine/Source/Developer/TraceInsights/Private/Insights/ViewModels/ThreadTimingTrack.cpp
+++ b/Engine/Source/Developer/TraceInsights/Private/Insights/ViewModels/ThreadTimingTrack.cpp
@@ -395,7 +395,8 @@ void FThreadTimingTrack::BuildDrawState(ITimingEventsTrackDrawStateBuilder& Buil
 					Timeline.EnumerateEventsDownSampled(Viewport.GetStartTime(), Viewport.GetEndTime(), SecondsPerPixel,
 						[this, &Builder, Timers](double StartTime, double EndTime, uint32 Depth, const Trace::FTimingProfilerEvent& Event)
 						{
-							const uint64 Type = Timers[Event.TimerIndex].Id;
+							//ensure(Timers[Event.TimerIndex].Id == Event.TimerIndex);
+							const uint64 Type = Event.TimerIndex;
 							const uint32 Color = 0;
 							Builder.AddEvent(StartTime, EndTime, Depth, Timers[Event.TimerIndex].Name, Type, Color);
 						});
@@ -405,7 +406,8 @@ void FThreadTimingTrack::BuildDrawState(ITimingEventsTrackDrawStateBuilder& Buil
 					Timeline.EnumerateEvents(Viewport.GetStartTime(), Viewport.GetEndTime(),
 						[this, &Builder, Timers](double StartTime, double EndTime, uint32 Depth, const Trace::FTimingProfilerEvent& Event)
 						{
-							const uint64 Type = Timers[Event.TimerIndex].Id;
+							//ensure(Timers[Event.TimerIndex].Id == Event.TimerIndex);
+							const uint64 Type = Event.TimerIndex;
 							const uint32 Color = 0;
 							Builder.AddEvent(StartTime, EndTime, Depth, Timers[Event.TimerIndex].Name, Type, Color);
 						});
@@ -417,6 +419,49 @@ void FThreadTimingTrack::BuildDrawState(ITimingEventsTrackDrawStateBuilder& Buil
 
 ////////////////////////////////////////////////////////////////////////////////////////////////////
 
+void FThreadTimingTrack::BuildFilteredDrawState(ITimingEventsTrackDrawStateBuilder& Builder, const ITimingTrackUpdateContext& Context)
+{
+	const TSharedPtr<ITimingEventFilter> EventFilterPtr = Context.GetEventFilter();
+	if (EventFilterPtr.IsValid() && FTimingEventFilter::CheckTypeName(*EventFilterPtr))
+	{
+		const FTimingEventFilter& EventFilter = static_cast<const FTimingEventFilter&>(*EventFilterPtr);
+		const uint64 FilterEventType = EventFilter.GetEventType();
+
+		TSharedPtr<const Trace::IAnalysisSession> Session = FInsightsManager::Get()->GetSession();
+		if (Session.IsValid() && Trace::ReadTimingProfilerProvider(*Session.Get()))
+		{
+			Trace::FAnalysisSessionReadScope SessionReadScope(*Session.Get());
+
+			const Trace::ITimingProfilerProvider& TimingProfilerProvider = *Trace::ReadTimingProfilerProvider(*Session.Get());
+
+			const FTimingTrackViewport& Viewport = Context.GetViewport();
+
+			TimingProfilerProvider.ReadTimers([this, &Builder, &Viewport, &TimingProfilerProvider, FilterEventType](const Trace::FTimingProfilerTimer* Timers, uint64 TimersCount)
+			{
+				TimingProfilerProvider.ReadTimeline(TimelineIndex, [this, &Builder, &Viewport, Timers, FilterEventType](const Trace::ITimingProfilerProvider::Timeline& Timeline)
+				{
+					// Note: Enumerating events for filtering should not use downsampling.
+					//const double SecondsPerPixel = 1.0 / Viewport.GetScaleX();
+					//Timeline.EnumerateEventsDownSampled(Viewport.GetStartTime(), Viewport.GetEndTime(), SecondsPerPixel,
+					Timeline.EnumerateEvents(Viewport.GetStartTime(), Viewport.GetEndTime(),
+						[this, &Builder, Timers, FilterEventType](double StartTime, double EndTime, uint32 Depth, const Trace::FTimingProfilerEvent& Event)
+					{
+						//ensure(Timers[Event.TimerIndex].Id == Event.TimerIndex);
+						const uint64 Type = Event.TimerIndex;
+						if (Type == FilterEventType)
+						{
+							const uint32 Color = 0;
+							Builder.AddEvent(StartTime, EndTime, Depth, Timers[Event.TimerIndex].Name, Type, Color);
+						}
+					});
+				});
+			});
+		}
+	}
+}
+
+////////////////////////////////////////////////////////////////////////////////////////////////////
+
 void FThreadTimingTrack::PostDraw(const ITimingTrackDrawContext& Context) const
 {
 	const TSharedPtr<const ITimingEvent> SelectedEventPtr = Context.GetSelectedEvent();
@@ -565,7 +610,8 @@ const TSharedPtr<const ITimingEvent> FThreadTimingTrack::SearchEvent(const FTimi
 
 	FindTimingProfilerEvent(InSearchParameters, [this, &FoundEvent](double InFoundStartTime, double InFoundEndTime, uint32 InFoundDepth, const Trace::FTimingProfilerEvent& InFoundEvent)
 	{
-		FoundEvent = MakeShared<FTimingEvent>(SharedThis(this), InFoundStartTime, InFoundEndTime, InFoundDepth);
+		const uint64 EventType = InFoundEvent.TimerIndex;
+		FoundEvent = MakeShared<FTimingEvent>(SharedThis(this), InFoundStartTime, InFoundEndTime, InFoundDepth, EventType);
 	});
 
 	return FoundEvent;
@@ -654,11 +700,16 @@ void FThreadTimingTrack::OnEventSelected(const ITimingEvent& InSelectedEvent) co
 	{
 		const FTimingEvent& TrackEvent = static_cast<const FTimingEvent&>(InSelectedEvent);
 
-		FindTimingProfilerEvent(TrackEvent, [](double InFoundStartTime, double InFoundEndTime, uint32 InFoundDepth, const Trace::FTimingProfilerEvent& InFoundEvent)
-		{
-			// Select the timer node corresponding to timing event type of selected timing event.
-			FTimingProfilerManager::Get()->SetSelectedTimer(InFoundEvent.TimerIndex);
-		});
+		//uint64 TimerType = uint64(-1);
+		//FindTimingProfilerEvent(TrackEvent, [&TimerType](double InFoundStartTime, double InFoundEndTime, uint32 InFoundDepth, const Trace::FTimingProfilerEvent& InFoundEvent)
+		//{
+		//	TimerType = InFoundEvent.TimerIndex;
+		//});
+
+		const uint64 TimerType = TrackEvent.GetType();
+
+		// Select the timer node corresponding to timing event type of selected timing event.
+		FTimingProfilerManager::Get()->SetSelectedTimer(TimerType);
 	}
 }
 
@@ -670,15 +721,20 @@ void FThreadTimingTrack::OnClipboardCopyEvent(const ITimingEvent& InSelectedEven
 	{
 		const FTimingEvent& TrackEvent = static_cast<const FTimingEvent&>(InSelectedEvent);
 
-		FindTimingProfilerEvent(TrackEvent, [](double InFoundStartTime, double InFoundEndTime, uint32 InFoundDepth, const Trace::FTimingProfilerEvent& InFoundEvent)
+		//uint64 TimerType = uint64(-1);
+		//FindTimingProfilerEvent(TrackEvent, [&TimerType](double InFoundStartTime, double InFoundEndTime, uint32 InFoundDepth, const Trace::FTimingProfilerEvent& InFoundEvent)
+		//{
+		//	TimerType = InFoundEvent.TimerIndex;
+		//});
+
+		const uint64 TimerType = TrackEvent.GetType();
+
+		const FTimerNodePtr TimerNodePtr = FTimingProfilerManager::Get()->GetTimerNode(TimerType);
+		if (TimerNodePtr.IsValid())
 		{
-			const FTimerNodePtr TimerNodePtr = FTimingProfilerManager::Get()->GetTimerNode(InFoundEvent.TimerIndex);
-			if (TimerNodePtr.IsValid())
-			{
-				// Copy name of selected timing event to clipboard.
-				FPlatformApplicationMisc::ClipboardCopy(*TimerNodePtr->GetName().ToString());
-			}
-		});
+			// Copy name of selected timing event to clipboard.
+			FPlatformApplicationMisc::ClipboardCopy(*TimerNodePtr->GetName().ToString());
+		}
 	}
 }
 
diff --git a/Engine/Source/Developer/TraceInsights/Private/Insights/ViewModels/ThreadTimingTrack.h b/Engine/Source/Developer/TraceInsights/Private/Insights/ViewModels/ThreadTimingTrack.h
--- a/Engine/Source/Developer/TraceInsights/Private/Insights/ViewModels/ThreadTimingTrack.h
+++ b/Engine/Source/Developer/TraceInsights/Private/Insights/ViewModels/ThreadTimingTrack.h
@@ -18,6 +18,10 @@ class FCpuTimingTrack;
 class STimingView;
 struct FSlateBrush;
 
+////////////////////////////////////////////////////////////////////////////////////////////////////
+
+
+
 ////////////////////////////////////////////////////////////////////////////////////////////////////
 
 class FThreadTimingSharedState : public Insights::ITimingViewExtender, public TSharedFromThis<FThreadTimingSharedState>
@@ -108,6 +112,7 @@ public:
 	//void SetThreadId(uint32 InThreadId) { ThreadId = InThreadId; }
 
 	virtual void BuildDrawState(ITimingEventsTrackDrawStateBuilder& Builder, const ITimingTrackUpdateContext& Context) override;
+	virtual void BuildFilteredDrawState(ITimingEventsTrackDrawStateBuilder& Builder, const ITimingTrackUpdateContext& Context) override;
 
 	virtual void PostDraw(const ITimingTrackDrawContext& Context) const override;
 
diff --git a/Engine/Source/Developer/TraceInsights/Private/Insights/ViewModels/TimingEventsTrack.cpp b/Engine/Source/Developer/TraceInsights/Private/Insights/ViewModels/TimingEventsTrack.cpp
--- a/Engine/Source/Developer/TraceInsights/Private/Insights/ViewModels/TimingEventsTrack.cpp
+++ b/Engine/Source/Developer/TraceInsights/Private/Insights/ViewModels/TimingEventsTrack.cpp
@@ -2,15 +2,29 @@
 
 #include "Insights/ViewModels/TimingEventsTrack.h"
 
+#include "Insights/Common/Stopwatch.h"
+#include "Insights/TimingProfilerCommon.h"
 #include "Insights/ViewModels/TimingEvent.h"
 #include "Insights/ViewModels/TimingTrackViewport.h"
 #include "Insights/ViewModels/TimingViewDrawHelper.h"
 
 #define LOCTEXT_NAMESPACE "TimingEventsTrack"
 
+////////////////////////////////////////////////////////////////////////////////////////////////////
+// FTimingEvent, FTimingEventFilter
 ////////////////////////////////////////////////////////////////////////////////////////////////////
 
 const FName FTimingEvent::TypeName = FName(TEXT("FTimingEvent"));
+const FName FTimingEventFilter::TypeName = FName(TEXT("FTimingEventFilter"));
+
+bool FTimingEventFilter::FilterTrack(const FBaseTimingTrack& InTrack) const
+{
+	return !bFilterByTrackType || InTrack.GetType() == TrackType;
+}
+
+////////////////////////////////////////////////////////////////////////////////////////////////////
+// FTimingEventsTrack
+////////////////////////////////////////////////////////////////////////////////////////////////////
 
 bool FTimingEventsTrack::bUseDownSampling = true;
 
@@ -20,6 +34,7 @@ FTimingEventsTrack::FTimingEventsTrack(const FName& InType, const FName& InSubTy
 	: FBaseTimingTrack(InType, InSubType, InName)
 	, NumLanes(0)
 	, DrawState(MakeShared<FTimingEventsTrackDrawState>())
+	, FilteredDrawState(MakeShared<FTimingEventsTrackDrawState>())
 {
 }
 
@@ -37,6 +52,7 @@ void FTimingEventsTrack::Reset()
 
 	NumLanes = 0;
 	DrawState->Reset();
+	FilteredDrawState->Reset();
 }
 
 ////////////////////////////////////////////////////////////////////////////////////////////////////
@@ -47,13 +63,93 @@ void FTimingEventsTrack::PreUpdate(const ITimingTrackUpdateContext& Context)
 	{
 		ClearDirtyFlag();
 
-		FTimingEventsTrackDrawStateBuilder Builder(*DrawState, Context.GetViewport());
+		int32 MaxDepth = -1;
+
+		{
+			FTimingEventsTrackDrawStateBuilder Builder(*DrawState, Context.GetViewport());
 
-		BuildDrawState(Builder, Context);
+			BuildDrawState(Builder, Context);
 
-		Builder.Flush();
+			Builder.Flush();
+
+			if (Builder.GetMaxDepth() > MaxDepth)
+			{
+				MaxDepth = Builder.GetMaxDepth();
+			}
+		}
 
-		SetNumLanes(Builder.GetMaxDepth() + 1);
+		const TSharedPtr<ITimingEventFilter> EventFilter = Context.GetEventFilter();
+		if (EventFilter.IsValid() && EventFilter->FilterTrack(*this))
+		{
+			const FTimingTrackViewport& Viewport = Context.GetViewport();
+
+			const bool bFastLastBuild = FilteredDrawStateInfo.LastBuildDuration < 0.005; // LastBuildDuration < 5ms
+			const bool bFilterPointerChanged = !FilteredDrawStateInfo.LastEventFilter.HasSameObject(EventFilter.Get());
+			const bool bFilterContentChanged = FilteredDrawStateInfo.LastFilterChangeNumber != EventFilter->GetChangeNumber();
+
+			if (bFastLastBuild || bFilterPointerChanged || bFilterContentChanged)
+			{
+				FilteredDrawStateInfo.LastEventFilter = EventFilter;
+				FilteredDrawStateInfo.LastFilterChangeNumber = EventFilter->GetChangeNumber();
+				FilteredDrawStateInfo.ViewportStartTime = Context.GetViewport().GetStartTime();
+				FilteredDrawStateInfo.ViewportScaleX = Context.GetViewport().GetScaleX();
+				FilteredDrawStateInfo.Counter = 0;
+			}
+			else
+			{
+				if (FilteredDrawStateInfo.ViewportStartTime == Viewport.GetStartTime() &&
+					FilteredDrawStateInfo.ViewportScaleX == Viewport.GetScaleX())
+				{
+					if (FilteredDrawStateInfo.Counter > 0)
+					{
+						FilteredDrawStateInfo.Counter--;
+					}
+				}
+				else
+				{
+					FilteredDrawStateInfo.ViewportStartTime = Context.GetViewport().GetStartTime();
+					FilteredDrawStateInfo.ViewportScaleX = Context.GetViewport().GetScaleX();
+
+					if (FilteredDrawStateInfo.Counter >= 0)
+					{
+						FilteredDrawStateInfo.Counter = 1; // wait
+					}
+				}
+			}
+
+			if (FilteredDrawStateInfo.Counter == 0)
+			{
+				FStopwatch Stopwatch;
+				Stopwatch.Start();
+				{
+					FTimingEventsTrackDrawStateBuilder Builder(*FilteredDrawState, Context.GetViewport());
+					BuildFilteredDrawState(Builder, Context);
+					Builder.Flush();
+				}
+				Stopwatch.Stop();
+				FilteredDrawStateInfo.LastBuildDuration = Stopwatch.GetAccumulatedTime();
+			}
+			else
+			{
+				FilteredDrawState->Reset();
+				FilteredDrawStateInfo.Opacity = 0.0f;
+				SetDirtyFlag();
+			}
+		}
+		else
+		{
+			FilteredDrawStateInfo.LastBuildDuration = 0.0;
+
+			if (FilteredDrawStateInfo.LastEventFilter.IsValid())
+			{
+				FilteredDrawStateInfo.LastEventFilter.Reset();
+				FilteredDrawStateInfo.LastFilterChangeNumber = 0;
+				FilteredDrawStateInfo.Counter = 0;
+				FilteredDrawState->Reset();
+			}
+		}
+
+		SetNumLanes(MaxDepth + 1);
 	}
 
 	UpdateTrackHeight(Context);
@@ -123,9 +219,8 @@ void FTimingEventsTrack::PostUpdate(const ITimingTrackUpdateContext& Context)
 
 void FTimingEventsTrack::Draw(const ITimingTrackDrawContext& Context) const
 {
-	const FTimingViewDrawHelper& Helper = *static_cast<const FTimingViewDrawHelper*>(&Context.GetHelper());
-	Helper.DrawEvents(*DrawState, *this);
-	Helper.DrawTrackHeader(*this);
+	DrawEvents(Context, 1.0f);
+	DrawHeader(Context);
 }
 
 ////////////////////////////////////////////////////////////////////////////////////////////////////
@@ -133,7 +228,25 @@ void FTimingEventsTrack::Draw(const ITimingTrackDrawContext& Context) const
 void FTimingEventsTrack::DrawEvents(const ITimingTrackDrawContext& Context, const float OffsetY) const
 {
 	const FTimingViewDrawHelper& Helper = *static_cast<const FTimingViewDrawHelper*>(&Context.GetHelper());
-	Helper.DrawEvents(*DrawState, *this, OffsetY);
+
+	if (Context.GetEventFilter().IsValid())
+	{
+		Helper.DrawFadedEvents(*DrawState, *this, OffsetY, 0.1f);
+
+		if (FilteredDrawStateInfo.Opacity == 1.0f)
+		{
+			Helper.DrawEvents(*FilteredDrawState, *this, OffsetY);
+		}
+		else
+		{
+			FilteredDrawStateInfo.Opacity = FMath::Min(1.0f, FilteredDrawStateInfo.Opacity + 0.05f);
+			Helper.DrawFadedEvents(*FilteredDrawState, *this, OffsetY, FilteredDrawStateInfo.Opacity);
+		}
+	}
+	else
+	{
+		Helper.DrawEvents(*DrawState, *this, OffsetY);
+	}
 }
 
 ////////////////////////////////////////////////////////////////////////////////////////////////////
@@ -205,4 +318,21 @@ const TSharedPtr<const ITimingEvent> FTimingEventsTrack::GetEvent(float InPosX,
 
 ////////////////////////////////////////////////////////////////////////////////////////////////////
 
+TSharedPtr<ITimingEventFilter> FTimingEventsTrack::GetFilterByEvent(const TSharedPtr<const ITimingEvent> InTimingEvent) const
+{
+	if (InTimingEvent.IsValid() && FTimingEvent::CheckTypeName(*InTimingEvent))
+	{
+		const FTimingEvent& Event = *StaticCastSharedPtr<const FTimingEvent, const ITimingEvent>(InTimingEvent);
+
+		TSharedRef<FTimingEventFilter> EventFilterRef = MakeShared<FTimingEventFilter>();
+		EventFilterRef->SetFilterByEventType(true);
+		EventFilterRef->SetEventType(Event.GetType());
+
+		return EventFilterRef;
+	}
+	return nullptr;
+}
+
+////////////////////////////////////////////////////////////////////////////////////////////////////
+
 #undef LOCTEXT_NAMESPACE
diff --git a/Engine/Source/Developer/TraceInsights/Private/Insights/ViewModels/TimingGraphTrack.cpp b/Engine/Source/Developer/TraceInsights/Private/Insights/ViewModels/TimingGraphTrack.cpp
--- a/Engine/Source/Developer/TraceInsights/Private/Insights/ViewModels/TimingGraphTrack.cpp
+++ b/Engine/Source/Developer/TraceInsights/Private/Insights/ViewModels/TimingGraphTrack.cpp
@@ -205,7 +205,7 @@ void FTimingGraphTrack::Update(const ITimingTrackUpdateContext& Context)
 		{
 			if (Series->IsVisible() && (bIsEntireGraphTrackDirty || Series->IsDirty()))
 			{
-				// Clear the flag before updating, becasue the update itself may furter need to set the series as dirty.
+				// Clear the flag before updating, becasue the update itself may further need to set the series as dirty.
 				Series->ClearDirtyFlag();
 
 				TSharedPtr<FTimingGraphSeries> TimingSeries = StaticCastSharedPtr<FTimingGraphSeries>(Series);
diff --git a/Engine/Source/Developer/TraceInsights/Private/Insights/ViewModels/TimingViewDrawHelper.cpp b/Engine/Source/Developer/TraceInsights/Private/Insights/ViewModels/TimingViewDrawHelper.cpp
--- a/Engine/Source/Developer/TraceInsights/Private/Insights/ViewModels/TimingViewDrawHelper.cpp
+++ b/Engine/Source/Developer/TraceInsights/Private/Insights/ViewModels/TimingViewDrawHelper.cpp
@@ -182,11 +182,6 @@ void FTimingEventsTrackDrawStateBuilder::AddEvent(double EventStartTime, double
 	//constexpr float BorderColorFactor = 0.75f; // darker border
 	constexpr float BorderColorFactor = 1.25f; // brighter border
 
-	//TODO if (HighlightedEventTypeId != FTimingEvent::InvalidTypeId && TypeId != HighlightedEventTypeId)
-	//{
-	//	EventColorFill.A = 0.1f;
-	//}
-
 	//////////////////////////////////////////////////
 
 	// Save X2, for current depth.
@@ -323,7 +318,6 @@ FTimingViewDrawHelper::FTimingViewDrawHelper(const FDrawContext& InDrawContext,
 	, EventFont(FCoreStyle::GetDefaultFontStyle("Regular", 8))
 	, ValidAreaX(0.0f)
 	, ValidAreaW(0.0f)
-	, HighlightedEventTypeId(uint64(-1))
 	, NumEvents(0)
 	, NumMergedBoxes(0)
 	, NumDrawBoxes(0)
@@ -420,6 +414,80 @@ void FTimingViewDrawHelper::DrawEvents(const FTimingEventsTrackDrawState& DrawSt
 
 ////////////////////////////////////////////////////////////////////////////////////////////////////
 
+void FTimingViewDrawHelper::DrawFadedEvents(const FTimingEventsTrackDrawState& DrawState, const FTimingEventsTrack& Track, const float OffsetY, const float Opacity) const
+{
+	const float TrackY = Track.GetPosY();
+	const float TrackH = Track.GetHeight();
+
+	if (TrackH > 0.0f &&
+		TrackY + TrackH > Viewport.GetTopOffset() &&
+		TrackY < Viewport.GetHeight() - Viewport.GetBottomOffset())
+	{
+		const FTimingViewLayout& Layout = Viewport.GetLayout();
+
+		NumEvents += DrawState.GetNumEvents();
+		NumMergedBoxes += DrawState.GetNumMergedBoxes();
+
+		const float TopLaneY = TrackY + OffsetY + Layout.TimelineDY;
+
+		// Draw filled boxes (merged borders).
+		//if (Layout.EventH > 0.0f)
+		{
+			const int32 EventFillLayerId = ReservedLayerId + ToInt32(EDrawLayer::EventFill);
+			const float EventFillH = Layout.EventH;
+			for (const FTimingEventsTrackDrawState::FBoxPrimitive& Box : DrawState.Boxes)
+			{
+				const float Y = TopLaneY + (Layout.EventH + Layout.EventDY) * Box.Depth;
+				DrawContext.DrawBox(EventFillLayerId, Box.X, Y, Box.W, EventFillH, WhiteBrush, Box.Color.CopyWithNewOpacity(Opacity));
+			}
+			NumDrawBoxes += DrawState.Boxes.Num();
+		}
+
+		// Draw filled boxes (event inside area).
+		if (Layout.EventH > 2.0f)
+		{
+			const int32 EventFillLayerId = ReservedLayerId + ToInt32(EDrawLayer::EventFill);
+			const float EventFillH = Layout.EventH - 2.0f;
+			for (const FTimingEventsTrackDrawState::FBoxPrimitive& Box : DrawState.InsideBoxes)
+			{
+				const float Y = TopLaneY + (Layout.EventH + Layout.EventDY) * Box.Depth + 1.0f;
+				DrawContext.DrawBox(EventFillLayerId, Box.X, Y, Box.W, EventFillH, WhiteBrush, Box.Color.CopyWithNewOpacity(Opacity));
+			}
+			NumDrawBoxes += DrawState.InsideBoxes.Num();
+		}
+
+		// Draw borders.
+		//if (Layout.EventH > 0.0f)
+		{
+			const int32 EventBorderLayerId = ReservedLayerId + ToInt32(EDrawLayer::EventBorder);
+			const float EventBorderH = Layout.EventH;
+			for (const FTimingEventsTrackDrawState::FBoxPrimitive& Box : DrawState.Borders)
+			{
+				const float Y = TopLaneY + (Layout.EventH + Layout.EventDY) * Box.Depth;
+				DrawContext.DrawBox(EventBorderLayerId, Box.X, Y, Box.W, EventBorderH, EventBorderBrush, Box.Color.CopyWithNewOpacity(Opacity));
+			}
+			NumDrawBorders += DrawState.Borders.Num();
+		}
+
+		// Draw texts.
+		if (Layout.EventH > 10.0f)
+		{
+			const FLinearColor WhiteColor(1.0f, 1.0f, 1.0f, Opacity);
+			const FLinearColor BlackColor(0.0f, 0.0f, 0.0f, Opacity);
+
+			const int32 EventTextLayerId = ReservedLayerId + ToInt32(EDrawLayer::EventText);
+			for (const FTimingEventsTrackDrawState::FTextPrimitive& Text : DrawState.Texts)
+			{
+				const float Y = TopLaneY + (Layout.EventH + Layout.EventDY) * Text.Depth + 1.0f;
+				DrawContext.DrawText(EventTextLayerId, Text.X, Y, Text.Text, EventFont, Text.bWhite ? WhiteColor : BlackColor);
+			}
+			NumDrawTexts += DrawState.Texts.Num();
+		}
+	}
+}
+
+////////////////////////////////////////////////////////////////////////////////////////////////////
+
 void FTimingViewDrawHelper::DrawTrackHeader(const FTimingEventsTrack& Track) const
 {
 	const float TrackY = Track.GetPosY();
diff --git a/Engine/Source/Developer/TraceInsights/Private/Insights/ViewModels/TimingViewDrawHelper.h b/Engine/Source/Developer/TraceInsights/Private/Insights/ViewModels/TimingViewDrawHelper.h
--- a/Engine/Source/Developer/TraceInsights/Private/Insights/ViewModels/TimingViewDrawHelper.h
+++ b/Engine/Source/Developer/TraceInsights/Private/Insights/ViewModels/TimingViewDrawHelper.h
@@ -173,15 +173,22 @@ public:
 
 	void DrawBackground() const;
 
+	//////////////////////////////////////////////////
+
 	void BeginDrawTracks() const;
+
 	// OffsetY = 1.0f is for the top horizontal line (which separates the timelines) added by DrawTrackHeader.
 	void DrawEvents(const FTimingEventsTrackDrawState& DrawState, const FTimingEventsTrack& Track, const float OffsetY = 1.0f) const;
+
+	void DrawFadedEvents(const FTimingEventsTrackDrawState& DrawState, const FTimingEventsTrack& Track, const float OffsetY = 1.0f, const float Opacity = 0.1f) const;
+
 	void DrawTrackHeader(const FTimingEventsTrack& Track) const;
+
 	void EndDrawTracks() const;
 
-	void DrawTimingEventHighlight(double StartTime, double EndTime, float Y, EDrawEventMode Mode) const;
+	//////////////////////////////////////////////////
 
-	void SetHighlightedEventTypeId(uint64 InHighlightedEventTypeId) { HighlightedEventTypeId = InHighlightedEventTypeId; }
+	void DrawTimingEventHighlight(double StartTime, double EndTime, float Y, EDrawEventMode Mode) const;
 
 	int32 GetNumEvents() const { return NumEvents; }
 	int32 GetNumMergedBoxes() const { return NumMergedBoxes; }
@@ -208,8 +215,6 @@ private:
 	mutable float ValidAreaX;
 	mutable float ValidAreaW;
 
-	uint64 HighlightedEventTypeId;
-
 	mutable int32 NumEvents;
 	mutable int32 NumMergedBoxes;
 	mutable int32 NumDrawBoxes;
diff --git a/Engine/Source/Developer/TraceInsights/Private/Insights/Widgets/STimersView.cpp b/Engine/Source/Developer/TraceInsights/Private/Insights/Widgets/STimersView.cpp
--- a/Engine/Source/Developer/TraceInsights/Private/Insights/Widgets/STimersView.cpp
+++ b/Engine/Source/Developer/TraceInsights/Private/Insights/Widgets/STimersView.cpp
@@ -822,7 +822,27 @@ void STimersView::TreeView_OnMouseButtonDoubleClick(FTimerNodePtr TimerNodePtr)
 		TSharedPtr<STimingView> TimingView = Wnd.IsValid() ? Wnd->GetTimingView() : nullptr;
 		if (TimingView.IsValid())
 		{
-			TimingView->SetHighlightedEventTypeId(TimerNodePtr->GetId());
+			bool bSameFilter = false;
+
+			const TSharedPtr<ITimingEventFilter> EventFilterPtr = TimingView->GetEventFilter();
+			if (EventFilterPtr.IsValid())
+			{
+				const FTimingEventFilter& EventFilter = static_cast<const FTimingEventFilter&>(*EventFilterPtr);
+				if (EventFilter.IsFilteringByEventType() &&
+					EventFilter.GetEventType() == TimerNodePtr->GetId())
+				{
+					bSameFilter = true;
+					TimingView->SetEventFilter(nullptr); // reset filter
+				}
+			}
+
+			if (!bSameFilter)
+			{
+				TSharedRef<FTimingEventFilter> EventFilterRef = MakeShared<FTimingEventFilter>();
+				EventFilterRef->SetFilterByEventType(true);
+				EventFilterRef->SetEventType(TimerNodePtr->GetId());
+				TimingView->SetEventFilter(EventFilterRef); // set new filter
+			}
 		}
 	}
 }
diff --git a/Engine/Source/Developer/TraceInsights/Private/Insights/Widgets/STimingView.cpp b/Engine/Source/Developer/TraceInsights/Private/Insights/Widgets/STimingView.cpp
--- a/Engine/Source/Developer/TraceInsights/Private/Insights/Widgets/STimingView.cpp
+++ b/Engine/Source/Developer/TraceInsights/Private/Insights/Widgets/STimingView.cpp
@@ -267,6 +267,7 @@ void STimingView::Reset()
 		HoveredEvent.Reset();
 		OnHoveredEventChangedDelegate.Broadcast(HoveredEvent);
 	}
+
 	if (SelectedTrack.IsValid())
 	{
 		SelectedTrack.Reset();
@@ -278,6 +279,11 @@ void STimingView::Reset()
 		OnSelectedEventChangedDelegate.Broadcast(SelectedEvent);
 	}
 
+	if (TimingEventFilter.IsValid())
+	{
+		TimingEventFilter.Reset();
+	}
+
 	bPreventThrottling = false;
 
 	Tooltip.Reset();
@@ -287,8 +293,6 @@ void STimingView::Reset()
 	TimeMarker = std::numeric_limits<double>::infinity();
 	bIsScrubbing = false;
 
-	HighlightedEventTypeId = uint64(-1);
-
 	//ThisGeometry
 
 	//////////////////////////////////////////////////
@@ -428,6 +432,7 @@ void STimingView::Tick(const FGeometry& AllottedGeometry, const double InCurrent
 		virtual const FVector2D& GetMousePosition() const override { return TimingView->GetMousePosition(); }
 		virtual const TSharedPtr<const ITimingEvent> GetHoveredEvent() const override { return TimingView->GetHoveredEvent(); }
 		virtual const TSharedPtr<const ITimingEvent> GetSelectedEvent() const override { return TimingView->GetSelectedEvent(); }
+		virtual const TSharedPtr<ITimingEventFilter> GetEventFilter() const override { return TimingView->GetEventFilter(); }
 		virtual double GetCurrentTime() const override { return CurrentTime; }
 		virtual float GetDeltaTime() const override { return DeltaTime; }
 
@@ -544,16 +549,26 @@ void STimingView::Tick(const FGeometry& AllottedGeometry, const double InCurrent
 	////////////////////////////////////////////////////////////////////////////////////////////////////
 
 	const float InitialScrollPosY = Viewport.GetScrollPosY();
-	const float InitialPinnedTrackPosY = SelectedTrack.IsValid() ? SelectedTrack->GetPosY() : 0.0f;
+
+	TSharedPtr<FBaseTimingTrack> SelectedScrollableTrack;
+	if (SelectedTrack.IsValid() && SelectedTrack->IsVisible())
+	{
+		if (ScrollableTracks.Contains(SelectedTrack))
+		{
+			SelectedScrollableTrack = SelectedTrack;
+		}
+	}
+
+	const float InitialPinnedTrackPosY = SelectedScrollableTrack.IsValid() ? SelectedScrollableTrack->GetPosY() : 0.0f;
 
 	// Update the Y position for visible scrollable tracks.
 	UpdatePositionForScrollableTracks();
 
 	// The selected track will be pinned (keeps Y pos fixed unless user scrolls vertically).
-	if (SelectedTrack.IsValid())
+	if (SelectedScrollableTrack.IsValid())
 	{
 		const float ScrollingDY = LastScrollPosY - InitialScrollPosY;
-		const float PinnedTrackPosY = SelectedTrack->GetPosY();
+		const float PinnedTrackPosY = SelectedScrollableTrack->GetPosY();
 		const float AdjustmentDY = InitialPinnedTrackPosY - PinnedTrackPosY + ScrollingDY;
 
 		if (!FMath::IsNearlyZero(AdjustmentDY, 0.5f))
@@ -848,6 +863,7 @@ int32 STimingView::OnPaint(const FPaintArgs& Args, const FGeometry& AllottedGeom
 		virtual const FVector2D& GetMousePosition() const override { return TimingView->GetMousePosition(); }
 		virtual const TSharedPtr<const ITimingEvent> GetHoveredEvent() const override { return TimingView->GetHoveredEvent(); }
 		virtual const TSharedPtr<const ITimingEvent> GetSelectedEvent() const override { return TimingView->GetSelectedEvent(); }
+		virtual const TSharedPtr<ITimingEventFilter> GetEventFilter() const override { return TimingView->GetEventFilter(); }
 		virtual FDrawContext& GetDrawContext() const override { return DrawContext; }
 		virtual const ITimingViewDrawHelper& GetHelper() const override { return Helper; }
 
@@ -862,8 +878,6 @@ int32 STimingView::OnPaint(const FPaintArgs& Args, const FGeometry& AllottedGeom
 
 	////////////////////////////////////////////////////////////////////////////////////////////////////
 
-	Helper.SetHighlightedEventTypeId(HighlightedEventTypeId);
-
 	// Draw background.
 	Helper.DrawBackground();
 
@@ -1375,12 +1389,12 @@ FReply STimingView::OnMouseButtonDown(const FGeometry& MyGeometry, const FPointe
 		Reply = FReply::Handled().CaptureMouse(SharedThis(this));
 	}
 
-	if(bPreventThrottling)
+	if (bPreventThrottling)
 	{
 		Reply.PreventThrottling();
 	}
 
-	if(bStartScrubbing)
+	if (bStartScrubbing)
 	{
 		bIsPanning = false;
 		bIsDragging = false;
@@ -1541,7 +1555,15 @@ FReply STimingView::OnMouseButtonDoubleClick(const FGeometry& MyGeometry, const
 			}
 			else
 			{
-				//TODO: HighlightedEventTypeId = HoveredTimingEvent.TypeId;
+				SetEventFilter(HoveredEvent->GetTrack()->GetFilterByEvent(HoveredEvent));
+			}
+		}
+		else
+		{
+			if (TimingEventFilter.IsValid())
+			{
+				TimingEventFilter.Reset();
+				Viewport.AddDirtyFlags(ETimingTrackViewportDirtyFlags::HInvalidated);
 			}
 		}
 
@@ -2316,6 +2338,7 @@ void STimingView::SetTimeMarkersVisible(bool bIsMarkersTrackVisible)
 		{
 			if (Viewport.GetScrollPosY() != 0.0f)
 			{
+				UE_LOG(TimingProfiler, Log, TEXT("SetTimeMarkersVisible!!!"));
 				Viewport.SetScrollPosY(Viewport.GetScrollPosY() + MarkersTrack->GetHeight());
 			}
 
@@ -2323,6 +2346,7 @@ void STimingView::SetTimeMarkersVisible(bool bIsMarkersTrackVisible)
 		}
 		else
 		{
+			UE_LOG(TimingProfiler, Log, TEXT("SetTimeMarkersVisible!!!"));
 			Viewport.SetScrollPosY(Viewport.GetScrollPosY() - MarkersTrack->GetHeight());
 		}
 	}
@@ -2341,6 +2365,7 @@ void STimingView::SetDrawOnlyBookmarks(bool bIsBookmarksTrack)
 		{
 			if (Viewport.GetScrollPosY() != 0.0f)
 			{
+				UE_LOG(TimingProfiler, Log, TEXT("SetDrawOnlyBookmarks!!!"));
 				Viewport.SetScrollPosY(Viewport.GetScrollPosY() + MarkersTrack->GetHeight() - PrevHeight);
 			}
 
@@ -2468,11 +2493,6 @@ void STimingView::OnSelectedTimingEventChanged()
 		}
 	}
 
-	if (!SelectedEvent.IsValid())
-	{
-		HighlightedEventTypeId = uint64(-1);
-	}
-
 	OnSelectedEventChangedDelegate.Broadcast(SelectedEvent);
 }
 
@@ -2703,6 +2723,14 @@ void STimingView::FrameSelection()
 
 ////////////////////////////////////////////////////////////////////////////////////////////////////
 
+void STimingView::SetEventFilter(const TSharedPtr<ITimingEventFilter> InEventFilter)
+{
+	TimingEventFilter = InEventFilter;
+	Viewport.AddDirtyFlags(ETimingTrackViewportDirtyFlags::HInvalidated);
+}
+
+////////////////////////////////////////////////////////////////////////////////////////////////////
+
 TSharedRef<SWidget> STimingView::MakeTracksFilterMenu()
 {
 	FMenuBuilder MenuBuilder(/*bInShouldCloseWindowAfterMenuSelection=*/true, nullptr);
diff --git a/Engine/Source/Developer/TraceInsights/Private/Insights/Widgets/STimingView.h b/Engine/Source/Developer/TraceInsights/Private/Insights/Widgets/STimingView.h
--- a/Engine/Source/Developer/TraceInsights/Private/Insights/Widgets/STimingView.h
+++ b/Engine/Source/Developer/TraceInsights/Private/Insights/Widgets/STimingView.h
@@ -278,14 +278,15 @@ public:
 
 	TSharedPtr<FTimingGraphTrack> GetMainTimingGraphTrack() { return GraphTrack; }
 
-	uint64 GetHighlightedEventTypeId() const { return HighlightedEventTypeId; }
-	void SetHighlightedEventTypeId(uint64 InHighlightedEventTypeId) { HighlightedEventTypeId = InHighlightedEventTypeId; }
-
 	const TSharedPtr<FBaseTimingTrack> GetHoveredTrack() const { return HoveredTrack; }
 	const TSharedPtr<const ITimingEvent> GetHoveredEvent() const { return HoveredEvent; }
+
 	const TSharedPtr<FBaseTimingTrack> GetSelectedTrack() const { return SelectedTrack; }
 	const TSharedPtr<const ITimingEvent> GetSelectedEvent() const { return SelectedEvent; }
 
+	const TSharedPtr<ITimingEventFilter> GetEventFilter() const { return TimingEventFilter; }
+	void SetEventFilter(const TSharedPtr<ITimingEventFilter> InEventFilter);
+
 	const TSharedPtr<FBaseTimingTrack> GetTrackAt(float InPosX, float InPosY) const;
 
 protected:
@@ -469,9 +470,12 @@ protected:
 
 	TSharedPtr<FBaseTimingTrack> HoveredTrack;
 	TSharedPtr<const ITimingEvent> HoveredEvent;
+
 	TSharedPtr<FBaseTimingTrack> SelectedTrack;
 	TSharedPtr<const ITimingEvent> SelectedEvent;
 
+	TSharedPtr<ITimingEventFilter> TimingEventFilter;
+
 	FTooltipDrawState Tooltip;
 
 	enum class ESelectionType
@@ -490,8 +494,6 @@ protected:
 	/** True of the user is currently dragging the time marker */
 	bool bIsScrubbing;
 
-	uint64 HighlightedEventTypeId;
-
 	////////////////////////////////////////////////////////////
 	// Misc
 
@@ -522,4 +524,5 @@ protected:
 	Insights::FHoveredEventChangedDelegate OnHoveredEventChangedDelegate;
 	Insights::FSelectedTrackChangedDelegate OnSelectedTrackChangedDelegate;
 	Insights::FSelectedEventChangedDelegate OnSelectedEventChangedDelegate;
+	//Insights::FFilteredEventChangedDelegate OnFilteredEventChangedDelegate;
 };
diff --git a/Engine/Source/Developer/TraceInsights/Public/Insights/ViewModels/BaseTimingTrack.h b/Engine/Source/Developer/TraceInsights/Public/Insights/ViewModels/BaseTimingTrack.h
--- a/Engine/Source/Developer/TraceInsights/Public/Insights/ViewModels/BaseTimingTrack.h
+++ b/Engine/Source/Developer/TraceInsights/Public/Insights/ViewModels/BaseTimingTrack.h
@@ -15,6 +15,7 @@ class FTimingTrackViewport;
 class ITimingViewDrawHelper;
 class FTooltipDrawState;
 class ITimingEvent;
+class ITimingEventFilter;
 
 struct FDrawContext;
 
@@ -27,6 +28,7 @@ public:
 	virtual const FVector2D& GetMousePosition() const = 0;
 	virtual const TSharedPtr<const ITimingEvent> GetHoveredEvent() const = 0;
 	virtual const TSharedPtr<const ITimingEvent> GetSelectedEvent() const = 0;
+	virtual const TSharedPtr<ITimingEventFilter> GetEventFilter() const = 0;
 	virtual double GetCurrentTime() const = 0;
 	virtual float GetDeltaTime() const = 0;
 };
@@ -40,6 +42,7 @@ public:
 	virtual const FVector2D& GetMousePosition() const = 0;
 	virtual const TSharedPtr<const ITimingEvent> GetHoveredEvent() const = 0;
 	virtual const TSharedPtr<const ITimingEvent> GetSelectedEvent() const = 0;
+	virtual const TSharedPtr<ITimingEventFilter> GetEventFilter() const = 0;
 	virtual FDrawContext& GetDrawContext() const = 0;
 	virtual const ITimingViewDrawHelper& GetHelper() const = 0;
 };
@@ -194,6 +197,9 @@ public:
 	// Search for an event using custom parameters.
 	virtual const TSharedPtr<const ITimingEvent> SearchEvent(const FTimingEventSearchParameters& InSearchParameters) const { return nullptr; }
 
+	// Get the filter object for filtering all events similar with a specified event. Used when double clicked on an event.
+	virtual TSharedPtr<ITimingEventFilter> GetFilterByEvent(const TSharedPtr<const ITimingEvent> InTimingEvent) const { return nullptr; }
+
 	// Allows tracks to update event stats that are slower to compute (called at a lower frequency than GetEventAtPosition or Search or SearchTimingEvent).
 	virtual void UpdateEventStats(ITimingEvent& InOutEvent) const {}
 
diff --git a/Engine/Source/Developer/TraceInsights/Public/Insights/ViewModels/ITimingEvent.h b/Engine/Source/Developer/TraceInsights/Public/Insights/ViewModels/ITimingEvent.h
--- a/Engine/Source/Developer/TraceInsights/Public/Insights/ViewModels/ITimingEvent.h
+++ b/Engine/Source/Developer/TraceInsights/Public/Insights/ViewModels/ITimingEvent.h
@@ -33,3 +33,14 @@ public:
 		return A.IsValid() && B.IsValid() && (*A).Equals(*B);
 	}
 };
+
+class TRACEINSIGHTS_API ITimingEventFilter
+{
+public:
+	virtual const FName& GetTypeName() const = 0;
+
+	virtual bool FilterTrack(const FBaseTimingTrack& InTrack) const = 0;
+	virtual bool FilterEvent(double InEventStartTime, double InEventEndTime, uint32 InEventDepth, const TCHAR* InEventName = nullptr, uint64 InEventType = 0, uint32 InEventColor = 0) const = 0;
+
+	virtual uint32 GetChangeNumber() const = 0;
+};
diff --git a/Engine/Source/Developer/TraceInsights/Public/Insights/ViewModels/TimingEvent.h b/Engine/Source/Developer/TraceInsights/Public/Insights/ViewModels/TimingEvent.h
--- a/Engine/Source/Developer/TraceInsights/Public/Insights/ViewModels/TimingEvent.h
+++ b/Engine/Source/Developer/TraceInsights/Public/Insights/ViewModels/TimingEvent.h
@@ -15,12 +15,13 @@ class FBaseTimingTrack;
 class TRACEINSIGHTS_API FTimingEvent : public ITimingEvent
 {
 public:
-	FTimingEvent(const TSharedRef<const FBaseTimingTrack> InTrack, double InStartTime, double InEndTime, uint32 InDepth)
+	FTimingEvent(const TSharedRef<const FBaseTimingTrack> InTrack, double InStartTime, double InEndTime, uint32 InDepth, uint64 InType = uint64(-1))
 		: Track(InTrack)
 		, StartTime(InStartTime)
 		, EndTime(InEndTime)
 		, ExclusiveTime(0.0)
 		, Depth(InDepth)
+		, Type(InType)
 	{}
 
 	virtual ~FTimingEvent() {}
@@ -65,6 +66,8 @@ public:
 	double GetExclusiveTime() const { return ExclusiveTime; }
 	void SetExclusiveTime(double InExclusiveTime) { ExclusiveTime = InExclusiveTime; }
 
+	uint64 GetType() const { return Type; }
+
 	static const FName& GetStaticTypeName() { return FTimingEvent::TypeName; }
 	static bool CheckTypeName(const ITimingEvent& Event) { return Event.GetTypeName() == FTimingEvent::TypeName; }
 
@@ -87,6 +90,100 @@ private:
 	// The depth of the event
 	uint32 Depth;
 
+	// The event type (category, group id, etc.).
+	uint64 Type;
+
+	static const FName TypeName;
+};
+
+////////////////////////////////////////////////////////////////////////////////////////////////////
+
+class TRACEINSIGHTS_API FTimingEventFilter : public ITimingEventFilter
+{
+public:
+	FTimingEventFilter()
+		: bFilterByTrackType(false)
+		, TrackType(NAME_None)
+		, bFilterByEventType(false)
+		, EventType(uint64(-1))
+		, bFilterByMinDuration(false)
+		, MinDuration(0.0)
+	{}
+
+	virtual ~FTimingEventFilter() {}
+
+	FTimingEventFilter(const FTimingEventFilter&) = default;
+	FTimingEventFilter& operator=(const FTimingEventFilter&) = default;
+
+	FTimingEventFilter(FTimingEventFilter&&) = default;
+	FTimingEventFilter& operator=(FTimingEventFilter&&) = default;
+
+	//////////////////////////////////////////////////
+	// ITimingEventFilter interface
+
+	virtual const FName& GetTypeName() const override { return FTimingEventFilter::TypeName; }
+
+	virtual bool FilterTrack(const FBaseTimingTrack& InTrack) const override;
+
+	virtual bool FilterEvent(double InEventStartTime, double InEventEndTime, uint32 InEventDepth, const TCHAR* InEventName, uint64 InEventType = 0, uint32 InEventColor = 0) const override
+	{
+		return (!bFilterByEventType || InEventType == EventType)
+			&& (!bFilterByMinDuration || InEventEndTime - InEventStartTime >= MinDuration)
+			&& (!bFilterByMaxDuration || InEventEndTime - InEventStartTime <= MaxDuration);
+	}
+
+	virtual uint32 GetChangeNumber() const override { return ChangeNumber; }
+
+	//////////////////////////////////////////////////
+
+	bool IsFilteringByTrackType() const { return bFilterByTrackType; }
+	const FName& GetTrackType() const { return TrackType; }
+	void SetFilterByTrackType(bool bInFilterByTrackType) { if (bFilterByTrackType != bInFilterByTrackType) { bFilterByTrackType = bInFilterByTrackType; ChangeNumber++; } }
+	void SetTrackType(const FName InTrackType) { if (TrackType != InTrackType) { TrackType = InTrackType; ChangeNumber++; } }
+
+	bool IsFilteringByEventType() const { return bFilterByEventType; }
+	uint64 GetEventType() const { return EventType; }
+	void SetFilterByEventType(bool bInFilterByEventType) { if (bFilterByEventType != bInFilterByEventType) { bFilterByEventType = bInFilterByEventType; ChangeNumber++; } }
+	void SetEventType(uint64 InEventType) { if (EventType != InEventType) { EventType = InEventType; ChangeNumber++; } }
+
+	bool IsFilteringByMinDuration() const { return bFilterByMinDuration; }
+	double GetMinDuration() const { return MinDuration; }
+
+	bool IsFilteringByMaxDuration() const { return bFilterByMaxDuration; }
+	double GetMaxDuration() const { return MaxDuration; }
+
+	static const FName& GetStaticTypeName() { return FTimingEventFilter::TypeName; }
+	static bool CheckTypeName(const ITimingEventFilter& EventFilter) { return EventFilter.GetTypeName() == FTimingEventFilter::TypeName; }
+
+private:
+	bool FilterEventByType(double InEventStartTime, double InEventEndTime, uint32 InEventDepth, const TCHAR* InEventName, uint64 InEventType, uint32 InEventColor) const
+	{
+		return InEventType == EventType;
+	}
+	bool FilterEventByMinDuration(double InEventStartTime, double InEventEndTime, uint32 InEventDepth, const TCHAR* InEventName, uint64 InEventType, uint32 InEventColor) const
+	{
+		return InEventEndTime - InEventStartTime >= MinDuration;
+	}
+	bool FilterEventByMaxDuration(double InEventStartTime, double InEventEndTime, uint32 InEventDepth, const TCHAR* InEventName, uint64 InEventType, uint32 InEventColor) const
+	{
+		return InEventEndTime - InEventStartTime <= MaxDuration;
+	}
+
+private:
+	bool bFilterByTrackType;
+	FName TrackType;
+
+	bool bFilterByEventType;
+	uint64 EventType;
+
+	bool bFilterByMinDuration;
+	double MinDuration;
+
+	bool bFilterByMaxDuration;
+	double MaxDuration;
+
+	uint32 ChangeNumber;
+
 	static const FName TypeName;
 };
 
diff --git a/Engine/Source/Developer/TraceInsights/Public/Insights/ViewModels/TimingEventsTrack.h b/Engine/Source/Developer/TraceInsights/Public/Insights/ViewModels/TimingEventsTrack.h
--- a/Engine/Source/Developer/TraceInsights/Public/Insights/ViewModels/TimingEventsTrack.h
+++ b/Engine/Source/Developer/TraceInsights/Public/Insights/ViewModels/TimingEventsTrack.h
@@ -38,9 +38,12 @@ public:
 
 	virtual const TSharedPtr<const ITimingEvent> GetEvent(float InPosX, float InPosY, const FTimingTrackViewport& Viewport) const override;
 
+	virtual TSharedPtr<ITimingEventFilter> GetFilterByEvent(const TSharedPtr<const ITimingEvent> InTimingEvent) const override;
+
 	//////////////////////////////////////////////////
 
 	virtual void BuildDrawState(ITimingEventsTrackDrawStateBuilder& Builder, const ITimingTrackUpdateContext& Context) = 0;
+	virtual void BuildFilteredDrawState(ITimingEventsTrackDrawStateBuilder& Builder, const ITimingTrackUpdateContext& Context) {}
 
 protected:
 	int32 GetNumLanes() const { return NumLanes; }
@@ -57,9 +60,22 @@ protected:
 private:
 	int32 NumLanes; // number of lanes (sub-tracks)
 	TSharedRef<struct FTimingEventsTrackDrawState> DrawState;
+	TSharedRef<struct FTimingEventsTrackDrawState> FilteredDrawState;
+
+	struct FFilteredDrawStateInfo
+	{
+		double ViewportStartTime = 0.0;
+		double ViewportScaleX = 0.0;
+		double LastBuildDuration = 0.0;
+		TWeakPtr<ITimingEventFilter> LastEventFilter;
+		uint32 LastFilterChangeNumber = 0;
+		uint32 Counter = 0;
+		mutable float Opacity = 0.0f;
+	};
+	FFilteredDrawStateInfo FilteredDrawStateInfo;
 
 public:
-	static bool bUseDownSampling;
+	static bool bUseDownSampling; // toggle to enable/disbale downsampling, for debugging purposes only
 };
 
 ////////////////////////////////////////////////////////////////////////////////////////////////////
