[Insights] Fixed static analysis warning.

diff --git a/Engine/Source/Developer/TraceInsights/Private/Insights/ViewModels/TimingEventsTrack.cpp b/Engine/Source/Developer/TraceInsights/Private/Insights/ViewModels/TimingEventsTrack.cpp
--- a/Engine/Source/Developer/TraceInsights/Private/Insights/ViewModels/TimingEventsTrack.cpp
+++ b/Engine/Source/Developer/TraceInsights/Private/Insights/ViewModels/TimingEventsTrack.cpp
@@ -109,11 +109,7 @@ void FTimingEventsTrack::PreUpdate(const ITimingTrackUpdateContext& Context)
 				{
 					FilteredDrawStateInfo.ViewportStartTime = Context.GetViewport().GetStartTime();
 					FilteredDrawStateInfo.ViewportScaleX = Context.GetViewport().GetScaleX();
-
-					if (FilteredDrawStateInfo.Counter >= 0)
-					{
-						FilteredDrawStateInfo.Counter = 1; // wait
-					}
+					FilteredDrawStateInfo.Counter = 1; // wait
 				}
 			}
 
