[Insights]
  - Timing view: Added warning red line on sides of timing view when trying to overscroll using key shortcuts (Ctrl + Left/Right/Up/Down) or using Shift + Mouse Wheel.

diff --git a/Engine/Source/Developer/TraceInsights/Private/Insights/Widgets/STimingView.cpp b/Engine/Source/Developer/TraceInsights/Private/Insights/Widgets/STimingView.cpp
--- a/Engine/Source/Developer/TraceInsights/Private/Insights/Widgets/STimingView.cpp
+++ b/Engine/Source/Developer/TraceInsights/Private/Insights/Widgets/STimingView.cpp
@@ -253,6 +253,11 @@ void STimingView::Reset()
 	bIsPanning = false;
 	PanningMode = EPanningMode::None;
 
+	OverscrollLeft = 0.0f;
+	OverscrollRight = 0.0f;
+	OverscrollTop = 0.0f;
+	OverscrollBottom = 0.0f;
+
 	bIsSelecting = false;
 	SelectionStartTime = 0.0;
 	SelectionEndTime = 0.0;
@@ -351,6 +356,24 @@ void STimingView::Tick(const FGeometry& AllottedGeometry, const double InCurrent
 
 	bPreventThrottling = false;
 
+	constexpr float OverscrollFadeSpeed = 2.0f;
+	if (OverscrollLeft > 0.0f)
+	{
+		OverscrollLeft = FMath::Max(0.0f, OverscrollLeft - InDeltaTime * OverscrollFadeSpeed);
+	}
+	if (OverscrollRight > 0.0f)
+	{
+		OverscrollRight = FMath::Max(0.0f, OverscrollRight - InDeltaTime * OverscrollFadeSpeed);
+	}
+	if (OverscrollTop > 0.0f)
+	{
+		OverscrollTop = FMath::Max(0.0f, OverscrollTop - InDeltaTime * OverscrollFadeSpeed);
+	}
+	if (OverscrollBottom > 0.0f)
+	{
+		OverscrollBottom = FMath::Max(0.0f, OverscrollBottom - InDeltaTime * OverscrollFadeSpeed);
+	}
+
 	const float ViewWidth = AllottedGeometry.GetLocalSize().X;
 	const float ViewHeight = AllottedGeometry.GetLocalSize().Y;
 
@@ -1060,6 +1083,51 @@ int32 STimingView::OnPaint(const FPaintArgs& Args, const FGeometry& AllottedGeom
 
 	//////////////////////////////////////////////////
 
+	constexpr float OverscrollLineSize = 1.0f;
+	constexpr int32 OverscrollLineCount = 8;
+
+	if (OverscrollLeft > 0.0f)
+	{
+		// TODO: single box with gradient opacity
+		const float OverscrollLineY = Viewport.GetTopOffset();
+		const float OverscrollLineH = Viewport.GetScrollableAreaHeight();
+		for (int32 LineIndex = 0; LineIndex < OverscrollLineCount; ++LineIndex)
+		{
+			const float Opacity = OverscrollLeft * static_cast<float>(OverscrollLineCount - LineIndex) / static_cast<float>(OverscrollLineCount);
+			DrawContext.DrawBox(LineIndex * OverscrollLineSize, OverscrollLineY, OverscrollLineSize, OverscrollLineH, WhiteBrush, FLinearColor(1.0f, 0.1f, 0.1f, Opacity));
+		}
+	}
+	if (OverscrollRight > 0.0f)
+	{
+		const float OverscrollLineY = Viewport.GetTopOffset();
+		const float OverscrollLineH = Viewport.GetScrollableAreaHeight();
+		for (int32 LineIndex = 0; LineIndex < OverscrollLineCount; ++LineIndex)
+		{
+			const float Opacity = OverscrollRight * static_cast<float>(OverscrollLineCount - LineIndex) / static_cast<float>(OverscrollLineCount);
+			DrawContext.DrawBox(ViewWidth - (1 + LineIndex) * OverscrollLineSize, OverscrollLineY, OverscrollLineSize, OverscrollLineH, WhiteBrush, FLinearColor(1.0f, 0.1f, 0.1f, Opacity));
+		}
+	}
+	if (OverscrollTop > 0.0f)
+	{
+		const float OverscrollLineY = Viewport.GetTopOffset();
+		for (int32 LineIndex = 0; LineIndex < OverscrollLineCount; ++LineIndex)
+		{
+			const float Opacity = OverscrollTop * static_cast<float>(OverscrollLineCount - LineIndex) / static_cast<float>(OverscrollLineCount);
+			DrawContext.DrawBox(0.0f, OverscrollLineY + LineIndex * OverscrollLineSize, ViewWidth, OverscrollLineSize, WhiteBrush, FLinearColor(1.0f, 0.1f, 0.1f, Opacity));
+		}
+	}
+	if (OverscrollBottom > 0.0f)
+	{
+		const float OverscrollLineY = ViewHeight - Viewport.GetBottomOffset();
+		for (int32 LineIndex = 0; LineIndex < OverscrollLineCount; ++LineIndex)
+		{
+			const float Opacity = OverscrollBottom * static_cast<float>(OverscrollLineCount - LineIndex) / static_cast<float>(OverscrollLineCount);
+			DrawContext.DrawBox(0.0f, OverscrollLineY - (1 + LineIndex) * OverscrollLineSize, ViewWidth, OverscrollLineSize, WhiteBrush, FLinearColor(1.0f, 0.1f, 0.1f, Opacity));
+		}
+	}
+
+	//////////////////////////////////////////////////
+
 	const bool bShouldDisplayDebugInfo = FInsightsManager::Get()->IsDebugInfoEnabled();
 	if (bShouldDisplayDebugInfo)
 	{
@@ -1715,15 +1783,16 @@ FReply STimingView::OnMouseWheel(const FGeometry& MyGeometry, const FPointerEven
 		{
 			// Scroll vertically.
 			constexpr float ScrollSpeedY = 16.0f * 3;
-			const float ScrollPosY = Viewport.GetScrollPosY() - ScrollSpeedY * MouseEvent.GetWheelDelta();
-			ScrollAtPosY(ScrollPosY);
+			const float NewScrollPosY = Viewport.GetScrollPosY() - ScrollSpeedY * MouseEvent.GetWheelDelta();
+			ScrollAtPosY(EnforceVerticalScrollLimits(NewScrollPosY));
 		}
 	}
 	else if (MouseEvent.GetModifierKeys().IsControlDown())
 	{
 		// Scroll horizontally.
 		const double ScrollSpeedX = Viewport.GetDurationForViewportDX(16.0 * 3);
-		ScrollAtTime(Viewport.GetStartTime() - ScrollSpeedX * MouseEvent.GetWheelDelta());
+		const double NewStartTime = Viewport.GetStartTime() - ScrollSpeedX * MouseEvent.GetWheelDelta();
+		ScrollAtTime(EnforceHorizontalScrollLimits(NewStartTime));
 	}
 	else
 	{
@@ -1732,7 +1801,6 @@ FReply STimingView::OnMouseWheel(const FGeometry& MyGeometry, const FPointerEven
 		//MousePosition = MyGeometry.AbsoluteToLocal(MouseEvent.GetScreenSpacePosition());
 		if (Viewport.RelativeZoomWithFixedX(Delta, MousePosition.X))
 		{
-			//Viewport.EnforceHorizontalScrollLimits(1.0);
 			UpdateHorizontalScrollBar();
 		}
 	}
@@ -1892,7 +1960,6 @@ FReply STimingView::OnKeyDown(const FGeometry& MyGeometry, const FKeyEvent& InKe
 		const double ScaleX = Viewport.GetScaleX() * 1.25;
 		if (Viewport.ZoomWithFixedX(ScaleX, Viewport.GetWidth() / 2))
 		{
-			//Viewport.EnforceHorizontalScrollLimits(1.0);
 			UpdateHorizontalScrollBar();
 		}
 		return FReply::Handled();
@@ -1904,7 +1971,6 @@ FReply STimingView::OnKeyDown(const FGeometry& MyGeometry, const FKeyEvent& InKe
 		const double ScaleX = Viewport.GetScaleX() / 1.25;
 		if (Viewport.ZoomWithFixedX(ScaleX, Viewport.GetWidth() / 2))
 		{
-			//Viewport.EnforceHorizontalScrollLimits(1.0);
 			UpdateHorizontalScrollBar();
 		}
 		return FReply::Handled();
@@ -1914,13 +1980,8 @@ FReply STimingView::OnKeyDown(const FGeometry& MyGeometry, const FKeyEvent& InKe
 		if (InKeyEvent.GetModifierKeys().IsControlDown())
 		{
 			// Scroll Left
-			const double DT = Viewport.GetDuration();
-			//ScrollAtTime(Viewport.GetStartTime() - DT * 0.05);
-			if (Viewport.ScrollAtTime(Viewport.GetStartTime() - DT * 0.05))
-			{
-				//Viewport.EnforceHorizontalScrollLimits(1.0);
-				UpdateHorizontalScrollBar();
-			}
+			const double NewStartTime = Viewport.GetStartTime() - Viewport.GetDuration() * 0.05;
+			ScrollAtTime(EnforceHorizontalScrollLimits(NewStartTime));
 		}
 		else
 		{
@@ -1933,13 +1994,8 @@ FReply STimingView::OnKeyDown(const FGeometry& MyGeometry, const FKeyEvent& InKe
 		if (InKeyEvent.GetModifierKeys().IsControlDown())
 		{
 			// Scroll Right
-			const double DT = Viewport.GetDuration();
-			//ScrollAtTime(Viewport.GetStartTime() + DT * 0.05);
-			if (Viewport.ScrollAtTime(Viewport.GetStartTime() + DT * 0.05))
-			{
-				//Viewport.EnforceHorizontalScrollLimits(1.0);
-				UpdateHorizontalScrollBar();
-			}
+			const double NewStartTime = Viewport.GetStartTime() + Viewport.GetDuration() * 0.05;
+			ScrollAtTime(EnforceHorizontalScrollLimits(NewStartTime));
 		}
 		else
 		{
@@ -1952,7 +2008,8 @@ FReply STimingView::OnKeyDown(const FGeometry& MyGeometry, const FKeyEvent& InKe
 		if (InKeyEvent.GetModifierKeys().IsControlDown())
 		{
 			// Scroll Up
-			ScrollAtPosY(Viewport.GetScrollPosY() - 16.0f * 3);
+			const float NewScrollPosY = Viewport.GetScrollPosY() - 16.0f * 3;
+			ScrollAtPosY(EnforceVerticalScrollLimits(NewScrollPosY));
 		}
 		else
 		{
@@ -1965,7 +2022,8 @@ FReply STimingView::OnKeyDown(const FGeometry& MyGeometry, const FKeyEvent& InKe
 		if (InKeyEvent.GetModifierKeys().IsControlDown())
 		{
 			// Scroll Down
-			ScrollAtPosY(Viewport.GetScrollPosY() + 16.0f * 3);
+			const float NewScrollPosY = Viewport.GetScrollPosY() + 16.0f * 3;
+			ScrollAtPosY(EnforceVerticalScrollLimits(NewScrollPosY));
 		}
 		else
 		{
@@ -2150,6 +2208,55 @@ void STimingView::BindCommands()
 
 ////////////////////////////////////////////////////////////////////////////////////////////////////
 
+double STimingView::EnforceHorizontalScrollLimits(const double InStartTime)
+{
+	double NewStartTime = InStartTime;
+
+	double MinT, MaxT;
+	Viewport.GetHorizontalScrollLimits(MinT, MaxT);
+
+	if (NewStartTime > MaxT)
+	{
+		NewStartTime = MaxT;
+		OverscrollRight = 1.0f;
+	}
+
+	if (NewStartTime < MinT)
+	{
+		NewStartTime = MinT;
+		OverscrollLeft = 1.0f;
+	}
+
+	return NewStartTime;
+}
+
+////////////////////////////////////////////////////////////////////////////////////////////////////
+
+float STimingView::EnforceVerticalScrollLimits(const float InScrollPosY)
+{
+	float NewScrollPosY = InScrollPosY;
+
+	const float DY = Viewport.GetScrollHeight() - Viewport.GetScrollableAreaHeight() + 7.0f; // +7 is to allow some space for the horizontal scrollbar
+	const float MinY = FMath::Min(DY, 0.0f);
+	const float MaxY = DY - MinY;
+
+	if (NewScrollPosY > MaxY)
+	{
+		NewScrollPosY = MaxY;
+		OverscrollBottom = 1.0f;
+	}
+
+	if (NewScrollPosY < MinY)
+	{
+		NewScrollPosY = MinY;
+		OverscrollTop = 1.0f;
+	}
+
+	return NewScrollPosY;
+}
+
+////////////////////////////////////////////////////////////////////////////////////////////////////
+
 void STimingView::HorizontalScrollBar_OnUserScrolled(float ScrollOffset)
 {
 	Viewport.OnUserScrolled(HorizontalScrollBar, ScrollOffset);
@@ -2203,7 +2310,7 @@ void STimingView::CenterOnTimeInterval(double IntervalStartTime, double Interval
 {
 	if (Viewport.CenterOnTimeInterval(IntervalStartTime, IntervalDuration))
 	{
-		Viewport.EnforceHorizontalScrollLimits(1.0);
+		Viewport.EnforceHorizontalScrollLimits(1.0); // 1.0 is to disable interpolation
 		UpdateHorizontalScrollBar();
 	}
 }
diff --git a/Engine/Source/Developer/TraceInsights/Private/Insights/Widgets/STimingView.h b/Engine/Source/Developer/TraceInsights/Private/Insights/Widgets/STimingView.h
--- a/Engine/Source/Developer/TraceInsights/Private/Insights/Widgets/STimingView.h
+++ b/Engine/Source/Developer/TraceInsights/Private/Insights/Widgets/STimingView.h
@@ -36,18 +36,6 @@ namespace Insights { class ITimingViewExtender; }
 /** A custom widget used to display timing events. */
 class STimingView : public SCompoundWidget, public Insights::ITimingViewSession
 {
-protected:
-	struct FAssetLoadingEventAggregationRow
-	{
-		FString Name;
-		int32 Count;
-		double Total;
-		double Min;
-		double Max;
-		double Avg;
-		double Med;
-	};
-
 public:
 	/** Default constructor. */
 	STimingView();
@@ -306,6 +294,9 @@ protected:
 
 	void UpdatePositionForScrollableTracks();
 
+	double EnforceHorizontalScrollLimits(const double InStartTime);
+	float EnforceVerticalScrollLimits(const float InScrollPosY);
+
 	/**
 	 * Called when the user scrolls the horizontal scrollbar.
 	 *
@@ -459,6 +450,11 @@ protected:
 	};
 	EPanningMode PanningMode;
 
+	float OverscrollLeft;
+	float OverscrollRight;
+	float OverscrollTop;
+	float OverscrollBottom;
+
 	////////////////////////////////////////////////////////////
 	// Selection
 
@@ -524,5 +520,4 @@ protected:
 	Insights::FHoveredEventChangedDelegate OnHoveredEventChangedDelegate;
 	Insights::FSelectedTrackChangedDelegate OnSelectedTrackChangedDelegate;
 	Insights::FSelectedEventChangedDelegate OnSelectedEventChangedDelegate;
-	//Insights::FFilteredEventChangedDelegate OnFilteredEventChangedDelegate;
 };
