[Insights]
  - Networking / Packet view: Added key shortcuts for selected package(s):
      * Left/Right --> selects previous/next package
      * Shift + Left/Right --> extends selection (multiple packages) toward left/right side
      * Ctrl + Shift + Left/Right --> shrink selection (multiple packages) from left/right side
  - Networking / Packet Content view: Added "Find Packet" UI on top of packet breakdown visualisation:
      * Previous / Next buttons + PacketIndex editbox.
  - Networking / Packet Content view: Added "Find Event" UI:
      * First / Previous / Next / Last buttons --> to navigate between (filtered) events
      * By NetId checkbox + NetId editbox --> to enable filtering by NetId
      * By Type checkbox + EventType readonly editbox --> to enable filtering by EventType. The EventType can be set by double clicking in NetStats or an event in Packet Content view.
      * "Highlight" checkbox --> to enable the Packet Content view to highlight only the filtered events (shows all other events in with faded colors).

diff --git a/Engine/Source/Developer/TraceInsights/Private/Insights/InsightsStyle.cpp b/Engine/Source/Developer/TraceInsights/Private/Insights/InsightsStyle.cpp
--- a/Engine/Source/Developer/TraceInsights/Private/Insights/InsightsStyle.cpp
+++ b/Engine/Source/Developer/TraceInsights/Private/Insights/InsightsStyle.cpp
@@ -3,9 +3,9 @@
 #include "InsightsStyle.h"
 #include "Styling/SlateStyleRegistry.h"
 
-#define IMAGE_BRUSH(RelativePath, ...)  FSlateImageBrush (FPaths::EngineContentDir() / "Editor/Slate"/ RelativePath + TEXT(".png"), __VA_ARGS__)
-#define BOX_BRUSH(RelativePath, ...)    FSlateBoxBrush   (FPaths::EngineContentDir() / "Editor/Slate"/ RelativePath + TEXT(".png"), __VA_ARGS__)
-#define BORDER_BRUSH(RelativePath, ...) FSlateBorderBrush(FPaths::EngineContentDir() / "Editor/Slate"/ RelativePath + TEXT(".png"), __VA_ARGS__)
+#define IMAGE_BRUSH(RelativePath, ...)  FSlateImageBrush (FPaths::EngineContentDir() / "Editor/Slate" / RelativePath + TEXT(".png"), __VA_ARGS__)
+#define BOX_BRUSH(RelativePath, ...)    FSlateBoxBrush   (FPaths::EngineContentDir() / "Editor/Slate" / RelativePath + TEXT(".png"), __VA_ARGS__)
+#define BORDER_BRUSH(RelativePath, ...) FSlateBorderBrush(FPaths::EngineContentDir() / "Editor/Slate" / RelativePath + TEXT(".png"), __VA_ARGS__)
 
 TSharedPtr<FSlateStyleSet> FInsightsStyle::StyleInstance = nullptr;
 
@@ -52,8 +52,8 @@ TSharedRef<FSlateStyleSet> FInsightsStyle::Create()
 	//////////////////////////////////////////////////
 	// Icons for major components
 
-	Style.Set("StartPage.Icon.Large", new IMAGE_BRUSH("/Icons/icon_tab_Tools_16x", FVector2D(32.0f, 32.0f)));
-	Style.Set("StartPage.Icon.Small", new IMAGE_BRUSH("/Icons/icon_tab_Tools_16x", FVector2D(16.0f, 16.0f)));
+	Style.Set("StartPage.Icon.Large", new IMAGE_BRUSH("Icons/icon_tab_Tools_16x", FVector2D(32.0f, 32.0f)));
+	Style.Set("StartPage.Icon.Small", new IMAGE_BRUSH("Icons/icon_tab_Tools_16x", FVector2D(16.0f, 16.0f)));
 
 	Style.Set("TimingProfiler.Icon.Large", new IMAGE_BRUSH("Icons/Profiler/profiler_stats_40x", FVector2D(32.0f, 32.0f)));
 	Style.Set("TimingProfiler.Icon.Small", new IMAGE_BRUSH("Icons/Profiler/profiler_stats_40x", FVector2D(16.0f, 16.0f)));
@@ -64,8 +64,8 @@ TSharedRef<FSlateStyleSet> FInsightsStyle::Create()
 	Style.Set("NetworkingProfiler.Icon.Large", new IMAGE_BRUSH("Icons/Profiler/profiler_stats_40x", FVector2D(32.0f, 32.0f)));
 	Style.Set("NetworkingProfiler.Icon.Small", new IMAGE_BRUSH("Icons/Profiler/profiler_stats_40x", FVector2D(16.0f, 16.0f)));
 
-	Style.Set("Toolbar.Icon.Large", new IMAGE_BRUSH("/Icons/icon_tab_Tools_16x", FVector2D(32.0f, 32.0f)));
-	Style.Set("Toolbar.Icon.Small", new IMAGE_BRUSH("/Icons/icon_tab_Tools_16x", FVector2D(16.0f, 16.0f)));
+	Style.Set("Toolbar.Icon.Large", new IMAGE_BRUSH("Icons/icon_tab_Tools_16x", FVector2D(32.0f, 32.0f)));
+	Style.Set("Toolbar.Icon.Small", new IMAGE_BRUSH("Icons/icon_tab_Tools_16x", FVector2D(16.0f, 16.0f)));
 
 	//////////////////////////////////////////////////
 	// Start Page buttons
@@ -114,6 +114,11 @@ TSharedRef<FSlateStyleSet> FInsightsStyle::Create()
 
 	//////////////////////////////////////////////////
 
+	Style.Set("FindFirst", new IMAGE_BRUSH("Animation/backward_end", FVector2D(20.0f, 20.0f)));
+	Style.Set("FindPrevious", new IMAGE_BRUSH("Animation/backward", FVector2D(20.0f, 20.0f)));
+	Style.Set("FindNext", new IMAGE_BRUSH("Animation/forward", FVector2D(20.0f, 20.0f)));
+	Style.Set("FindLast", new IMAGE_BRUSH("Animation/forward_end", FVector2D(20.0f, 20.0f)));
+
 	return StyleRef;
 }
 
diff --git a/Engine/Source/Developer/TraceInsights/Private/Insights/NetworkingProfiler/ViewModels/PacketContentViewDrawHelper.cpp b/Engine/Source/Developer/TraceInsights/Private/Insights/NetworkingProfiler/ViewModels/PacketContentViewDrawHelper.cpp
--- a/Engine/Source/Developer/TraceInsights/Private/Insights/NetworkingProfiler/ViewModels/PacketContentViewDrawHelper.cpp
+++ b/Engine/Source/Developer/TraceInsights/Private/Insights/NetworkingProfiler/ViewModels/PacketContentViewDrawHelper.cpp
@@ -81,16 +81,6 @@ void FPacketContentViewDrawStateBuilder::AddEvent(const Trace::FNetProfilerConte
 		LastBox.AddDefaulted(Depth + 1 - LastBox.Num());
 	}
 
-	constexpr float Y0 = 0.0f;
-	constexpr float EventH = 14.0f;
-	constexpr float EventDY = 2.0f;
-
-	const float EventY = Y0 + (EventH + EventDY) * Depth;
-	//if (EventY < -EventH || EventY > Viewport.GetHeight())
-	//{
-	//	return;
-	//}
-
 	// Ensure we have enough slots in array. See LastEventX2[Depth] usage.
 	while (LastEventX2.Num() <= Depth)
 	{
@@ -130,38 +120,33 @@ void FPacketContentViewDrawStateBuilder::AddEvent(const Trace::FNetProfilerConte
 
 	//////////////////////////////////////////////////
 
-	if (EventW > 2.0f && EventH > 2.0f)
-	{
-		// Fill inside of the event box.
-		DrawState.Boxes.AddUninitialized();
-		FPacketContentViewDrawState::FBox& DrawBox = DrawState.Boxes.Last();
-		DrawBox.X = EventX1 + 1.0f;
-		DrawBox.Y = EventY + 1.0f;
-		DrawBox.W = EventW - 2.0f;
-		DrawBox.H = EventH - 2.0f;
-		DrawBox.Color = EventColorFill;
-	}
-
 	// Save X2, for current depth.
 	LastEventX2[Depth] = EventX2;
 
-	// Add border around the event box.
 	if (EventW > 2.0f)
 	{
 		FBoxData& Box = LastBox[Depth];
 		if (Box.X1 < Box.X2)
 		{
-			FlushBox(Box, EventY, EventH);
+			FlushBox(Box, Depth);
 			Box.Reset();
 		}
 
+		// Fill inside of the event box.
+		DrawState.InsideBoxes.AddUninitialized();
+		FPacketContentViewDrawState::FBoxPrimitive& InsideBox = DrawState.InsideBoxes.Last();
+		InsideBox.Depth = Depth;
+		InsideBox.X = EventX1 + 1.0f;
+		InsideBox.W = EventW - 2.0f;
+		InsideBox.Color = EventColorFill;
+
+		// Add border around the event box.
 		DrawState.Borders.AddUninitialized();
-		FPacketContentViewDrawState::FBox& DrawBox = DrawState.Borders.Last();
-		DrawBox.X = EventX1;
-		DrawBox.Y = EventY;
-		DrawBox.W = EventW;
-		DrawBox.H = EventH;
-		DrawBox.Color = FLinearColor(EventColorFill.R * BorderColorFactor, EventColorFill.G * BorderColorFactor, EventColorFill.B * BorderColorFactor, EventColorFill.A);
+		FPacketContentViewDrawState::FBoxPrimitive& BorderBox = DrawState.Borders.Last();
+		BorderBox.Depth = Depth;
+		BorderBox.X = EventX1;
+		BorderBox.W = EventW;
+		BorderBox.Color = FLinearColor(EventColorFill.R * BorderColorFactor, EventColorFill.G * BorderColorFactor, EventColorFill.B * BorderColorFactor, EventColorFill.A);
 	}
 	else // 1px or 2px boxes
 	{
@@ -181,7 +166,7 @@ void FPacketContentViewDrawStateBuilder::AddEvent(const Trace::FNetProfilerConte
 			// Flush previous box, if any.
 			if (Box.X1 < Box.X2)
 			{
-				FlushBox(Box, EventY, EventH);
+				FlushBox(Box, Depth);
 			}
 
 			// Start new "merge box".
@@ -231,9 +216,9 @@ void FPacketContentViewDrawStateBuilder::AddEvent(const Trace::FNetProfilerConte
 			const bool bIsDarkColor = (EventColorFill.ComputeLuminance() < 0.4f);
 
 			DrawState.Texts.AddDefaulted();
-			FPacketContentViewDrawState::FText& DrawText = DrawState.Texts.Last();
+			FPacketContentViewDrawState::FTextPrimitive& DrawText = DrawState.Texts.Last();
+			DrawText.Depth = Depth;
 			DrawText.X = EventX1 + 2.0f;
-			DrawText.Y = EventY + 1.0f;
 			DrawText.Text = Name.Left(LastWholeCharacterIndex + 1);
 			DrawText.bWhite = bIsDarkColor;
 		}
@@ -242,14 +227,13 @@ void FPacketContentViewDrawStateBuilder::AddEvent(const Trace::FNetProfilerConte
 
 ////////////////////////////////////////////////////////////////////////////////////////////////////
 
-void FPacketContentViewDrawStateBuilder::FlushBox(const FBoxData& Box, const float EventY, const float EventH)
+void FPacketContentViewDrawStateBuilder::FlushBox(const FBoxData& Box, const int32 Depth)
 {
 	DrawState.Boxes.AddUninitialized();
-	FPacketContentViewDrawState::FBox& DrawBox = DrawState.Boxes.Last();
+	FPacketContentViewDrawState::FBoxPrimitive& DrawBox = DrawState.Boxes.Last();
+	DrawBox.Depth = Depth;
 	DrawBox.X = Box.X1;
 	DrawBox.W = Box.X2 - Box.X1;
-	DrawBox.Y = EventY;
-	DrawBox.H = EventH;
 	DrawBox.Color = Box.LinearColor;
 }
 
@@ -257,18 +241,13 @@ void FPacketContentViewDrawStateBuilder::FlushBox(const FBoxData& Box, const flo
 
 void FPacketContentViewDrawStateBuilder::Flush()
 {
-	constexpr float Y0 = 0.0f;
-	constexpr float EventH = 14.0f;
-	constexpr float EventDY = 2.0f;
-
 	// Flush merged boxes.
 	for (int32 Depth = 0; Depth <= MaxDepth; ++Depth)
 	{
 		const FBoxData& Box = LastBox[Depth];
 		if (Box.X1 < Box.X2)
 		{
-			const float EventY = Y0 + (EventH + EventDY) * Depth;
-			FlushBox(Box, EventY, EventH);
+			FlushBox(Box, Depth);
 		}
 	}
 }
@@ -285,6 +264,9 @@ FPacketContentViewDrawHelper::FPacketContentViewDrawHelper(const FDrawContext& I
 	, HoveredEventBorderBrush(FInsightsStyle::Get().GetBrush("HoveredEventBorder"))
 	, SelectedEventBorderBrush(FInsightsStyle::Get().GetBrush("SelectedEventBorder"))
 	, EventFont(FCoreStyle::GetDefaultFontStyle("Regular", 8))
+	, LayoutPosY(0.0f)
+	, LayoutEventH(14.0f)
+	, LayoutEventDY(2.0f)
 {
 }
 
@@ -343,28 +325,54 @@ FLinearColor FPacketContentViewDrawHelper::GetColorByType(int32 Type)
 
 ////////////////////////////////////////////////////////////////////////////////////////////////////
 
-void FPacketContentViewDrawHelper::Draw(const FPacketContentViewDrawState& DrawState) const
+void FPacketContentViewDrawHelper::Draw(const FPacketContentViewDrawState& DrawState, const float Opacity) const
 {
-	// Draw filled boxes.
-	for (const FPacketContentViewDrawState::FBox& Box : DrawState.Boxes)
+	// Draw filled boxes (merged borders).
+	//if (LayoutEventH > 0.0f)
 	{
-		DrawContext.DrawBox(Box.X, Box.Y, Box.W, Box.H, WhiteBrush, Box.Color);
+		const float EventFillH = LayoutEventH;
+		for (const FPacketContentViewDrawState::FBoxPrimitive& Box : DrawState.Boxes)
+		{
+			const float Y = LayoutPosY + (LayoutEventH + LayoutEventDY) * Box.Depth;
+			DrawContext.DrawBox(Box.X, Y, Box.W, EventFillH, WhiteBrush, Box.Color.CopyWithNewOpacity(Opacity));
+		}
+		DrawContext.LayerId++;
+	}
+
+	// Draw filled boxes (event inside area).
+	if (LayoutEventH > 2.0f)
+	{
+		const float EventFillH = LayoutEventH - 2.0f;
+		for (const FPacketContentViewDrawState::FBoxPrimitive& Box : DrawState.InsideBoxes)
+		{
+			const float Y = LayoutPosY + (LayoutEventH + LayoutEventDY) * Box.Depth + 1.0f;
+			DrawContext.DrawBox(Box.X, Y, Box.W, EventFillH, WhiteBrush, Box.Color.CopyWithNewOpacity(Opacity));
+		}
+		DrawContext.LayerId++;
 	}
-	DrawContext.LayerId++;
 
 	// Draw borders.
-	for (const FPacketContentViewDrawState::FBox& Box : DrawState.Borders)
+	//if (LayoutEventH > 0.0f)
 	{
-		DrawContext.DrawBox(Box.X, Box.Y, Box.W, Box.H, EventBorderBrush, Box.Color);
+		const float EventBorderH = LayoutEventH;
+		for (const FPacketContentViewDrawState::FBoxPrimitive& Box : DrawState.Borders)
+		{
+			const float Y = LayoutPosY + (LayoutEventH + LayoutEventDY) * Box.Depth;
+			DrawContext.DrawBox(Box.X, Y, Box.W, EventBorderH, EventBorderBrush, Box.Color.CopyWithNewOpacity(Opacity));
+		}
+		DrawContext.LayerId++;
 	}
-	DrawContext.LayerId++;
 
 	// Draw texts.
-	//if (EventH > 10.0f)
+	if (LayoutEventH > 10.0f)
 	{
-		for (const FPacketContentViewDrawState::FText& Text : DrawState.Texts)
+		const FLinearColor WhiteColor(1.0f, 1.0f, 1.0f, Opacity);
+		const FLinearColor BlackColor(0.0f, 0.0f, 0.0f, Opacity);
+
+		for (const FPacketContentViewDrawState::FTextPrimitive& Text : DrawState.Texts)
 		{
-			DrawContext.DrawText(Text.X, Text.Y, Text.Text, EventFont, Text.bWhite ? FLinearColor::White : FLinearColor::Black);
+			const float Y = LayoutPosY + (LayoutEventH + LayoutEventDY) * Text.Depth + 1.0f;
+			DrawContext.DrawText(Text.X, Y, Text.Text, EventFont, Text.bWhite ? WhiteColor : BlackColor);
 		}
 		DrawContext.LayerId++;
 	}
@@ -409,19 +417,14 @@ void FPacketContentViewDrawHelper::DrawEventHighlight(const FNetworkPacketEvent&
 	}
 
 	const float EventW = EventX2 - EventX1;
-
-	constexpr float Y0 = 0.0f;
-	constexpr float EventH = 14.0f;
-	constexpr float EventDY = 2.0f;
-
-	const float EventY = Y0 + (EventH + EventDY) * Event.Level;
+	const float EventY = LayoutPosY + (LayoutEventH + LayoutEventDY) * Event.Level;
 
 	if (Mode == EHighlightMode::Hovered)
 	{
 		const FLinearColor Color(1.0f, 1.0f, 0.0f, 1.0f); // yellow
 
 		// Draw border around the timing event box.
-		DrawContext.DrawBox(EventX1 - 2.0f, EventY - 2.0f, EventW + 4.0f, EventH + 4.0f, HoveredEventBorderBrush, Color);
+		DrawContext.DrawBox(EventX1 - 2.0f, EventY - 2.0f, EventW + 4.0f, LayoutEventH + 4.0f, HoveredEventBorderBrush, Color);
 	}
 	else // EHighlightMode::Selected or EHighlightMode::SelectedAndHovered
 	{
@@ -433,7 +436,7 @@ void FPacketContentViewDrawHelper::DrawEventHighlight(const FNetworkPacketEvent&
 		const FLinearColor Color(S, S, Blue, 1.0f);
 
 		// Draw border around the timing event box.
-		DrawContext.DrawBox(EventX1 - 2.0f, EventY - 2.0f, EventW + 4.0f, EventH + 4.0f, SelectedEventBorderBrush, Color);
+		DrawContext.DrawBox(EventX1 - 2.0f, EventY - 2.0f, EventW + 4.0f, LayoutEventH + 4.0f, SelectedEventBorderBrush, Color);
 	}
 	DrawContext.LayerId++;
 }
diff --git a/Engine/Source/Developer/TraceInsights/Private/Insights/NetworkingProfiler/ViewModels/PacketContentViewDrawHelper.h b/Engine/Source/Developer/TraceInsights/Private/Insights/NetworkingProfiler/ViewModels/PacketContentViewDrawHelper.h
--- a/Engine/Source/Developer/TraceInsights/Private/Insights/NetworkingProfiler/ViewModels/PacketContentViewDrawHelper.h
+++ b/Engine/Source/Developer/TraceInsights/Private/Insights/NetworkingProfiler/ViewModels/PacketContentViewDrawHelper.h
@@ -86,24 +86,30 @@ struct FNetworkPacketEvent
 
 struct FPacketContentViewDrawState
 {
-	struct FBox
+	struct FBoxPrimitive
 	{
+		int32 Depth;
 		float X;
-		float Y;
 		float W;
-		float H;
 		FLinearColor Color;
 	};
 
-	struct FText
+	struct FTextPrimitive
 	{
+		int32 Depth;
 		float X;
-		float Y;
 		FString Text;
 		bool bWhite;
 	};
 
 	FPacketContentViewDrawState()
+		: Events()
+		, Boxes()
+		, InsideBoxes()
+		, Borders()
+		, Texts()
+		, NumLanes(0)
+		, NumMergedBoxes(0)
 	{
 	}
 
@@ -111,17 +117,26 @@ struct FPacketContentViewDrawState
 	{
 		Events.Reset();
 		Boxes.Reset();
+		InsideBoxes.Reset();
 		Borders.Reset();
 		Texts.Reset();
+		NumLanes = 0;
 		NumMergedBoxes = 0;
 	}
 
+	int32 GetNumLanes() const { return NumLanes; }
+
+	int32 GetNumEvents() const { return Events.Num(); }
 	int32 GetNumMergedBoxes() const { return NumMergedBoxes; }
+	int32 GetTotalNumBoxes() const { return Boxes.Num() + InsideBoxes.Num(); }
 
 	TArray<FNetworkPacketEvent> Events;
-	TArray<FBox> Boxes;
-	TArray<FBox> Borders;
-	TArray<FText> Texts;
+	TArray<FBoxPrimitive> Boxes;
+	TArray<FBoxPrimitive> InsideBoxes;
+	TArray<FBoxPrimitive> Borders;
+	TArray<FTextPrimitive> Texts;
+
+	int32 NumLanes;
 
 	// Debug stats.
 	int32 NumMergedBoxes;
@@ -155,8 +170,10 @@ public:
 	void AddEvent(const Trace::FNetProfilerContentEvent& Event, const TCHAR* Name, uint32 NetId);
 	void Flush();
 
+	int32 GetMaxDepth() const { return MaxDepth; }
+
 private:
-	void FlushBox(const FBoxData& Box, const float EventY, const float EventH);
+	void FlushBox(const FBoxData& Box, const int32 Depth);
 
 private:
 	FPacketContentViewDrawState& DrawState; // cached draw state to build
@@ -194,8 +211,16 @@ public:
 	const FSlateBrush* GetWhiteBrush() const { return WhiteBrush; }
 	const FSlateFontInfo& GetEventFont() const { return EventFont; }
 
+	float GetLayoutPosY() const { return LayoutPosY; }
+	float GetLayoutEventH() const { return LayoutEventH; }
+	float GetLayoutEventDY() const { return LayoutEventDY; }
+
+	void SetLayoutPosY(const float InLayoutPosY) { LayoutPosY = InLayoutPosY; }
+	void SetLayoutEventH(const float InLayoutEventH) { LayoutEventH = InLayoutEventH; }
+	void SetLayoutEventDY(const float InLayoutEventDY) { LayoutEventDY = InLayoutEventDY; }
+
 	void DrawBackground() const;
-	void Draw(const FPacketContentViewDrawState& DrawState) const;
+	void Draw(const FPacketContentViewDrawState& DrawState, const float Opacity = 1.0f) const;
 	void DrawEventHighlight(const FNetworkPacketEvent& Event, EHighlightMode Mode) const;
 
 	static FLinearColor GetColorByType(int32 Type);
@@ -209,6 +234,10 @@ private:
 	const FSlateBrush* HoveredEventBorderBrush;
 	const FSlateBrush* SelectedEventBorderBrush;
 	const FSlateFontInfo EventFont;
+
+	float LayoutPosY;
+	float LayoutEventH;
+	float LayoutEventDY;
 };
 
 ////////////////////////////////////////////////////////////////////////////////////////////////////
diff --git a/Engine/Source/Developer/TraceInsights/Private/Insights/NetworkingProfiler/ViewModels/PacketViewDrawHelper.cpp b/Engine/Source/Developer/TraceInsights/Private/Insights/NetworkingProfiler/ViewModels/PacketViewDrawHelper.cpp
--- a/Engine/Source/Developer/TraceInsights/Private/Insights/NetworkingProfiler/ViewModels/PacketViewDrawHelper.cpp
+++ b/Engine/Source/Developer/TraceInsights/Private/Insights/NetworkingProfiler/ViewModels/PacketViewDrawHelper.cpp
@@ -14,6 +14,52 @@
 
 #include <limits>
 
+////////////////////////////////////////////////////////////////////////////////////////////////////
+// FNetworkPacketAggregatedSample
+////////////////////////////////////////////////////////////////////////////////////////////////////
+
+void FNetworkPacketAggregatedSample::AddPacket(const int32 PacketIndex, const Trace::FNetProfilerPacket& Packet)
+{
+	NumPackets++;
+
+	const double TimeStamp = static_cast<double>(Packet.TimeStamp);
+	if (TimeStamp < StartTime)
+	{
+		StartTime = TimeStamp;
+	}
+	if (TimeStamp > EndTime)
+	{
+		EndTime = TimeStamp;
+	}
+
+	if (Packet.TotalPacketSizeInBytes > LargestPacket.TotalSizeInBytes)
+	{
+		LargestPacket.Index = PacketIndex;
+		LargestPacket.SequenceNumber = Packet.SequenceNumber;
+		LargestPacket.ContentSizeInBits = Packet.ContentSizeInBits;
+		LargestPacket.TotalSizeInBytes = Packet.TotalPacketSizeInBytes;
+		LargestPacket.TimeStamp = TimeStamp;
+		LargestPacket.Status = Packet.DeliveryStatus;
+	}
+
+	switch (AggregatedStatus)
+	{
+	case Trace::ENetProfilerDeliveryStatus::Unknown:
+		AggregatedStatus = Packet.DeliveryStatus;
+		break;
+
+	case Trace::ENetProfilerDeliveryStatus::Dropped:
+		break;
+
+	case Trace::ENetProfilerDeliveryStatus::Delivered:
+		if (Packet.DeliveryStatus == Trace::ENetProfilerDeliveryStatus::Dropped)
+		{
+			AggregatedStatus = Trace::ENetProfilerDeliveryStatus::Dropped;
+		}
+		break;
+	}
+}
+
 ////////////////////////////////////////////////////////////////////////////////////////////////////
 // FNetworkPacketSeriesBuilder
 ////////////////////////////////////////////////////////////////////////////////////////////////////
@@ -35,7 +81,7 @@ FNetworkPacketSeriesBuilder::FNetworkPacketSeriesBuilder(FNetworkPacketSeries& I
 
 ////////////////////////////////////////////////////////////////////////////////////////////////////
 
-void FNetworkPacketSeriesBuilder::AddPacket(int32 PacketIndex, const Trace::FNetProfilerPacket& Packet)
+void FNetworkPacketSeriesBuilder::AddPacket(const int32 PacketIndex, const Trace::FNetProfilerPacket& Packet)
 {
 	NumAddedPackets++;
 
@@ -43,45 +89,7 @@ void FNetworkPacketSeriesBuilder::AddPacket(int32 PacketIndex, const Trace::FNet
 	if (SampleIndex >= 0 && SampleIndex < NumSamples)
 	{
 		FNetworkPacketAggregatedSample& Sample = Series.Samples[SampleIndex];
-		Sample.NumPackets++;
-
-		const double TimeStamp = static_cast<double>(Packet.TimeStamp);
-		if (TimeStamp < Sample.StartTime)
-		{
-			Sample.StartTime = TimeStamp;
-		}
-		if (TimeStamp > Sample.EndTime)
-		{
-			Sample.EndTime = TimeStamp;
-		}
-
-		if (Packet.TotalPacketSizeInBytes > Sample.LargestPacket.TotalSizeInBytes)
-		{
-			Sample.LargestPacket.Index = PacketIndex;
-			Sample.LargestPacket.SequenceNumber = Packet.SequenceNumber;
-			Sample.LargestPacket.ContentSizeInBits = Packet.ContentSizeInBits;
-			Sample.LargestPacket.TotalSizeInBytes = Packet.TotalPacketSizeInBytes;
-			Sample.LargestPacket.TimeStamp = TimeStamp;
-			Sample.LargestPacket.Status = Packet.DeliveryStatus;
-		}
-
-		switch (Sample.AggregatedStatus)
-		{
-			case Trace::ENetProfilerDeliveryStatus::Unknown:
-				Sample.AggregatedStatus = Packet.DeliveryStatus;
-				break;
-
-			case Trace::ENetProfilerDeliveryStatus::Dropped:
-				break;
-
-			case Trace::ENetProfilerDeliveryStatus::Delivered:
-				if (Packet.DeliveryStatus == Trace::ENetProfilerDeliveryStatus::Dropped)
-				{
-					Sample.AggregatedStatus = Trace::ENetProfilerDeliveryStatus::Dropped;
-				}
-				break;
-		}
-
+		Sample.AddPacket(PacketIndex, Packet);
 		Series.NumAggregatedPackets++;
 	}
 }
diff --git a/Engine/Source/Developer/TraceInsights/Private/Insights/NetworkingProfiler/ViewModels/PacketViewDrawHelper.h b/Engine/Source/Developer/TraceInsights/Private/Insights/NetworkingProfiler/ViewModels/PacketViewDrawHelper.h
--- a/Engine/Source/Developer/TraceInsights/Private/Insights/NetworkingProfiler/ViewModels/PacketViewDrawHelper.h
+++ b/Engine/Source/Developer/TraceInsights/Private/Insights/NetworkingProfiler/ViewModels/PacketViewDrawHelper.h
@@ -93,6 +93,8 @@ struct FNetworkPacketAggregatedSample
 	FNetworkPacketAggregatedSample(FNetworkPacketAggregatedSample&&) = default;
 	FNetworkPacketAggregatedSample& operator=(FNetworkPacketAggregatedSample&&) = default;
 
+	void AddPacket(const int32 PacketIndex, const Trace::FNetProfilerPacket& Packet);
+
 	bool Equals(const FNetworkPacketAggregatedSample& Other) const
 	{
 		return NumPackets == Other.NumPackets
diff --git a/Engine/Source/Developer/TraceInsights/Private/Insights/NetworkingProfiler/Widgets/SNetStatsView.cpp b/Engine/Source/Developer/TraceInsights/Private/Insights/NetworkingProfiler/Widgets/SNetStatsView.cpp
--- a/Engine/Source/Developer/TraceInsights/Private/Insights/NetworkingProfiler/Widgets/SNetStatsView.cpp
+++ b/Engine/Source/Developer/TraceInsights/Private/Insights/NetworkingProfiler/Widgets/SNetStatsView.cpp
@@ -24,13 +24,16 @@
 #include "Insights/NetworkingProfiler/ViewModels/NetStatsViewColumnFactory.h"
 #include "Insights/NetworkingProfiler/Widgets/SNetStatsViewTooltip.h"
 #include "Insights/NetworkingProfiler/Widgets/SNetStatsTableRow.h"
+#include "Insights/NetworkingProfiler/Widgets/SNetworkingProfilerWindow.h"
+#include "Insights/NetworkingProfiler/Widgets/SPacketContentView.h"
 
 #define LOCTEXT_NAMESPACE "SNetStatsView"
 
 ////////////////////////////////////////////////////////////////////////////////////////////////////
 
 SNetStatsView::SNetStatsView()
-	: Table(MakeShared<Insights::FTable>())
+	: ProfilerWindow()
+	, Table(MakeShared<Insights::FTable>())
 	, bExpansionSaved(false)
 	, bFilterOutZeroCountEvents(false)
 	, GroupingMode(ENetEventGroupingMode::Flat)
@@ -65,8 +68,10 @@ SNetStatsView::~SNetStatsView()
 ////////////////////////////////////////////////////////////////////////////////////////////////////
 
 BEGIN_SLATE_FUNCTION_BUILD_OPTIMIZATION
-void SNetStatsView::Construct(const FArguments& InArgs)
+void SNetStatsView::Construct(const FArguments& InArgs, TSharedPtr<SNetworkingProfilerWindow> InProfilerWindow)
 {
+	ProfilerWindow = InProfilerWindow;
+
 	SAssignNew(ExternalScrollbar, SScrollBar)
 	.AlwaysShowScrollbar(true);
 
@@ -839,6 +844,25 @@ void SNetStatsView::TreeView_OnMouseButtonDoubleClick(FNetEventNodePtr NetEventN
 		const bool bIsGroupExpanded = TreeView->IsItemExpanded(NetEventNodePtr);
 		TreeView->SetItemExpansion(NetEventNodePtr, !bIsGroupExpanded);
 	}
+	else
+	{
+		TSharedPtr<SPacketContentView> PacketContentView = ProfilerWindow.IsValid() ? ProfilerWindow->GetPacketContentView() : nullptr;
+		if (PacketContentView.IsValid())
+		{
+			const uint32 EventTypeIndex = NetEventNodePtr->GetId();
+			const uint32 FilterEventTypeIndex = PacketContentView->GetFilterEventTypeIndex();
+
+			if (EventTypeIndex == FilterEventTypeIndex && PacketContentView->IsFilterByEventTypeEnabled())
+			{
+				PacketContentView->DisableFilterEventType();
+			}
+			else
+			{
+				PacketContentView->EnableFilterEventType(EventTypeIndex);
+				//PacketContentView->FindFirstEvent();
+			}
+		}
+	}
 }
 
 ////////////////////////////////////////////////////////////////////////////////////////////////////
diff --git a/Engine/Source/Developer/TraceInsights/Private/Insights/NetworkingProfiler/Widgets/SNetStatsView.h b/Engine/Source/Developer/TraceInsights/Private/Insights/NetworkingProfiler/Widgets/SNetStatsView.h
--- a/Engine/Source/Developer/TraceInsights/Private/Insights/NetworkingProfiler/Widgets/SNetStatsView.h
+++ b/Engine/Source/Developer/TraceInsights/Private/Insights/NetworkingProfiler/Widgets/SNetStatsView.h
@@ -22,6 +22,7 @@
 ////////////////////////////////////////////////////////////////////////////////////////////////////
 
 class FMenuBuilder;
+class SNetworkingProfilerWindow;
 
 namespace Trace
 {
@@ -63,7 +64,7 @@ public:
 	 * Construct this widget
 	 * @param InArgs - The declaration data for this widget
 	 */
-	void Construct(const FArguments& InArgs);
+	void Construct(const FArguments& InArgs, TSharedPtr<SNetworkingProfilerWindow> InProfilerWindow);
 
 	TSharedPtr<Insights::FTable> GetTable() const { return Table; }
 
@@ -253,6 +254,8 @@ protected:
 	////////////////////////////////////////////////////////////////////////////////////////////////////
 
 protected:
+	TSharedPtr<SNetworkingProfilerWindow> ProfilerWindow;
+
 	/** Table view model. */
 	TSharedPtr<Insights::FTable> Table;
 
diff --git a/Engine/Source/Developer/TraceInsights/Private/Insights/NetworkingProfiler/Widgets/SNetworkingProfilerWindow.cpp b/Engine/Source/Developer/TraceInsights/Private/Insights/NetworkingProfiler/Widgets/SNetworkingProfilerWindow.cpp
--- a/Engine/Source/Developer/TraceInsights/Private/Insights/NetworkingProfiler/Widgets/SNetworkingProfilerWindow.cpp
+++ b/Engine/Source/Developer/TraceInsights/Private/Insights/NetworkingProfiler/Widgets/SNetworkingProfilerWindow.cpp
@@ -179,7 +179,7 @@ TSharedRef<SDockTab> SNetworkingProfilerWindow::SpawnTab_NetStatsView(const FSpa
 		.ShouldAutosize(false)
 		.TabRole(ETabRole::PanelTab)
 		[
-			SAssignNew(NetStatsView, SNetStatsView)
+			SAssignNew(NetStatsView, SNetStatsView, SharedThis(this))
 		];
 
 	DockTab->SetOnTabClosed(SDockTab::FOnTabClosedCallback::CreateRaw(this, &SNetworkingProfilerWindow::OnNetStatsViewTabClosed));
@@ -747,13 +747,21 @@ END_SLATE_FUNCTION_BUILD_OPTIMIZATION
 
 void SNetworkingProfilerWindow::GameInstance_OnSelectionChanged(TSharedPtr<FGameInstanceItem> NewGameInstance, ESelectInfo::Type SelectInfo)
 {
+	const bool bSameValue = (!SelectedGameInstance.IsValid() && !NewGameInstance.IsValid()) ||
+							(SelectedGameInstance.IsValid() && NewGameInstance.IsValid() &&
+								SelectedGameInstance->GetIndex() == NewGameInstance->GetIndex());
+
 	SelectedGameInstance = NewGameInstance;
-	UpdateAvailableConnections();
-	if (PacketView.IsValid())
+
+	if (!bSameValue)
 	{
-		PacketView->SetConnection(GetSelectedGameInstanceIndex(), GetSelectedConnectionIndex(), GetSelectedConnectionMode());
+		UpdateAvailableConnections();
+		if (PacketView.IsValid())
+		{
+			PacketView->SetConnection(GetSelectedGameInstanceIndex(), GetSelectedConnectionIndex(), GetSelectedConnectionMode());
+		}
+		UpdateAggregatedNetStats();
 	}
-	UpdateAggregatedNetStats();
 }
 
 ////////////////////////////////////////////////////////////////////////////////////////////////////
@@ -804,13 +812,21 @@ END_SLATE_FUNCTION_BUILD_OPTIMIZATION
 
 void SNetworkingProfilerWindow::Connection_OnSelectionChanged(TSharedPtr<FConnectionItem> NewConnection, ESelectInfo::Type SelectInfo)
 {
+	const bool bSameValue = (!SelectedConnection.IsValid() && !NewConnection.IsValid()) ||
+							(SelectedConnection.IsValid() && NewConnection.IsValid() &&
+								SelectedConnection->GetIndex() == NewConnection->GetIndex());
+
 	SelectedConnection = NewConnection;
-	UpdateAvailableConnectionModes();
-	if (PacketView.IsValid())
+
+	if (!bSameValue)
 	{
-		PacketView->SetConnection(GetSelectedGameInstanceIndex(), GetSelectedConnectionIndex(), GetSelectedConnectionMode());
+		UpdateAvailableConnectionModes();
+		if (PacketView.IsValid())
+		{
+			PacketView->SetConnection(GetSelectedGameInstanceIndex(), GetSelectedConnectionIndex(), GetSelectedConnectionMode());
+		}
+		UpdateAggregatedNetStats();
 	}
-	UpdateAggregatedNetStats();
 }
 
 ////////////////////////////////////////////////////////////////////////////////////////////////////
@@ -861,12 +877,20 @@ END_SLATE_FUNCTION_BUILD_OPTIMIZATION
 
 void SNetworkingProfilerWindow::ConnectionMode_OnSelectionChanged(TSharedPtr<FConnectionModeItem> NewConnectionMode, ESelectInfo::Type SelectInfo)
 {
+	const bool bSameValue = (!SelectedConnectionMode.IsValid() && !NewConnectionMode.IsValid()) ||
+							(SelectedConnectionMode.IsValid() && NewConnectionMode.IsValid() &&
+								SelectedConnectionMode->Mode == NewConnectionMode->Mode);
+
 	SelectedConnectionMode = NewConnectionMode;
-	if (PacketView.IsValid())
+
+	if (!bSameValue)
 	{
-		PacketView->SetConnection(GetSelectedGameInstanceIndex(), GetSelectedConnectionIndex(), GetSelectedConnectionMode());
+		if (PacketView.IsValid())
+		{
+			PacketView->SetConnection(GetSelectedGameInstanceIndex(), GetSelectedConnectionIndex(), GetSelectedConnectionMode());
+		}
+		UpdateAggregatedNetStats();
 	}
-	UpdateAggregatedNetStats();
 }
 
 ////////////////////////////////////////////////////////////////////////////////////////////////////
diff --git a/Engine/Source/Developer/TraceInsights/Private/Insights/NetworkingProfiler/Widgets/SPacketContentView.cpp b/Engine/Source/Developer/TraceInsights/Private/Insights/NetworkingProfiler/Widgets/SPacketContentView.cpp
--- a/Engine/Source/Developer/TraceInsights/Private/Insights/NetworkingProfiler/Widgets/SPacketContentView.cpp
+++ b/Engine/Source/Developer/TraceInsights/Private/Insights/NetworkingProfiler/Widgets/SPacketContentView.cpp
@@ -8,6 +8,10 @@
 #include "Framework/MultiBox/MultiBoxBuilder.h"
 #include "Rendering/DrawElements.h"
 #include "Styling/CoreStyle.h"
+#include "Widgets/Images/SImage.h"
+#include "Widgets/Input/SCheckBox.h"
+#include "Widgets/Input/SButton.h"
+#include "Widgets/Input/SEditableTextBox.h"
 #include "Widgets/Layout/SScrollBar.h"
 
 // Insights
@@ -17,6 +21,7 @@
 #include "Insights/InsightsManager.h"
 #include "Insights/InsightsStyle.h"
 #include "Insights/NetworkingProfiler/Widgets/SNetworkingProfilerWindow.h"
+#include "Insights/NetworkingProfiler/Widgets/SPacketView.h"
 
 ////////////////////////////////////////////////////////////////////////////////////////////////////
 
@@ -27,6 +32,7 @@
 SPacketContentView::SPacketContentView()
 	: ProfilerWindow()
 	, DrawState(MakeShared<FPacketContentViewDrawState>())
+	, FilteredDrawState(MakeShared<FPacketContentViewDrawState>())
 {
 	Reset();
 }
@@ -55,7 +61,17 @@ void SPacketContentView::Reset()
 	PacketIndex = 0;
 	PacketBitSize = 0;
 
+	bFilterByEventType = false;
+	FilterEventTypeIndex = 0;
+	FilterEventName = FText::GetEmpty();
+
+	bFilterByNetId = false;
+	FilterNetId = 0;
+
+	bHighlightFilteredEvents = false;
+
 	DrawState->Reset();
+	FilteredDrawState->Reset();
 	bIsStateDirty = true;
 
 	MousePosition = FVector2D::ZeroVector;
@@ -95,6 +111,222 @@ void SPacketContentView::Construct(const FArguments& InArgs, TSharedPtr<SNetwork
 		SNew(SOverlay)
 		.Visibility(EVisibility::SelfHitTestInvisible)
 
+		+ SOverlay::Slot()
+		.HAlign(HAlign_Center)
+		.VAlign(VAlign_Top)
+		.Padding(FMargin(0, 0, 0, 0))
+		[
+			SNew(SHorizontalBox)
+
+			//////////////////////////////////////////////////
+			// Find Packet
+
+			+ SHorizontalBox::Slot()
+			.AutoWidth()
+			.VAlign(VAlign_Center)
+			.Padding(0.0f, 0.0f, 0.0f, 0.0f)
+			[
+				SNew(STextBlock)
+				.Text(LOCTEXT("FindPacketText", "Find Packet:"))
+			]
+
+			+ SHorizontalBox::Slot()
+			.AutoWidth()
+			.Padding(4.0f, 0.0f, 0.0f, 0.0f)
+			[
+				SNew(SButton)
+				.ToolTipText(LOCTEXT("PreviousPacketToolTip", "Previous Packet"))
+				.ContentPadding(0.0f)
+				.OnClicked(this, &SPacketContentView::FindPreviousPacket_OnClicked)
+				.Content()
+				[
+					SNew(SImage)
+					.Image(FInsightsStyle::GetBrush("FindPrevious"))
+				]
+			]
+
+			+ SHorizontalBox::Slot()
+			.AutoWidth()
+			.Padding(4.0f, 0.0f, 0.0f, 0.0f)
+			[
+				SNew(SEditableTextBox)
+				.RevertTextOnEscape(true)
+				.SelectAllTextWhenFocused(true)
+				.Text(this, &SPacketContentView::GetPacketText)
+				.OnTextCommitted(this, &SPacketContentView::Packet_OnTextCommitted)
+				.MinDesiredWidth(30.0f)
+			]
+
+			+ SHorizontalBox::Slot()
+			.AutoWidth()
+			.Padding(4.0f, 0.0f, 0.0f, 0.0f)
+			[
+				SNew(SButton)
+				.ToolTipText(LOCTEXT("NextPacketToolTip", "Next Packet"))
+				.ContentPadding(0.0f)
+				.OnClicked(this, &SPacketContentView::FindNextPacket_OnClicked)
+				.Content()
+				[
+					SNew(SImage)
+					.Image(FInsightsStyle::GetBrush("FindNext"))
+				]
+			]
+
+			//////////////////////////////////////////////////
+			// Find Event
+
+			+ SHorizontalBox::Slot()
+			.AutoWidth()
+			.VAlign(VAlign_Center)
+			.Padding(12.0f, 0.0f, 0.0f, 0.0f)
+			[
+				SNew(STextBlock)
+				.Text(LOCTEXT("FindEventText", "Find Event:"))
+			]
+
+			+ SHorizontalBox::Slot()
+			.AutoWidth()
+			.Padding(4.0f, 0.0f, 0.0f, 0.0f)
+			[
+				SNew(SButton)
+				.ToolTipText(LOCTEXT("FindFirstEventToolTip", "First Event"))
+				.ContentPadding(0.0f)
+				.OnClicked(this, &SPacketContentView::FindFirstEvent)
+				.Content()
+				[
+					SNew(SImage)
+					.Image(FInsightsStyle::GetBrush("FindFirst"))
+				]
+			]
+
+			+ SHorizontalBox::Slot()
+			.AutoWidth()
+			.Padding(4.0f, 0.0f, 0.0f, 0.0f)
+			[
+				SNew(SButton)
+				.ToolTipText(LOCTEXT("FindPreviousEventToolTip", "Previous Event"))
+				.ContentPadding(0.0f)
+				.OnClicked(this, &SPacketContentView::FindPreviousEvent)
+				.Content()
+				[
+					SNew(SImage)
+					.Image(FInsightsStyle::GetBrush("FindPrevious"))
+				]
+			]
+
+			+ SHorizontalBox::Slot()
+			.AutoWidth()
+			.Padding(4.0f, 0.0f, 0.0f, 0.0f)
+			[
+				SNew(SButton)
+				.ToolTipText(LOCTEXT("FindNextEventToolTip", "Next Event"))
+				.ContentPadding(0.0f)
+				.OnClicked(this, &SPacketContentView::FindNextEvent)
+				.Content()
+				[
+					SNew(SImage)
+					.Image(FInsightsStyle::GetBrush("FindNext"))
+				]
+			]
+
+			+ SHorizontalBox::Slot()
+			.AutoWidth()
+			.Padding(4.0f, 0.0f, 0.0f, 0.0f)
+			[
+				SNew(SButton)
+				.ToolTipText(LOCTEXT("FindLastEventToolTip", "Last Event"))
+				.ContentPadding(0.0f)
+				.OnClicked(this, &SPacketContentView::FindLastEvent)
+				.Content()
+				[
+					SNew(SImage)
+					.Image(FInsightsStyle::GetBrush("FindLast"))
+				]
+			]
+
+			//////////////////////////////////////////////////
+			// By NetId
+
+			+ SHorizontalBox::Slot()
+			.AutoWidth()
+			.VAlign(VAlign_Center)
+			.Padding(12.0f, 0.0f, 0.0f, 0.0f)
+			[
+				SNew(SCheckBox)
+				.ToolTipText(LOCTEXT("FilterByNetId_Tooltip", "Filter events that have the specified NetId."))
+				.IsChecked(this, &SPacketContentView::FilterByNetId_IsChecked)
+				.OnCheckStateChanged(this, &SPacketContentView::FilterByNetId_OnCheckStateChanged)
+				[
+					SNew(STextBlock)
+					.Text(LOCTEXT("FilterByNetId_Text", "By NetId:"))
+				]
+			]
+
+			+ SHorizontalBox::Slot()
+			.AutoWidth()
+			.Padding(4.0f, 0.0f, 0.0f, 0.0f)
+			[
+				SNew(SEditableTextBox)
+				.RevertTextOnEscape(true)
+				.SelectAllTextWhenFocused(true)
+				.Text(this, &SPacketContentView::GetFilterNetIdText)
+				.OnTextCommitted(this, &SPacketContentView::FilterNetId_OnTextCommitted)
+				.MinDesiredWidth(40.0f)
+			]
+
+			//////////////////////////////////////////////////
+			// By Event Type
+
+			+ SHorizontalBox::Slot()
+			.AutoWidth()
+			.VAlign(VAlign_Center)
+			.Padding(12.0f, 0.0f, 0.0f, 0.0f)
+			[
+				SNew(SCheckBox)
+				.ToolTipText(LOCTEXT("FilterByEventType_Tooltip",
+					"Filter events that have the specified type.\n\n"
+					"To set the event type:\n"
+					"\tdouble click either an event in the Packet Content view\n"
+					"\tor an event type in the NetStats tree view."))
+				.IsChecked(this, &SPacketContentView::FilterByEventType_IsChecked)
+				.OnCheckStateChanged(this, &SPacketContentView::FilterByEventType_OnCheckStateChanged)
+				[
+					SNew(STextBlock)
+					.Text(LOCTEXT("FilterByEventType_Text", "By Type:"))
+				]
+			]
+
+			+ SHorizontalBox::Slot()
+			.AutoWidth()
+			.Padding(4.0f, 0.0f, 0.0f, 0.0f)
+			[
+				SNew(SEditableTextBox)
+				.Text(this, &SPacketContentView::GetFilterEventTypeText)
+				.IsReadOnly(true)
+				.MinDesiredWidth(120.0f)
+			]
+
+			//////////////////////////////////////////////////
+			// Highlight Filtered Events
+
+			+ SHorizontalBox::Slot()
+			.AutoWidth()
+			.VAlign(VAlign_Center)
+			.Padding(12.0f, 0.0f, 0.0f, 0.0f)
+			[
+				SNew(SCheckBox)
+				.ToolTipText(LOCTEXT("HighlightFilteredEvents_Tooltip", "Highlight filtered events."))
+				.IsChecked(this, &SPacketContentView::HighlightFilteredEvents_IsChecked)
+				.OnCheckStateChanged(this, &SPacketContentView::HighlightFilteredEvents_OnCheckStateChanged)
+				[
+					SNew(STextBlock)
+					.Text(LOCTEXT("HighlightFilteredEvents_Text", "Highlight"))
+				]
+			]
+
+			//////////////////////////////////////////////////
+		]
+
 		+ SOverlay::Slot()
 		.VAlign(VAlign_Bottom)
 		.Padding(FMargin(0, 0, 0, 0))
@@ -116,6 +348,222 @@ void SPacketContentView::Construct(const FArguments& InArgs, TSharedPtr<SNetwork
 
 ////////////////////////////////////////////////////////////////////////////////////////////////////
 
+FReply SPacketContentView::FindPreviousPacket_OnClicked()
+{
+	TSharedPtr<SPacketView> PacketView = ProfilerWindow->GetPacketView();
+	if (PacketView.IsValid())
+	{
+		PacketView->SelectPreviousPacket();
+	}
+
+	return FReply::Handled();
+}
+
+////////////////////////////////////////////////////////////////////////////////////////////////////
+
+FReply SPacketContentView::FindNextPacket_OnClicked()
+{
+	TSharedPtr<SPacketView> PacketView = ProfilerWindow->GetPacketView();
+	if (PacketView.IsValid())
+	{
+		PacketView->SelectNextPacket();
+	}
+
+	return FReply::Handled();
+}
+
+////////////////////////////////////////////////////////////////////////////////////////////////////
+
+FText SPacketContentView::GetPacketText() const
+{
+	return FText::AsNumber(PacketIndex);
+}
+
+////////////////////////////////////////////////////////////////////////////////////////////////////
+
+void SPacketContentView::Packet_OnTextCommitted(const FText& InNewText, ETextCommit::Type InTextCommit)
+{
+	if (InNewText.IsNumeric())
+	{
+		int32 NewPacketIndex = 0;
+		TTypeFromString<int32>::FromString(NewPacketIndex, *InNewText.ToString());
+
+		TSharedPtr<SPacketView> PacketView = ProfilerWindow->GetPacketView();
+		if (PacketView.IsValid())
+		{
+			PacketView->SetSelectedPacket(NewPacketIndex);
+		}
+	}
+}
+
+////////////////////////////////////////////////////////////////////////////////////////////////////
+
+FReply SPacketContentView::FindFirstEvent()
+{
+	if (FilteredDrawState->Events.Num() > 0)
+	{
+		SelectedEvent.Set(FilteredDrawState->Events[0]);
+		OnSelectedEventChanged();
+		BringEventIntoView(SelectedEvent);
+	}
+
+	return FReply::Handled();
+}
+
+////////////////////////////////////////////////////////////////////////////////////////////////////
+
+FReply SPacketContentView::FindPreviousEvent()
+{
+	if (!SelectedEvent.IsValid())
+	{
+		return FindFirstEvent();
+	}
+
+	const int32 EventCount = FilteredDrawState->Events.Num();
+	for (int32 EventIndex = EventCount - 1; EventIndex >= 0; --EventIndex)
+	{
+		const FNetworkPacketEvent& Event = FilteredDrawState->Events[EventIndex];
+		if (Event.Equals(SelectedEvent.Event))
+		{
+			if (EventIndex > 0)
+			{
+				SelectedEvent.Set(FilteredDrawState->Events[EventIndex - 1]);
+				OnSelectedEventChanged();
+				break;
+			}
+		}
+		else if (Event.BitOffset <= SelectedEvent.Event.BitOffset)
+		{
+			if (Event.BitOffset < SelectedEvent.Event.BitOffset || Event.Level < SelectedEvent.Event.Level)
+			{
+				SelectedEvent.Set(Event);
+				OnSelectedEventChanged();
+				break;
+			}
+		}
+	}
+
+	BringEventIntoView(SelectedEvent);
+	return FReply::Handled();
+}
+
+////////////////////////////////////////////////////////////////////////////////////////////////////
+
+FReply SPacketContentView::FindNextEvent()
+{
+	if (!SelectedEvent.IsValid())
+	{
+		return FindLastEvent();
+	}
+
+	const int32 EventCount = FilteredDrawState->Events.Num();
+	for (int32 EventIndex = 0; EventIndex < EventCount; ++EventIndex)
+	{
+		const FNetworkPacketEvent& Event = FilteredDrawState->Events[EventIndex];
+		if (Event.Equals(SelectedEvent.Event))
+		{
+			if (EventIndex < EventCount - 1)
+			{
+				SelectedEvent.Set(FilteredDrawState->Events[EventIndex + 1]);
+				OnSelectedEventChanged();
+				break;
+			}
+		}
+		else if (Event.BitOffset >= SelectedEvent.Event.BitOffset)
+		{
+			if (Event.BitOffset > SelectedEvent.Event.BitOffset || Event.Level > SelectedEvent.Event.Level)
+			{
+				SelectedEvent.Set(Event);
+				OnSelectedEventChanged();
+				break;
+			}
+		}
+	}
+
+	BringEventIntoView(SelectedEvent);
+	return FReply::Handled();
+}
+
+////////////////////////////////////////////////////////////////////////////////////////////////////
+
+FReply SPacketContentView::FindLastEvent()
+{
+	if (FilteredDrawState->Events.Num() > 0)
+	{
+		SelectedEvent.Set(FilteredDrawState->Events.Last());
+		OnSelectedEventChanged();
+		BringEventIntoView(SelectedEvent);
+	}
+
+	return FReply::Handled();
+}
+
+////////////////////////////////////////////////////////////////////////////////////////////////////
+
+ECheckBoxState SPacketContentView::FilterByNetId_IsChecked() const
+{
+	return bFilterByNetId ? ECheckBoxState::Checked : ECheckBoxState::Unchecked;
+}
+
+////////////////////////////////////////////////////////////////////////////////////////////////////
+
+void SPacketContentView::FilterByNetId_OnCheckStateChanged(ECheckBoxState NewState)
+{
+	bFilterByNetId = (NewState == ECheckBoxState::Checked);
+	bIsStateDirty = true;
+}
+
+////////////////////////////////////////////////////////////////////////////////////////////////////
+
+FText SPacketContentView::GetFilterNetIdText() const
+{
+	return FText::AsNumber(FilterNetId);
+}
+
+////////////////////////////////////////////////////////////////////////////////////////////////////
+
+void SPacketContentView::FilterNetId_OnTextCommitted(const FText& InNewText, ETextCommit::Type InTextCommit)
+{
+	if (InNewText.IsNumeric())
+	{
+		int32 NewNetId = 0;
+		TTypeFromString<int32>::FromString(NewNetId, *InNewText.ToString());
+		SetFilterNetId(NewNetId);
+	}
+}
+
+////////////////////////////////////////////////////////////////////////////////////////////////////
+
+ECheckBoxState SPacketContentView::FilterByEventType_IsChecked() const
+{
+	return bFilterByEventType ? ECheckBoxState::Checked : ECheckBoxState::Unchecked;
+}
+
+////////////////////////////////////////////////////////////////////////////////////////////////////
+
+void SPacketContentView::FilterByEventType_OnCheckStateChanged(ECheckBoxState NewState)
+{
+	bFilterByEventType = (NewState == ECheckBoxState::Checked);
+	bIsStateDirty = true;
+}
+
+////////////////////////////////////////////////////////////////////////////////////////////////////
+
+ECheckBoxState SPacketContentView::HighlightFilteredEvents_IsChecked() const
+{
+	return bHighlightFilteredEvents ? ECheckBoxState::Checked : ECheckBoxState::Unchecked;
+}
+
+////////////////////////////////////////////////////////////////////////////////////////////////////
+
+void SPacketContentView::HighlightFilteredEvents_OnCheckStateChanged(ECheckBoxState NewState)
+{
+	bHighlightFilteredEvents = (NewState == ECheckBoxState::Checked);
+	bIsStateDirty = true;
+}
+
+////////////////////////////////////////////////////////////////////////////////////////////////////
+
 void SPacketContentView::Tick(const FGeometry& AllottedGeometry, const double InCurrentTime, const float InDeltaTime)
 {
 	if (ThisGeometry != AllottedGeometry || bIsViewportDirty)
@@ -169,6 +617,7 @@ void SPacketContentView::ResetPacket()
 	UpdateHorizontalScrollBar();
 
 	DrawState->Reset();
+	FilteredDrawState->Reset();
 	bIsStateDirty = true;
 
 	HoveredEvent.Reset();
@@ -191,6 +640,7 @@ void SPacketContentView::SetPacket(uint32 InGameInstanceIndex, uint32 InConnecti
 	UpdateHorizontalScrollBar();
 
 	DrawState->Reset();
+	FilteredDrawState->Reset();
 	bIsStateDirty = true;
 
 	HoveredEvent.Reset();
@@ -199,6 +649,64 @@ void SPacketContentView::SetPacket(uint32 InGameInstanceIndex, uint32 InConnecti
 
 ////////////////////////////////////////////////////////////////////////////////////////////////////
 
+void SPacketContentView::SetFilterNetId(const uint32 InNetId)
+{
+	FilterNetId = InNetId;
+
+	if (bFilterByNetId)
+	{
+		bIsStateDirty = true;
+	}
+}
+
+////////////////////////////////////////////////////////////////////////////////////////////////////
+
+void SPacketContentView::SetFilterEventType(const uint32 InEventTypeIndex, const FText& InEventName)
+{
+	FilterEventTypeIndex = InEventTypeIndex;
+	FilterEventName = InEventName;
+
+	if (bFilterByEventType)
+	{
+		bIsStateDirty = true;
+	}
+}
+
+////////////////////////////////////////////////////////////////////////////////////////////////////
+
+void SPacketContentView::EnableFilterEventType(const uint32 InEventTypeIndex)
+{
+	FText EventName;
+
+	TSharedPtr<const Trace::IAnalysisSession> Session = FInsightsManager::Get()->GetSession();
+	if (Session.IsValid())
+	{
+		Trace::FAnalysisSessionReadScope SessionReadScope(*Session.Get());
+		const Trace::INetProfilerProvider& NetProfilerProvider = Trace::ReadNetProfilerProvider(*Session.Get());
+
+		NetProfilerProvider.ReadEventType(InEventTypeIndex, [&EventName](const Trace::FNetProfilerEventType& EventType)
+		{
+			EventName = FText::FromString(EventType.Name);
+		});
+	}
+
+	bFilterByEventType = true;
+	SetFilterEventType(InEventTypeIndex, EventName);
+	bHighlightFilteredEvents = true;
+	bIsStateDirty = true;
+}
+
+////////////////////////////////////////////////////////////////////////////////////////////////////
+
+void SPacketContentView::DisableFilterEventType()
+{
+	bFilterByEventType = false;
+	bHighlightFilteredEvents = false;
+	bIsStateDirty = true;
+}
+
+////////////////////////////////////////////////////////////////////////////////////////////////////
+
 void SPacketContentView::UpdateState()
 {
 	FStopwatch Stopwatch;
@@ -207,6 +715,7 @@ void SPacketContentView::UpdateState()
 	if (PacketBitSize > 0)
 	{
 		FPacketContentViewDrawStateBuilder Builder(*DrawState, Viewport);
+		FPacketContentViewDrawStateBuilder FilteredDrawStateBuilder(*FilteredDrawState, Viewport);
 
 		TSharedPtr<const Trace::IAnalysisSession> Session = FInsightsManager::Get()->GetSession();
 		if (Session.IsValid())
@@ -214,25 +723,13 @@ void SPacketContentView::UpdateState()
 			Trace::FAnalysisSessionReadScope SessionReadScope(*Session.Get());
 			const Trace::INetProfilerProvider& NetProfilerProvider = Trace::ReadNetProfilerProvider(*Session.Get());
 
-			//uint32 StartEventIndex;
-			//uint32 EndEventIndex;
-			//NetProfilerProvider.ReadConnection(ConnectionIndex, [&StartEventIndex, &EndEventIndex](const Trace::FNetProfilerConnection& Connection)
-			//{
-			//	StartEventIndex = Connection.StartEventIndex;
-			//	EndEventIndex = Connection.EndEventIndex;
-			//});
-			//NetProfilerProvider.EnumeratePacketContentEventsByIndex(ConnectionIndex, ConnectionMode, StartEventIndex, uint32 EndEventIndex, [](const Trace::FNetProfilerContentEvent& Event)
-			//{
-			//});
-
 			const FAxisViewportDouble& ViewportX = Viewport.GetHorizontalAxisViewport();
 
-			const int64 MinPos = static_cast<int64>(FMath::FloorToDouble(ViewportX.GetValueAtOffset(0.0f)));
-			const int64 MaxPos = static_cast<int64>(FMath::CeilToDouble(ViewportX.GetValueAtOffset(ViewportX.GetSize())));
-
+			//const int64 StartPos = static_cast<int64>(FMath::FloorToDouble(ViewportX.GetValueAtOffset(0.0f)));
+			//const int64 EndPos = static_cast<int64>(FMath::CeilToDouble(ViewportX.GetValueAtOffset(ViewportX.GetSize())));
 			const uint32 StartPos = 0;
 			const uint32 EndPos = PacketBitSize;
-			NetProfilerProvider.EnumeratePacketContentEventsByPosition(ConnectionIndex, ConnectionMode, PacketIndex, StartPos, EndPos, [this, &Builder, &NetProfilerProvider](const Trace::FNetProfilerContentEvent& Event)
+			NetProfilerProvider.EnumeratePacketContentEventsByPosition(ConnectionIndex, ConnectionMode, PacketIndex, StartPos, EndPos, [this, &Builder, &FilteredDrawStateBuilder, &NetProfilerProvider](const Trace::FNetProfilerContentEvent& Event)
 			{
 				const TCHAR* Name = nullptr;
 				NetProfilerProvider.ReadName(Event.NameIndex, [&Name](const Trace::FNetProfilerName& NetProfilerName)
@@ -250,10 +747,17 @@ void SPacketContentView::UpdateState()
 				}
 
 				Builder.AddEvent(Event, Name, NetId);
+
+				if ((!bFilterByEventType || FilterEventTypeIndex == Event.EventTypeIndex) &&
+					(!bFilterByNetId || (Event.ObjectInstanceIndex != 0 && FilterNetId == NetId)))
+				{
+					FilteredDrawStateBuilder.AddEvent(Event, Name, NetId);
+				}
 			});
 		}
 
 		Builder.Flush();
+		FilteredDrawStateBuilder.Flush();
 	}
 
 	Stopwatch.Stop();
@@ -358,12 +862,8 @@ void SPacketContentView::OnSelectedEventChanged()
 void SPacketContentView::SelectHoveredEvent()
 {
 	SelectedEvent = HoveredEvent;
-	//if (SelectedEvent.IsValid())
-	//{
-	//	LastSelectionType = ESelectionType::TimingEvent;
-	//	BringIntoView(SelectedEvent.StartPos, SelectedEvent.EndPos);
-	//}
 	OnSelectedEventChanged();
+	BringEventIntoView(SelectedEvent);
 }
 
 ////////////////////////////////////////////////////////////////////////////////////////////////////
@@ -379,10 +879,10 @@ FNetworkPacketEventRef SPacketContentView::GetEventAtMousePosition(float X, floa
 			const float EventX1 = ViewportX.GetRoundedOffsetForValue(static_cast<double>(Event.BitOffset));
 			const float EventX2 = ViewportX.GetRoundedOffsetForValue(static_cast<double>(Event.BitOffset + Event.BitSize));
 
-			constexpr float Y0 = 0.0f;
+			constexpr float EventsPosY = 32.0f;
 			constexpr float EventH = 14.0f;
 			constexpr float EventDY = 2.0f;
-			const float EventY = Y0 + (EventH + EventDY) * Event.Level;
+			const float EventY = EventsPosY + (EventH + EventDY) * Event.Level;
 
 			constexpr float ToleranceX = 1.0f;
 
@@ -414,10 +914,22 @@ int32 SPacketContentView::OnPaint(const FPaintArgs& Args, const FGeometry& Allot
 
 		FPacketContentViewDrawHelper Helper(DrawContext, Viewport);
 
+		Helper.SetLayoutPosY(32.0f);
+		//Helper.SetLayoutEventH(14.0f);
+		//Helper.SetLayoutEventDY(2.0f);
+
 		Helper.DrawBackground();
 
-		// Draw the cached state (the network events of the package breakdown).
-		Helper.Draw(*DrawState);
+		if (bHighlightFilteredEvents && (bFilterByNetId || bFilterByEventType))
+		{
+			Helper.Draw(*DrawState, 0.1f);
+			Helper.Draw(*FilteredDrawState);
+		}
+		else
+		{
+			// Draw the events contained by the network packet using the cached draw state.
+			Helper.Draw(*DrawState);
+		}
 
 		if (!FNetworkPacketEventRef::AreEquals(SelectedEvent, HoveredEvent))
 		{
@@ -725,7 +1237,20 @@ FReply SPacketContentView::OnMouseWheel(const FGeometry& MyGeometry, const FPoin
 
 FReply SPacketContentView::OnMouseButtonDoubleClick(const FGeometry& MyGeometry, const FPointerEvent& MouseEvent)
 {
-	return FReply::Unhandled();
+	// Select the hovered timing event (if any).
+	UpdateHoveredEvent();
+	SelectHoveredEvent();
+
+	if (SelectedEvent.IsValid())
+	{
+		EnableFilterEventType(SelectedEvent.Event.EventTypeIndex);
+	}
+	else
+	{
+		DisableFilterEventType();
+	}
+
+	return FReply::Handled();
 }
 
 ////////////////////////////////////////////////////////////////////////////////////////////////////
@@ -748,6 +1273,38 @@ FCursorReply SPacketContentView::OnCursorQuery(const FGeometry& MyGeometry, cons
 
 ////////////////////////////////////////////////////////////////////////////////////////////////////
 
+FReply SPacketContentView::OnKeyDown(const FGeometry& MyGeometry, const FKeyEvent& InKeyEvent)
+{
+	if (InKeyEvent.GetKey() == EKeys::Left)
+	{
+		if (InKeyEvent.GetModifierKeys().IsShiftDown())
+		{
+			FindFirstEvent();
+		}
+		else
+		{
+			FindPreviousEvent();
+		}
+		return FReply::Handled();
+	}
+	else if (InKeyEvent.GetKey() == EKeys::Right)
+	{
+		if (InKeyEvent.GetModifierKeys().IsShiftDown())
+		{
+			FindLastEvent();
+		}
+		else
+		{
+			FindNextEvent();
+		}
+		return FReply::Handled();
+	}
+
+	return SCompoundWidget::OnKeyDown(MyGeometry, InKeyEvent);
+}
+
+////////////////////////////////////////////////////////////////////////////////////////////////////
+
 void SPacketContentView::BindCommands()
 {
 }
@@ -781,4 +1338,48 @@ void SPacketContentView::ZoomHorizontally(const float Delta, const float X)
 
 ////////////////////////////////////////////////////////////////////////////////////////////////////
 
+void SPacketContentView::BringIntoView(const float X1, const float X2)
+{
+	FAxisViewportDouble& ViewportX = Viewport.GetHorizontalAxisViewport();
+
+	// Increase interval with 8% (of view size) on each side.
+	const float DX = ViewportX.GetSize() * 0.08f;
+
+	float NewPos = ViewportX.GetPos();
+
+	const float MinPos = X2 + DX - ViewportX.GetSize();
+	if (NewPos < MinPos)
+	{
+		NewPos = MinPos;
+	}
+
+	const float MaxPos = X1 - DX;
+	if (NewPos > MaxPos)
+	{
+		NewPos = MaxPos;
+	}
+
+	if (NewPos != ViewportX.GetPos())
+	{
+		ViewportX.ScrollAtPos(NewPos);
+		UpdateHorizontalScrollBar();
+		bIsStateDirty = true;
+	}
+}
+
+////////////////////////////////////////////////////////////////////////////////////////////////////
+
+void SPacketContentView::BringEventIntoView(const FNetworkPacketEventRef& EventRef)
+{
+	if (EventRef.IsValid())
+	{
+		const FAxisViewportDouble& ViewportX = Viewport.GetHorizontalAxisViewport();
+		const float X1 = ViewportX.GetPosForValue(static_cast<double>(EventRef.Event.BitOffset));
+		const float X2 = ViewportX.GetPosForValue(static_cast<double>(EventRef.Event.BitOffset + SelectedEvent.Event.BitSize));
+		BringIntoView(X1, X2);
+	}
+}
+
+////////////////////////////////////////////////////////////////////////////////////////////////////
+
 #undef LOCTEXT_NAMESPACE
diff --git a/Engine/Source/Developer/TraceInsights/Private/Insights/NetworkingProfiler/Widgets/SPacketContentView.h b/Engine/Source/Developer/TraceInsights/Private/Insights/NetworkingProfiler/Widgets/SPacketContentView.h
--- a/Engine/Source/Developer/TraceInsights/Private/Insights/NetworkingProfiler/Widgets/SPacketContentView.h
+++ b/Engine/Source/Developer/TraceInsights/Private/Insights/NetworkingProfiler/Widgets/SPacketContentView.h
@@ -7,6 +7,7 @@
 #include "Input/Reply.h"
 #include "Layout/Geometry.h"
 #include "Rendering/RenderingCommon.h"
+#include "Styling/SlateTypes.h"
 #include "TraceServices/Model/NetProfiler.h"
 #include "Widgets/DeclarativeSyntaxSupport.h"
 #include "Widgets/SCompoundWidget.h"
@@ -52,6 +53,12 @@ struct FNetworkPacketEventRef
 		return *this;
 	}
 
+	void Set(const FNetworkPacketEvent& InEvent)
+	{
+		Event = InEvent;
+		bIsValid = true;
+	}
+
 	void Reset()
 	{
 		bIsValid = false;
@@ -125,10 +132,48 @@ public:
 
 	virtual FCursorReply OnCursorQuery(const FGeometry& MyGeometry, const FPointerEvent& CursorEvent) const override;
 
+	virtual bool SupportsKeyboardFocus() const override { return true; }
+	virtual FReply OnKeyDown(const FGeometry& MyGeometry, const FKeyEvent& InKeyEvent) override;
+	//virtual FReply OnKeyUp(const FGeometry& MyGeometry, const FKeyEvent& InKeyEvent) override;
+
 	void ResetPacket();
 	void SetPacket(uint32 InGameInstanceIndex, uint32 InConnectionIndex, Trace::ENetProfilerConnectionMode InConnectionMode, uint32 InPacketIndex, int64 InPacketBitSize);
 
+	bool IsFilterByNetIdEnabled() const { return bFilterByNetId; }
+	uint32 GetFilterNetId() const { return FilterNetId; }
+	void SetFilterNetId(const uint32 InNetId);
+
+	bool IsFilterByEventTypeEnabled() const { return bFilterByEventType; }
+	uint32 GetFilterEventTypeIndex() const { return FilterEventTypeIndex; }
+	const FText& GetFilterEventName() const { return FilterEventName; }
+	void SetFilterEventType(const uint32 InEventTypeIndex, const FText& InEventName);
+	void EnableFilterEventType(const uint32 InEventTypeIndex);
+	void DisableFilterEventType();
+
+	FReply FindFirstEvent();
+	FReply FindPreviousEvent();
+	FReply FindNextEvent();
+	FReply FindLastEvent();
+
 private:
+	FReply FindPreviousPacket_OnClicked();
+	FReply FindNextPacket_OnClicked();
+
+	FText GetPacketText() const;
+	void Packet_OnTextCommitted(const FText& InNewText, ETextCommit::Type InTextCommit);
+
+	ECheckBoxState FilterByNetId_IsChecked() const;
+	void FilterByNetId_OnCheckStateChanged(ECheckBoxState NewState);
+	FText GetFilterNetIdText() const;
+	void FilterNetId_OnTextCommitted(const FText& InNewText, ETextCommit::Type InTextCommit);
+
+	ECheckBoxState FilterByEventType_IsChecked() const;
+	void FilterByEventType_OnCheckStateChanged(ECheckBoxState NewState);
+	FText GetFilterEventTypeText() const { return FilterEventName; }
+
+	ECheckBoxState HighlightFilteredEvents_IsChecked() const;
+	void HighlightFilteredEvents_OnCheckStateChanged(ECheckBoxState NewState);
+
 	//void ShowContextMenu(const FPointerEvent& MouseEvent);
 
 	/** Binds our UI commands to delegates. */
@@ -144,6 +189,8 @@ private:
 	void UpdateHorizontalScrollBar();
 
 	void ZoomHorizontally(const float Delta, const float X);
+	void BringIntoView(const float X1, const float X2);
+	void BringEventIntoView(const FNetworkPacketEventRef& EventRef);
 
 	////////////////////////////////////////////////////////////////////////////////////////////////////
 
@@ -162,15 +209,28 @@ private:
 	FPacketContentViewport Viewport;
 	bool bIsViewportDirty;
 
+	float NetEventsOffsetY;
+
 	uint32 GameInstanceIndex;
 	uint32 ConnectionIndex;
 	Trace::ENetProfilerConnectionMode ConnectionMode;
 	uint32 PacketIndex;
 	int64 PacketBitSize; // total number of bits; [bit]
 
-	/** Cached draw state of the packet content (i.e. all it needs to render). */
+	bool bFilterByEventType;
+	uint32 FilterEventTypeIndex;
+	FText FilterEventName;
+
+	bool bFilterByNetId;
+	uint32 FilterNetId;
+
+	bool bHighlightFilteredEvents;
+
+	/** Cached draw state of the packet content. */
 	TSharedRef<FPacketContentViewDrawState> DrawState;
+	TSharedRef<FPacketContentViewDrawState> FilteredDrawState;
 	bool bIsStateDirty;
+
 	//////////////////////////////////////////////////
 
 	TSharedPtr<SScrollBar> HorizontalScrollBar;
diff --git a/Engine/Source/Developer/TraceInsights/Private/Insights/NetworkingProfiler/Widgets/SPacketView.cpp b/Engine/Source/Developer/TraceInsights/Private/Insights/NetworkingProfiler/Widgets/SPacketView.cpp
--- a/Engine/Source/Developer/TraceInsights/Private/Insights/NetworkingProfiler/Widgets/SPacketView.cpp
+++ b/Engine/Source/Developer/TraceInsights/Private/Insights/NetworkingProfiler/Widgets/SPacketView.cpp
@@ -16,6 +16,7 @@
 #include "Insights/Common/PaintUtils.h"
 #include "Insights/Common/Stopwatch.h"
 #include "Insights/Common/TimeUtils.h"
+#include "Insights/TimingProfilerCommon.h" // for UE_LOG
 #include "Insights/InsightsManager.h"
 #include "Insights/InsightsStyle.h"
 #include "Insights/NetworkingProfiler/NetworkingProfilerManager.h"
@@ -215,6 +216,33 @@ void SPacketView::Tick(const FGeometry& AllottedGeometry, const double InCurrent
 
 ////////////////////////////////////////////////////////////////////////////////////////////////////
 
+bool SPacketView::IsConnectionValid(const Trace::INetProfilerProvider& NetProfilerProvider, const uint32 InGameInstanceIndex, const uint32 InConnectionIndex, const Trace::ENetProfilerConnectionMode InConnectionMode)
+{
+	//TODO: return NetProfilerProvider.IsConnectionValid(GameInstanceIndex, ConnectionIndex, ConnectionMode);
+
+	const uint32 GameInstanceCount = NetProfilerProvider.GetGameInstanceCount();
+	if (InGameInstanceIndex >= GameInstanceCount)
+	{
+		return false;
+	}
+
+	bool bIsValidConnection = false;
+	NetProfilerProvider.ReadConnections(InGameInstanceIndex, [InConnectionIndex, InConnectionMode, &bIsValidConnection](const Trace::FNetProfilerConnection& Connection)
+	{
+		if (InConnectionIndex == Connection.ConnectionIndex)
+		{
+			if ((InConnectionMode == Trace::ENetProfilerConnectionMode::Outgoing && Connection.bHasOutgoingData) ||
+				(InConnectionMode == Trace::ENetProfilerConnectionMode::Incoming && Connection.bHasIncomingData))
+			{
+				bIsValidConnection = true;
+			}
+		}
+	});
+	return bIsValidConnection;
+}
+
+////////////////////////////////////////////////////////////////////////////////////////////////////
+
 void SPacketView::UpdateState()
 {
 	FStopwatch Stopwatch;
@@ -235,45 +263,27 @@ void SPacketView::UpdateState()
 		Trace::FAnalysisSessionReadScope SessionReadScope(*Session.Get());
 		const Trace::INetProfilerProvider& NetProfilerProvider = Trace::ReadNetProfilerProvider(*Session.Get());
 
-		// TODO: if (NetProfilerProvider.IsConnectionValid(GameInstanceIndex, ConnectionIndex, ConnectionMode))
-		const uint32 GameInstanceCount = NetProfilerProvider.GetGameInstanceCount();
-		if (GameInstanceIndex < GameInstanceCount)
+		if (IsConnectionValid(NetProfilerProvider, GameInstanceIndex, ConnectionIndex, ConnectionMode))
 		{
-			bool bIsValidConnection = false;
-			NetProfilerProvider.ReadConnections(GameInstanceIndex, [this, &bIsValidConnection](const Trace::FNetProfilerConnection& Connection)
-			{
-				if (ConnectionIndex == Connection.ConnectionIndex)
-				{
-					if ((ConnectionMode == Trace::ENetProfilerConnectionMode::Outgoing && Connection.bHasOutgoingData) ||
-						(ConnectionMode == Trace::ENetProfilerConnectionMode::Incoming && Connection.bHasIncomingData))
-					{
-						bIsValidConnection = true;
-					}
-				}
-			});
+			const uint32 NumPackets = NetProfilerProvider.GetPacketCount(ConnectionIndex, ConnectionMode);
+
+			ViewportX.SetMinMaxInterval(0, NumPackets);
 
-			if (bIsValidConnection)
+			if (NumPackets > 0)
 			{
-				const uint32 NumPackets = NetProfilerProvider.GetPacketCount(ConnectionIndex, ConnectionMode);
+				const int32 MinIndex = ViewportX.GetValueAtOffset(0.0f);
+				const int32 MaxIndex = ViewportX.GetValueAtOffset(ViewportX.GetSize());
 
-				ViewportX.SetMinMaxInterval(0, NumPackets);
+				const uint32 PacketStartIndex = static_cast<uint32>(FMath::Max(0, MinIndex));
+				const uint32 PacketEndIndex = FMath::Min(NumPackets - 1, static_cast<uint32>(FMath::Max(0, MaxIndex)));
 
-				if (NumPackets > 0)
+				if (PacketStartIndex <= PacketEndIndex)
 				{
-					const int32 MinIndex = ViewportX.GetValueAtOffset(0.0f);
-					const int32 MaxIndex = ViewportX.GetValueAtOffset(ViewportX.GetSize());
-
-					const uint32 PacketStartIndex = static_cast<uint32>(FMath::Max(0, MinIndex));
-					const uint32 PacketEndIndex = FMath::Min(NumPackets - 1, static_cast<uint32>(FMath::Max(0, MaxIndex)));
-
-					if (PacketStartIndex <= PacketEndIndex)
+					int32 PacketIndex = PacketStartIndex;
+					NetProfilerProvider.EnumeratePackets(ConnectionIndex, ConnectionMode, PacketStartIndex, PacketEndIndex, [&Builder, &PacketIndex](const Trace::FNetProfilerPacket& Packet)
 					{
-						int32 PacketIndex = PacketStartIndex;
-						NetProfilerProvider.EnumeratePackets(ConnectionIndex, ConnectionMode, PacketStartIndex, PacketEndIndex, [&Builder, &PacketIndex](const Trace::FNetProfilerPacket& Packet)
-						{
-							Builder.AddPacket(PacketIndex++, Packet);
-						});
-					}
+						Builder.AddPacket(PacketIndex++, Packet);
+					});
 				}
 			}
 		}
@@ -325,6 +335,36 @@ void SPacketView::UpdateState()
 
 ////////////////////////////////////////////////////////////////////////////////////////////////////
 
+FNetworkPacketSampleRef SPacketView::GetSample(const int32 InPacketIndex)
+{
+	FNetworkPacketSampleRef SampleRef;
+	SampleRef.Series = PacketSeries;
+
+	TSharedPtr<const Trace::IAnalysisSession> Session = FInsightsManager::Get()->GetSession();
+	if (Session.IsValid())
+	{
+		Trace::FAnalysisSessionReadScope SessionReadScope(*Session.Get());
+		const Trace::INetProfilerProvider& NetProfilerProvider = Trace::ReadNetProfilerProvider(*Session.Get());
+
+		if (IsConnectionValid(NetProfilerProvider, GameInstanceIndex, ConnectionIndex, ConnectionMode))
+		{
+			const uint32 NumPackets = NetProfilerProvider.GetPacketCount(ConnectionIndex, ConnectionMode);
+			if (InPacketIndex >= 0 && InPacketIndex < static_cast<int32>(NumPackets))
+			{
+				NetProfilerProvider.EnumeratePackets(ConnectionIndex, ConnectionMode, InPacketIndex, InPacketIndex, [InPacketIndex, &SampleRef](const Trace::FNetProfilerPacket& Packet)
+				{
+					SampleRef.Sample = MakeShared<FNetworkPacketAggregatedSample>();
+					SampleRef.Sample->AddPacket(InPacketIndex, Packet);
+				});
+			}
+		}
+	}
+
+	return SampleRef;
+}
+
+////////////////////////////////////////////////////////////////////////////////////////////////////
+
 FNetworkPacketSampleRef SPacketView::GetSampleAtMousePosition(float X, float Y)
 {
 	if (!bIsStateDirty)
@@ -488,8 +528,8 @@ const TCHAR* AggregatedStatusToString(Trace::ENetProfilerDeliveryStatus Status)
 {
 	switch (Status)
 	{
-		case Trace::ENetProfilerDeliveryStatus::Delivered:  return TEXT("all packages are Delivered");
-		case Trace::ENetProfilerDeliveryStatus::Dropped:    return TEXT("at least one Dropped package");
+		case Trace::ENetProfilerDeliveryStatus::Delivered:  return TEXT("all packets are Delivered");
+		case Trace::ENetProfilerDeliveryStatus::Dropped:    return TEXT("at least one Dropped packet");
 		case Trace::ENetProfilerDeliveryStatus::Unknown:
 		default:                                            return TEXT("Unknown");
 	}
@@ -1101,22 +1141,300 @@ FReply SPacketView::OnKeyDown(const FGeometry& MyGeometry, const FKeyEvent& InKe
 	{
 		if (InKeyEvent.GetModifierKeys().IsControlDown())
 		{
-			// Select all.
-			SelectedSample.Reset();
-			SelectionStartPacketIndex = 0;
 			const FAxisViewportInt32& ViewportX = Viewport.GetHorizontalAxisViewport();
-			SelectionEndPacketIndex = ViewportX.GetMaxValue();
+			const int32 LastPacketIndex = ViewportX.GetMaxValue();
+
+			// Select all packets.
+			SelectionStartPacketIndex = 0;
+			SelectionEndPacketIndex = LastPacketIndex;
+
 			LastSelectedPacketIndex = 0;
+
+			SelectedSample.Reset();
+			if (SelectionEndPacketIndex == SelectionStartPacketIndex + 1)
+			{
+				SelectedSample = GetSample(SelectionStartPacketIndex);
+			}
 			OnSelectionChanged();
+
 			return FReply::Handled();
 		}
 	}
+	else if (InKeyEvent.GetKey() == EKeys::Left)
+	{
+		if (InKeyEvent.GetModifierKeys().IsShiftDown())
+		{
+			if (InKeyEvent.GetModifierKeys().IsControlDown())
+			{
+				ShrinkRightSideOfSelectedInterval();
+			}
+			else
+			{
+				ExtendLeftSideOfSelectedInterval();
+			}
+		}
+		else
+		{
+			SelectPreviousPacket();
+		}
+		return FReply::Handled();
+	}
+	else if (InKeyEvent.GetKey() == EKeys::Right)
+	{
+		if (InKeyEvent.GetModifierKeys().IsShiftDown())
+		{
+			if (InKeyEvent.GetModifierKeys().IsControlDown())
+			{
+				ShrinkLeftSideOfSelectedInterval();
+			}
+			else
+			{
+				ExtendRightSideOfSelectedInterval();
+			}
+		}
+		else
+		{
+			SelectNextPacket();
+		}
+		return FReply::Handled();
+	}
 
 	return SCompoundWidget::OnKeyDown(MyGeometry, InKeyEvent);
 }
 
 ////////////////////////////////////////////////////////////////////////////////////////////////////
 
+void SPacketView::EnsurePacketIsVisible(const int InPacketIndex)
+{
+	FAxisViewportInt32& ViewportX = Viewport.GetHorizontalAxisViewport();
+	const float LeftX = ViewportX.GetPosForValue(InPacketIndex);
+	const float RightX = LeftX + ViewportX.GetSampleSize();
+
+	if (RightX > ViewportX.GetPos() + ViewportX.GetSize())
+	{
+		const int32 VisibleSampleCount = FMath::FloorToInt(ViewportX.GetSize() / ViewportX.GetSampleSize());
+		ViewportX.ScrollAtValue(InPacketIndex - VisibleSampleCount + 1);
+		UpdateHorizontalScrollBar();
+		bIsStateDirty = true;
+	}
+
+	if (LeftX < ViewportX.GetPos())
+	{
+		ViewportX.ScrollAtValue(InPacketIndex);
+		UpdateHorizontalScrollBar();
+		bIsStateDirty = true;
+	}
+}
+
+////////////////////////////////////////////////////////////////////////////////////////////////////
+
+void SPacketView::SetSelectedPacket(const int32 InPacketIndex)
+{
+	const FAxisViewportInt32& ViewportX = Viewport.GetHorizontalAxisViewport();
+	const int32 PacketCount = ViewportX.GetMaxValue();
+	if (PacketCount > 0)
+	{
+		const int32 PacketIndex = FMath::Min(FMath::Max(0, InPacketIndex), PacketCount - 1);
+		SelectionStartPacketIndex = PacketIndex;
+		SelectionEndPacketIndex = PacketIndex + 1;
+		LastSelectedPacketIndex = PacketIndex;
+		EnsurePacketIsVisible(LastSelectedPacketIndex);
+		UpdateSelectedSample();
+	}
+}
+
+////////////////////////////////////////////////////////////////////////////////////////////////////
+
+void SPacketView::SelectPreviousPacket()
+{
+	const FAxisViewportInt32& ViewportX = Viewport.GetHorizontalAxisViewport();
+	const int32 PacketCount = ViewportX.GetMaxValue();
+
+	if (PacketCount > 0)
+	{
+		if (SelectionStartPacketIndex >= SelectionEndPacketIndex) // no selection?
+		{
+			// Select the first packet.
+			SelectionStartPacketIndex = 0;
+			SelectionEndPacketIndex = 1;
+		}
+		else
+		{
+			if (SelectionStartPacketIndex > 0)
+			{
+				// Select the previous packet.
+				SelectionStartPacketIndex--;
+			}
+			SelectionEndPacketIndex = SelectionStartPacketIndex + 1;
+		}
+
+		LastSelectedPacketIndex = SelectionStartPacketIndex;
+		EnsurePacketIsVisible(LastSelectedPacketIndex);
+		UpdateSelectedSample();
+	}
+}
+
+////////////////////////////////////////////////////////////////////////////////////////////////////
+
+void SPacketView::SelectNextPacket()
+{
+	const FAxisViewportInt32& ViewportX = Viewport.GetHorizontalAxisViewport();
+	const int32 PacketCount = ViewportX.GetMaxValue();
+
+	if (PacketCount > 0)
+	{
+		if (SelectionStartPacketIndex >= SelectionEndPacketIndex) // no selection?
+		{
+			// Select the last packet.
+			SelectionStartPacketIndex = PacketCount - 1;
+			SelectionEndPacketIndex = PacketCount;
+		}
+		else
+		{
+			if (SelectionEndPacketIndex < PacketCount)
+			{
+				// Select the next packet.
+				SelectionEndPacketIndex++;
+			}
+			SelectionStartPacketIndex = SelectionEndPacketIndex - 1;
+		}
+
+		LastSelectedPacketIndex = SelectionStartPacketIndex;
+		EnsurePacketIsVisible(LastSelectedPacketIndex);
+		UpdateSelectedSample();
+	}
+}
+
+////////////////////////////////////////////////////////////////////////////////////////////////////
+
+void SPacketView::ExtendLeftSideOfSelectedInterval()
+{
+	const FAxisViewportInt32& ViewportX = Viewport.GetHorizontalAxisViewport();
+	const int32 PacketCount = ViewportX.GetMaxValue();
+
+	if (PacketCount > 0)
+	{
+		bool bSelectionChanged = false;
+
+		if (SelectionStartPacketIndex >= SelectionEndPacketIndex) // no selection?
+		{
+			// Select the first packet.
+			SelectionStartPacketIndex = 0;
+			SelectionEndPacketIndex = 1;
+			bSelectionChanged = true;
+		}
+		else
+		{
+			if (SelectionStartPacketIndex > 0)
+			{
+				// Extend left side of selected interval.
+				SelectionStartPacketIndex--;
+				bSelectionChanged = true;
+			}
+		}
+
+		if (bSelectionChanged)
+		{
+			LastSelectedPacketIndex = SelectionStartPacketIndex;
+			EnsurePacketIsVisible(LastSelectedPacketIndex);
+			UpdateSelectedSample();
+		}
+	}
+}
+
+////////////////////////////////////////////////////////////////////////////////////////////////////
+
+void SPacketView::ShrinkLeftSideOfSelectedInterval()
+{
+	const FAxisViewportInt32& ViewportX = Viewport.GetHorizontalAxisViewport();
+	const int32 PacketCount = ViewportX.GetMaxValue();
+
+	if (PacketCount > 0)
+	{
+		if (SelectionStartPacketIndex + 1 < SelectionEndPacketIndex)
+		{
+			// Shrink left side of selected interval.
+			SelectionStartPacketIndex++;
+
+			LastSelectedPacketIndex = SelectionStartPacketIndex;
+			EnsurePacketIsVisible(LastSelectedPacketIndex);
+			UpdateSelectedSample();
+		}
+	}
+}
+
+////////////////////////////////////////////////////////////////////////////////////////////////////
+
+void SPacketView::ExtendRightSideOfSelectedInterval()
+{
+	const FAxisViewportInt32& ViewportX = Viewport.GetHorizontalAxisViewport();
+	const int32 PacketCount = ViewportX.GetMaxValue();
+
+	if (PacketCount > 0)
+	{
+		bool bSelectionChanged = false;
+
+		if (SelectionStartPacketIndex >= SelectionEndPacketIndex) // no selection?
+		{
+			// Select the last packet.
+			SelectionStartPacketIndex = PacketCount - 1;
+			SelectionEndPacketIndex = PacketCount;
+			bSelectionChanged = true;
+		}
+		else
+		{
+			if (SelectionEndPacketIndex < PacketCount)
+			{
+				// Extend right side of selected interval.
+				SelectionEndPacketIndex++;
+				bSelectionChanged = true;
+			}
+		}
+
+		if (bSelectionChanged)
+		{
+			LastSelectedPacketIndex = SelectionEndPacketIndex - 1;
+			EnsurePacketIsVisible(LastSelectedPacketIndex);
+			UpdateSelectedSample();
+		}
+	}
+}
+
+////////////////////////////////////////////////////////////////////////////////////////////////////
+
+void SPacketView::ShrinkRightSideOfSelectedInterval()
+{
+	const FAxisViewportInt32& ViewportX = Viewport.GetHorizontalAxisViewport();
+	const int32 PacketCount = ViewportX.GetMaxValue();
+
+	if (PacketCount > 0)
+	{
+		if (SelectionEndPacketIndex - 1 > SelectionStartPacketIndex)
+		{
+			// Shrink right side of selected interval.
+			SelectionEndPacketIndex--;
+
+			LastSelectedPacketIndex = SelectionEndPacketIndex - 1;
+			EnsurePacketIsVisible(LastSelectedPacketIndex);
+			UpdateSelectedSample();
+		}
+	}
+}
+
+////////////////////////////////////////////////////////////////////////////////////////////////////
+
+void SPacketView::UpdateSelectedSample()
+{
+	SelectedSample.Reset();
+	if (SelectionEndPacketIndex == SelectionStartPacketIndex + 1)
+	{
+		SelectedSample = GetSample(SelectionStartPacketIndex);
+	}
+	OnSelectionChanged();
+}
+
+////////////////////////////////////////////////////////////////////////////////////////////////////
+
 void SPacketView::ShowContextMenu(const FPointerEvent& MouseEvent)
 {
 	const bool bShouldCloseWindowAfterMenuSelection = true;
diff --git a/Engine/Source/Developer/TraceInsights/Private/Insights/NetworkingProfiler/Widgets/SPacketView.h b/Engine/Source/Developer/TraceInsights/Private/Insights/NetworkingProfiler/Widgets/SPacketView.h
--- a/Engine/Source/Developer/TraceInsights/Private/Insights/NetworkingProfiler/Widgets/SPacketView.h
+++ b/Engine/Source/Developer/TraceInsights/Private/Insights/NetworkingProfiler/Widgets/SPacketView.h
@@ -132,12 +132,25 @@ public:
 	virtual FReply OnKeyDown(const FGeometry& MyGeometry, const FKeyEvent& InKeyEvent) override;
 	//virtual FReply OnKeyUp(const FGeometry& MyGeometry, const FKeyEvent& InKeyEvent) override;
 
+	void EnsurePacketIsVisible(const int InPacketIndex);
+	void SetSelectedPacket(const int32 InPacketIndex);
+	void SelectPreviousPacket();
+	void SelectNextPacket();
+	void ExtendLeftSideOfSelectedInterval();
+	void ShrinkLeftSideOfSelectedInterval();
+	void ExtendRightSideOfSelectedInterval();
+	void ShrinkRightSideOfSelectedInterval();
+
 private:
+	void UpdateSelectedSample();
+
+	bool IsConnectionValid(const Trace::INetProfilerProvider& NetProfilerProvider, const uint32 InGameInstanceIndex, const uint32 InConnectionIndex, const Trace::ENetProfilerConnectionMode InConnectionMode);
 	void UpdateState();
 
 	void DrawHorizontalAxisGrid(FDrawContext& DrawContext, const FSlateBrush* Brush, const FSlateFontInfo& Font) const;
 	void DrawVerticalAxisGrid(FDrawContext& DrawContext, const FSlateBrush* Brush, const FSlateFontInfo& Font) const;
 
+	FNetworkPacketSampleRef GetSample(const int32 InPacketIndex);
 	FNetworkPacketSampleRef GetSampleAtMousePosition(float X, float Y);
 	void SelectSampleAtMousePosition(float X, float Y, const FPointerEvent& MouseEvent);
 	void OnSelectionChanged();
@@ -172,7 +185,7 @@ private:
 	bool bIsViewportDirty;
 
 	/** Cached info for the packet series. */
-	TSharedPtr<FNetworkPacketSeries> PacketSeries;
+	TSharedRef<FNetworkPacketSeries> PacketSeries;
 	bool bIsStateDirty;
 
 	bool bIsAutoZoomEnabled;
