[Insights]
  - Networking / Packet view: Added highlight for packets with at least one event matching the filter (by NetId and/or by [REDACTED_PERSON]). The packets with no event matching the filter will be displayed with a faded color.
  - (previous commit)  Networking: Fixed packge selection being reset if selection is made too soon.

diff --git a/Engine/Source/Developer/TraceInsights/Private/Insights/NetworkingProfiler/ViewModels/PacketViewDrawHelper.cpp b/Engine/Source/Developer/TraceInsights/Private/Insights/NetworkingProfiler/ViewModels/PacketViewDrawHelper.cpp
--- a/Engine/Source/Developer/TraceInsights/Private/Insights/NetworkingProfiler/ViewModels/PacketViewDrawHelper.cpp
+++ b/Engine/Source/Developer/TraceInsights/Private/Insights/NetworkingProfiler/ViewModels/PacketViewDrawHelper.cpp
@@ -81,7 +81,7 @@ FNetworkPacketSeriesBuilder::FNetworkPacketSeriesBuilder(FNetworkPacketSeries& I
 
 ////////////////////////////////////////////////////////////////////////////////////////////////////
 
-void FNetworkPacketSeriesBuilder::AddPacket(const int32 PacketIndex, const Trace::FNetProfilerPacket& Packet)
+FNetworkPacketAggregatedSample* FNetworkPacketSeriesBuilder::AddPacket(const int32 PacketIndex, const Trace::FNetProfilerPacket& Packet)
 {
 	NumAddedPackets++;
 
@@ -91,7 +91,10 @@ void FNetworkPacketSeriesBuilder::AddPacket(const int32 PacketIndex, const Trace
 		FNetworkPacketAggregatedSample& Sample = Series.Samples[SampleIndex];
 		Sample.AddPacket(PacketIndex, Packet);
 		Series.NumAggregatedPackets++;
+		return &Sample;
 	}
+
+	return nullptr;
 }
 
 ////////////////////////////////////////////////////////////////////////////////////////////////////
@@ -186,8 +189,12 @@ void FPacketViewDrawHelper::DrawCached(const FNetworkPacketSeries& Series) const
 		const float H = ValueY - BaselineY;
 		const float Y = ViewHeight - H;
 
-		const FLinearColor ColorFill = GetColorByStatus(Sample.AggregatedStatus);
-		const FLinearColor ColorBorder(ColorFill.R * 0.75f, ColorFill.G * 0.75f, ColorFill.B * 0.75f, 1.0);
+		FLinearColor ColorFill = GetColorByStatus(Sample.AggregatedStatus);
+		if (!Sample.bAtLeastOnePacketMatchesFilter)
+		{
+			ColorFill.A = 0.1f;
+		}
+		const FLinearColor ColorBorder(ColorFill.R * 0.75f, ColorFill.G * 0.75f, ColorFill.B * 0.75f, ColorFill.A);
 
 		if (SampleW > 2.0f)
 		{
diff --git a/Engine/Source/Developer/TraceInsights/Private/Insights/NetworkingProfiler/ViewModels/PacketViewDrawHelper.h b/Engine/Source/Developer/TraceInsights/Private/Insights/NetworkingProfiler/ViewModels/PacketViewDrawHelper.h
--- a/Engine/Source/Developer/TraceInsights/Private/Insights/NetworkingProfiler/ViewModels/PacketViewDrawHelper.h
+++ b/Engine/Source/Developer/TraceInsights/Private/Insights/NetworkingProfiler/ViewModels/PacketViewDrawHelper.h
@@ -79,12 +79,15 @@ struct FNetworkPacketAggregatedSample
 
 	FNetworkPacket LargestPacket;
 
+	bool bAtLeastOnePacketMatchesFilter;
+
 	FNetworkPacketAggregatedSample()
 		: NumPackets(0)
 		, StartTime(DBL_MAX)
 		, EndTime(-DBL_MAX)
 		, AggregatedStatus(Trace::ENetProfilerDeliveryStatus::Unknown)
 		, LargestPacket()
+		, bAtLeastOnePacketMatchesFilter(true)
 	{}
 
 	FNetworkPacketAggregatedSample(const FNetworkPacketAggregatedSample&) = default;
@@ -144,7 +147,7 @@ public:
 	FNetworkPacketSeriesBuilder(const FNetworkPacketSeriesBuilder&) = delete;
 	FNetworkPacketSeriesBuilder& operator=(const FNetworkPacketSeriesBuilder&) = delete;
 
-	void AddPacket(int32 PacketIndex, const Trace::FNetProfilerPacket& Packet);
+	FNetworkPacketAggregatedSample* AddPacket(int32 PacketIndex, const Trace::FNetProfilerPacket& Packet);
 
 	int32 GetNumAddedPackets() const { return NumAddedPackets; }
 
diff --git a/Engine/Source/Developer/TraceInsights/Private/Insights/NetworkingProfiler/Widgets/SNetStatsView.h b/Engine/Source/Developer/TraceInsights/Private/Insights/NetworkingProfiler/Widgets/SNetStatsView.h
--- a/Engine/Source/Developer/TraceInsights/Private/Insights/NetworkingProfiler/Widgets/SNetStatsView.h
+++ b/Engine/Source/Developer/TraceInsights/Private/Insights/NetworkingProfiler/Widgets/SNetStatsView.h
@@ -90,8 +90,6 @@ public:
 
 	void SelectNetEventNode(uint64 Id);
 
-	//const TSet<FNetEventNodePtr>& GetNetEventNodes() const { return NetEventNodes; }
-	//const TMap<uint64, FNetEventNodePtr> GetNetEventNodesIdMap() const { return NetEventNodesIdMap; }
 	const FNetEventNodePtr* GetNetEventNode(uint64 Id) const { return NetEventNodesIdMap.Find(Id); }
 
 protected:
diff --git a/Engine/Source/Developer/TraceInsights/Private/Insights/NetworkingProfiler/Widgets/SNetworkingProfilerWindow.cpp b/Engine/Source/Developer/TraceInsights/Private/Insights/NetworkingProfiler/Widgets/SNetworkingProfilerWindow.cpp
--- a/Engine/Source/Developer/TraceInsights/Private/Insights/NetworkingProfiler/Widgets/SNetworkingProfilerWindow.cpp
+++ b/Engine/Source/Developer/TraceInsights/Private/Insights/NetworkingProfiler/Widgets/SNetworkingProfilerWindow.cpp
@@ -953,17 +953,18 @@ void SNetworkingProfilerWindow::SetSelectedBitRange(uint32 StartPos, uint32 EndP
 
 ////////////////////////////////////////////////////////////////////////////////////////////////////
 
-void SNetworkingProfilerWindow::SetSelectedEventType(const uint64 InEventTypeId)
+void SNetworkingProfilerWindow::SetSelectedEventTypeIndex(const uint32 InEventTypeIndex)
 {
-	if (InEventTypeId != SelectedEventTypeId)
+	if (InEventTypeIndex != SelectedEventTypeIndex)
 	{
-		SelectedEventTypeId = InEventTypeId;
+		SelectedEventTypeIndex = InEventTypeIndex;
 
-		if (SelectedEventTypeId != InvalidEventTypeId)
+		if (SelectedEventTypeIndex != InvalidEventTypeIndex)
 		{
 			if (NetStatsView)
 			{
-				NetStatsView->SelectNetEventNode(SelectedEventTypeId);
+				const uint64 NodeId = static_cast<uint64>(SelectedEventTypeIndex);
+				NetStatsView->SelectNetEventNode(NodeId);
 			}
 		}
 	}
diff --git a/Engine/Source/Developer/TraceInsights/Private/Insights/NetworkingProfiler/Widgets/SNetworkingProfilerWindow.h b/Engine/Source/Developer/TraceInsights/Private/Insights/NetworkingProfiler/Widgets/SNetworkingProfilerWindow.h
--- a/Engine/Source/Developer/TraceInsights/Private/Insights/NetworkingProfiler/Widgets/SNetworkingProfilerWindow.h
+++ b/Engine/Source/Developer/TraceInsights/Private/Insights/NetworkingProfiler/Widgets/SNetworkingProfilerWindow.h
@@ -134,7 +134,7 @@ public:
 
 	void SetSelectedPacket(uint32 StartIndex, uint32 EndIndex, uint32 SinglePacketBitSize = 0);
 	void SetSelectedBitRange(uint32 StartPos, uint32 EndPos);
-	void SetSelectedEventType(uint64 EventTypeId);
+	void SetSelectedEventTypeIndex(uint32 InEventTypeIndex);
 
 private:
 	TSharedRef<SDockTab> SpawnTab_Toolbar(const FSpawnTabArgs& Args);
@@ -300,11 +300,12 @@ private:
 
 	uint32 SelectedPacketStartIndex;
 	uint32 SelectedPacketEndIndex;
+
 	uint32 SelectionStartPosition;
 	uint32 SelectionEndPosition;
 
-	static const uint64 InvalidEventTypeId = uint64(-1);
-	uint64 SelectedEventTypeId;
+	static const uint32 InvalidEventTypeIndex = uint32(-1);
+	uint32 SelectedEventTypeIndex;
 
 	/** The handle to the active update duration tick */
 	TWeakPtr<FActiveTimerHandle> ActiveTimerHandle;
diff --git a/Engine/Source/Developer/TraceInsights/Private/Insights/NetworkingProfiler/Widgets/SPacketContentView.cpp b/Engine/Source/Developer/TraceInsights/Private/Insights/NetworkingProfiler/Widgets/SPacketContentView.cpp
--- a/Engine/Source/Developer/TraceInsights/Private/Insights/NetworkingProfiler/Widgets/SPacketContentView.cpp
+++ b/Engine/Source/Developer/TraceInsights/Private/Insights/NetworkingProfiler/Widgets/SPacketContentView.cpp
@@ -511,6 +511,12 @@ void SPacketContentView::FilterByNetId_OnCheckStateChanged(ECheckBoxState NewSta
 {
 	bFilterByNetId = (NewState == ECheckBoxState::Checked);
 	bIsStateDirty = true;
+
+	TSharedPtr<SPacketView> PacketView = ProfilerWindow ? ProfilerWindow->GetPacketView() : nullptr;
+	if (PacketView.IsValid())
+	{
+		PacketView->InvalidateState();
+	}
 }
 
 ////////////////////////////////////////////////////////////////////////////////////////////////////
@@ -545,6 +551,12 @@ void SPacketContentView::FilterByEventType_OnCheckStateChanged(ECheckBoxState Ne
 {
 	bFilterByEventType = (NewState == ECheckBoxState::Checked);
 	bIsStateDirty = true;
+
+	TSharedPtr<SPacketView> PacketView = ProfilerWindow ? ProfilerWindow->GetPacketView() : nullptr;
+	if (PacketView.IsValid())
+	{
+		PacketView->InvalidateState();
+	}
 }
 
 ////////////////////////////////////////////////////////////////////////////////////////////////////
@@ -656,6 +668,12 @@ void SPacketContentView::SetFilterNetId(const uint32 InNetId)
 	if (bFilterByNetId)
 	{
 		bIsStateDirty = true;
+
+		TSharedPtr<SPacketView> PacketView = ProfilerWindow ? ProfilerWindow->GetPacketView() : nullptr;
+		if (PacketView.IsValid())
+		{
+			PacketView->InvalidateState();
+		}
 	}
 }
 
@@ -669,6 +687,12 @@ void SPacketContentView::SetFilterEventType(const uint32 InEventTypeIndex, const
 	if (bFilterByEventType)
 	{
 		bIsStateDirty = true;
+
+		TSharedPtr<SPacketView> PacketView = ProfilerWindow ? ProfilerWindow->GetPacketView() : nullptr;
+		if (PacketView.IsValid())
+		{
+			PacketView->InvalidateState();
+		}
 	}
 }
 
@@ -852,8 +876,7 @@ void SPacketContentView::OnSelectedEventChanged()
 	if (SelectedEvent.IsValid() && ProfilerWindow.IsValid())
 	{
 		// Select the node coresponding to net event type of selected net event instance.
-		const uint64 EventTypeId = static_cast<uint64>(SelectedEvent.Event.EventTypeIndex);
-		ProfilerWindow->SetSelectedEventType(EventTypeId);
+		ProfilerWindow->SetSelectedEventTypeIndex(SelectedEvent.Event.EventTypeIndex);
 	}
 }
 
diff --git a/Engine/Source/Developer/TraceInsights/Private/Insights/NetworkingProfiler/Widgets/SPacketView.cpp b/Engine/Source/Developer/TraceInsights/Private/Insights/NetworkingProfiler/Widgets/SPacketView.cpp
--- a/Engine/Source/Developer/TraceInsights/Private/Insights/NetworkingProfiler/Widgets/SPacketView.cpp
+++ b/Engine/Source/Developer/TraceInsights/Private/Insights/NetworkingProfiler/Widgets/SPacketView.cpp
@@ -16,7 +16,6 @@
 #include "Insights/Common/PaintUtils.h"
 #include "Insights/Common/Stopwatch.h"
 #include "Insights/Common/TimeUtils.h"
-#include "Insights/TimingProfilerCommon.h" // for UE_LOG
 #include "Insights/InsightsManager.h"
 #include "Insights/InsightsStyle.h"
 #include "Insights/NetworkingProfiler/NetworkingProfilerManager.h"
@@ -248,6 +247,33 @@ void SPacketView::UpdateState()
 	FStopwatch Stopwatch;
 	Stopwatch.Start();
 
+	//////////////////////////////////////////////////
+
+	struct FPacketFilter
+	{
+		bool bByNetId = false;
+		uint32 NetId = 0;
+		bool bByEventType = false;
+		uint32 EventTypeIndex = 0;
+	};
+
+	FPacketFilter Filter;
+
+	if (ProfilerWindow)
+	{
+		TSharedPtr<SPacketContentView> PacketContentView = ProfilerWindow->GetPacketContentView();
+		if (PacketContentView)
+		{
+			Filter.bByNetId = PacketContentView->IsFilterByNetIdEnabled();
+			Filter.NetId = PacketContentView->GetFilterNetId();
+
+			Filter.bByEventType = PacketContentView->IsFilterByEventTypeEnabled();
+			Filter.EventTypeIndex = PacketContentView->GetFilterEventTypeIndex();
+		}
+	}
+
+	//////////////////////////////////////////////////
+
 	// Reset stats.
 	PacketSeries->NumAggregatedPackets = 0;
 	NumUpdatedPackets = 0;
@@ -280,9 +306,60 @@ void SPacketView::UpdateState()
 				if (PacketStartIndex <= PacketEndIndex)
 				{
 					int32 PacketIndex = PacketStartIndex;
-					NetProfilerProvider.EnumeratePackets(ConnectionIndex, ConnectionMode, PacketStartIndex, PacketEndIndex, [&Builder, &PacketIndex](const Trace::FNetProfilerPacket& Packet)
+					NetProfilerProvider.EnumeratePackets(ConnectionIndex, ConnectionMode, PacketStartIndex, PacketEndIndex, [this, &Builder, &PacketIndex, &Filter, &NetProfilerProvider](const Trace::FNetProfilerPacket& Packet)
 					{
-						Builder.AddPacket(PacketIndex++, Packet);
+						FNetworkPacketAggregatedSample* SamplePtr = Builder.AddPacket(PacketIndex++, Packet);
+
+						if (SamplePtr)
+						{
+							if (SamplePtr->NumPackets == 1)
+							{
+								SamplePtr->bAtLeastOnePacketMatchesFilter = !Filter.bByNetId && !Filter.bByEventType;
+							}
+
+							if (!SamplePtr->bAtLeastOnePacketMatchesFilter && (Filter.bByNetId || Filter.bByEventType))
+							{
+								bool bFilterMatch = false;
+
+								const uint32 StartPos = 0;
+								const uint32 EndPos = Packet.ContentSizeInBits;
+								NetProfilerProvider.EnumeratePacketContentEventsByPosition(ConnectionIndex, ConnectionMode, PacketIndex - 1, StartPos, EndPos, [this, &bFilterMatch, &Filter, &NetProfilerProvider](const Trace::FNetProfilerContentEvent& Event)
+								{
+									if (!bFilterMatch)
+									{
+										bool bEventMatchesFilter = true;
+										if (Filter.bByEventType && Filter.EventTypeIndex != Event.EventTypeIndex)
+										{
+											bEventMatchesFilter = false;
+										}
+										if (bEventMatchesFilter && Filter.bByNetId)
+										{
+											uint32 NetId = uint32(-1);
+											if (Event.ObjectInstanceIndex != 0)
+											{
+												NetProfilerProvider.ReadObject(GameInstanceIndex, Event.ObjectInstanceIndex, [&NetId](const Trace::FNetProfilerObjectInstance& ObjectInstance)
+												{
+													NetId = ObjectInstance.NetId;
+												});
+											}
+											if (Filter.NetId != NetId)
+											{
+												bEventMatchesFilter = false;
+											}
+										}
+										if (bEventMatchesFilter)
+										{
+											bFilterMatch = true;
+										}
+									}
+								});
+
+								if (bFilterMatch)
+								{
+									SamplePtr->bAtLeastOnePacketMatchesFilter = true;
+								}
+							}
+						}
 					});
 				}
 			}
diff --git a/Engine/Source/Developer/TraceInsights/Private/Insights/NetworkingProfiler/Widgets/SPacketView.h b/Engine/Source/Developer/TraceInsights/Private/Insights/NetworkingProfiler/Widgets/SPacketView.h
--- a/Engine/Source/Developer/TraceInsights/Private/Insights/NetworkingProfiler/Widgets/SPacketView.h
+++ b/Engine/Source/Developer/TraceInsights/Private/Insights/NetworkingProfiler/Widgets/SPacketView.h
@@ -141,6 +141,8 @@ public:
 	void ExtendRightSideOfSelectedInterval();
 	void ShrinkRightSideOfSelectedInterval();
 
+	void InvalidateState() { bIsStateDirty = true; }
+
 private:
 	void UpdateSelectedSample();
 
