Insights/Services changes to support latest anim changes

Insights:
- ITimingViewDrawHelper::GetTrackNameTextColor now takes a FBaseTimingTrack to allow it to be used with graph tracks.
- Graph tracks can now be constructed with a proper type and name, as well as a subtype.
- Graph tracks can now specify their fill color seperately from their line color.
- Exposed the various FGraphTrack::CanExecute_* UI functions so derived tracks can prevent those items from being used in context menus.
- Graph tracks now optionally draw their baseline.
- Moved graph track background drawing to PreDraw() to correctly sort.
- Moved AutoZoom functionality from FTimingGraphTrack::UpdateStatsCounterSeries into FGraphTrack, so now all tracks can auto-zoom if required. Fixed some issues arising from tracks being zero-size when they 'scale-in' at first.
- Fixed high-DPI issues with graph tracks
- Fixed the various pixel offsets when drawing graph tracks. We not apply a 0.5 pixel (note not slate unit) offset and disable pixel snapping when drawing lines.
- Added the ability to add disconnected lines via FGraphTrackBuilder.
- Fixed off-by-one issues with GraphTrack drawing.
- Added key color values to graph context menu for easier correlation.
- Exposed TimeUtils so it can be used by plugin tracks.

Services:
- Added GetFrameFromTime to IFrameProvider.
- Changed all references to FrameStartTimes to use TArray64
- Fixed large copy incurred in passing frame times to TCounterData::GetIterator().
- Temporarily added LowerBound64 to Frames.cpp, as Algo doesnt support TArray64 yet.
- For TPointTimeline, added a feature to allow iteration +/- 1 outside of the specified range, useful when iterating the timeline for drawing.

diff --git a/Engine/Source/Developer/TraceInsights/Private/Insights/Common/TimeUtils.cpp b/Engine/Source/Developer/TraceInsights/Private/Insights/Common/TimeUtils.cpp
--- a/Engine/Source/Developer/TraceInsights/Private/Insights/Common/TimeUtils.cpp
+++ b/Engine/Source/Developer/TraceInsights/Private/Insights/Common/TimeUtils.cpp
@@ -1,6 +1,6 @@
 // Copyright 1998-2019 Epic Games, Inc. All Rights Reserved.
 
-#include "TimeUtils.h"
+#include "Insights/Common/TimeUtils.h"
 
 #include <limits>
 //#include <math.h>
diff --git a/Engine/Source/Developer/TraceInsights/Private/Insights/Common/TimeUtils.h b/Engine/Source/Developer/TraceInsights/Private/Insights/Common/TimeUtils.h
deleted file mode 100644
diff --git a/Engine/Source/Developer/TraceInsights/Private/Insights/ViewModels/GraphSeries.cpp b/Engine/Source/Developer/TraceInsights/Private/Insights/ViewModels/GraphSeries.cpp
--- a/Engine/Source/Developer/TraceInsights/Private/Insights/ViewModels/GraphSeries.cpp
+++ b/Engine/Source/Developer/TraceInsights/Private/Insights/ViewModels/GraphSeries.cpp
@@ -16,10 +16,6 @@ FGraphSeries::FGraphSeries()
 	, bIsVisible(true)
 	, bIsDirty(false)
 	, bAutoZoom(false)
-	, TargetAutoZoomLowValue(0.0)
-	, TargetAutoZoomHighValue(1.0)
-	, AutoZoomLowValue(0.0)
-	, AutoZoomHighValue(1.0)
 	, BaselineY(0.0)
 	, ScaleY(1.0)
 	, Color(0.0f, 0.5f, 1.0f, 1.0f)
@@ -105,4 +101,59 @@ FString FGraphSeries::FormatValue(double Value) const
 
 ////////////////////////////////////////////////////////////////////////////////////////////////////
 
+void FGraphSeries::UpdateAutoZoom(const float InTopY, const float InBottomY, const double InMinEventValue, const double InMaxEventValue)
+{
+	if(IsAutoZoomEnabled())
+	{
+		const double LowValue = GetValueForY(InBottomY);
+		const double HighValue = GetValueForY(InTopY);
+
+		double MinValue = InMinEventValue;
+		double MaxValue = InMaxEventValue;
+
+		// If MinValue == MaxValue, we keep the previous baseline and scale, but only if the min/max value is already visible.
+		if (MinValue == MaxValue && (MinValue < LowValue || MaxValue > HighValue))
+		{
+			MinValue = FMath::Min(MinValue, LowValue);
+			MaxValue = FMath::Max(MaxValue, HighValue);
+		}
+
+		if (MinValue < MaxValue)
+		{
+			constexpr bool bIsAutoZoomAnimated = true;
+			if (bIsAutoZoomAnimated)
+			{
+				// Interpolate the min-max interval (animating the vertical position and scale of the graph series).
+				constexpr double InterpolationSpeed = 0.5;
+				const double NewMinValue = InterpolationSpeed * MinValue + (1.0 - InterpolationSpeed) * LowValue;
+				const double NewMaxValue = InterpolationSpeed * MaxValue + (1.0 - InterpolationSpeed) * HighValue;
+
+				// Check if we reach the target min-max interval.
+				const double ErrorTolerance = 0.5 / GetScaleY(); // delta value for dy ~= 0.5 pixels
+				if (!FMath::IsNearlyEqual(NewMinValue, MinValue, ErrorTolerance) ||
+					!FMath::IsNearlyEqual(NewMaxValue, MaxValue, ErrorTolerance))
+				{
+					MinValue = NewMinValue;
+					MaxValue = NewMaxValue;
+
+					// Request a new update so we can further interpolate the min-max interval.
+					SetDirtyFlag();
+				}
+			}
+
+			double NewBaselineY;
+			double NewScaleY;
+			ComputeBaselineAndScale(MinValue, MaxValue, InTopY, InBottomY, NewBaselineY, NewScaleY);
+			SetBaselineY(NewBaselineY);
+			SetScaleY(NewScaleY);
+		}
+		else
+		{
+			// If MinValue == MaxValue, we keep the previous baseline and scale.
+		}
+	}
+}
+
+////////////////////////////////////////////////////////////////////////////////////////////////////
+
 #undef LOCTEXT_NAMESPACE
diff --git a/Engine/Source/Developer/TraceInsights/Private/Insights/ViewModels/GraphTrack.cpp b/Engine/Source/Developer/TraceInsights/Private/Insights/ViewModels/GraphTrack.cpp
--- a/Engine/Source/Developer/TraceInsights/Private/Insights/ViewModels/GraphTrack.cpp
+++ b/Engine/Source/Developer/TraceInsights/Private/Insights/ViewModels/GraphTrack.cpp
@@ -6,6 +6,7 @@
 #include "Framework/MultiBox/MultiBoxBuilder.h"
 #include "Rendering/DrawElements.h"
 #include "Styling/CoreStyle.h"
+#include "EditorFontGlyphs.h"
 
 // Insights
 #include "Insights/Common/PaintUtils.h"
@@ -32,12 +33,36 @@ FGraphTrack::FGraphTrack(const FName& InSubType)
 	, PointBrush(FEditorStyle::GetBrush("Graph.ExecutionBubble"))
 	, BorderBrush(FInsightsStyle::Get().GetBrush("SingleBorder"))
 	, Font(FCoreStyle::GetDefaultFontStyle("Regular", 8))
+	, bDrawDebugInfo(false)
 	, bDrawPoints(true)
 	, bDrawPointsWithBorder(true)
 	, bDrawLines(true)
 	, bDrawPolygon(true)
 	, bUseEventDuration(true)
 	, bDrawBoxes(true)
+	, bDrawBaseline(true)
+	, NumAddedEvents(0)
+	, NumDrawPoints(0)
+	, NumDrawLines(0)
+	, NumDrawBoxes(0)
+{
+}
+
+FGraphTrack::FGraphTrack(const FName& InSubType, const FString& InName)
+	: FBaseTimingTrack(FName(TEXT("Graph")), InSubType, InName)
+	//, AllSeries()
+	, WhiteBrush(FInsightsStyle::Get().GetBrush("WhiteBrush"))
+	, PointBrush(FEditorStyle::GetBrush("Graph.ExecutionBubble"))
+	, BorderBrush(FInsightsStyle::Get().GetBrush("SingleBorder"))
+	, Font(FCoreStyle::GetDefaultFontStyle("Regular", 8))
+	, bDrawDebugInfo(false)
+	, bDrawPoints(true)
+	, bDrawPointsWithBorder(true)
+	, bDrawLines(true)
+	, bDrawPolygon(true)
+	, bUseEventDuration(true)
+	, bDrawBoxes(true)
+	, bDrawBaseline(true)
 	, NumAddedEvents(0)
 	, NumDrawPoints(0)
 	, NumDrawLines(0)
@@ -92,7 +117,10 @@ void FGraphTrack::UpdateStats()
 		if (Series->IsVisible())
 		{
 			NumDrawPoints += Series->Points.Num();
-			NumDrawLines += Series->LinePoints.Num() / 2;
+			for(int32 BatchIndex = 0; BatchIndex < Series->LinePoints.Num(); ++BatchIndex)
+			{
+				NumDrawLines += Series->LinePoints[BatchIndex].Num() / 2;
+			}
 			NumDrawBoxes += Series->Boxes.Num();
 		}
 	}
@@ -100,20 +128,29 @@ void FGraphTrack::UpdateStats()
 
 ////////////////////////////////////////////////////////////////////////////////////////////////////
 
-void FGraphTrack::Draw(const ITimingTrackDrawContext& Context) const
+void FGraphTrack::PreDraw(const ITimingTrackDrawContext& Context) const
 {
 	FDrawContext& DrawContext = Context.GetDrawContext();
 	const FTimingTrackViewport& Viewport = Context.GetViewport();
 
 	FDrawHelpers::DrawBackground(DrawContext, WhiteBrush, Viewport, GetPosY(), GetHeight());
+}
+
+////////////////////////////////////////////////////////////////////////////////////////////////////
+
+void FGraphTrack::Draw(const ITimingTrackDrawContext& Context) const
+{
+	FDrawContext& DrawContext = Context.GetDrawContext();
+	const FTimingTrackViewport& Viewport = Context.GetViewport();
 
 	// Set clipping area.
 	{
 		const FVector2D AbsPos = DrawContext.Geometry.GetAbsolutePosition();
+		const float Scale = DrawContext.Geometry.GetAccumulatedLayoutTransform().GetScale();
 		const float L = AbsPos.X;
-		const float R = AbsPos.X + Viewport.GetWidth();
-		const float T = AbsPos.Y + GetPosY();
-		const float B = AbsPos.Y + GetPosY() + GetHeight();
+		const float R = AbsPos.X + (Viewport.GetWidth() * Scale);
+		const float T = AbsPos.Y + (GetPosY() * Scale);
+		const float B = AbsPos.Y + ((GetPosY() + GetHeight()) * Scale);
 		const FSlateClippingZone ClipZone(FVector2D(L, T), FVector2D(R, T), FVector2D(L, B), FVector2D(R, B));
 		DrawContext.ElementList.PushClip(ClipZone);
 	}
@@ -121,22 +158,24 @@ void FGraphTrack::Draw(const ITimingTrackDrawContext& Context) const
 	// Draw top line.
 	//DrawContext.DrawBox(0.0f, GetPosY(), Viewport.Width, 1.0f, WhiteBrush, FLinearColor(0.05f, 0.05f, 0.05f, 1.0f));
 
-	// Draw baseline (Value == 0), for the first series only.
-	if (AllSeries.Num() > 0)
-	{
-		const TSharedPtr<FGraphSeries>& Series = AllSeries[0];
-		const float BaselineY = FMath::RoundToFloat(static_cast<float>(Series->GetBaselineY()));
-		if (BaselineY >= 0.0f && BaselineY < GetHeight())
-		{
-			DrawContext.DrawBox(0.0f, GetPosY() + BaselineY, Viewport.GetWidth(), 1.0f, WhiteBrush, FLinearColor(0.05f, 0.05f, 0.05f, 1.0f));
-			DrawContext.LayerId++;
-		}
-	}
-
+	bool bDrawnBaseline = false;
 	for (const TSharedPtr<FGraphSeries>& Series : AllSeries)
 	{
 		if (Series->IsVisible())
 		{
+			// Draw baseline (Value == 0), for the first visible series only.
+			if(bDrawBaseline && !bDrawnBaseline)
+			{
+				const float BaselineY = FMath::RoundToFloat(static_cast<float>(Series->GetBaselineY()));
+				if (BaselineY >= 0.0f && BaselineY < GetHeight())
+				{
+					DrawContext.DrawBox(0.0f, GetPosY() + BaselineY, Viewport.GetWidth(), 1.0f, WhiteBrush, FLinearColor(0.05f, 0.05f, 0.05f, 1.0f));
+					DrawContext.LayerId++;
+				}
+
+				bDrawnBaseline = true;
+			}
+
 			DrawSeries(*Series, DrawContext, Viewport);
 		}
 	}
@@ -149,6 +188,22 @@ void FGraphTrack::Draw(const ITimingTrackDrawContext& Context) const
 
 void FGraphTrack::DrawSeries(const FGraphSeries& Series, FDrawContext& DrawContext, const FTimingTrackViewport& Viewport) const
 {
+	// Set clipping area. This area tries to take into account the optional border area that allows
+	// graph tracks to optionally act like 'event' tracks if required (with respect to layout, anyway).
+	// The GetBorderY() - 1.0f calculation is desgiend to avoid clipping the line rasterization, as the custom verts
+	// of the fill and the outer lines appear to get rasterized differently, the latter missing one pixel 
+	// on its upper side.
+	{
+		const FVector2D AbsPos = DrawContext.Geometry.GetAbsolutePosition();
+		const float Scale = DrawContext.Geometry.GetAccumulatedLayoutTransform().GetScale();
+		const float L = AbsPos.X;
+		const float R = AbsPos.X + (Viewport.GetWidth() * Scale);
+		const float T = AbsPos.Y + ((GetPosY() + (GetBorderY() - 1.0f)) * Scale);
+		const float B = AbsPos.Y + ((GetPosY() + (GetHeight() - (GetBorderY() - 1.0f))) * Scale);
+		const FSlateClippingZone ClipZone(FVector2D(L, T), FVector2D(R, T), FVector2D(L, B), FVector2D(R, B));
+		DrawContext.ElementList.PushClip(ClipZone);
+	}
+
 	if (bDrawBoxes)
 	{
 		int32 NumBoxes = Series.Boxes.Num();
@@ -169,6 +224,11 @@ void FGraphTrack::DrawSeries(const FGraphSeries& Series, FDrawContext& DrawConte
 		DrawContext.LayerId++;
 	}
 
+	const float LocalPosY = FMath::RoundToFloat(GetPosY());
+
+	FPaintGeometry Geo = DrawContext.Geometry.ToPaintGeometry();
+	Geo.AppendTransform(FSlateLayoutTransform(FVector2D(0.0f, LocalPosY * DrawContext.Geometry.GetAccumulatedLayoutTransform().GetScale())));
+
 	if (bDrawPolygon)
 	{
 		FSlateResourceHandle ResourceHandle = FSlateApplication::Get().GetRenderer()->GetResourceHandle(*WhiteBrush);
@@ -177,119 +237,131 @@ void FGraphTrack::DrawSeries(const FGraphSeries& Series, FDrawContext& DrawConte
 		FVector2D AtlasOffset = ResourceProxy ? ResourceProxy->StartUV : FVector2D(0.f, 0.f);
 		FVector2D AtlasUVSize = ResourceProxy ? ResourceProxy->SizeUV : FVector2D(1.f, 1.f);
 
-		const FVector2D Pos = DrawContext.Geometry.GetAbsolutePosition();
 		const FVector2D Size = DrawContext.Geometry.GetLocalSize();
-		const float Scale = DrawContext.Geometry.Scale;
 
-		FSlateRenderTransform RenderTransform;
-		FColor FillColor(Series.Color.R * 255, Series.Color.G * 255, Series.Color.B * 255, 42);
+		const FSlateRenderTransform& RenderTransform = Geo.GetAccumulatedRenderTransform();
+		FColor FillColor(Series.FillColor.R * 255, Series.FillColor.G * 255, Series.FillColor.B * 255, 42);
 
-		TArray<SlateIndex> Indices;
-		TArray<FSlateVertex> Verts;
-
-		Indices.Reserve(Series.LinePoints.Num() * 6);
-		Verts.Reserve(Series.LinePoints.Num() * 2);
+		for(int32 BatchIndex = 0; BatchIndex < Series.LinePoints.Num(); ++BatchIndex)
+		{
+			TArray<SlateIndex> Indices;
+			TArray<FSlateVertex> Verts;
 
-		const float TopV = GetPosY() / Size.Y;
-		const float BottomV = (GetPosY() + GetHeight()) / Size.Y;
-		const float BaselineY = static_cast<float>(Series.GetBaselineY());
-		const float BaselineV = FMath::Clamp<float>((GetPosY() + BaselineY) / Size.Y, TopV, BottomV);
+			Indices.Reserve(Series.LinePoints[BatchIndex].Num() * 6);
+			Verts.Reserve(Series.LinePoints[BatchIndex].Num() * 2);
 
-		int32 PrevSide = 0;
+			const float TopV = 0.0f;
+			const float BottomV = GetHeight() / Size.Y;
+			const float BaselineY = static_cast<float>(Series.GetBaselineY());
+			const float BaselineV = FMath::Clamp<float>(BaselineY / Size.Y, TopV, BottomV);
 
-		for (int32 PointIndex = 0; PointIndex < Series.LinePoints.Num(); ++PointIndex)
-		{
-			const FVector2D& LinePoint = Series.LinePoints[PointIndex];
+			int32 PrevSide = 0;
 
-			// When crossing baseline the polygon needs to be intersected.
-			int32 CrtSide = (LinePoint.Y > BaselineY) ? 1 : (LinePoint.Y < BaselineY) ? -1 : 0;
-			if (PrevSide != 0 && CrtSide + PrevSide == 0) // alternating sides?
-			{
-				// Compute intersection point.
-				const FVector2D& PrevLinePoint = Series.LinePoints[PointIndex - 1];
-				float X = PrevLinePoint.X + (LinePoint.X - PrevLinePoint.X) / ((BaselineY - LinePoint.Y) / (PrevLinePoint.Y - BaselineY) + 1.0f);
-
-				// Add an interesection point vertex.
-				FVector2D UV(X / Size.X, BaselineV);
-				Verts.Add(FSlateVertex::Make<ESlateVertexRounding::Disabled>(RenderTransform, Pos + UV * Size * Scale, AtlasOffset + UV * AtlasUVSize, FillColor));
-
-				// Add a value point vertex.
-				UV.X = LinePoint.X / Size.X;
-				UV.Y = TopV + LinePoint.Y / Size.Y;
-				Verts.Add(FSlateVertex::Make<ESlateVertexRounding::Disabled>(RenderTransform, Pos + UV * Size * Scale, AtlasOffset + UV * AtlasUVSize, FillColor));
-
-				// Add a baseline vertex.
-				UV.Y = BaselineV;
-				Verts.Add(FSlateVertex::Make<ESlateVertexRounding::Disabled>(RenderTransform, Pos + UV * Size * Scale, AtlasOffset + FVector2D(UV.X, 0.5f) * AtlasUVSize, FillColor));
-
-				int32 Index0 = Verts.Num() - 5;
-				int32 Index1 = Verts.Num() - 4;
-				int32 Index2 = Verts.Num() - 3;
-				int32 Index3 = Verts.Num() - 2;
-				int32 Index4 = Verts.Num() - 1;
-
-				Indices.Add(Index0);
-				Indices.Add(Index1);
-				Indices.Add(Index2);
-
-				Indices.Add(Index2);
-				Indices.Add(Index3);
-				Indices.Add(Index4);
-			}
-			else
+			for (int32 PointIndex = 0; PointIndex < Series.LinePoints[BatchIndex].Num(); ++PointIndex)
 			{
-				// Add a value point vertex.
-				FVector2D UV(LinePoint.X / Size.X, TopV + LinePoint.Y / Size.Y);
-				Verts.Add(FSlateVertex::Make<ESlateVertexRounding::Disabled>(RenderTransform, Pos + UV * Size * Scale, AtlasOffset + UV * AtlasUVSize, FillColor));
+				const FVector2D& LinePoint = Series.LinePoints[BatchIndex][PointIndex];
 
-				// Add a baseline vertex.
-				UV.Y = BaselineV;
-				Verts.Add(FSlateVertex::Make<ESlateVertexRounding::Disabled>(RenderTransform, Pos + UV * Size * Scale, AtlasOffset + FVector2D(UV.X, 0.5f) * AtlasUVSize, FillColor));
-
-				if (Verts.Num() >= 4)
+				// When crossing baseline the polygon needs to be intersected.
+				int32 CrtSide = (LinePoint.Y > BaselineY) ? 1 : (LinePoint.Y < BaselineY) ? -1 : 0;
+				if (PrevSide != 0 && CrtSide + PrevSide == 0) // alternating sides?
 				{
-					int32 Index0 = Verts.Num() - 4;
-					int32 Index1 = Verts.Num() - 3;
-					int32 Index2 = Verts.Num() - 2;
-					int32 Index3 = Verts.Num() - 1;
+					// Compute intersection point.
+					const FVector2D& PrevLinePoint = Series.LinePoints[BatchIndex][PointIndex - 1];
+					float X = PrevLinePoint.X + (LinePoint.X - PrevLinePoint.X) / ((BaselineY - LinePoint.Y) / (PrevLinePoint.Y - BaselineY) + 1.0f);
+
+					// Add an intersection point vertex.
+					FVector2D UV(X / Size.X, BaselineV);
+					Verts.Add(FSlateVertex::Make<ESlateVertexRounding::Disabled>(RenderTransform, UV * Size, AtlasOffset + UV * AtlasUVSize, FillColor));
+
+					// Add a value point vertex.
+					UV.X = LinePoint.X / Size.X;
+					UV.Y = TopV + LinePoint.Y / Size.Y;
+					Verts.Add(FSlateVertex::Make<ESlateVertexRounding::Disabled>(RenderTransform, UV * Size, AtlasOffset + UV * AtlasUVSize, FillColor));
+
+					// Add a baseline vertex.
+					UV.Y = BaselineV;
+					Verts.Add(FSlateVertex::Make<ESlateVertexRounding::Disabled>(RenderTransform, UV * Size, AtlasOffset + FVector2D(UV.X, 0.5f) * AtlasUVSize, FillColor));
+
+					int32 Index0 = Verts.Num() - 5;
+					int32 Index1 = Verts.Num() - 4;
+					int32 Index2 = Verts.Num() - 3;
+					int32 Index3 = Verts.Num() - 2;
+					int32 Index4 = Verts.Num() - 1;
 
 					Indices.Add(Index0);
 					Indices.Add(Index1);
 					Indices.Add(Index2);
 
 					Indices.Add(Index2);
-					Indices.Add(Index1);
 					Indices.Add(Index3);
+					Indices.Add(Index4);
+				}
+				else
+				{
+					// Add a value point vertex.
+					FVector2D UV(LinePoint.X / Size.X, TopV + LinePoint.Y / Size.Y);
+					Verts.Add(FSlateVertex::Make<ESlateVertexRounding::Disabled>(RenderTransform, UV * Size, AtlasOffset + UV * AtlasUVSize, FillColor));
+
+					// Add a baseline vertex.
+					UV.Y = BaselineV;
+					Verts.Add(FSlateVertex::Make<ESlateVertexRounding::Disabled>(RenderTransform, UV * Size, AtlasOffset + FVector2D(UV.X, 0.5f) * AtlasUVSize, FillColor));
+
+					if (Verts.Num() >= 4)
+					{
+						int32 Index0 = Verts.Num() - 4;
+						int32 Index1 = Verts.Num() - 3;
+						int32 Index2 = Verts.Num() - 2;
+						int32 Index3 = Verts.Num() - 1;
+
+						Indices.Add(Index0);
+						Indices.Add(Index1);
+						Indices.Add(Index2);
+
+						Indices.Add(Index2);
+						Indices.Add(Index1);
+						Indices.Add(Index3);
+					}
 				}
+				PrevSide = CrtSide;
 			}
-			PrevSide = CrtSide;
-		}
 
-		if (Indices.Num() > 0)
-		{
-			FSlateDrawElement::MakeCustomVerts(
-				DrawContext.ElementList,
-				DrawContext.LayerId,
-				ResourceHandle,
-				Verts,
-				Indices,
-				nullptr,
-				0,
-				0, ESlateDrawEffect::PreMultipliedAlpha);
-			DrawContext.LayerId++;
+			if (Indices.Num() > 0)
+			{
+				FSlateDrawElement::MakeCustomVerts(
+					DrawContext.ElementList,
+					DrawContext.LayerId,
+					ResourceHandle,
+					Verts,
+					Indices,
+					nullptr,
+					0,
+					0, ESlateDrawEffect::PreMultipliedAlpha);
+				DrawContext.LayerId++;
+			}
 		}
 	}
 
-	const float LocalPosY = FMath::RoundToFloat(GetPosY());
-
 	if (bDrawLines)
 	{
-		FPaintGeometry Geo = DrawContext.Geometry.ToPaintGeometry();
-		Geo.AppendTransform(FSlateLayoutTransform(FVector2D(1.0f, LocalPosY + 1.0f)));
-		FSlateDrawElement::MakeLines(DrawContext.ElementList, DrawContext.LayerId, Geo, Series.LinePoints, DrawContext.DrawEffects, Series.Color, false, 1.0f);
+		// Find scale to get the size of a pixel
+		const float InvScale = 1.0f / DrawContext.Geometry.GetAccumulatedLayoutTransform().GetScale();
+		FPaintGeometry LineGeo = Geo;
+		LineGeo.AppendTransform(FSlateLayoutTransform(FVector2D(0.5f * InvScale, 0.5f * InvScale)));
+
+		for(int32 BatchIndex = 0; BatchIndex < Series.LinePoints.Num(); ++BatchIndex)
+		{
+			if(Series.LinePoints[BatchIndex].Num())
+			{
+				// Disable pixel snapping here so lines line up with boxes/polys correctly
+				FSlateDrawElement::MakeLines(DrawContext.ElementList, DrawContext.LayerId, LineGeo, Series.LinePoints[BatchIndex], DrawContext.DrawEffects | ESlateDrawEffect::NoPixelSnapping, Series.Color, false, 1.0f);
+			}
+		}
 		DrawContext.LayerId++;
 	}
 
+	// Restore clipping.
+	DrawContext.ElementList.PopClip();
+
 	if (bDrawPoints)
 	{
 		int32 NumPoints = Series.Points.Num();
@@ -303,8 +375,8 @@ void FGraphTrack::DrawSeries(const FGraphSeries& Series, FDrawContext& DrawConte
 			for (int32 Index = 0; Index < NumPoints; ++Index)
 			{
 				const FVector2D& Pt = Series.Points[Index];
-				const float PtX = Pt.X - PointVisualSize / 2.0f - 0.5f;
-				const float PtY = LocalPosY + Pt.Y - PointVisualSize / 2.0f - 0.5f;
+				const float PtX = (Pt.X - PointVisualSize / 2.0f - 0.5f);
+				const float PtY = (LocalPosY + Pt.Y - PointVisualSize / 2.0f - 0.5f);
 				DrawContext.DrawBox(PtX, PtY, PointVisualSize + 2.0f, PointVisualSize + 2.0f, PointBrush, Series.BorderColor);
 			}
 			DrawContext.LayerId++;
@@ -314,8 +386,8 @@ void FGraphTrack::DrawSeries(const FGraphSeries& Series, FDrawContext& DrawConte
 		for (int32 Index = 0; Index < NumPoints; ++Index)
 		{
 			const FVector2D& Pt = Series.Points[Index];
-			const float PtX = Pt.X - PointVisualSize / 2.0f + 0.5f;
-			const float PtY = LocalPosY + Pt.Y - PointVisualSize / 2.0f + 0.5f;
+			const float PtX = (Pt.X - PointVisualSize / 2.0f + 0.5f);
+			const float PtY = (LocalPosY + Pt.Y - PointVisualSize / 2.0f + 0.5f);
 			DrawContext.DrawBox(PtX, PtY, PointVisualSize, PointVisualSize, PointBrush, Series.Color);
 		}
 		DrawContext.LayerId++;
@@ -354,20 +426,58 @@ void FGraphTrack::DrawSeries(const FGraphSeries& Series, FDrawContext& DrawConte
 		DrawContext.LayerId++;
 
 #endif
+	}
 
-		if (false) // for debugging only
+	if (bDrawDebugInfo) // for debugging only
+	{
+		// Find scale to get the size of a pixel
+		const float PixelUnit = 1.0f / DrawContext.Geometry.GetAccumulatedLayoutTransform().GetScale();
+		FPaintGeometry LineGeo = Geo;
+		LineGeo.AppendTransform(FSlateLayoutTransform(FVector2D(0.5f * PixelUnit, 0.5f * PixelUnit)));
+
+		// Draw white corner at (0, 0) using MakeLines.
 		{
-			// Draw points (cross).
-			for (int32 Index = 0; Index < NumPoints; ++Index)
-			{
-				const FVector2D& Pt = Series.Points[Index];
-				const float PtX = Pt.X;
-				const float PtY = LocalPosY + Pt.Y;
-				DrawContext.DrawBox(PtX - 5.0f, PtY, 11.0f, 1.0f, WhiteBrush, Series.Color);
-				DrawContext.DrawBox(PtX, PtY - 5.0f, 1.0f, 11.0f, WhiteBrush, Series.Color);
-			}
-			DrawContext.LayerId++;
+			TArray<FVector2D> HLine;
+			HLine.Add(FVector2D(0.0f, 0.0f));
+			HLine.Add(FVector2D(10.0f, 0.0f));
+			FSlateDrawElement::MakeLines(DrawContext.ElementList, DrawContext.LayerId, LineGeo, HLine, DrawContext.DrawEffects | ESlateDrawEffect::NoPixelSnapping, FLinearColor::White, false, 1.0f);
+			TArray<FVector2D> VLine;
+			VLine.Add(FVector2D(0.0f, 0.0f));
+			VLine.Add(FVector2D(0.0f, 10.0f));
+			FSlateDrawElement::MakeLines(DrawContext.ElementList, DrawContext.LayerId, LineGeo, VLine, DrawContext.DrawEffects | ESlateDrawEffect::NoPixelSnapping, FLinearColor::White, false, 1.0f);
+		}
+		DrawContext.LayerId++;
+		// Draw corner at (0, 0) using DrawBox.
+		DrawContext.DrawBox(2.0f, LocalPosY, 6.0f, PixelUnit, WhiteBrush, Series.Color);
+		DrawContext.DrawBox(0.0f, LocalPosY + 2.0f, PixelUnit, 6.0f, WhiteBrush, Series.Color);
+		DrawContext.LayerId++;
+		const int32 NumDbgPoints = Series.Points.Num();
+		// Draw white cross lines (15x15) at each point.
+		for (int32 Index = 0; Index < NumDbgPoints; ++Index)
+		{
+			const FVector2D& Pt = Series.Points[Index];
+			const float PtX = Pt.X;
+			const float PtY = Pt.Y;
+			TArray<FVector2D> HLine;
+			HLine.Add(FVector2D(PtX - 7.0f, PtY));
+			HLine.Add(FVector2D(PtX + 8.0f, PtY));
+			FSlateDrawElement::MakeLines(DrawContext.ElementList, DrawContext.LayerId, LineGeo, HLine, DrawContext.DrawEffects | ESlateDrawEffect::NoPixelSnapping, FLinearColor::White, false, 1.0f);
+			TArray<FVector2D> VLine;
+			VLine.Add(FVector2D(PtX, PtY - 7.0f));
+			VLine.Add(FVector2D(PtX, PtY + 8.0f));
+			FSlateDrawElement::MakeLines(DrawContext.ElementList, DrawContext.LayerId, LineGeo, VLine, DrawContext.DrawEffects | ESlateDrawEffect::NoPixelSnapping, FLinearColor::White, false, 1.0f);
+		}
+		DrawContext.LayerId++;
+		// Draw cross lines (11x11) at each point.
+		for (int32 Index = 0; Index < NumDbgPoints; ++Index)
+		{
+			const FVector2D& Pt = Series.Points[Index];
+			const float PtX = Pt.X;
+			const float PtY = LocalPosY + Pt.Y;
+			DrawContext.DrawBox(PtX - 5.0f, PtY, 11.0f, PixelUnit, WhiteBrush, Series.Color);
+			DrawContext.DrawBox(PtX, PtY - 5.0f, PixelUnit, 11.0f, WhiteBrush, Series.Color);
 		}
+		DrawContext.LayerId++;
 	}
 }
 
@@ -489,6 +599,25 @@ void FGraphTrack::BuildContextMenu(FMenuBuilder& MenuBuilder)
 {
 	MenuBuilder.BeginSection("Misc");
 	{
+		if (true)
+		{
+			FUIAction Action_ShowDebugInfo
+			(
+				FExecuteAction::CreateSP(this, &FGraphTrack::ContextMenu_ShowDebugInfo_Execute),
+				FCanExecuteAction::CreateSP(this, &FGraphTrack::ContextMenu_ShowDebugInfo_CanExecute),
+				FIsActionChecked::CreateSP(this, &FGraphTrack::ContextMenu_ShowDebugInfo_IsChecked)
+			);
+			MenuBuilder.AddMenuEntry
+			(
+				LOCTEXT("ContextMenu_ShowDebugInfo", "Show Debug Info"),
+				LOCTEXT("ContextMenu_ShowDebugInfo_Desc", "Show Debug Info."),
+				FSlateIcon(),
+				Action_ShowDebugInfo,
+				NAME_None,
+				EUserInterfaceActionType::ToggleButton
+			);
+		}
+
 		FUIAction Action_ShowPoints
 		(
 			FExecuteAction::CreateSP(this, &FGraphTrack::ContextMenu_ShowPoints_Execute),
@@ -599,11 +728,29 @@ void FGraphTrack::BuildContextMenu(FMenuBuilder& MenuBuilder)
 			);
 			MenuBuilder.AddMenuEntry
 			(
-				Series->GetName(),
-				Series->GetDescription(),
-				FSlateIcon(),
 				Action_ShowSeries,
+				SNew(SHorizontalBox)
+				+SHorizontalBox::Slot()
+				.AutoWidth()
+				.VAlign(VAlign_Center)
+				.Padding(2.0f, 0.0f)
+				[
+					SNew(STextBlock)
+					.Text(Series->GetName())
+				]
+				+SHorizontalBox::Slot()
+				.FillWidth(1.0f)
+				.HAlign(HAlign_Right)
+				.VAlign(VAlign_Center)
+				.Padding(4.0f, 0.0f)
+				[
+					SNew(STextBlock)
+					.Font(FEditorStyle::Get().GetFontStyle("FontAwesome.9"))
+					.ColorAndOpacity_Lambda([SeriesPtr = Series.Get()](){ return FSlateColor(SeriesPtr->GetColor()); })
+					.Text(FEditorFontGlyphs::Circle)
+				],
 				NAME_None,
+				Series->GetDescription(),
 				EUserInterfaceActionType::ToggleButton
 			);
 		}
@@ -613,6 +760,28 @@ void FGraphTrack::BuildContextMenu(FMenuBuilder& MenuBuilder)
 
 ////////////////////////////////////////////////////////////////////////////////////////////////////
 
+void FGraphTrack::ContextMenu_ShowDebugInfo_Execute()
+{
+	bDrawDebugInfo = !bDrawDebugInfo;
+	SetDirtyFlag();
+}
+
+////////////////////////////////////////////////////////////////////////////////////////////////////
+
+bool FGraphTrack::ContextMenu_ShowDebugInfo_CanExecute()
+{
+	return true;
+}
+
+////////////////////////////////////////////////////////////////////////////////////////////////////
+
+bool FGraphTrack::ContextMenu_ShowDebugInfo_IsChecked()
+{
+	return bDrawDebugInfo;
+}
+
+////////////////////////////////////////////////////////////////////////////////////////////////////
+
 void FGraphTrack::ContextMenu_ShowPoints_Execute()
 {
 	bDrawPoints = !bDrawPoints;
diff --git a/Engine/Source/Developer/TraceInsights/Private/Insights/ViewModels/GraphTrackBuilder.cpp b/Engine/Source/Developer/TraceInsights/Private/Insights/ViewModels/GraphTrackBuilder.cpp
--- a/Engine/Source/Developer/TraceInsights/Private/Insights/ViewModels/GraphTrackBuilder.cpp
+++ b/Engine/Source/Developer/TraceInsights/Private/Insights/ViewModels/GraphTrackBuilder.cpp
@@ -34,7 +34,7 @@ FGraphTrackBuilder::~FGraphTrackBuilder()
 
 ////////////////////////////////////////////////////////////////////////////////////////////////////
 
-void FGraphTrackBuilder::AddEvent(double Time, double Duration, double Value)
+void FGraphTrackBuilder::AddEvent(double Time, double Duration, double Value, bool bConnected)
 {
 	Track.NumAddedEvents++;
 
@@ -47,11 +47,11 @@ void FGraphTrackBuilder::AddEvent(double Time, double Duration, double Value)
 
 	if (Track.bDrawLines || Track.bDrawPolygon)
 	{
-		AddConnectedLine(Time, Value);
+		AddConnectedLine(Time, Value, !bConnected);
 
 		if (Track.bUseEventDuration && Duration != 0.0)
 		{
-			AddConnectedLine(Time + Duration, Value);
+			AddConnectedLine(Time + Duration, Value, false);
 		}
 	}
 
@@ -149,6 +149,7 @@ void FGraphTrackBuilder::EndPoints()
 void FGraphTrackBuilder::BeginConnectedLines()
 {
 	Series.LinePoints.Reset();
+	Series.LinePoints.AddDefaulted();
 
 	LinesCurrentX = -FLT_MAX;
 	LinesMinY = FLT_MAX;
@@ -160,7 +161,7 @@ void FGraphTrackBuilder::BeginConnectedLines()
 
 ////////////////////////////////////////////////////////////////////////////////////////////////////
 
-void FGraphTrackBuilder::AddConnectedLine(double Time, double Value)
+void FGraphTrackBuilder::AddConnectedLine(double Time, double Value, bool bNewBatch)
 {
 	if (bIsLastLineAdded)
 	{
@@ -172,6 +173,22 @@ void FGraphTrackBuilder::AddConnectedLine(double Time, double Value)
 
 	//ensure(X >= LinesCurrentX); // we are assuming events are already sorted by Time
 
+	// Start a new batch
+	if(bNewBatch)
+	{
+		if(Series.LinePoints.Num() > 0)
+		{
+			Series.LinePoints.AddDefaulted();
+		}
+
+		// Reset the reduction
+		LinesCurrentX = -FLT_MAX;
+		LinesMinY = FLT_MAX;
+		LinesMaxY = -FLT_MAX;
+		LinesFirstY = FLT_MAX;
+		LinesLastY = FLT_MAX;
+	}
+
 	if (X < 0)
 	{
 		LinesCurrentX = X;
@@ -189,8 +206,8 @@ void FGraphTrackBuilder::AddConnectedLine(double Time, double Value)
 			{
 				FlushConnectedLine();
 
-				Series.LinePoints.Add(FVector2D(LinesCurrentX, LinesLastY));
-				Series.LinePoints.Add(FVector2D(X, Y));
+				Series.LinePoints.Last().Add(FVector2D(LinesCurrentX, LinesLastY));
+				Series.LinePoints.Last().Add(FVector2D(X, Y));
 			}
 
 			// Reset the "reduction line" so last FlushConnectedLine() call will do nothing.
@@ -206,8 +223,8 @@ void FGraphTrackBuilder::AddConnectedLine(double Time, double Value)
 		{
 			FlushConnectedLine();
 
-			Series.LinePoints.Add(FVector2D(LinesCurrentX, LinesLastY));
-			Series.LinePoints.Add(FVector2D(X, Y));
+			Series.LinePoints.Last().Add(FVector2D(LinesCurrentX, LinesLastY));
+			Series.LinePoints.Last().Add(FVector2D(X, Y));
 		}
 
 		// Advance the "reduction line".
@@ -238,8 +255,8 @@ void FGraphTrackBuilder::FlushConnectedLine()
 {
 	if (LinesCurrentX >= 0.0f && LinesMinY != LinesMaxY)
 	{
-		Series.LinePoints.Add(FVector2D(LinesCurrentX, LinesMaxY));
-		Series.LinePoints.Add(FVector2D(LinesCurrentX, LinesMinY));
+		Series.LinePoints.Last().Add(FVector2D(LinesCurrentX, LinesMaxY));
+		Series.LinePoints.Last().Add(FVector2D(LinesCurrentX, LinesMinY));
 	}
 }
 
diff --git a/Engine/Source/Developer/TraceInsights/Private/Insights/ViewModels/TimingGraphTrack.cpp b/Engine/Source/Developer/TraceInsights/Private/Insights/ViewModels/TimingGraphTrack.cpp
--- a/Engine/Source/Developer/TraceInsights/Private/Insights/ViewModels/TimingGraphTrack.cpp
+++ b/Engine/Source/Developer/TraceInsights/Private/Insights/ViewModels/TimingGraphTrack.cpp
@@ -147,8 +147,6 @@ TSharedPtr<FTimingGraphSeries> FTimingGraphTrack::AddStatsCounterSeries(uint32 C
 	Series->SetScaleY(1.0);
 
 	Series->EnableAutoZoom();
-	Series->SetTargetAutoZoomRange(0.0, 1.0);
-	Series->SetAutoZoomRange(0.0, 1.0);
 
 	AllSeries.Add(Series);
 	return Series;
@@ -307,50 +305,7 @@ void FTimingGraphTrack::UpdateStatsCounterSeries(FTimingGraphSeries& Series, con
 
 				const float TopY = 4.0f;
 				const float BottomY = GetHeight() - 4.0f;
-
-				const double LowValue = Series.GetValueForY(BottomY);
-				const double HighValue = Series.GetValueForY(TopY);
-
-				// If MinValue == MaxValue, we keep the previous baseline and scale, but only if the min/max value is already visible.
-				if (MinValue == MaxValue && (MinValue < LowValue || MaxValue > HighValue))
-				{
-					MinValue = FMath::Min(MinValue, LowValue);
-					MaxValue = FMath::Max(MaxValue, HighValue);
-				}
-
-				if (MinValue < MaxValue)
-				{
-					constexpr bool bIsAutoZoomAnimated = true;
-					if (bIsAutoZoomAnimated)
-					{
-						// Interpolate the min-max interval (animating the vertical position and scale of the graph series).
-						constexpr double InterpolationSpeed = 0.5;
-						const double NewMinValue = InterpolationSpeed * MinValue + (1.0 - InterpolationSpeed) * LowValue;
-						const double NewMaxValue = InterpolationSpeed * MaxValue + (1.0 - InterpolationSpeed) * HighValue;
-
-						// Check if we reach the target min-max interval.
-						const double ErrorTolerance = 0.5 / Series.GetScaleY(); // delta value for dy ~= 0.5 pixels
-						if (!FMath::IsNearlyEqual(NewMinValue, MinValue, ErrorTolerance) ||
-							!FMath::IsNearlyEqual(NewMaxValue, MaxValue, ErrorTolerance))
-						{
-							MinValue = NewMinValue;
-							MaxValue = NewMaxValue;
-
-							// Request a new update so we can further interpolate the min-max interval.
-							Series.SetDirtyFlag();
-						}
-					}
-
-					double BaselineY;
-					double ScaleY;
-					Series.ComputeBaselineAndScale(MinValue, MaxValue, TopY, BottomY, BaselineY, ScaleY);
-					Series.SetBaselineY(BaselineY);
-					Series.SetScaleY(ScaleY);
-				}
-				else
-				{
-					// If MinValue == MaxValue, we keep the previous baseline and scale.
-				}
+				Series.UpdateAutoZoom(TopY, BottomY, MinValue, MaxValue);
 			}
 
 			if (Counter.IsFloatingPoint())
diff --git a/Engine/Source/Developer/TraceInsights/Private/Insights/ViewModels/TimingViewDrawHelper.cpp b/Engine/Source/Developer/TraceInsights/Private/Insights/ViewModels/TimingViewDrawHelper.cpp
--- a/Engine/Source/Developer/TraceInsights/Private/Insights/ViewModels/TimingViewDrawHelper.cpp
+++ b/Engine/Source/Developer/TraceInsights/Private/Insights/ViewModels/TimingViewDrawHelper.cpp
@@ -645,7 +645,7 @@ void FTimingViewDrawHelper::DrawTimingEventHighlight(double StartTime, double En
 
 ////////////////////////////////////////////////////////////////////////////////////////////////////
 
-FLinearColor FTimingViewDrawHelper::GetTrackNameTextColor(const FTimingEventsTrack& Track) const
+FLinearColor FTimingViewDrawHelper::GetTrackNameTextColor(const FBaseTimingTrack& Track) const
 {
 	return  Track.IsHovered() ?  FLinearColor(1.0f, 1.0f, 0.0f, 1.0f) :
 			Track.IsSelected() ? FLinearColor(1.0f, 1.0f, 0.5f, 1.0f) :
diff --git a/Engine/Source/Developer/TraceInsights/Private/Insights/ViewModels/TimingViewDrawHelper.h b/Engine/Source/Developer/TraceInsights/Private/Insights/ViewModels/TimingViewDrawHelper.h
--- a/Engine/Source/Developer/TraceInsights/Private/Insights/ViewModels/TimingViewDrawHelper.h
+++ b/Engine/Source/Developer/TraceInsights/Private/Insights/ViewModels/TimingViewDrawHelper.h
@@ -164,7 +164,7 @@ public:
 	virtual const FSlateBrush* GetWhiteBrush() const override { return WhiteBrush; }
 	virtual const FSlateFontInfo& GetEventFont() const override { return EventFont; }
 	virtual FLinearColor GetEdgeColor() const override { return EdgeColor; }
-	virtual FLinearColor GetTrackNameTextColor(const FTimingEventsTrack& Track) const override;
+	virtual FLinearColor GetTrackNameTextColor(const FBaseTimingTrack& Track) const override;
 	virtual int32 GetHeaderBackgroundLayerId() const override { return ReservedLayerId + ToInt32(EDrawLayer::HeaderBackground); }
 	virtual int32 GetHeaderTextLayerId() const override { return ReservedLayerId + ToInt32(EDrawLayer::HeaderText); }
 
diff --git a/Engine/Source/Developer/TraceInsights/Private/Insights/Widgets/STimingView.cpp b/Engine/Source/Developer/TraceInsights/Private/Insights/Widgets/STimingView.cpp
--- a/Engine/Source/Developer/TraceInsights/Private/Insights/Widgets/STimingView.cpp
+++ b/Engine/Source/Developer/TraceInsights/Private/Insights/Widgets/STimingView.cpp
@@ -954,28 +954,64 @@ int32 STimingView::OnPaint(const FPaintArgs& Args, const FGeometry& AllottedGeom
 
 		Helper.BeginDrawTracks();
 
-		for (const TSharedPtr<FBaseTimingTrack>& TrackPtr : ScrollableTracks)
+		const FVector2D Position = AllottedGeometry.GetAbsolutePosition();
+		const float Scale = AllottedGeometry.GetAccumulatedLayoutTransform().GetScale();
+
 		{
-			if (TrackPtr->IsVisible())
+			const float L = Position.X;
+			const float R = Position.X + (Viewport.GetWidth() * Scale);
+			const float T = Position.Y + (Viewport.GetTopOffset() * Scale);
+			const float B = Position.Y + ((Viewport.GetHeight() - Viewport.GetBottomOffset()) * Scale);
+			const FSlateClippingZone ClipZone(FVector2D(L, T), FVector2D(R, T), FVector2D(L, B), FVector2D(R, B));
+			DrawContext.ElementList.PushClip(ClipZone);
+
+			for (const TSharedPtr<FBaseTimingTrack>& TrackPtr : ScrollableTracks)
 			{
-				TrackPtr->Draw(TimingDrawContext);
+				if (TrackPtr->IsVisible())
+				{
+					TrackPtr->Draw(TimingDrawContext);
+				}
 			}
+
+			DrawContext.ElementList.PopClip();
 		}
 
-		for (const TSharedPtr<FBaseTimingTrack>& TrackPtr : TopDockedTracks)
 		{
-			if (TrackPtr->IsVisible())
+			const float L = Position.X;
+			const float R = Position.X + (Viewport.GetWidth() * Scale);
+			const float T = Position.Y;
+			const float B = Position.Y + (Viewport.GetTopOffset() * Scale);
+			const FSlateClippingZone ClipZone(FVector2D(L, T), FVector2D(R, T), FVector2D(L, B), FVector2D(R, B));
+			DrawContext.ElementList.PushClip(ClipZone);
+
+			for (const TSharedPtr<FBaseTimingTrack>& TrackPtr : TopDockedTracks)
 			{
-				TrackPtr->Draw(TimingDrawContext);
+				if (TrackPtr->IsVisible())
+				{
+					TrackPtr->Draw(TimingDrawContext);
+				}
 			}
+
+			DrawContext.ElementList.PopClip();
 		}
 
-		for (const TSharedPtr<FBaseTimingTrack>& TrackPtr : BottomDockedTracks)
 		{
-			if (TrackPtr->IsVisible())
+			const float L = Position.X;
+			const float R = Position.X + (Viewport.GetWidth() * Scale);
+			const float T = Position.Y + ((Viewport.GetHeight() - Viewport.GetBottomOffset()) * Scale);
+			const float B = Position.Y + (Viewport.GetHeight() * Scale);
+			const FSlateClippingZone ClipZone(FVector2D(L, T), FVector2D(R, T), FVector2D(L, B), FVector2D(R, B));
+			DrawContext.ElementList.PushClip(ClipZone);
+
+			for (const TSharedPtr<FBaseTimingTrack>& TrackPtr : BottomDockedTracks)
 			{
-				TrackPtr->Draw(TimingDrawContext);
+				if (TrackPtr->IsVisible())
+				{
+					TrackPtr->Draw(TimingDrawContext);
+				}
 			}
+
+			DrawContext.ElementList.PopClip();
 		}
 
 		for (const TSharedPtr<FBaseTimingTrack>& TrackPtr : ForegroundTracks)
diff --git a/Engine/Source/Developer/TraceInsights/Public/Insights/Common/TimeUtils.h b/Engine/Source/Developer/TraceInsights/Public/Insights/Common/TimeUtils.h
new file mode 100644
--- /dev/null
+++ b/Engine/Source/Developer/TraceInsights/Public/Insights/Common/TimeUtils.h
@@ -0,0 +1,69 @@
+// Copyright 1998-2019 Epic Games, Inc. All Rights Reserved.
+
+#pragma once
+
+#include "CoreMinimal.h"
+#include "Containers/UnrealString.h"
+
+/*
+double is a 64 bit IEEE 754 double precision Floating Point Number
+1 bit for the sign, 11 bits for the exponent, and 52* bits for the value
+15.5 decimal digits of precision (max value*: ‭4`503`599`627`370`495)
+|------------------------------------------------------------------------------
+| 900`000.000`000`001     --> up to 10 days   with    1 ns precision
+|  90`000.000`000`000`1   --> up to 25 hours  with  0.1 ns precision (100 ps)
+|   9`000.000`000`000`01  --> up to 2.5 hours with 0.01 ns precision (10 ps)
+|     900.000`000`000`001 --> up to 15 min    with    1 ps precision
+|------------------------------------------------------------------------------
+|   3`600.000`000`000`001 --> 1 ps precision at 1 hour : NOT OK (16 digits)!!!
+|  86`400.000`000`000`001 --> 1 ps precision at 1 day  : NOT OK (17 digits)!!!
+|------------------------------------------------------------------------------
+*/
+
+namespace TimeUtils
+{
+	static constexpr double Picosecond = 0.000000000001;
+	static constexpr double Nanosecond = 0.000000001;
+	static constexpr double Microsecond = 0.000001;
+	static constexpr double Milisecond = 0.001;
+	static constexpr double Second = 1.0;
+	static constexpr double Minute = 60.0;
+	static constexpr double Hour = 3600.0;
+	static constexpr double Day = 86400.0;
+	static constexpr double Week = 604800.0;
+
+	struct TRACEINSIGHTS_API FTimeSplit
+	{
+		union
+		{
+			int32 Units[8]; // Units[0] = Days, Units[1] = Hours, ..., Units[7] = Picoseconds
+			struct
+			{
+				int32 Days;
+				int32 Hours;
+				int32 Minutes;
+				int32 Seconds;
+				int32 Miliseconds;
+				int32 Microseconds;
+				int32 Nanoseconds;
+				int32 Picoseconds;
+			};
+		};
+		bool bIsZero;
+		bool bIsNegative;
+		bool bIsInfinite;
+	};
+
+	TRACEINSIGHTS_API FString FormatTimeValue(const double Duration, const int32 NumDigits = 1);
+	TRACEINSIGHTS_API FString FormatTimeAuto(const double Duration, const int32 NumDigits = 1);
+	TRACEINSIGHTS_API FString FormatTimeMs(const double Duration, const int32 NumDigits = 2, bool bAddTimeUnit = false);
+	TRACEINSIGHTS_API FString FormatTime(const double Time, const double Precision = 0.0);
+	TRACEINSIGHTS_API FString FormatTimeHMS(const double Time, const double Precision = 0.0);
+	TRACEINSIGHTS_API void SplitTime(const double Time, FTimeSplit& OutTimeSplit);
+	TRACEINSIGHTS_API FString FormatTimeSplit(const FTimeSplit& TimeSplit, const double Precision = 0.0);
+	TRACEINSIGHTS_API FString FormatTimeSplit(const double Time, const double Precision = 0.0);
+
+	void TestTimeFormatting();
+	void TestTimeAutoFormatting();
+	void TestOptimizationIssue();
+}
diff --git a/Engine/Source/Developer/TraceInsights/Public/Insights/ViewModels/GraphSeries.h b/Engine/Source/Developer/TraceInsights/Public/Insights/ViewModels/GraphSeries.h
--- a/Engine/Source/Developer/TraceInsights/Public/Insights/ViewModels/GraphSeries.h
+++ b/Engine/Source/Developer/TraceInsights/Public/Insights/ViewModels/GraphSeries.h
@@ -52,6 +52,14 @@ public:
 	{
 		Color = InColor;
 		BorderColor = InBorderColor;
+		FillColor = InColor;
+	}
+
+	void SetColor(FLinearColor InColor, FLinearColor InBorderColor, FLinearColor InFillColor)
+	{
+		Color = InColor;
+		BorderColor = InBorderColor;
+		FillColor = InFillColor;
 	}
 
 	/**
@@ -65,7 +73,7 @@ public:
 	 * @return The scale between Value units and viewport units; in pixels (Slate units) / Value unit.
 	 */
 	double GetScaleY() const { return ScaleY; }
-	void SetScaleY(const double InScaleY) { ScaleY = InScaleY; }
+	void SetScaleY(const double InScaleY) { ScaleY = FMath::Max(InScaleY, (double)KINDA_SMALL_NUMBER); }
 
 	/**
 	 * @param Value a value; in Value units
@@ -94,18 +102,6 @@ public:
 	bool IsAutoZoomEnabled() const { return bAutoZoom; }
 	void EnableAutoZoom() { bAutoZoom = true; }
 
-	/** Target low value of auto zoom interval (corresponding to bottom of the track) */
-	double GetTargetAutoZoomLowValue() const { return TargetAutoZoomLowValue; }
-	/** Target high value of auto zoom interval (corresponding to top of the track) */
-	double GetTargetAutoZoomHighValue() const { return TargetAutoZoomHighValue; }
-	void SetTargetAutoZoomRange(double LowValue, double HighValue) { TargetAutoZoomLowValue = LowValue; TargetAutoZoomHighValue = HighValue; }
-
-	/** Current auto zoom low value */
-	double GetAutoZoomLowValue() const { return AutoZoomLowValue; }
-	/** Current auto zoom high value */
-	double GetAutoZoomHighValue() const { return AutoZoomHighValue; }
-	void SetAutoZoomRange(double LowValue, double HighValue) { AutoZoomLowValue = LowValue; AutoZoomHighValue = HighValue; }
-
 	/**
 	 * Compute BaselineY and ScaleY so the [Low, High] Value range will correspond to [Top, Bottom] Y position range.
 	 * GetYForValue(InHighValue) == InTopY
@@ -114,7 +110,7 @@ public:
 	void ComputeBaselineAndScale(const double InLowValue, const double InHighValue, const float InTopY, const float InBottomY, double& OutBaselineY, double& OutScaleY) const
 	{
 		ensure(InLowValue < InHighValue);
-		ensure(InTopY < InBottomY);
+		ensure(InTopY <= InBottomY);
 		const double InvRange = 1.0 / (InHighValue - InLowValue);
 		OutScaleY = static_cast<double>(InBottomY - InTopY) * InvRange;
 		//OutBaselineY = (InHighValue * static_cast<double>(InBottomY) - InLowValue * static_cast<double>(InTopY)) * InvRange;
@@ -132,6 +128,9 @@ public:
 	 */
 	const FGraphSeriesEvent* GetEvent(const float PosX, const float PosY, const FTimingTrackViewport& Viewport, bool bCheckLine, bool bCheckBox) const;
 
+	/** Update the track's auto-zoom. Does nothing if IsAutoZoomEnabled() is false. */
+	void UpdateAutoZoom(const float InTopY, const float InBottomY, const double InMinEventValue, const double InMaxEventValue);
+
 	virtual FString FormatValue(double Value) const;
 
 private:
@@ -142,21 +141,18 @@ private:
 	bool bIsDirty;
 
 	bool bAutoZoom;
-	double TargetAutoZoomLowValue; // target low value of auto zoom interval (corresponding to bottom of the track)
-	double TargetAutoZoomHighValue; // target high value of auto zoom interval (corresponding to top of the track)
-	double AutoZoomLowValue; // current auto zoom low value
-	double AutoZoomHighValue; // current auto zoom high value
 
 	double BaselineY; // Y position (in viewport local space) of the baseline (with Value == 0); in pixels (Slate units)
 	double ScaleY; // scale between Value units and viewport units; in pixels (Slate units) / Value unit
 
 	FLinearColor Color;
+	FLinearColor FillColor;
 	FLinearColor BorderColor;
 
 protected:
 	TArray<FGraphSeriesEvent> Events; // reduced list of events; used to identify an event at a certain screen position (ex.: the event hovered by mouse)
 	TArray<FVector2D> Points; // reduced list of points; for drawing points
-	TArray<FVector2D> LinePoints; // reduced list of points; for drawing the connected line and filled polygon
+	TArray<TArray<FVector2D>> LinePoints;	// reduced list of points; for drawing the connected line and filled polygon, split into disconnected batches
 	TArray<FBox> Boxes; // reduced list of boxes; for drawing boxes
 };
 
diff --git a/Engine/Source/Developer/TraceInsights/Public/Insights/ViewModels/GraphTrack.h b/Engine/Source/Developer/TraceInsights/Public/Insights/ViewModels/GraphTrack.h
--- a/Engine/Source/Developer/TraceInsights/Public/Insights/ViewModels/GraphTrack.h
+++ b/Engine/Source/Developer/TraceInsights/Public/Insights/ViewModels/GraphTrack.h
@@ -30,6 +30,7 @@ private:
 
 public:
 	explicit FGraphTrack(const FName& InSubType = NAME_None);
+	explicit FGraphTrack(const FName& InSubType, const FString& InName);
 	virtual ~FGraphTrack();
 
 	//////////////////////////////////////////////////
@@ -39,6 +40,7 @@ public:
 
 	virtual void PostUpdate(const ITimingTrackUpdateContext& Context) override;
 
+	virtual void PreDraw(const ITimingTrackDrawContext& Context) const override;
 	virtual void Draw(const ITimingTrackDrawContext& Context) const override;
 	virtual void DrawEvent(const ITimingTrackDrawContext& Context, const ITimingEvent& InTimingEvent, EDrawEventMode InDrawMode) const override;
 
@@ -65,33 +67,44 @@ protected:
 
 	void DrawSeries(const FGraphSeries& Series, FDrawContext& DrawContext, const FTimingTrackViewport& Viewport) const;
 
+	void UpdateAutoZoom();
+
+	virtual bool ContextMenu_ShowPoints_CanExecute();
+	virtual bool ContextMenu_ShowPointsWithBorder_CanExecute();
+	virtual bool ContextMenu_ShowLines_CanExecute();
+	virtual bool ContextMenu_ShowPolygon_CanExecute();
+	virtual bool ContextMenu_UseEventDuration_CanExecute();
+	virtual bool ContextMenu_ShowBars_CanExecute();
+	virtual bool ContextMenu_ShowSeries_CanExecute(FGraphSeries* Series);
+
+	// Get the Y value that is used to provided a clipping border between adjacent 
+	// graph tracks
+	virtual float GetBorderY() const { return 0.0f; }
+
 private:
+	bool ContextMenu_ShowDebugInfo_CanExecute();
+	void ContextMenu_ShowDebugInfo_Execute();
+	bool ContextMenu_ShowDebugInfo_IsChecked();
+
 	void ContextMenu_ShowPoints_Execute();
-	bool ContextMenu_ShowPoints_CanExecute();
 	bool ContextMenu_ShowPoints_IsChecked();
 
 	void ContextMenu_ShowPointsWithBorder_Execute();
-	bool ContextMenu_ShowPointsWithBorder_CanExecute();
 	bool ContextMenu_ShowPointsWithBorder_IsChecked();
 
 	void ContextMenu_ShowLines_Execute();
-	bool ContextMenu_ShowLines_CanExecute();
 	bool ContextMenu_ShowLines_IsChecked();
 
 	void ContextMenu_ShowPolygon_Execute();
-	bool ContextMenu_ShowPolygon_CanExecute();
 	bool ContextMenu_ShowPolygon_IsChecked();
 
 	void ContextMenu_UseEventDuration_Execute();
-	bool ContextMenu_UseEventDuration_CanExecute();
 	bool ContextMenu_UseEventDuration_IsChecked();
 
 	void ContextMenu_ShowBars_Execute();
-	bool ContextMenu_ShowBars_CanExecute();
 	bool ContextMenu_ShowBars_IsChecked();
 
 	void ContextMenu_ShowSeries_Execute(FGraphSeries* Series);
-	bool ContextMenu_ShowSeries_CanExecute(FGraphSeries* Series);
 	bool ContextMenu_ShowSeries_IsChecked(FGraphSeries* Series);
 
 protected:
@@ -103,12 +116,14 @@ protected:
 	const FSlateBrush* BorderBrush;
 	const FSlateFontInfo Font;
 
+	bool bDrawDebugInfo;
 	bool bDrawPoints;
 	bool bDrawPointsWithBorder;
 	bool bDrawLines;
 	bool bDrawPolygon;
 	bool bUseEventDuration;
 	bool bDrawBoxes;
+	bool bDrawBaseline;
 
 	// Stats
 	int32 NumAddedEvents; // total event count
diff --git a/Engine/Source/Developer/TraceInsights/Public/Insights/ViewModels/GraphTrackBuilder.h b/Engine/Source/Developer/TraceInsights/Public/Insights/ViewModels/GraphTrackBuilder.h
--- a/Engine/Source/Developer/TraceInsights/Public/Insights/ViewModels/GraphTrackBuilder.h
+++ b/Engine/Source/Developer/TraceInsights/Public/Insights/ViewModels/GraphTrackBuilder.h
@@ -37,7 +37,7 @@ public:
 	FGraphSeries& GetSeries() const { return Series; }
 	const FTimingTrackViewport& GetViewport() const { return Viewport; }
 
-	void AddEvent(double Time, double Duration, double Value);
+	void AddEvent(double Time, double Duration, double Value, bool bConnected = true);
 
 private:
 	void BeginPoints();
@@ -46,7 +46,7 @@ private:
 	void EndPoints();
 
 	void BeginConnectedLines();
-	void AddConnectedLine(double Time, double Value);
+	void AddConnectedLine(double Time, double Value, bool bNewBatch);
 	void FlushConnectedLine();
 	void EndConnectedLines();
 
diff --git a/Engine/Source/Developer/TraceInsights/Public/Insights/ViewModels/ITimingViewDrawHelper.h b/Engine/Source/Developer/TraceInsights/Public/Insights/ViewModels/ITimingViewDrawHelper.h
--- a/Engine/Source/Developer/TraceInsights/Public/Insights/ViewModels/ITimingViewDrawHelper.h
+++ b/Engine/Source/Developer/TraceInsights/Public/Insights/ViewModels/ITimingViewDrawHelper.h
@@ -4,7 +4,7 @@
 
 #include "CoreMinimal.h"
 
-class FTimingEventsTrack;
+class FBaseTimingTrack;
 struct FSlateBrush;
 struct FSlateFontInfo;
 
@@ -15,7 +15,7 @@ public:
 	virtual const FSlateBrush* GetWhiteBrush() const = 0;
 	virtual const FSlateFontInfo& GetEventFont() const = 0;
 	virtual FLinearColor GetEdgeColor() const = 0;
-	virtual FLinearColor GetTrackNameTextColor(const FTimingEventsTrack& Track) const = 0;
+	virtual FLinearColor GetTrackNameTextColor(const FBaseTimingTrack& Track) const = 0;
 	virtual int32 GetHeaderBackgroundLayerId() const = 0;
 	virtual int32 GetHeaderTextLayerId() const = 0;
 };
\ No newline at end of file
diff --git a/Engine/Source/Developer/TraceServices/Private/Analyzers/CsvProfilerTraceAnalysis.cpp b/Engine/Source/Developer/TraceServices/Private/Analyzers/CsvProfilerTraceAnalysis.cpp
--- a/Engine/Source/Developer/TraceServices/Private/Analyzers/CsvProfilerTraceAnalysis.cpp
+++ b/Engine/Source/Developer/TraceServices/Private/Analyzers/CsvProfilerTraceAnalysis.cpp
@@ -402,7 +402,7 @@ void FCsvProfilerAnalyzer::HandleEventEvent(const FOnEventContext& Context)
 
 uint32 FCsvProfilerAnalyzer::GetFrameNumberForTimestamp(ETraceFrameType FrameType, double Timestamp) const
 {
-	const TArray<double>& FrameStartTimes = FrameProvider.GetFrameStartTimes(FrameType);
+	const TArray64<double>& FrameStartTimes = FrameProvider.GetFrameStartTimes(FrameType);
 
 	if (FrameStartTimes.Num() == 0 || Timestamp < FrameStartTimes[0])
 	{
diff --git a/Engine/Source/Developer/TraceServices/Private/Model/Counters.cpp b/Engine/Source/Developer/TraceServices/Private/Model/Counters.cpp
--- a/Engine/Source/Developer/TraceServices/Private/Model/Counters.cpp
+++ b/Engine/Source/Developer/TraceServices/Private/Model/Counters.cpp
@@ -9,7 +9,7 @@ namespace Trace
 
 const FName FCounterProvider::ProviderName("CounterProvider");
 
-FCounter::FCounter(ILinearAllocator& Allocator, const TArray<double>& InFrameStartTimes)
+FCounter::FCounter(ILinearAllocator& Allocator, const TArray64<double>& InFrameStartTimes)
 	: FrameStartTimes(InFrameStartTimes)
 	, IntCounterData(Allocator)
 	, DoubleCounterData(Allocator)
@@ -18,13 +18,13 @@ FCounter::FCounter(ILinearAllocator& Allocator, const TArray<double>& InFrameSta
 }
 
 template<typename CounterType, typename EnumerationType>
-static void EnumerateCounterValuesInternal(const TCounterData<CounterType>& CounterData, const TArray<double>& FrameStartTimes, bool bResetEveryFrame, double IntervalStart, double IntervalEnd, bool bIncludeExternalBounds, TFunctionRef<void(double, EnumerationType)> Callback)
+static void EnumerateCounterValuesInternal(const TCounterData<CounterType>& CounterData, const TArray64<double>& FrameStartTimes, bool bResetEveryFrame, double IntervalStart, double IntervalEnd, bool bIncludeExternalBounds, TFunctionRef<void(double, EnumerationType)> Callback)
 {
 	if (!CounterData.Num())
 	{
 		return;
 	}
-	TArray<double> NoFrameStartTimes;
+	TArray64<double> NoFrameStartTimes;
 	auto CounterIterator = bResetEveryFrame ? CounterData.GetIterator(FrameStartTimes) : CounterData.GetIterator(NoFrameStartTimes);
 	bool bFirstValue = true;
 	bool bFirstEnumeratedValue = true;
diff --git a/Engine/Source/Developer/TraceServices/Private/Model/CountersPrivate.h b/Engine/Source/Developer/TraceServices/Private/Model/CountersPrivate.h
--- a/Engine/Source/Developer/TraceServices/Private/Model/CountersPrivate.h
+++ b/Engine/Source/Developer/TraceServices/Private/Model/CountersPrivate.h
@@ -24,7 +24,7 @@ template<typename ValueType>
 class TCounterDataIterator
 {
 public:
-	TCounterDataIterator(const TCounterData<ValueType>& Outer, const TArray<double>& FrameStartTimes)
+	TCounterDataIterator(const TCounterData<ValueType>& Outer, const TArray64<double>& FrameStartTimes)
 		: FrameStartTimesIterator(FrameStartTimes.CreateConstIterator())
 		, TimestampsIterator(Outer.Timestamps.GetIterator())
 		, OpTypesIterator(Outer.OpTypes.GetIterator())
@@ -92,7 +92,7 @@ private:
 		}
 	}
 
-	TArray<double>::TConstIterator FrameStartTimesIterator;
+	TArray64<double>::TConstIterator FrameStartTimesIterator;
 	TPagedArray<double>::TIterator TimestampsIterator;
 	TPagedArray<ECounterOpType>::TIterator OpTypesIterator;
 	typename TPagedArray<ValueType>::TIterator OpArgumentsIterator;
@@ -145,7 +145,7 @@ public:
 		return Timestamps.Num();
 	}
 
-	TIterator GetIterator(TArray<double> FrameStartTimes) const
+	TIterator GetIterator(const TArray64<double>& FrameStartTimes) const
 	{
 		return TIterator(*this, FrameStartTimes);
 	}
@@ -163,7 +163,7 @@ class FCounter
 	: public IEditableCounter
 {
 public:
-	FCounter(ILinearAllocator& Allocator, const TArray<double>& FrameStartTimes);
+	FCounter(ILinearAllocator& Allocator, const TArray64<double>& FrameStartTimes);
 	virtual const TCHAR* GetName() const override { return Name; }
 	virtual void SetName(const TCHAR* InName) override { Name = InName; }
 	virtual const TCHAR* GetDescription() const override { return Description; }
@@ -181,7 +181,7 @@ public:
 	virtual void SetValue(double Time, double Value) override;
 
 private:
-	const TArray<double>& FrameStartTimes;
+	const TArray64<double>& FrameStartTimes;
 	TCounterData<int64> IntCounterData;
 	TCounterData<double> DoubleCounterData;
 	const TCHAR* Name = nullptr;
diff --git a/Engine/Source/Developer/TraceServices/Private/Model/Frames.cpp b/Engine/Source/Developer/TraceServices/Private/Model/Frames.cpp
--- a/Engine/Source/Developer/TraceServices/Private/Model/Frames.cpp
+++ b/Engine/Source/Developer/TraceServices/Private/Model/Frames.cpp
@@ -40,6 +40,46 @@ void FFrameProvider::EnumerateFrames(ETraceFrameType FrameType, uint64 Start, ui
 	}
 }
 
+// Pulled from Algo/BinarySearch.h, which by default truncates to int32.
+// @TODO: remove when LowerBound supports 64 bit results
+template <typename RangeValueType, typename PredicateValueType, typename ProjectionType, typename SortPredicateType>
+static int64 LowerBound64(RangeValueType* First, const int64 Num, const PredicateValueType& Value, ProjectionType Projection, SortPredicateType SortPredicate)
+{
+	// Current start of sequence to check
+	int64 Start = 0;
+	// Size of sequence to check
+	int64 Size = Num;
+
+	// With this method, if Size is even it will do one more comparison than necessary, but because Size can be predicted by the CPU it is faster in practice
+	while (Size > 0)
+	{
+		const int64 LeftoverSize = Size % 2;
+		Size = Size / 2;
+
+		const int64 CheckIndex = Start + Size;
+		const int64 StartIfLess = CheckIndex + LeftoverSize;
+
+		auto&& CheckValue = Invoke(Projection, First[CheckIndex]);
+		Start = SortPredicate(CheckValue, Value) ? StartIfLess : Start;
+	}
+	return Start;
+}
+
+bool FFrameProvider::GetFrameFromTime(ETraceFrameType FrameType, double Time, FFrame& OutFrame) const
+{
+	int64 LowerBound = LowerBound64(FrameStartTimes[FrameType].GetData(), FrameStartTimes[FrameType].Num(), Time, FIdentityFunctor(), TLess<double>());
+	if(FrameStartTimes[FrameType].IsValidIndex(LowerBound) && LowerBound > 0)
+	{
+		OutFrame.Index = LowerBound;
+		OutFrame.StartTime = FrameStartTimes[FrameType][LowerBound - 1];
+		OutFrame.EndTime = FrameStartTimes[FrameType][LowerBound];
+		OutFrame.FrameType = FrameType;
+		return true;
+	}
+
+	return false;
+}
+
 const FFrame* FFrameProvider::GetFrame(ETraceFrameType FrameType, uint64 Index) const
 {
 	Session.ReadAccessCheck();
diff --git a/Engine/Source/Developer/TraceServices/Private/Model/FramesPrivate.h b/Engine/Source/Developer/TraceServices/Private/Model/FramesPrivate.h
--- a/Engine/Source/Developer/TraceServices/Private/Model/FramesPrivate.h
+++ b/Engine/Source/Developer/TraceServices/Private/Model/FramesPrivate.h
@@ -19,7 +19,8 @@ public:
 
 	virtual uint64 GetFrameCount(ETraceFrameType FrameType) const override;
 	virtual void EnumerateFrames(ETraceFrameType FrameType, uint64 Start, uint64 End, TFunctionRef<void(const FFrame&)> Callback) const override;
-	virtual const TArray<double>& GetFrameStartTimes(ETraceFrameType FrameType) const override { return FrameStartTimes[FrameType]; }
+	virtual const TArray64<double>& GetFrameStartTimes(ETraceFrameType FrameType) const override { return FrameStartTimes[FrameType]; }
+	virtual bool GetFrameFromTime(ETraceFrameType FrameType, double Time, FFrame& OutFrame) const override;
 	virtual const FFrame* GetFrame(ETraceFrameType FrameType, uint64 Index) const override;
 	void BeginFrame(ETraceFrameType FrameType, double Time);
 	void EndFrame(ETraceFrameType FrameType, double Time);
@@ -27,7 +28,7 @@ public:
 private:
 	IAnalysisSession& Session;
 	TArray<TPagedArray<FFrame>> Frames;
-	TArray<double> FrameStartTimes[TraceFrameType_Count];
+	TArray64<double> FrameStartTimes[TraceFrameType_Count];
 };
 
 }
\ No newline at end of file
diff --git a/Engine/Source/Developer/TraceServices/Public/Model/PointTimeline.h b/Engine/Source/Developer/TraceServices/Public/Model/PointTimeline.h
--- a/Engine/Source/Developer/TraceServices/Public/Model/PointTimeline.h
+++ b/Engine/Source/Developer/TraceServices/Public/Model/PointTimeline.h
@@ -47,8 +47,10 @@ public:
 	
 	virtual void EnumerateEvents(double IntervalStart, double IntervalEnd, typename ITimeline<EventType>::EventRangeCallback Callback) const override
 	{
+		bool bIntersected = false;
 		auto EventIterator = Events.GetIteratorFromPage(0);
 		const FEventInternal* Event = EventIterator.GetCurrentItem();
+		const FEventInternal* PreviousEvent = Event;
 		while (Event)
 		{
 			const FEventPage* Page = EventIterator.GetCurrentPage();
@@ -58,14 +60,28 @@ public:
 				{
 					return;
 				}
+
+				PreviousEvent = Event;
 				Event = EventIterator.GetCurrentItem();
 			}
 			else
 			{
 				if (!(Event->Time < IntervalStart || IntervalEnd < Event->Time))
+				{
+					if(bEnumerateOutsideRange && PreviousEvent != Event)
+					{
+						Callback(PreviousEvent->Time, PreviousEvent->Time, 0, PreviousEvent->Event);
+					}
+
+					Callback(Event->Time, Event->Time, 0, Event->Event);
+					bIntersected = true;
+				}
+				else if(bEnumerateOutsideRange && bIntersected)
 				{
 					Callback(Event->Time, Event->Time, 0, Event->Event);
 				}
+
+				PreviousEvent = Event;
 				Event = EventIterator.NextItem();
 			}
 		}
@@ -99,6 +115,14 @@ public:
 		return Index;
 	}
 
+	// Set whether we enumerate one event each side of the passed-in time range.
+	// This is useful for cases such as drawing where connected events need to be drawn outside of
+	// the viewport.
+	void SetEnumerateOutsideRange(bool bInEnumerateOutsideRange)
+	{
+		bEnumerateOutsideRange = bInEnumerateOutsideRange;
+	}
+
 private:
 	struct FEventInternal
 	{
@@ -120,6 +144,7 @@ private:
 	FEventPage* CurrentPage = nullptr;
 	double LastTime = 0.0;
 	uint64 ModCount = 0;
+	bool bEnumerateOutsideRange = false;
 };
 
 }
\ No newline at end of file
diff --git a/Engine/Source/Developer/TraceServices/Public/TraceServices/Model/Frames.h b/Engine/Source/Developer/TraceServices/Public/TraceServices/Model/Frames.h
--- a/Engine/Source/Developer/TraceServices/Public/TraceServices/Model/Frames.h
+++ b/Engine/Source/Developer/TraceServices/Public/TraceServices/Model/Frames.h
@@ -6,6 +6,7 @@
 #include "ProfilingDebugging/MiscTrace.h"
 #include "Templates/Function.h"
 #include "Delegates/Delegate.h"
+#include "Containers/ContainersFwd.h"
 
 namespace Trace
 {
@@ -25,7 +26,8 @@ public:
 	virtual ~IFrameProvider() = default;
 	virtual uint64 GetFrameCount(ETraceFrameType FrameType) const = 0;
 	virtual void EnumerateFrames(ETraceFrameType FrameType, uint64 Start, uint64 End, TFunctionRef<void(const FFrame&)> Callback) const = 0;
-	virtual const TArray<double>& GetFrameStartTimes(ETraceFrameType FrameType) const = 0;
+	virtual const TArray64<double>& GetFrameStartTimes(ETraceFrameType FrameType) const = 0;
+	virtual bool GetFrameFromTime(ETraceFrameType FrameType, double Time, FFrame& OutFrame) const = 0;
 	virtual const FFrame* GetFrame(ETraceFrameType FrameType, uint64 Index) const = 0;
 };
 
