[Insights] Minor code style changes.

diff --git a/Engine/Source/Developer/TraceInsights/Private/Insights/ViewModels/GraphTrack.cpp b/Engine/Source/Developer/TraceInsights/Private/Insights/ViewModels/GraphTrack.cpp
--- a/Engine/Source/Developer/TraceInsights/Private/Insights/ViewModels/GraphTrack.cpp
+++ b/Engine/Source/Developer/TraceInsights/Private/Insights/ViewModels/GraphTrack.cpp
@@ -2,11 +2,11 @@
 
 #include "Insights/ViewModels/GraphTrack.h"
 
+#include "EditorFontGlyphs.h"
 #include "EditorStyleSet.h"
 #include "Framework/MultiBox/MultiBoxBuilder.h"
 #include "Rendering/DrawElements.h"
 #include "Styling/CoreStyle.h"
-#include "EditorFontGlyphs.h"
 
 // Insights
 #include "Insights/Common/PaintUtils.h"
@@ -48,6 +48,8 @@ FGraphTrack::FGraphTrack(const FName& InSubType)
 {
 }
 
+////////////////////////////////////////////////////////////////////////////////////////////////////
+
 FGraphTrack::FGraphTrack(const FName& InSubType, const FString& InName)
 	: FBaseTimingTrack(FName(TEXT("Graph")), InSubType, InName)
 	//, AllSeries()
@@ -117,7 +119,7 @@ void FGraphTrack::UpdateStats()
 		if (Series->IsVisible())
 		{
 			NumDrawPoints += Series->Points.Num();
-			for(int32 BatchIndex = 0; BatchIndex < Series->LinePoints.Num(); ++BatchIndex)
+			for (int32 BatchIndex = 0; BatchIndex < Series->LinePoints.Num(); ++BatchIndex)
 			{
 				NumDrawLines += Series->LinePoints[BatchIndex].Num() / 2;
 			}
@@ -164,7 +166,7 @@ void FGraphTrack::Draw(const ITimingTrackDrawContext& Context) const
 		if (Series->IsVisible())
 		{
 			// Draw baseline (Value == 0), for the first visible series only.
-			if(bDrawBaseline && !bDrawnBaseline)
+			if (bDrawBaseline && !bDrawnBaseline)
 			{
 				const float BaselineY = FMath::RoundToFloat(static_cast<float>(Series->GetBaselineY()));
 				if (BaselineY >= 0.0f && BaselineY < GetHeight())
@@ -191,7 +193,7 @@ void FGraphTrack::DrawSeries(const FGraphSeries& Series, FDrawContext& DrawConte
 	// Set clipping area. This area tries to take into account the optional border area that allows
 	// graph tracks to optionally act like 'event' tracks if required (with respect to layout, anyway).
 	// The GetBorderY() - 1.0f calculation is desgiend to avoid clipping the line rasterization, as the custom verts
-	// of the fill and the outer lines appear to get rasterized differently, the latter missing one pixel 
+	// of the fill and the outer lines appear to get rasterized differently, the latter missing one pixel
 	// on its upper side.
 	{
 		const FVector2D AbsPos = DrawContext.Geometry.GetAbsolutePosition();
@@ -240,15 +242,18 @@ void FGraphTrack::DrawSeries(const FGraphSeries& Series, FDrawContext& DrawConte
 		const FVector2D Size = DrawContext.Geometry.GetLocalSize();
 
 		const FSlateRenderTransform& RenderTransform = Geo.GetAccumulatedRenderTransform();
+
 		FColor FillColor(Series.FillColor.R * 255, Series.FillColor.G * 255, Series.FillColor.B * 255, 42);
 
-		for(int32 BatchIndex = 0; BatchIndex < Series.LinePoints.Num(); ++BatchIndex)
+		for (int32 BatchIndex = 0; BatchIndex < Series.LinePoints.Num(); ++BatchIndex)
 		{
+			const TArray<FVector2D>& LinePoints = Series.LinePoints[BatchIndex];
+
 			TArray<SlateIndex> Indices;
 			TArray<FSlateVertex> Verts;
 
-			Indices.Reserve(Series.LinePoints[BatchIndex].Num() * 6);
-			Verts.Reserve(Series.LinePoints[BatchIndex].Num() * 2);
+			Indices.Reserve(LinePoints.Num() * 6);
+			Verts.Reserve(LinePoints.Num() * 2);
 
 			const float TopV = 0.0f;
 			const float BottomV = GetHeight() / Size.Y;
@@ -257,17 +262,17 @@ void FGraphTrack::DrawSeries(const FGraphSeries& Series, FDrawContext& DrawConte
 
 			int32 PrevSide = 0;
 
-			for (int32 PointIndex = 0; PointIndex < Series.LinePoints[BatchIndex].Num(); ++PointIndex)
+			for (int32 PointIndex = 0; PointIndex < LinePoints.Num(); ++PointIndex)
 			{
-				const FVector2D& LinePoint = Series.LinePoints[BatchIndex][PointIndex];
+				const FVector2D& LinePoint = LinePoints[PointIndex];
 
 				// When crossing baseline the polygon needs to be intersected.
 				int32 CrtSide = (LinePoint.Y > BaselineY) ? 1 : (LinePoint.Y < BaselineY) ? -1 : 0;
 				if (PrevSide != 0 && CrtSide + PrevSide == 0) // alternating sides?
 				{
 					// Compute intersection point.
-					const FVector2D& PrevLinePoint = Series.LinePoints[BatchIndex][PointIndex - 1];
-					float X = PrevLinePoint.X + (LinePoint.X - PrevLinePoint.X) / ((BaselineY - LinePoint.Y) / (PrevLinePoint.Y - BaselineY) + 1.0f);
+					const FVector2D& PrevLinePoint = LinePoints[PointIndex - 1];
+					const float X = PrevLinePoint.X + (LinePoint.X - PrevLinePoint.X) / ((BaselineY - LinePoint.Y) / (PrevLinePoint.Y - BaselineY) + 1.0f);
 
 					// Add an intersection point vertex.
 					FVector2D UV(X / Size.X, BaselineV);
@@ -348,12 +353,15 @@ void FGraphTrack::DrawSeries(const FGraphSeries& Series, FDrawContext& DrawConte
 		FPaintGeometry LineGeo = Geo;
 		LineGeo.AppendTransform(FSlateLayoutTransform(FVector2D(0.5f * InvScale, 0.5f * InvScale)));
 
-		for(int32 BatchIndex = 0; BatchIndex < Series.LinePoints.Num(); ++BatchIndex)
+		// Disable pixel snapping here so lines line up with boxes/polys correctly.
+		const ESlateDrawEffect LineDrawEffects = DrawContext.DrawEffects | ESlateDrawEffect::NoPixelSnapping;
+
+		for (int32 BatchIndex = 0; BatchIndex < Series.LinePoints.Num(); ++BatchIndex)
 		{
-			if(Series.LinePoints[BatchIndex].Num())
+			const TArray<FVector2D>& LinePoints = Series.LinePoints[BatchIndex];
+			if (LinePoints.Num() > 0)
 			{
-				// Disable pixel snapping here so lines line up with boxes/polys correctly
-				FSlateDrawElement::MakeLines(DrawContext.ElementList, DrawContext.LayerId, LineGeo, Series.LinePoints[BatchIndex], DrawContext.DrawEffects | ESlateDrawEffect::NoPixelSnapping, Series.Color, false, 1.0f);
+				FSlateDrawElement::MakeLines(DrawContext.ElementList, DrawContext.LayerId, LineGeo, LinePoints, LineDrawEffects, Series.Color, false, 1.0f);
 			}
 		}
 		DrawContext.LayerId++;
@@ -364,7 +372,7 @@ void FGraphTrack::DrawSeries(const FGraphSeries& Series, FDrawContext& DrawConte
 
 	if (bDrawPoints)
 	{
-		int32 NumPoints = Series.Points.Num();
+		const int32 NumPoints = Series.Points.Num();
 
 #define INSIGHTS_GRAPH_TRACK_DRAW_POINTS_AS_RECTANGLES 0
 #if !INSIGHTS_GRAPH_TRACK_DRAW_POINTS_AS_RECTANGLES
@@ -375,8 +383,8 @@ void FGraphTrack::DrawSeries(const FGraphSeries& Series, FDrawContext& DrawConte
 			for (int32 Index = 0; Index < NumPoints; ++Index)
 			{
 				const FVector2D& Pt = Series.Points[Index];
-				const float PtX = (Pt.X - PointVisualSize / 2.0f - 0.5f);
-				const float PtY = (LocalPosY + Pt.Y - PointVisualSize / 2.0f - 0.5f);
+				const float PtX = Pt.X - PointVisualSize / 2.0f - 0.5f;
+				const float PtY = LocalPosY + Pt.Y - PointVisualSize / 2.0f - 0.5f;
 				DrawContext.DrawBox(PtX, PtY, PointVisualSize + 2.0f, PointVisualSize + 2.0f, PointBrush, Series.BorderColor);
 			}
 			DrawContext.LayerId++;
@@ -386,8 +394,8 @@ void FGraphTrack::DrawSeries(const FGraphSeries& Series, FDrawContext& DrawConte
 		for (int32 Index = 0; Index < NumPoints; ++Index)
 		{
 			const FVector2D& Pt = Series.Points[Index];
-			const float PtX = (Pt.X - PointVisualSize / 2.0f + 0.5f);
-			const float PtY = (LocalPosY + Pt.Y - PointVisualSize / 2.0f + 0.5f);
+			const float PtX = Pt.X - PointVisualSize / 2.0f + 0.5f;
+			const float PtY = LocalPosY + Pt.Y - PointVisualSize / 2.0f + 0.5f;
 			DrawContext.DrawBox(PtX, PtY, PointVisualSize, PointVisualSize, PointBrush, Series.Color);
 		}
 		DrawContext.LayerId++;
@@ -435,39 +443,49 @@ void FGraphTrack::DrawSeries(const FGraphSeries& Series, FDrawContext& DrawConte
 		FPaintGeometry LineGeo = Geo;
 		LineGeo.AppendTransform(FSlateLayoutTransform(FVector2D(0.5f * PixelUnit, 0.5f * PixelUnit)));
 
+		// Disable pixel snapping here so lines line up with boxes/polys correctly.
+		const ESlateDrawEffect LineDrawEffects = DrawContext.DrawEffects | ESlateDrawEffect::NoPixelSnapping;
+
 		// Draw white corner at (0, 0) using MakeLines.
 		{
 			TArray<FVector2D> HLine;
 			HLine.Add(FVector2D(0.0f, 0.0f));
 			HLine.Add(FVector2D(10.0f, 0.0f));
-			FSlateDrawElement::MakeLines(DrawContext.ElementList, DrawContext.LayerId, LineGeo, HLine, DrawContext.DrawEffects | ESlateDrawEffect::NoPixelSnapping, FLinearColor::White, false, 1.0f);
+			FSlateDrawElement::MakeLines(DrawContext.ElementList, DrawContext.LayerId, LineGeo, HLine, LineDrawEffects, FLinearColor::White, false, 1.0f);
+
 			TArray<FVector2D> VLine;
 			VLine.Add(FVector2D(0.0f, 0.0f));
 			VLine.Add(FVector2D(0.0f, 10.0f));
-			FSlateDrawElement::MakeLines(DrawContext.ElementList, DrawContext.LayerId, LineGeo, VLine, DrawContext.DrawEffects | ESlateDrawEffect::NoPixelSnapping, FLinearColor::White, false, 1.0f);
+			FSlateDrawElement::MakeLines(DrawContext.ElementList, DrawContext.LayerId, LineGeo, VLine, LineDrawEffects, FLinearColor::White, false, 1.0f);
 		}
 		DrawContext.LayerId++;
+
 		// Draw corner at (0, 0) using DrawBox.
 		DrawContext.DrawBox(2.0f, LocalPosY, 6.0f, PixelUnit, WhiteBrush, Series.Color);
 		DrawContext.DrawBox(0.0f, LocalPosY + 2.0f, PixelUnit, 6.0f, WhiteBrush, Series.Color);
 		DrawContext.LayerId++;
+
 		const int32 NumDbgPoints = Series.Points.Num();
+
 		// Draw white cross lines (15x15) at each point.
 		for (int32 Index = 0; Index < NumDbgPoints; ++Index)
 		{
 			const FVector2D& Pt = Series.Points[Index];
 			const float PtX = Pt.X;
 			const float PtY = Pt.Y;
+
 			TArray<FVector2D> HLine;
 			HLine.Add(FVector2D(PtX - 7.0f, PtY));
 			HLine.Add(FVector2D(PtX + 8.0f, PtY));
-			FSlateDrawElement::MakeLines(DrawContext.ElementList, DrawContext.LayerId, LineGeo, HLine, DrawContext.DrawEffects | ESlateDrawEffect::NoPixelSnapping, FLinearColor::White, false, 1.0f);
+			FSlateDrawElement::MakeLines(DrawContext.ElementList, DrawContext.LayerId, LineGeo, HLine, LineDrawEffects, FLinearColor::White, false, 1.0f);
+
 			TArray<FVector2D> VLine;
 			VLine.Add(FVector2D(PtX, PtY - 7.0f));
 			VLine.Add(FVector2D(PtX, PtY + 8.0f));
-			FSlateDrawElement::MakeLines(DrawContext.ElementList, DrawContext.LayerId, LineGeo, VLine, DrawContext.DrawEffects | ESlateDrawEffect::NoPixelSnapping, FLinearColor::White, false, 1.0f);
+			FSlateDrawElement::MakeLines(DrawContext.ElementList, DrawContext.LayerId, LineGeo, VLine, LineDrawEffects, FLinearColor::White, false, 1.0f);
 		}
 		DrawContext.LayerId++;
+
 		// Draw cross lines (11x11) at each point.
 		for (int32 Index = 0; Index < NumDbgPoints; ++Index)
 		{
@@ -599,7 +617,7 @@ void FGraphTrack::BuildContextMenu(FMenuBuilder& MenuBuilder)
 {
 	MenuBuilder.BeginSection("Misc");
 	{
-		if (true)
+		if (false) // debug functionality
 		{
 			FUIAction Action_ShowDebugInfo
 			(
@@ -730,7 +748,7 @@ void FGraphTrack::BuildContextMenu(FMenuBuilder& MenuBuilder)
 			(
 				Action_ShowSeries,
 				SNew(SHorizontalBox)
-				+SHorizontalBox::Slot()
+				+ SHorizontalBox::Slot()
 				.AutoWidth()
 				.VAlign(VAlign_Center)
 				.Padding(2.0f, 0.0f)
@@ -738,7 +756,7 @@ void FGraphTrack::BuildContextMenu(FMenuBuilder& MenuBuilder)
 					SNew(STextBlock)
 					.Text(Series->GetName())
 				]
-				+SHorizontalBox::Slot()
+				+ SHorizontalBox::Slot()
 				.FillWidth(1.0f)
 				.HAlign(HAlign_Right)
 				.VAlign(VAlign_Center)
diff --git a/Engine/Source/Developer/TraceInsights/Private/Insights/ViewModels/TimingViewDrawHelper.cpp b/Engine/Source/Developer/TraceInsights/Private/Insights/ViewModels/TimingViewDrawHelper.cpp
--- a/Engine/Source/Developer/TraceInsights/Private/Insights/ViewModels/TimingViewDrawHelper.cpp
+++ b/Engine/Source/Developer/TraceInsights/Private/Insights/ViewModels/TimingViewDrawHelper.cpp
@@ -3,7 +3,6 @@
 #include "TimingViewDrawHelper.h"
 
 #include "Fonts/FontMeasure.h"
-#include "Fonts/SlateFontInfo.h"
 #include "Framework/Application/SlateApplication.h"
 #include "Rendering/DrawElements.h"
 #include "Styling/CoreStyle.h"
@@ -275,7 +274,7 @@ void FTimingEventsTrackDrawStateBuilder::AddEvent(double EventStartTime, double
 
 ////////////////////////////////////////////////////////////////////////////////////////////////////
 
-void FTimingEventsTrackDrawStateBuilder::FlushBox(const FBoxData& Box, int32 Depth)
+void FTimingEventsTrackDrawStateBuilder::FlushBox(const FBoxData& Box, const int32 Depth)
 {
 	DrawState.Boxes.AddUninitialized();
 	FTimingEventsTrackDrawState::FBoxPrimitive& DrawBox = DrawState.Boxes.Last();
diff --git a/Engine/Source/Developer/TraceInsights/Private/Insights/ViewModels/TimingViewDrawHelper.h b/Engine/Source/Developer/TraceInsights/Private/Insights/ViewModels/TimingViewDrawHelper.h
--- a/Engine/Source/Developer/TraceInsights/Private/Insights/ViewModels/TimingViewDrawHelper.h
+++ b/Engine/Source/Developer/TraceInsights/Private/Insights/ViewModels/TimingViewDrawHelper.h
@@ -7,7 +7,7 @@
 #include "Math/Color.h"
 
 // Insights
-#include "Insights/ViewModels/TimingEventsTrack.h"
+#include "Insights/ViewModels/TimingEventsTrack.h" // for ITimingEventsTrackDrawStateBuilder
 #include "Insights/ViewModels/ITimingViewDrawHelper.h"
 
 ////////////////////////////////////////////////////////////////////////////////////////////////////
@@ -121,7 +121,6 @@ private:
 
 private:
 	FTimingEventsTrackDrawState& DrawState; // cached draw state to build
-
 	const FTimingTrackViewport& Viewport;
 
 	int32 MaxDepth;
diff --git a/Engine/Source/Developer/TraceInsights/Public/Insights/ViewModels/GraphTrack.h b/Engine/Source/Developer/TraceInsights/Public/Insights/ViewModels/GraphTrack.h
--- a/Engine/Source/Developer/TraceInsights/Public/Insights/ViewModels/GraphTrack.h
+++ b/Engine/Source/Developer/TraceInsights/Public/Insights/ViewModels/GraphTrack.h
@@ -67,8 +67,6 @@ protected:
 
 	void DrawSeries(const FGraphSeries& Series, FDrawContext& DrawContext, const FTimingTrackViewport& Viewport) const;
 
-	void UpdateAutoZoom();
-
 	virtual bool ContextMenu_ShowPoints_CanExecute();
 	virtual bool ContextMenu_ShowPointsWithBorder_CanExecute();
 	virtual bool ContextMenu_ShowLines_CanExecute();
@@ -77,8 +75,7 @@ protected:
 	virtual bool ContextMenu_ShowBars_CanExecute();
 	virtual bool ContextMenu_ShowSeries_CanExecute(FGraphSeries* Series);
 
-	// Get the Y value that is used to provided a clipping border between adjacent 
-	// graph tracks
+	// Get the Y value that is used to provide a clipping border between adjacent graph tracks.
 	virtual float GetBorderY() const { return 0.0f; }
 
 private:
