Added a browsing widget for graph series to ease management of graph tracks with many series

diff --git a/Engine/Source/Developer/TraceInsights/Private/Insights/ViewModels/GraphTrack.cpp b/Engine/Source/Developer/TraceInsights/Private/Insights/ViewModels/GraphTrack.cpp
--- a/Engine/Source/Developer/TraceInsights/Private/Insights/ViewModels/GraphTrack.cpp
+++ b/Engine/Source/Developer/TraceInsights/Private/Insights/ViewModels/GraphTrack.cpp
@@ -2,11 +2,11 @@
 
 #include "Insights/ViewModels/GraphTrack.h"
 
-#include "EditorFontGlyphs.h"
 #include "EditorStyleSet.h"
 #include "Framework/MultiBox/MultiBoxBuilder.h"
 #include "Rendering/DrawElements.h"
 #include "Styling/CoreStyle.h"
+#include "Widgets/Layout/SBox.h"
 
 // Insights
 #include "Insights/Common/PaintUtils.h"
@@ -19,6 +19,7 @@
 #include "Insights/ViewModels/TimingEvent.h"
 #include "Insights/ViewModels/TimingTrackViewport.h"
 #include "Insights/ViewModels/TooltipDrawState.h"
+#include "Insights/Widgets/SGraphSeriesList.h"
 
 #define LOCTEXT_NAMESPACE "GraphTrack"
 
@@ -736,42 +737,13 @@ void FGraphTrack::BuildContextMenu(FMenuBuilder& MenuBuilder)
 
 	MenuBuilder.BeginSection("Series", LOCTEXT("ContextMenu_Header_Series", "Series"));
 	{
-		for (const TSharedPtr<FGraphSeries>& Series : AllSeries)
-		{
-			FUIAction Action_ShowSeries
-			(
-				FExecuteAction::CreateSP(this, &FGraphTrack::ContextMenu_ShowSeries_Execute, Series.Get()),
-				FCanExecuteAction::CreateSP(this, &FGraphTrack::ContextMenu_ShowSeries_CanExecute, Series.Get()),
-				FIsActionChecked::CreateSP(this, &FGraphTrack::ContextMenu_ShowSeries_IsChecked, Series.Get())
-			);
-			MenuBuilder.AddMenuEntry
-			(
-				Action_ShowSeries,
-				SNew(SHorizontalBox)
-				+ SHorizontalBox::Slot()
-				.AutoWidth()
-				.VAlign(VAlign_Center)
-				.Padding(2.0f, 0.0f)
-				[
-					SNew(STextBlock)
-					.Text(Series->GetName())
-				]
-				+ SHorizontalBox::Slot()
-				.FillWidth(1.0f)
-				.HAlign(HAlign_Right)
-				.VAlign(VAlign_Center)
-				.Padding(4.0f, 0.0f)
-				[
-					SNew(STextBlock)
-					.Font(FEditorStyle::Get().GetFontStyle("FontAwesome.9"))
-					.ColorAndOpacity_Lambda([SeriesPtr = Series.Get()](){ return FSlateColor(SeriesPtr->GetColor()); })
-					.Text(FEditorFontGlyphs::Circle)
-				],
-				NAME_None,
-				Series->GetDescription(),
-				EUserInterfaceActionType::ToggleButton
-			);
-		}
+		MenuBuilder.AddWidget(
+			SNew(SBox)
+			.MaxDesiredHeight(150.0f)
+			[
+				SNew(SGraphSeriesList, SharedThis(this))
+			], 
+			FText::GetEmpty(), true);
 	}
 	MenuBuilder.EndSection();
 }
diff --git a/Engine/Source/Developer/TraceInsights/Private/Insights/Widgets/SGraphSeriesList.cpp b/Engine/Source/Developer/TraceInsights/Private/Insights/Widgets/SGraphSeriesList.cpp
new file mode 100644
--- /dev/null
+++ b/Engine/Source/Developer/TraceInsights/Private/Insights/Widgets/SGraphSeriesList.cpp
@@ -0,0 +1,183 @@
+// Copyright 1998-2019 Epic Games, Inc. All Rights Reserved.
+
+#include "SGraphSeriesList.h"
+#include "Widgets/SBoxPanel.h"
+#include "Widgets/Input/SSearchBox.h"
+#include "Widgets/Text/STextBlock.h"
+#include "Widgets/Input/SCheckBox.h"
+#include "Widgets/Views/SHeaderRow.h"
+#include "Widgets/Layout/SScrollBorder.h"
+#include "Framework/Application/SlateApplication.h"
+#include "EditorStyleSet.h"
+#include "EditorFontGlyphs.h"
+
+#include "Insights/ViewModels/GraphSeries.h"
+#include "Insights/ViewModels/GraphTrack.h"
+
+#define LOCTEXT_NAMESPACE "SGraphSeriesList"
+
+// A list entry widget for a graph series
+class SGraphSeriesListEntry : public STableRow<TSharedPtr<FGraphSeries>>
+{
+public:
+	SLATE_BEGIN_ARGS(SGraphSeriesListEntry) {}
+
+	SLATE_ARGUMENT(TSharedPtr<FGraphSeries>, GraphSeries)
+
+	SLATE_ATTRIBUTE(FText, SearchText)
+
+	SLATE_END_ARGS()
+
+	void Construct(const FArguments& InArgs, const TSharedRef<STableViewBase>& InOwnerTable)
+	{
+		GraphSeries = InArgs._GraphSeries;
+		SearchText = InArgs._SearchText;
+
+		STableRow<TSharedPtr<FGraphSeries>>::Construct(STableRow<TSharedPtr<FGraphSeries>>::FArguments(), InOwnerTable);
+	}
+
+	virtual void ConstructChildren( ETableViewMode::Type InOwnerTableMode, const TAttribute<FMargin>& InPadding, const TSharedRef<SWidget>& InContent ) override
+	{
+		ChildSlot
+		[
+			SNew(SHorizontalBox)
+			+SHorizontalBox::Slot()
+			.AutoWidth()
+			.Padding(2.0f)
+			.HAlign(HAlign_Center)
+			.VAlign(VAlign_Center)
+			[
+				SNew(SCheckBox)
+				.IsChecked_Lambda([this](){ return GraphSeries.Pin()->IsVisible() ? ECheckBoxState::Checked : ECheckBoxState::Unchecked; })
+				.OnCheckStateChanged_Lambda([this](ECheckBoxState InCheckBoxState) { GraphSeries.Pin()->SetVisibility(InCheckBoxState == ECheckBoxState::Checked); })
+				.Content()
+				[
+					SNew(STextBlock)
+					.Text(GraphSeries.Pin()->GetName())
+					.HighlightText(SearchText)
+				]
+			]
+			+SHorizontalBox::Slot()
+			.FillWidth(1.0f)
+			.Padding(2.0f)
+			.HAlign(HAlign_Right)
+			.VAlign(VAlign_Center)
+			[
+				SNew(STextBlock)
+				.Font(FEditorStyle::Get().GetFontStyle("FontAwesome.9"))
+				.ColorAndOpacity_Lambda([this](){ return FSlateColor(GraphSeries.Pin()->GetColor()); })
+				.Text(FEditorFontGlyphs::Circle)
+			]
+		];
+	}
+
+	// The series we represent
+	TWeakPtr<FGraphSeries> GraphSeries;
+
+	// The search text to highlight
+	TAttribute<FText> SearchText;
+};
+
+void SGraphSeriesList::Construct(const FArguments& InArgs, const TSharedRef<FGraphTrack>& InGraphTrack)
+{
+	GraphTrack = InGraphTrack;
+
+	FilteredSeries = InGraphTrack->GetSeries();
+
+	ListView = SNew(SListView<TSharedPtr<FGraphSeries>>)
+		.IsFocusable(true)
+		.ItemHeight(20.0f)
+		.SelectionMode(ESelectionMode::None)
+		.ListItemsSource(&FilteredSeries)
+		.OnGenerateRow(this, &SGraphSeriesList::OnGenerateRow);
+
+	ChildSlot
+	[
+		SNew(SVerticalBox)
+		+SVerticalBox::Slot()
+		.AutoHeight()
+		[
+			SNew(SHorizontalBox)
+			+SHorizontalBox::Slot()
+			.AutoWidth()
+			.Padding(2.0f, 0.0f)
+			[
+				// Checkbox for bulk operations
+				SNew(SCheckBox)
+				.IsChecked_Lambda([this]()
+				{
+					int32 NumVisible = 0;
+					for(TSharedPtr<FGraphSeries> Series : FilteredSeries)
+					{
+						if(Series->IsVisible())
+						{
+							NumVisible++;
+						}
+					}
+
+					if(NumVisible == 0)
+					{
+						return ECheckBoxState::Unchecked;
+					}
+					else if(NumVisible == FilteredSeries.Num())
+					{
+						return ECheckBoxState::Checked;
+					}
+
+					return ECheckBoxState::Undetermined;
+				})
+				.OnCheckStateChanged_Lambda([this](ECheckBoxState InCheckBoxState) 
+				{
+					const bool bVisible = InCheckBoxState != ECheckBoxState::Unchecked;
+					for(TSharedPtr<FGraphSeries> Series : FilteredSeries)
+					{
+						Series->SetVisibility(bVisible);
+					}
+				})
+			]
+			+SHorizontalBox::Slot()
+			.FillWidth(1.0f)
+			[
+				// Search box allows for filtering
+				SAssignNew(SearchBox, SSearchBox)
+				.OnTextChanged_Lambda([this](const FText& InText){ SearchText = InText; RefreshFilter(); })
+			]
+		]
+		+SVerticalBox::Slot()
+		.FillHeight(1.0f)
+		[
+			SNew(SScrollBorder, ListView.ToSharedRef())
+			[
+				ListView.ToSharedRef()
+			]
+		]
+	];
+
+	// Set focus to the search box on creation
+	FSlateApplication::Get().SetKeyboardFocus(SearchBox);
+	FSlateApplication::Get().SetUserFocus(0, SearchBox);
+}
+
+TSharedRef<ITableRow> SGraphSeriesList::OnGenerateRow(TSharedPtr<FGraphSeries> Item, const TSharedRef<STableViewBase>& OwnerTable)
+{
+	return SNew(SGraphSeriesListEntry, OwnerTable)
+		.GraphSeries(Item)
+		.SearchText_Lambda([this](){ return SearchText; });
+}
+
+void SGraphSeriesList::RefreshFilter()
+{
+	FilteredSeries.Reset();
+
+	for(TSharedPtr<FGraphSeries> Series : GraphTrack.Pin()->GetSeries())
+	{
+		if(SearchText.IsEmpty() || Series->GetName().ToString().Contains(SearchText.ToString()))
+		{
+			FilteredSeries.Add(Series);
+		}
+	}
+
+	ListView->RequestListRefresh();
+}
+
+#undef LOCTEXT_NAMESPACE
\ No newline at end of file
diff --git a/Engine/Source/Developer/TraceInsights/Private/Insights/Widgets/SGraphSeriesList.h b/Engine/Source/Developer/TraceInsights/Private/Insights/Widgets/SGraphSeriesList.h
new file mode 100644
--- /dev/null
+++ b/Engine/Source/Developer/TraceInsights/Private/Insights/Widgets/SGraphSeriesList.h
@@ -0,0 +1,45 @@
+// Copyright 1998-2019 Epic Games, Inc. All Rights Reserved.
+
+#pragma once
+
+#include "CoreMinimal.h"
+#include "Widgets/SCompoundWidget.h"
+#include "Widgets/Views/SListView.h"
+
+class SSearchBox;
+class FGraphTrack;
+class FGraphSeries;
+
+// A widget representing all the series in a graph track, allowing management of their visibility
+class SGraphSeriesList : public SCompoundWidget
+{
+public:
+	SLATE_BEGIN_ARGS(SGraphSeriesList) {}
+
+	SLATE_END_ARGS()
+
+	void Construct(const FArguments& InArgs, const TSharedRef<FGraphTrack>& InGraphTrack);
+
+private:
+	// Generate a row for the list
+	TSharedRef<ITableRow> OnGenerateRow(TSharedPtr<FGraphSeries> Item, const TSharedRef<STableViewBase>& OwnerTable);
+
+	// Refresh the list as the filter has changed
+	void RefreshFilter();
+
+private:
+	// The graph track we are operating on
+	TWeakPtr<FGraphTrack> GraphTrack;
+
+	// The search widget
+	TSharedPtr<SSearchBox> SearchBox;
+
+	// The list view widget
+	TSharedPtr<SListView<TSharedPtr<FGraphSeries>>> ListView;
+
+	// Text we are searching for
+	FText SearchText;
+
+	// Filtered list of series
+	TArray<TSharedPtr<FGraphSeries>> FilteredSeries;
+};
