[Insights]
  - Implemented new simple RTTI mechanism for ITimingEvent, ITimingEventFilter and FBaseTimingTrack class hierarchies. It replaces GetType(), GetSubType() functionality.
  - Implemented new ITimingEventFilter classes (by track, by event type, by min duration, by max duration, by frame index).
  - Timing view: Added Game and Rendering Frame timing tracks to correlate frame index and duration of each frame with gpu/cpu timelines.
  - Added debug code to measure how long it takes to compute aggregation stats for Timers and Counters.

diff --git a/Engine/Plugins/Animation/GameplayInsights/Source/GameplayInsights/Private/AnimNodesTrack.cpp b/Engine/Plugins/Animation/GameplayInsights/Source/GameplayInsights/Private/AnimNodesTrack.cpp
--- a/Engine/Plugins/Animation/GameplayInsights/Source/GameplayInsights/Private/AnimNodesTrack.cpp
+++ b/Engine/Plugins/Animation/GameplayInsights/Source/GameplayInsights/Private/AnimNodesTrack.cpp
@@ -32,11 +32,10 @@
 
 #define LOCTEXT_NAMESPACE "AnimNodesTrack"
 
-const FName FAnimNodesTrack::TypeName(TEXT("Events"));
-const FName FAnimNodesTrack::SubTypeName(TEXT("Animation.AnimNodes"));
+INSIGHTS_IMPLEMENT_RTTI(FAnimNodesTrack)
 
 FAnimNodesTrack::FAnimNodesTrack(const FAnimationSharedData& InSharedData, uint64 InObjectID, const TCHAR* InName)
-	: TGameplayTrackMixin<FTimingEventsTrack>(InObjectID, FAnimNodesTrack::TypeName, FAnimNodesTrack::SubTypeName, FText::Format(LOCTEXT("TrackNameFormat", "Graph - {0}"), FText::FromString(FString(InName))))
+	: TGameplayTrackMixin<FTimingEventsTrack>(InObjectID, FText::Format(LOCTEXT("TrackNameFormat", "Graph - {0}"), FText::FromString(FString(InName))))
 	, SharedData(InSharedData)
 #if WITH_EDITOR
 	, InstanceClass(nullptr)
diff --git a/Engine/Plugins/Animation/GameplayInsights/Source/GameplayInsights/Private/AnimationSharedData.cpp b/Engine/Plugins/Animation/GameplayInsights/Source/GameplayInsights/Private/AnimationSharedData.cpp
--- a/Engine/Plugins/Animation/GameplayInsights/Source/GameplayInsights/Private/AnimationSharedData.cpp
+++ b/Engine/Plugins/Animation/GameplayInsights/Source/GameplayInsights/Private/AnimationSharedData.cpp
@@ -68,13 +68,7 @@ void FAnimationSharedData::Tick(Insights::ITimingViewSession& InTimingViewSessio
 
 				auto FindSkeletalMeshPoseTrack = [](const FBaseTimingTrack& InTrack)
 				{
-					if (InTrack.GetType() == FSkeletalMeshPoseTrack::TypeName &&
-						InTrack.GetSubType() == FSkeletalMeshPoseTrack::SubTypeName)
-					{
-						return true;
-					}
-
-					return false;
+					return InTrack.Is<FSkeletalMeshPoseTrack>();
 				};
 
 				TSharedPtr<FSkeletalMeshPoseTrack> ExistingSkeletalMeshPoseTrack = StaticCastSharedPtr<FSkeletalMeshPoseTrack>(ObjectEventsTrack->GetGameplayTrack().FindChildTrack(InObjectInfo.Id, FindSkeletalMeshPoseTrack));
@@ -97,10 +91,9 @@ void FAnimationSharedData::Tick(Insights::ITimingViewSession& InTimingViewSessio
 
 				auto FindTickRecordTrackWithAssetId = [&InAssetId](const FBaseTimingTrack& InTrack)
 				{
-					if (InTrack.GetType() == FAnimationTickRecordsTrack::TypeName &&
-						InTrack.GetSubType() == FAnimationTickRecordsTrack::SubTypeName)
+					if (InTrack.Is<FAnimationTickRecordsTrack>())
 					{
-						const FAnimationTickRecordsTrack& AnimationTickRecordsTrack = *static_cast<const FAnimationTickRecordsTrack*>(&InTrack);
+						const FAnimationTickRecordsTrack& AnimationTickRecordsTrack = InTrack.As<FAnimationTickRecordsTrack>();
 						return AnimationTickRecordsTrack.GetAssetId() == InAssetId;
 					}
 
@@ -129,13 +122,7 @@ void FAnimationSharedData::Tick(Insights::ITimingViewSession& InTimingViewSessio
 
 				auto FindAnimNodesTrack = [](const FBaseTimingTrack& InTrack)
 				{
-					if (InTrack.GetType() == FAnimNodesTrack::TypeName &&
-						InTrack.GetSubType() == FAnimNodesTrack::SubTypeName)
-					{
-						return true;
-					}
-
-					return false;
+					return InTrack.Is<FAnimNodesTrack>();
 				};
 
 				TSharedPtr<FAnimNodesTrack> ExistingAnimNodesTrack = StaticCastSharedPtr<FAnimNodesTrack>(ObjectEventsTrack->GetGameplayTrack().FindChildTrack(InObjectInfo.Id, FindAnimNodesTrack));
diff --git a/Engine/Plugins/Animation/GameplayInsights/Source/GameplayInsights/Private/AnimationTickRecordsTrack.cpp b/Engine/Plugins/Animation/GameplayInsights/Source/GameplayInsights/Private/AnimationTickRecordsTrack.cpp
--- a/Engine/Plugins/Animation/GameplayInsights/Source/GameplayInsights/Private/AnimationTickRecordsTrack.cpp
+++ b/Engine/Plugins/Animation/GameplayInsights/Source/GameplayInsights/Private/AnimationTickRecordsTrack.cpp
@@ -25,8 +25,7 @@
 
 #define LOCTEXT_NAMESPACE "AnimationTickRecordsTrack"
 
-const FName FAnimationTickRecordsTrack::TypeName(TEXT("Graph"));
-const FName FAnimationTickRecordsTrack::SubTypeName(TEXT("Animation.TickRecords"));
+INSIGHTS_IMPLEMENT_RTTI(FAnimationTickRecordsTrack)
 
 FString FTickRecordSeries::FormatValue(double Value) const
 {
@@ -68,7 +67,7 @@ static FLinearColor MakeSeriesColor(FTickRecordSeries::ESeriesType InSeed, bool
 }
 
 FAnimationTickRecordsTrack::FAnimationTickRecordsTrack(const FAnimationSharedData& InSharedData, uint64 InObjectID, uint64 InAssetId, int32 InNodeId, const TCHAR* InName)
-	: TGameplayTrackMixin<FGraphTrack>(InObjectID, FAnimationTickRecordsTrack::SubTypeName, MakeTrackName(InSharedData.GetGameplaySharedData(), InAssetId, InName))
+	: TGameplayTrackMixin<FGraphTrack>(InObjectID, MakeTrackName(InSharedData.GetGameplaySharedData(), InAssetId, InName))
 	, SharedData(InSharedData)
 	, AssetId(InAssetId)
 	, NodeId(InNodeId)
diff --git a/Engine/Plugins/Animation/GameplayInsights/Source/GameplayInsights/Private/AnimationTickRecordsTrack.h b/Engine/Plugins/Animation/GameplayInsights/Source/GameplayInsights/Private/AnimationTickRecordsTrack.h
--- a/Engine/Plugins/Animation/GameplayInsights/Source/GameplayInsights/Private/AnimationTickRecordsTrack.h
+++ b/Engine/Plugins/Animation/GameplayInsights/Source/GameplayInsights/Private/AnimationTickRecordsTrack.h
@@ -42,10 +42,9 @@ public:
 
 class FAnimationTickRecordsTrack : public TGameplayTrackMixin<FGraphTrack>
 {
-public:
-	static const FName TypeName;
-	static const FName SubTypeName;
+	INSIGHTS_DECLARE_RTTI(FAnimationTickRecordsTrack, TGameplayTrackMixin<FGraphTrack>)
 
+public:
 	FAnimationTickRecordsTrack(const FAnimationSharedData& InSharedData, uint64 InObjectId, uint64 InAssetId, int32 InNodeId, const TCHAR* InName);
 
 	virtual void PreUpdate(const ITimingTrackUpdateContext& Context) override;
diff --git a/Engine/Plugins/Animation/GameplayInsights/Source/GameplayInsights/Private/GameplaySharedData.cpp b/Engine/Plugins/Animation/GameplayInsights/Source/GameplayInsights/Private/GameplaySharedData.cpp
--- a/Engine/Plugins/Animation/GameplayInsights/Source/GameplayInsights/Private/GameplaySharedData.cpp
+++ b/Engine/Plugins/Animation/GameplayInsights/Source/GameplayInsights/Private/GameplaySharedData.cpp
@@ -84,7 +84,7 @@ static void UpdateTrackOrderRecursive(TSharedRef<FObjectEventsTrack> InTrack, in
 
 	for(FGameplayTrack* ChildTrack : InTrack->GetGameplayTrack().GetChildTracks())
 	{
-		if(ChildTrack->GetTimingTrack()->GetType() == FObjectEventsTrack::TypeName)
+		if(ChildTrack->GetTimingTrack()->Is<FObjectEventsTrack>())
 		{
 			ChildTrack->SetIndent(InTrack->GetGameplayTrack().GetIndent() + 1);
 			UpdateTrackOrderRecursive(StaticCastSharedRef<FObjectEventsTrack>(ChildTrack->GetTimingTrack()), InOutOrder);
@@ -93,7 +93,7 @@ static void UpdateTrackOrderRecursive(TSharedRef<FObjectEventsTrack> InTrack, in
 
 	for(FGameplayTrack* ChildTrack : InTrack->GetGameplayTrack().GetChildTracks())
 	{
-		if(ChildTrack->GetTimingTrack()->GetType() != FObjectEventsTrack::TypeName)
+		if(!ChildTrack->GetTimingTrack()->Is<FObjectEventsTrack>())
 		{
 			ChildTrack->SetIndent(InTrack->GetGameplayTrack().GetIndent() + 1);
 			ChildTrack->GetTimingTrack()->SetOrder(InOutOrder++);
diff --git a/Engine/Plugins/Animation/GameplayInsights/Source/GameplayInsights/Private/GameplayTrack.h b/Engine/Plugins/Animation/GameplayInsights/Source/GameplayInsights/Private/GameplayTrack.h
--- a/Engine/Plugins/Animation/GameplayInsights/Source/GameplayInsights/Private/GameplayTrack.h
+++ b/Engine/Plugins/Animation/GameplayInsights/Source/GameplayInsights/Private/GameplayTrack.h
@@ -64,15 +64,11 @@ private:
 template <class Base>
 class TGameplayTrackMixin : public Base
 {
-public:
-	TGameplayTrackMixin(uint64 InObjectId, const FName& InSubType, const FText& InName)
-		: Base(InSubType, InName.ToString())
-		, GameplayTrack(*this, InObjectId)
-	{
-	}
+	//INSIGHTS_DECLARE_RTTI(TGameplayTrackMixin, Base) -- Commented to skip this class in the Simple RTTI class hierarchy. It will get same type name as its Base class.
 
-	TGameplayTrackMixin(uint64 InObjectId, const FName& InType, const FName& InSubType, const FText& InName)
-		: Base(InType, InSubType, InName.ToString())
+public:
+	TGameplayTrackMixin(uint64 InObjectId, const FText& InName)
+		: Base(InName.ToString())
 		, GameplayTrack(*this, InObjectId)
 	{
 	}
@@ -85,4 +81,7 @@ public:
 
 private:
 	FGameplayTrack GameplayTrack;
-};
\ No newline at end of file
+};
+
+//template <class Base>
+//INSIGHTS_IMPLEMENT_RTTI(TGameplayTrackMixin<Base>) // Note: All templeted classes will return same type name as FName(TEXT"TGameplayTrackMixin<Base>") !!!
diff --git a/Engine/Plugins/Animation/GameplayInsights/Source/GameplayInsights/Private/ObjectEventsTrack.cpp b/Engine/Plugins/Animation/GameplayInsights/Source/GameplayInsights/Private/ObjectEventsTrack.cpp
--- a/Engine/Plugins/Animation/GameplayInsights/Source/GameplayInsights/Private/ObjectEventsTrack.cpp
+++ b/Engine/Plugins/Animation/GameplayInsights/Source/GameplayInsights/Private/ObjectEventsTrack.cpp
@@ -11,11 +11,10 @@
 
 #define LOCTEXT_NAMESPACE "ObjectEventsTrack"
 
-const FName FObjectEventsTrack::TypeName(TEXT("Events"));
-const FName FObjectEventsTrack::SubTypeName(TEXT("Gameplay.ObjectEvents"));
+INSIGHTS_IMPLEMENT_RTTI(FObjectEventsTrack)
 
 FObjectEventsTrack::FObjectEventsTrack(const FGameplaySharedData& InSharedData, uint64 InObjectID, const TCHAR* InName)
-	: TGameplayTrackMixin<FTimingEventsTrack>(InObjectID, TypeName, SubTypeName, MakeTrackName(InSharedData, InObjectID, InName))
+	: TGameplayTrackMixin<FTimingEventsTrack>(InObjectID, MakeTrackName(InSharedData, InObjectID, InName))
 	, SharedData(InSharedData)
 {
 }
diff --git a/Engine/Plugins/Animation/GameplayInsights/Source/GameplayInsights/Private/ObjectEventsTrack.h b/Engine/Plugins/Animation/GameplayInsights/Source/GameplayInsights/Private/ObjectEventsTrack.h
--- a/Engine/Plugins/Animation/GameplayInsights/Source/GameplayInsights/Private/ObjectEventsTrack.h
+++ b/Engine/Plugins/Animation/GameplayInsights/Source/GameplayInsights/Private/ObjectEventsTrack.h
@@ -11,10 +11,9 @@ class FTimingEventSearchParameters;
 
 class FObjectEventsTrack : public TGameplayTrackMixin<FTimingEventsTrack>
 {
-public:
-	static const FName TypeName;
-	static const FName SubTypeName;
+	INSIGHTS_DECLARE_RTTI(FObjectEventsTrack, TGameplayTrackMixin<FTimingEventsTrack>)
 
+public:
 	FObjectEventsTrack(const FGameplaySharedData& InSharedData, uint64 InObjectID, const TCHAR* InName);
 
 	virtual void BuildDrawState(ITimingEventsTrackDrawStateBuilder& Builder, const ITimingTrackUpdateContext& Context) override;
diff --git a/Engine/Plugins/Animation/GameplayInsights/Source/GameplayInsights/Private/SkeletalMeshPoseTrack.cpp b/Engine/Plugins/Animation/GameplayInsights/Source/GameplayInsights/Private/SkeletalMeshPoseTrack.cpp
--- a/Engine/Plugins/Animation/GameplayInsights/Source/GameplayInsights/Private/SkeletalMeshPoseTrack.cpp
+++ b/Engine/Plugins/Animation/GameplayInsights/Source/GameplayInsights/Private/SkeletalMeshPoseTrack.cpp
@@ -19,14 +19,12 @@
 #include "GameFramework/Actor.h"
 #endif
 
-
 #define LOCTEXT_NAMESPACE "SkeletalMeshPoseTrack"
 
-const FName FSkeletalMeshPoseTrack::TypeName(TEXT("Events"));
-const FName FSkeletalMeshPoseTrack::SubTypeName(TEXT("Animation.SkeletalMeshPose"));
+INSIGHTS_IMPLEMENT_RTTI(FSkeletalMeshPoseTrack)
 
 FSkeletalMeshPoseTrack::FSkeletalMeshPoseTrack(const FAnimationSharedData& InSharedData, uint64 InObjectID, const TCHAR* InName)
-	: TGameplayTrackMixin<FTimingEventsTrack>(InObjectID, FSkeletalMeshPoseTrack::TypeName, FSkeletalMeshPoseTrack::SubTypeName, FText::Format(LOCTEXT("TrackNameFormat", "Pose - {0}"), FText::FromString(FString(InName))))
+	: TGameplayTrackMixin<FTimingEventsTrack>(InObjectID, FText::Format(LOCTEXT("TrackNameFormat", "Pose - {0}"), FText::FromString(FString(InName))))
 	, SharedData(InSharedData)
 	, Color(FLinearColor::MakeRandomColor())
 	, bDrawPose(false)
diff --git a/Engine/Source/Developer/TraceInsights/Private/Insights/ViewModels/BaseTimingTrack.cpp b/Engine/Source/Developer/TraceInsights/Private/Insights/ViewModels/BaseTimingTrack.cpp
--- a/Engine/Source/Developer/TraceInsights/Private/Insights/ViewModels/BaseTimingTrack.cpp
+++ b/Engine/Source/Developer/TraceInsights/Private/Insights/ViewModels/BaseTimingTrack.cpp
@@ -5,6 +5,8 @@
 
 #define LOCTEXT_NAMESPACE "BaseTimingTrack"
 
+INSIGHTS_IMPLEMENT_RTTI(FBaseTimingTrack)
+
 void FBaseTimingTrack::BuildContextMenu(FMenuBuilder& MenuBuilder)
 {
 	MenuBuilder.BeginSection(TEXT("Empty"), FText::FromString(GetName()));
diff --git a/Engine/Source/Developer/TraceInsights/Private/Insights/ViewModels/FileActivityTimingTrack.cpp b/Engine/Source/Developer/TraceInsights/Private/Insights/ViewModels/FileActivityTimingTrack.cpp
--- a/Engine/Source/Developer/TraceInsights/Private/Insights/ViewModels/FileActivityTimingTrack.cpp
+++ b/Engine/Source/Developer/TraceInsights/Private/Insights/ViewModels/FileActivityTimingTrack.cpp
@@ -518,11 +518,15 @@ void FFileActivitySharedState::ToggleBackgroundEvents()
 // FFileActivityTimingTrack
 ////////////////////////////////////////////////////////////////////////////////////////////////////
 
+INSIGHTS_IMPLEMENT_RTTI(FFileActivityTimingTrack)
+
+////////////////////////////////////////////////////////////////////////////////////////////////////
+
 void FFileActivityTimingTrack::InitTooltip(FTooltipDrawState& InOutTooltip, const ITimingEvent& InTooltipEvent) const
 {
-	if (InTooltipEvent.CheckTrack(this) && FTimingEvent::CheckTypeName(InTooltipEvent))
+	if (InTooltipEvent.CheckTrack(this) && InTooltipEvent.Is<FTimingEvent>())
 	{
-		const FTimingEvent& TooltipEvent = static_cast<const FTimingEvent&>(InTooltipEvent);
+		const FTimingEvent& TooltipEvent = InTooltipEvent.As<FTimingEvent>();
 
 		auto MatchEvent = [&TooltipEvent](double InStartTime, double InEndTime, uint32 InDepth)
 		{
diff --git a/Engine/Source/Developer/TraceInsights/Private/Insights/ViewModels/FileActivityTimingTrack.h b/Engine/Source/Developer/TraceInsights/Private/Insights/ViewModels/FileActivityTimingTrack.h
--- a/Engine/Source/Developer/TraceInsights/Private/Insights/ViewModels/FileActivityTimingTrack.h
+++ b/Engine/Source/Developer/TraceInsights/Private/Insights/ViewModels/FileActivityTimingTrack.h
@@ -84,9 +84,11 @@ private:
 
 class FFileActivityTimingTrack : public FTimingEventsTrack
 {
+	INSIGHTS_DECLARE_RTTI(FFileActivityTimingTrack, FTimingEventsTrack)
+
 public:
-	explicit FFileActivityTimingTrack(FFileActivitySharedState& InSharedState, const FName SubType, const FString& InName)
-		: FTimingEventsTrack(FName(TEXT("FileActivity")), SubType, InName)
+	explicit FFileActivityTimingTrack(FFileActivitySharedState& InSharedState, const FString& InName)
+		: FTimingEventsTrack(InName)
 		, SharedState(InSharedState)
 	{
 	}
@@ -107,7 +109,7 @@ class FOverviewFileActivityTimingTrack : public FFileActivityTimingTrack
 {
 public:
 	explicit FOverviewFileActivityTimingTrack(FFileActivitySharedState& InSharedState)
-		: FFileActivityTimingTrack(InSharedState, FName(TEXT("Overview")), TEXT("I/O Overview"))
+		: FFileActivityTimingTrack(InSharedState, TEXT("I/O Overview"))
 	{
 	}
 
@@ -121,7 +123,7 @@ class FDetailedFileActivityTimingTrack : public FFileActivityTimingTrack
 {
 public:
 	explicit FDetailedFileActivityTimingTrack(FFileActivitySharedState& InSharedState)
-		: FFileActivityTimingTrack(InSharedState, FName(TEXT("Detailed")), TEXT("I/O Activity"))
+		: FFileActivityTimingTrack(InSharedState, TEXT("I/O Activity"))
 	{
 	}
 
diff --git a/Engine/Source/Developer/TraceInsights/Private/Insights/ViewModels/FrameTimingTrack.cpp b/Engine/Source/Developer/TraceInsights/Private/Insights/ViewModels/FrameTimingTrack.cpp
new file mode 100644
--- /dev/null
+++ b/Engine/Source/Developer/TraceInsights/Private/Insights/ViewModels/FrameTimingTrack.cpp
@@ -0,0 +1,437 @@
+// Copyright 1998-2019 Epic Games, Inc. All Rights Reserved.
+
+#include "FrameTimingTrack.h"
+
+#include "Fonts/FontMeasure.h"
+#include "Framework/MultiBox/MultiBoxBuilder.h"
+#include "HAL/PlatformApplicationMisc.h"
+#include "Styling/SlateBrush.h"
+
+// Insights
+#include "Insights/Common/PaintUtils.h"
+#include "Insights/Common/TimeUtils.h"
+#include "Insights/InsightsManager.h"
+#include "Insights/ITimingViewSession.h"
+#include "Insights/ViewModels/FrameTrackHelper.h"
+#include "Insights/ViewModels/TimerNode.h"
+#include "Insights/ViewModels/TimingEvent.h"
+#include "Insights/ViewModels/TimingEventSearch.h"
+#include "Insights/ViewModels/TimingTrackViewport.h"
+#include "Insights/ViewModels/TooltipDrawState.h"
+#include "Insights/Widgets/STimingView.h"
+
+#define LOCTEXT_NAMESPACE "FrameTimingTrack"
+
+////////////////////////////////////////////////////////////////////////////////////////////////////
+// FFrameSharedState
+////////////////////////////////////////////////////////////////////////////////////////////////////
+
+TSharedPtr<FFrameTimingTrack> FFrameSharedState::GetFrameTrack(uint32 InFrameType)
+{
+	TSharedPtr<FFrameTimingTrack>*const TrackPtrPtr = FrameTracks.Find(InFrameType);
+	return TrackPtrPtr ? *TrackPtrPtr : nullptr;
+}
+
+////////////////////////////////////////////////////////////////////////////////////////////////////
+
+bool FFrameSharedState::IsFrameTrackVisible(uint32 InFrameType) const
+{
+	const TSharedPtr<FFrameTimingTrack>*const TrackPtrPtr = FrameTracks.Find(InFrameType);
+	return TrackPtrPtr && (*TrackPtrPtr)->IsVisible();
+}
+
+////////////////////////////////////////////////////////////////////////////////////////////////////
+
+void FFrameSharedState::OnBeginSession(Insights::ITimingViewSession& InSession)
+{
+	if (&InSession != TimingView)
+	{
+		return;
+	}
+
+	bShowHideAllFrameTracks = false;
+	FrameTracks.Reset();
+}
+
+////////////////////////////////////////////////////////////////////////////////////////////////////
+
+void FFrameSharedState::OnEndSession(Insights::ITimingViewSession& InSession)
+{
+	if (&InSession != TimingView)
+	{
+		return;
+	}
+
+	bShowHideAllFrameTracks = false;
+	FrameTracks.Reset();
+}
+
+////////////////////////////////////////////////////////////////////////////////////////////////////
+
+void FFrameSharedState::Tick(Insights::ITimingViewSession& InSession, const Trace::IAnalysisSession& InAnalysisSession)
+{
+	if (&InSession != TimingView)
+	{
+		return;
+	}
+
+	{
+		bool bTracksOrderChanged = false;
+		int32 Order = 10;
+
+		// Iterate through frame types.
+		for (uint32 FrameType = 0; FrameType < TraceFrameType_Count; ++FrameType)
+		{
+			TSharedPtr<FFrameTimingTrack>* TrackPtrPtr = FrameTracks.Find(FrameType);
+			if (!TrackPtrPtr)
+			{
+				const FString TrackName = FString::Printf(TEXT("%s Frames"), FFrameTrackDrawHelper::FrameTypeToString(FrameType));
+				TSharedPtr<FFrameTimingTrack> Track = MakeShared<FFrameTimingTrack>(*this, TrackName, FrameType);
+				Track->SetOrder(Order);
+				FrameTracks.Add(FrameType, Track);
+				Track->SetVisibilityFlag(bShowHideAllFrameTracks);
+				InSession.AddScrollableTrack(Track);
+			}
+			else
+			{
+				TSharedPtr<FFrameTimingTrack> Track = *TrackPtrPtr;
+				if (Track->GetOrder() != Order)
+				{
+					Track->SetOrder(Order);
+					bTracksOrderChanged = true;
+				}
+			}
+
+			Order += 10;
+		}
+
+		if (bTracksOrderChanged)
+		{
+			InSession.InvalidateScrollableTracksOrder();
+		}
+	}
+}
+
+////////////////////////////////////////////////////////////////////////////////////////////////////
+
+void FFrameSharedState::ExtendFilterMenu(Insights::ITimingViewSession& InSession, FMenuBuilder& InOutMenuBuilder)
+{
+	if (&InSession != TimingView)
+	{
+		return;
+	}
+
+	InOutMenuBuilder.BeginSection("FrameTracks", LOCTEXT("FrameTracksHeading", "Frames"));
+	{
+		//TODO: MenuBuilder.AddMenuEntry(Commands.ShowAllFrameTracks);
+		InOutMenuBuilder.AddMenuEntry(
+			LOCTEXT("ShowAllFrameTracks", "Frame Tracks - R"),
+			LOCTEXT("ShowAllFrameTracks_Tooltip", "Show/hide all Frame tracks"),
+			FSlateIcon(),
+			FUIAction(FExecuteAction::CreateSP(this, &FFrameSharedState::ShowHideAllFrameTracks_Execute),
+					  FCanExecuteAction(),
+					  FIsActionChecked::CreateSP(this, &FFrameSharedState::ShowHideAllFrameTracks_IsChecked)),
+			NAME_None,
+			EUserInterfaceActionType::ToggleButton
+		);
+	}
+	InOutMenuBuilder.EndSection();
+}
+
+////////////////////////////////////////////////////////////////////////////////////////////////////
+
+bool FFrameSharedState::ShowHideAllFrameTracks_IsChecked() const
+{
+	return bShowHideAllFrameTracks;
+}
+
+////////////////////////////////////////////////////////////////////////////////////////////////////
+
+void FFrameSharedState::ShowHideAllFrameTracks_Execute()
+{
+	bShowHideAllFrameTracks = !bShowHideAllFrameTracks;
+
+	for (const auto& KV : FrameTracks)
+	{
+		FFrameTimingTrack& Track = *KV.Value;
+		Track.SetVisibilityFlag(bShowHideAllFrameTracks);
+	}
+
+	if (TimingView)
+	{
+		TimingView->OnTrackVisibilityChanged();
+	}
+}
+
+////////////////////////////////////////////////////////////////////////////////////////////////////
+// FFrameTimingTrack
+////////////////////////////////////////////////////////////////////////////////////////////////////
+
+INSIGHTS_IMPLEMENT_RTTI(FFrameTimingTrack)
+
+////////////////////////////////////////////////////////////////////////////////////////////////////
+
+void FFrameTimingTrack::BuildDrawState(ITimingEventsTrackDrawStateBuilder& Builder, const ITimingTrackUpdateContext& Context)
+{
+	TSharedPtr<const Trace::IAnalysisSession> Session = FInsightsManager::Get()->GetSession();
+	if (Session.IsValid())
+	{
+		Trace::FAnalysisSessionReadScope SessionReadScope(*Session.Get());
+		const Trace::IFrameProvider& FramesProvider = Trace::ReadFrameProvider(*Session.Get());
+		const FTimingTrackViewport& Viewport = Context.GetViewport();
+
+		uint64 StartIndex = 0;
+		uint64 EndIndex = FramesProvider.GetFrameCount(static_cast<ETraceFrameType>(FrameType));
+
+		const uint32 Color = FFrameTrackDrawHelper::GetColor32ByFrameType(FrameType);
+
+		FramesProvider.EnumerateFrames(static_cast<ETraceFrameType>(FrameType), StartIndex, EndIndex, [this, Color, &Builder](const Trace::FFrame& Frame)
+		{
+			constexpr uint32 Depth = 0;
+			const FString FrameName = GetShortFrameName(Frame.Index);
+			const uint64 IndexAsFrameType = Frame.Index;
+			Builder.AddEvent(Frame.StartTime, Frame.EndTime, Depth, *FrameName, IndexAsFrameType, Color);
+		});
+	}
+}
+
+////////////////////////////////////////////////////////////////////////////////////////////////////
+
+void FFrameTimingTrack::BuildFilteredDrawState(ITimingEventsTrackDrawStateBuilder& Builder, const ITimingTrackUpdateContext& Context)
+{
+	const TSharedPtr<ITimingEventFilter> EventFilterPtr = Context.GetEventFilter();
+	if (EventFilterPtr.IsValid() && EventFilterPtr->FilterTrack(*this))
+	{
+		TSharedPtr<const Trace::IAnalysisSession> Session = FInsightsManager::Get()->GetSession();
+		if (Session.IsValid())
+		{
+			Trace::FAnalysisSessionReadScope SessionReadScope(*Session.Get());
+			const Trace::IFrameProvider& FramesProvider = Trace::ReadFrameProvider(*Session.Get());
+			const FTimingTrackViewport& Viewport = Context.GetViewport();
+
+			uint64 StartIndex = 0;
+			uint64 EndIndex = FramesProvider.GetFrameCount(static_cast<ETraceFrameType>(FrameType));
+
+			const uint32 Color = FFrameTrackDrawHelper::GetColor32ByFrameType(FrameType);
+
+			FramesProvider.EnumerateFrames(static_cast<ETraceFrameType>(FrameType), StartIndex, EndIndex, [this, Color, &Builder, &EventFilterPtr](const Trace::FFrame& Frame)
+			{
+				constexpr uint32 Depth = 0;
+				const uint64 IndexAsEventType = Frame.Index;
+
+				FTimingEvent TimingEvent(SharedThis(this), Frame.StartTime, Frame.EndTime, Depth, IndexAsEventType);
+
+				if (EventFilterPtr->FilterEvent(TimingEvent))
+				{
+					const FString FrameName = GetShortFrameName(Frame.Index);
+					Builder.AddEvent(Frame.StartTime, Frame.EndTime, Depth, *FrameName, IndexAsEventType, Color);
+				}
+			});
+		}
+	}
+}
+
+////////////////////////////////////////////////////////////////////////////////////////////////////
+
+void FFrameTimingTrack::PostDraw(const ITimingTrackDrawContext& Context) const
+{
+	const TSharedPtr<const ITimingEvent> SelectedEventPtr = Context.GetSelectedEvent();
+	if (SelectedEventPtr.IsValid() &&
+		SelectedEventPtr->CheckTrack(this) &&
+		SelectedEventPtr->Is<FTimingEvent>())
+	{
+		const FTimingEvent& SelectedEvent = SelectedEventPtr->As<FTimingEvent>();
+		const ITimingViewDrawHelper& Helper = Context.GetHelper();
+		DrawSelectedEventInfo(SelectedEvent, Context.GetViewport(), Context.GetDrawContext(), Helper.GetWhiteBrush(), Helper.GetEventFont());
+	}
+}
+
+////////////////////////////////////////////////////////////////////////////////////////////////////
+
+void FFrameTimingTrack::DrawSelectedEventInfo(const FTimingEvent& SelectedEvent, const FTimingTrackViewport& Viewport, const FDrawContext& DrawContext, const FSlateBrush* WhiteBrush, const FSlateFontInfo& Font) const
+{
+	FindFrame(SelectedEvent, [this, &SelectedEvent, &Font, &Viewport, &DrawContext, &WhiteBrush](double InFoundStartTime, double InFoundEndTime, uint32 InFoundDepth, const Trace::FFrame& InFoundFrame)
+	{
+		//const double Duration = SelectedEvent.GetDuration();
+		const double Duration = InFoundFrame.EndTime - InFoundFrame.StartTime;
+		const FString Str = GetCompleteFrameName(InFoundFrame.Index, Duration);
+
+		const TSharedRef<FSlateFontMeasure> FontMeasureService = FSlateApplication::Get().GetRenderer()->GetFontMeasureService();
+		const FVector2D Size = FontMeasureService->Measure(Str, Font);
+		const float X = Viewport.GetWidth() - Size.X - 23.0f;
+		const float Y = Viewport.GetHeight() - Size.Y - 18.0f;
+
+		const FLinearColor BackgroundColor(0.05f, 0.05f, 0.05f, 1.0f);
+		const FLinearColor TextColor(0.7f, 0.7f, 0.7f, 1.0f);
+
+		DrawContext.DrawBox(X - 8.0f, Y - 2.0f, Size.X + 16.0f, Size.Y + 4.0f, WhiteBrush, BackgroundColor);
+		DrawContext.LayerId++;
+
+		DrawContext.DrawText(X, Y, Str, Font, TextColor);
+		DrawContext.LayerId++;
+	});
+}
+
+////////////////////////////////////////////////////////////////////////////////////////////////////
+
+void FFrameTimingTrack::InitTooltip(FTooltipDrawState& InOutTooltip, const ITimingEvent& InTooltipEvent) const
+{
+	if (InTooltipEvent.CheckTrack(this) && InTooltipEvent.Is<FTimingEvent>())
+	{
+		const FTimingEvent& TooltipEvent = InTooltipEvent.As<FTimingEvent>();
+
+		FindFrame(TooltipEvent, [this, &InOutTooltip, &TooltipEvent](double InFoundStartTime, double InFoundEndTime, uint32 InFoundDepth, const Trace::FFrame& InFoundFrame)
+		{
+			InOutTooltip.ResetContent();
+
+			const FString FrameName = GetFrameName(InFoundFrame.Index);
+			InOutTooltip.AddTitle(FrameName);
+
+			//const double Duration = TooltipEvent.GetDuration();
+			const double Duration = InFoundFrame.EndTime - InFoundFrame.StartTime;
+			InOutTooltip.AddNameValueTextLine(TEXT("Duration:"), TimeUtils::FormatTimeAuto(Duration));
+
+			InOutTooltip.UpdateLayout();
+		});
+	}
+}
+
+////////////////////////////////////////////////////////////////////////////////////////////////////
+
+const TSharedPtr<const ITimingEvent> FFrameTimingTrack::SearchEvent(const FTimingEventSearchParameters& InSearchParameters) const
+{
+	TSharedPtr<const ITimingEvent> FoundEvent;
+
+	FindFrame(InSearchParameters, [this, &FoundEvent](double InFoundStartTime, double InFoundEndTime, uint32 InFoundDepth, const Trace::FFrame& InFoundFrame)
+	{
+		const uint64 IndexAsEventType = InFoundFrame.Index; // storing the frame index as event type
+		FoundEvent = MakeShared<FTimingEvent>(SharedThis(this), InFoundStartTime, InFoundEndTime, InFoundDepth, IndexAsEventType);
+	});
+
+	return FoundEvent;
+}
+
+////////////////////////////////////////////////////////////////////////////////////////////////////
+
+void FFrameTimingTrack::OnEventSelected(const ITimingEvent& InSelectedEvent) const
+{
+	//if (InSelectedEvent.CheckTrack(this) && InSelectedEvent.Is<FTimingEvent>())
+	//{
+	//	const FTimingEvent& TrackEvent = InSelectedEvent.As<FTimingEvent>();
+	//	const uint64 FrameIndex = TrackEvent.GetType(); // the frame index was stored as event type
+	//}
+}
+
+////////////////////////////////////////////////////////////////////////////////////////////////////
+
+void FFrameTimingTrack::OnClipboardCopyEvent(const ITimingEvent& InSelectedEvent) const
+{
+	if (InSelectedEvent.CheckTrack(this) && InSelectedEvent.Is<FTimingEvent>())
+	{
+		const FTimingEvent& TrackEvent = InSelectedEvent.As<FTimingEvent>();
+
+		const uint64 FrameIndex = TrackEvent.GetType(); // using event type as frame index
+		const FString FrameName = GetFrameName(FrameIndex);
+
+		// Copy name of selected frame to clipboard.
+		FPlatformApplicationMisc::ClipboardCopy(*FrameName);
+	}
+}
+
+////////////////////////////////////////////////////////////////////////////////////////////////////
+
+void FFrameTimingTrack::BuildContextMenu(FMenuBuilder& MenuBuilder)
+{
+	FText SectionTitle = FText::FromString(GetName());
+	MenuBuilder.BeginSection(TEXT("Empty"), SectionTitle);
+	{
+		MenuBuilder.AddMenuEntry(
+			LOCTEXT("ContextMenu_NA", "N/A"),
+			LOCTEXT("ContextMenu_NA_Desc", "No actions available."),
+			FSlateIcon(),
+			FUIAction(FExecuteAction(), FCanExecuteAction::CreateLambda([](){ return false; })),
+			NAME_None,
+			EUserInterfaceActionType::Button
+		);
+	}
+	MenuBuilder.EndSection();
+}
+
+////////////////////////////////////////////////////////////////////////////////////////////////////
+
+bool FFrameTimingTrack::FindFrame(const FTimingEvent& InTimingEvent, TFunctionRef<void(double, double, uint32, const Trace::FFrame&)> InFoundPredicate) const
+{
+	auto MatchEvent = [&InTimingEvent](double InStartTime, double InEndTime, uint32 InDepth)
+	{ 
+		return InDepth == InTimingEvent.GetDepth()
+			&& InStartTime == InTimingEvent.GetStartTime()
+			&& InEndTime == InTimingEvent.GetEndTime();
+	};
+
+	FTimingEventSearchParameters SearchParameters(InTimingEvent.GetStartTime(), InTimingEvent.GetEndTime(), ETimingEventSearchFlags::StopAtFirstMatch, MatchEvent);
+	SearchParameters.SearchHandle = &InTimingEvent.GetSearchHandle();
+	return FindFrame(SearchParameters, InFoundPredicate);
+}
+
+////////////////////////////////////////////////////////////////////////////////////////////////////
+
+bool FFrameTimingTrack::FindFrame(const FTimingEventSearchParameters& InParameters, TFunctionRef<void(double, double, uint32, const Trace::FFrame&)> InFoundPredicate) const
+{
+	return TTimingEventSearch<Trace::FFrame>::Search(
+		InParameters,
+
+		[this](TTimingEventSearch<Trace::FFrame>::FContext& InContext)
+		{
+			TSharedPtr<const Trace::IAnalysisSession> Session = FInsightsManager::Get()->GetSession();
+			if (Session.IsValid())
+			{
+				Trace::FAnalysisSessionReadScope SessionReadScope(*Session.Get());
+				const Trace::IFrameProvider& FramesProvider = Trace::ReadFrameProvider(*Session.Get());
+
+				uint64 StartIndex = 0;
+				uint64 EndIndex = FramesProvider.GetFrameCount(static_cast<ETraceFrameType>(FrameType));
+
+				FramesProvider.EnumerateFrames(static_cast<ETraceFrameType>(FrameType), StartIndex, EndIndex, [&InContext](const Trace::FFrame& Frame)
+				{
+					if (Frame.StartTime < InContext.GetParameters().EndTime &&
+						Frame.EndTime > InContext.GetParameters().StartTime)
+					{
+						constexpr uint32 Depth = 0;
+						InContext.Check(Frame.StartTime, Frame.EndTime, Depth, Frame);
+					}
+				});
+			}
+		},
+
+		[&InFoundPredicate](double InFoundStartTime, double InFoundEndTime, uint32 InFoundDepth, const Trace::FFrame& InEvent)
+		{
+			InFoundPredicate(InFoundStartTime, InFoundEndTime, InFoundDepth, InEvent);
+		},
+		
+		SearchCache);
+}
+
+////////////////////////////////////////////////////////////////////////////////////////////////////
+
+const FString FFrameTimingTrack::GetShortFrameName(const uint64 InFrameIndex) const
+{
+	return FString::Printf(TEXT("%d"), InFrameIndex);
+}
+
+////////////////////////////////////////////////////////////////////////////////////////////////////
+
+const FString FFrameTimingTrack::GetFrameName(const uint64 InFrameIndex) const
+{
+	return FString::Printf(TEXT("%s Frame %d"), FFrameTrackDrawHelper::FrameTypeToString(FrameType), InFrameIndex);
+}
+
+////////////////////////////////////////////////////////////////////////////////////////////////////
+
+const FString FFrameTimingTrack::GetCompleteFrameName(const uint64 InFrameIndex, const double InFrameDuration) const
+{
+	return FString::Printf(TEXT("%s Frame %d (%s)"), FFrameTrackDrawHelper::FrameTypeToString(FrameType), InFrameIndex, *TimeUtils::FormatTimeAuto(InFrameDuration));
+}
+ 
+////////////////////////////////////////////////////////////////////////////////////////////////////
+
+#undef LOCTEXT_NAMESPACE
diff --git a/Engine/Source/Developer/TraceInsights/Private/Insights/ViewModels/FrameTimingTrack.h b/Engine/Source/Developer/TraceInsights/Private/Insights/ViewModels/FrameTimingTrack.h
new file mode 100644
--- /dev/null
+++ b/Engine/Source/Developer/TraceInsights/Private/Insights/ViewModels/FrameTimingTrack.h
@@ -0,0 +1,110 @@
+// Copyright 1998-2019 Epic Games, Inc. All Rights Reserved.
+
+#pragma once
+
+#include "CoreMinimal.h"
+#include "Fonts/SlateFontInfo.h"
+#include "TraceServices/Model/Frames.h"
+
+// Insights
+#include "Insights/ITimingViewExtender.h"
+#include "Insights/ViewModels/TimingEventSearch.h" // for TTimingEventSearchCache
+#include "Insights/ViewModels/TimingEventsTrack.h"
+
+class FTimingEvent;
+class FTimingEventSearchParameters;
+class FFrameTimingTrack;
+class STimingView;
+struct FSlateBrush;
+
+////////////////////////////////////////////////////////////////////////////////////////////////////
+
+class FFrameSharedState : public Insights::ITimingViewExtender, public TSharedFromThis<FFrameSharedState>
+{
+public:
+	explicit FFrameSharedState(STimingView* InTimingView) : TimingView(InTimingView) {}
+	virtual ~FFrameSharedState() = default;
+
+	TSharedPtr<FFrameTimingTrack> GetFrameTrack(uint32 InFrameType);
+
+	bool IsFrameTrackVisible(uint32 InFrameType) const;
+
+	//////////////////////////////////////////////////
+	// ITimingViewExtender interface
+
+	virtual void OnBeginSession(Insights::ITimingViewSession& InSession) override;
+	virtual void OnEndSession(Insights::ITimingViewSession& InSession) override;
+	virtual void Tick(Insights::ITimingViewSession& InSession, const Trace::IAnalysisSession& InAnalysisSession) override;
+	virtual void ExtendFilterMenu(Insights::ITimingViewSession& InSession, FMenuBuilder& InMenuBuilder) override;
+
+	//////////////////////////////////////////////////
+
+	void ShowHideAllFrameTracks() { ShowHideAllFrameTracks_Execute(); }
+
+private:
+	bool ShowHideAllFrameTracks_IsChecked() const;
+	void ShowHideAllFrameTracks_Execute();
+
+private:
+	STimingView* TimingView;
+
+	bool bShowHideAllFrameTracks;
+
+	/** Maps frame type to track pointer. */
+	TMap<uint32, TSharedPtr<FFrameTimingTrack>> FrameTracks;
+};
+
+////////////////////////////////////////////////////////////////////////////////////////////////////
+
+class FFrameTimingTrack : public FTimingEventsTrack
+{
+	INSIGHTS_DECLARE_RTTI(FFrameTimingTrack, FTimingEventsTrack)
+
+public:
+	explicit FFrameTimingTrack(FFrameSharedState& InSharedState, const FString& InName, uint32 InFrameType)
+		: FTimingEventsTrack(InName)
+		, SharedState(InSharedState)
+		, FrameType(InFrameType)
+	{
+	}
+
+	virtual ~FFrameTimingTrack() {}
+
+	//////////////////////////////////////////////////
+
+	uint32 GetFrameType() const { return FrameType; }
+	//void SetFrameType(uint32 InFrameType) { FrameType = InFrameType; }
+
+	virtual void BuildDrawState(ITimingEventsTrackDrawStateBuilder& Builder, const ITimingTrackUpdateContext& Context) override;
+	virtual void BuildFilteredDrawState(ITimingEventsTrackDrawStateBuilder& Builder, const ITimingTrackUpdateContext& Context) override;
+
+	virtual void PostDraw(const ITimingTrackDrawContext& Context) const override;
+
+	virtual void InitTooltip(FTooltipDrawState& InOutTooltip, const ITimingEvent& InTooltipEvent) const override;
+
+	virtual const TSharedPtr<const ITimingEvent> SearchEvent(const FTimingEventSearchParameters& InSearchParameters) const override;
+
+	virtual void OnEventSelected(const ITimingEvent& InSelectedEvent) const override;
+	virtual void OnClipboardCopyEvent(const ITimingEvent& InSelectedEvent) const override;
+	virtual void BuildContextMenu(FMenuBuilder& MenuBuilder) override;
+
+	const FString GetShortFrameName(const uint64 FrameIndex) const;
+	const FString GetFrameName(const uint64 FrameIndex) const;
+	const FString GetCompleteFrameName(const uint64 FrameIndex, const double FrameDuration) const;
+
+private:
+	void DrawSelectedEventInfo(const FTimingEvent& SelectedEvent, const FTimingTrackViewport& Viewport, const FDrawContext& DrawContext, const FSlateBrush* WhiteBrush, const FSlateFontInfo& Font) const;
+
+	bool FindFrame(const FTimingEvent& InTimingEvent, TFunctionRef<void(double, double, uint32, const Trace::FFrame&)> InFoundPredicate) const;
+	bool FindFrame(const FTimingEventSearchParameters& InParameters, TFunctionRef<void(double, double, uint32, const Trace::FFrame&)> InFoundPredicate) const;
+
+private:
+	FFrameSharedState& SharedState;
+
+	uint32 FrameType; // ETraceFrameType
+
+	// Search cache
+	mutable TTimingEventSearchCache<Trace::FFrame> SearchCache;
+};
+
+////////////////////////////////////////////////////////////////////////////////////////////////////
diff --git a/Engine/Source/Developer/TraceInsights/Private/Insights/ViewModels/FrameTrackHelper.cpp b/Engine/Source/Developer/TraceInsights/Private/Insights/ViewModels/FrameTrackHelper.cpp
--- a/Engine/Source/Developer/TraceInsights/Private/Insights/ViewModels/FrameTrackHelper.cpp
+++ b/Engine/Source/Developer/TraceInsights/Private/Insights/ViewModels/FrameTrackHelper.cpp
@@ -103,22 +103,38 @@ void FFrameTrackDrawHelper::DrawBackground() const
 
 ////////////////////////////////////////////////////////////////////////////////////////////////////
 
-FLinearColor FFrameTrackDrawHelper::GetColorByFrameType(int32 FrameType)
+const TCHAR* FFrameTrackDrawHelper::FrameTypeToString(int32 FrameType)
 {
-	const float Alpha = 0.9f;
 	switch (FrameType)
 	{
-	case TraceFrameType_Game:
-		return FLinearColor(0.75, 1.0, 1.0, Alpha);
+	case TraceFrameType_Game:      return TEXT("Game");
+	case TraceFrameType_Rendering: return TEXT("Rendering");
+	default:                       return TEXT("Misc");
+	}
+}
 
-	case TraceFrameType_Rendering:
-		return FLinearColor(1.0, 0.75, 0.75, Alpha);
+////////////////////////////////////////////////////////////////////////////////////////////////////
 
-	//case 2:
-	//	return FLinearColor(0.75, 0.75, 1.0, Alpha);
+uint32 FFrameTrackDrawHelper::GetColor32ByFrameType(int32 FrameType)
+{
+	switch (FrameType)
+	{
+	case TraceFrameType_Game	:	return 0xFF5555FF;
+	case TraceFrameType_Rendering:	return 0xFFFF5555;
+	default:						return 0xFF666666;
+	}
+}
+
+////////////////////////////////////////////////////////////////////////////////////////////////////
 
-	default:
-		return FLinearColor(1.0, 1.0, 1.0, Alpha);
+FLinearColor FFrameTrackDrawHelper::GetColorByFrameType(int32 FrameType)
+{
+	constexpr float Alpha = 0.9f;
+	switch (FrameType)
+	{
+	case TraceFrameType_Game:		return FLinearColor(0.75f, 1.0f, 1.0f, Alpha);
+	case TraceFrameType_Rendering:	return FLinearColor(1.0f, 0.75f, 0.75f, Alpha);
+	default:						return FLinearColor(1.0f, 1.0f, 1.0f, Alpha);
 	}
 }
 
diff --git a/Engine/Source/Developer/TraceInsights/Private/Insights/ViewModels/FrameTrackHelper.h b/Engine/Source/Developer/TraceInsights/Private/Insights/ViewModels/FrameTrackHelper.h
--- a/Engine/Source/Developer/TraceInsights/Private/Insights/ViewModels/FrameTrackHelper.h
+++ b/Engine/Source/Developer/TraceInsights/Private/Insights/ViewModels/FrameTrackHelper.h
@@ -136,6 +136,8 @@ public:
 	void DrawHoveredSample(const FFrameTrackSample& Sample) const;
 	void DrawHighlightedInterval(const FFrameTrackSeries& Series, const double StartTime, const double EndTime) const;
 
+	static const TCHAR* FrameTypeToString(int32 FrameType);
+	static uint32 GetColor32ByFrameType(int32 FrameType);
 	static FLinearColor GetColorByFrameType(int32 FrameType);
 
 	int32 GetNumFrames() const { return NumFrames; }
diff --git a/Engine/Source/Developer/TraceInsights/Private/Insights/ViewModels/GraphTrack.cpp b/Engine/Source/Developer/TraceInsights/Private/Insights/ViewModels/GraphTrack.cpp
--- a/Engine/Source/Developer/TraceInsights/Private/Insights/ViewModels/GraphTrack.cpp
+++ b/Engine/Source/Developer/TraceInsights/Private/Insights/ViewModels/GraphTrack.cpp
@@ -27,8 +27,12 @@
 // FGraphTrack
 ////////////////////////////////////////////////////////////////////////////////////////////////////
 
-FGraphTrack::FGraphTrack(const FName& InSubType)
-	: FBaseTimingTrack(FName(TEXT("Graph")), InSubType)
+INSIGHTS_IMPLEMENT_RTTI(FGraphTrack)
+
+////////////////////////////////////////////////////////////////////////////////////////////////////
+
+FGraphTrack::FGraphTrack()
+	: FBaseTimingTrack()
 	//, AllSeries()
 	, WhiteBrush(FInsightsStyle::Get().GetBrush("WhiteBrush"))
 	, PointBrush(FEditorStyle::GetBrush("Graph.ExecutionBubble"))
@@ -51,8 +55,8 @@ FGraphTrack::FGraphTrack(const FName& InSubType)
 
 ////////////////////////////////////////////////////////////////////////////////////////////////////
 
-FGraphTrack::FGraphTrack(const FName& InSubType, const FString& InName)
-	: FBaseTimingTrack(FName(TEXT("Graph")), InSubType, InName)
+FGraphTrack::FGraphTrack(const FString& InName)
+	: FBaseTimingTrack(InName)
 	//, AllSeries()
 	, WhiteBrush(FInsightsStyle::Get().GetBrush("WhiteBrush"))
 	, PointBrush(FEditorStyle::GetBrush("Graph.ExecutionBubble"))
@@ -505,9 +509,9 @@ void FGraphTrack::DrawSeries(const FGraphSeries& Series, FDrawContext& DrawConte
 void FGraphTrack::DrawEvent(const ITimingTrackDrawContext& Context, const ITimingEvent& InTimingEvent, EDrawEventMode InDrawMode) const
 {
 	ensure(InTimingEvent.CheckTrack(this));
-	ensure(FGraphTrackEvent::CheckTypeName(InTimingEvent));
+	ensure(InTimingEvent.Is<FGraphTrackEvent>());
 
-	const FGraphTrackEvent& GraphEvent = static_cast<const FGraphTrackEvent&>(InTimingEvent);
+	const FGraphTrackEvent& GraphEvent = InTimingEvent.As<FGraphTrackEvent>();
 	const TSharedPtr<const FGraphSeries> Series = GraphEvent.GetSeries();
 
 	const FTimingTrackViewport& Viewport = Context.GetViewport();
@@ -572,9 +576,9 @@ void FGraphTrack::DrawEvent(const ITimingTrackDrawContext& Context, const ITimin
 
 void FGraphTrack::InitTooltip(FTooltipDrawState& InOutTooltip, const ITimingEvent& InTooltipEvent) const
 {
-	if (InTooltipEvent.CheckTrack(this) && FGraphTrackEvent::CheckTypeName(InTooltipEvent))
+	if (InTooltipEvent.CheckTrack(this) && InTooltipEvent.Is<FGraphTrackEvent>())
 	{
-		const FGraphTrackEvent& TooltipEvent = static_cast<const FGraphTrackEvent&>(InTooltipEvent);
+		const FGraphTrackEvent& TooltipEvent = InTooltipEvent.As<FGraphTrackEvent>();
 		const TSharedRef<const FGraphSeries> Series = TooltipEvent.GetSeries();
 
 		InOutTooltip.ResetContent();
@@ -928,8 +932,12 @@ bool FGraphTrack::ContextMenu_ShowSeries_IsChecked(FGraphSeries* Series)
 // FRandomGraphTrack
 ////////////////////////////////////////////////////////////////////////////////////////////////////
 
+INSIGHTS_IMPLEMENT_RTTI(FRandomGraphTrack)
+
+////////////////////////////////////////////////////////////////////////////////////////////////////
+
 FRandomGraphTrack::FRandomGraphTrack()
-	: FGraphTrack(FName(TEXT("Random")))
+	: FGraphTrack()
 {
 	bDrawPoints = true;
 	bDrawPointsWithBorder = true;
diff --git a/Engine/Source/Developer/TraceInsights/Private/Insights/ViewModels/GraphTrackEvent.cpp b/Engine/Source/Developer/TraceInsights/Private/Insights/ViewModels/GraphTrackEvent.cpp
--- a/Engine/Source/Developer/TraceInsights/Private/Insights/ViewModels/GraphTrackEvent.cpp
+++ b/Engine/Source/Developer/TraceInsights/Private/Insights/ViewModels/GraphTrackEvent.cpp
@@ -8,7 +8,7 @@
 
 ////////////////////////////////////////////////////////////////////////////////////////////////////
 
-const FName FGraphTrackEvent::TypeName = FName(TEXT("FGraphTrackEvent"));
+INSIGHTS_IMPLEMENT_RTTI(FGraphTrackEvent)
 
 ////////////////////////////////////////////////////////////////////////////////////////////////////
 
diff --git a/Engine/Source/Developer/TraceInsights/Private/Insights/ViewModels/LoadingTimingTrack.cpp b/Engine/Source/Developer/TraceInsights/Private/Insights/ViewModels/LoadingTimingTrack.cpp
--- a/Engine/Source/Developer/TraceInsights/Private/Insights/ViewModels/LoadingTimingTrack.cpp
+++ b/Engine/Source/Developer/TraceInsights/Private/Insights/ViewModels/LoadingTimingTrack.cpp
@@ -86,10 +86,9 @@ void FLoadingSharedState::Tick(Insights::ITimingViewSession& InSession, const Tr
 			{
 				if (!LoadingTracks.Contains(LoadingTimelineIndex))
 				{
-					const TCHAR* const GroupName = ThreadInfo.GroupName ? ThreadInfo.GroupName : ThreadInfo.Name;
+					//const TCHAR* const GroupName = ThreadInfo.GroupName ? ThreadInfo.GroupName : ThreadInfo.Name;
 					const FString TrackName(ThreadInfo.Name && *ThreadInfo.Name ? FString::Printf(TEXT("Loading - %s"), ThreadInfo.Name) : FString::Printf(TEXT("Loading - Thread %u"), ThreadInfo.Id));
-
-					TSharedRef<FLoadingTimingTrack> LoadingThreadTrack = MakeShared<FLoadingTimingTrack>(*this, LoadingTimelineIndex, GroupName, TrackName);
+					TSharedRef<FLoadingTimingTrack> LoadingThreadTrack = MakeShared<FLoadingTimingTrack>(*this, LoadingTimelineIndex, TrackName);
 					LoadingThreadTrack->SetOrder(-100 + LoadingTracks.Num());
 					LoadingThreadTrack->SetVisibilityFlag(bShowHideAllLoadingTracks);
 					InSession.AddScrollableTrack(LoadingThreadTrack);
@@ -221,6 +220,10 @@ void FLoadingSharedState::SetColorSchema(int32 Schema)
 // FLoadingTimingTrack
 ////////////////////////////////////////////////////////////////////////////////////////////////////
 
+INSIGHTS_IMPLEMENT_RTTI(FLoadingTimingTrack)
+
+////////////////////////////////////////////////////////////////////////////////////////////////////
+
 void FLoadingTimingTrack::BuildDrawState(ITimingEventsTrackDrawStateBuilder& Builder, const ITimingTrackUpdateContext& Context)
 {
 	TSharedPtr<const Trace::IAnalysisSession> Session = FInsightsManager::Get()->GetSession();
@@ -263,9 +266,9 @@ void FLoadingTimingTrack::BuildDrawState(ITimingEventsTrackDrawStateBuilder& Bui
 
 void FLoadingTimingTrack::InitTooltip(FTooltipDrawState& InOutTooltip, const ITimingEvent& InTooltipEvent) const
 {
-	if (InTooltipEvent.CheckTrack(this) && FTimingEvent::CheckTypeName(InTooltipEvent))
+	if (InTooltipEvent.CheckTrack(this) && InTooltipEvent.Is<FTimingEvent>())
 	{
-		const FTimingEvent& TooltipEvent = static_cast<const FTimingEvent&>(InTooltipEvent);
+		const FTimingEvent& TooltipEvent = InTooltipEvent.As<FTimingEvent>();
 
 		auto MatchEvent = [&TooltipEvent](double InStartTime, double InEndTime, uint32 InDepth)
 		{
diff --git a/Engine/Source/Developer/TraceInsights/Private/Insights/ViewModels/LoadingTimingTrack.h b/Engine/Source/Developer/TraceInsights/Private/Insights/ViewModels/LoadingTimingTrack.h
--- a/Engine/Source/Developer/TraceInsights/Private/Insights/ViewModels/LoadingTimingTrack.h
+++ b/Engine/Source/Developer/TraceInsights/Private/Insights/ViewModels/LoadingTimingTrack.h
@@ -65,9 +65,11 @@ private:
 
 class FLoadingTimingTrack : public FTimingEventsTrack
 {
+	INSIGHTS_DECLARE_RTTI(FLoadingTimingTrack, FTimingEventsTrack)
+
 public:
-	explicit FLoadingTimingTrack(FLoadingSharedState& InSharedState, uint32 InTimelineIndex, const FName& InSubType, const FString& InName)
-		: FTimingEventsTrack(FName(TEXT("Loading")), InSubType, InName)
+	explicit FLoadingTimingTrack(FLoadingSharedState& InSharedState, uint32 InTimelineIndex, const FString& InName)
+		: FTimingEventsTrack(InName)
 		, SharedState(InSharedState)
 		, TimelineIndex(InTimelineIndex)
 	{
diff --git a/Engine/Source/Developer/TraceInsights/Private/Insights/ViewModels/MarkersTimingTrack.cpp b/Engine/Source/Developer/TraceInsights/Private/Insights/ViewModels/MarkersTimingTrack.cpp
--- a/Engine/Source/Developer/TraceInsights/Private/Insights/ViewModels/MarkersTimingTrack.cpp
+++ b/Engine/Source/Developer/TraceInsights/Private/Insights/ViewModels/MarkersTimingTrack.cpp
@@ -22,8 +22,12 @@
 // FMarkersTimingTrack
 ////////////////////////////////////////////////////////////////////////////////////////////////////
 
+INSIGHTS_IMPLEMENT_RTTI(FMarkersTimingTrack)
+
+////////////////////////////////////////////////////////////////////////////////////////////////////
+
 FMarkersTimingTrack::FMarkersTimingTrack()
-	: FBaseTimingTrack(FName("Markers"))
+	: FBaseTimingTrack()
 	//, TimeMarkerBoxes()
 	//, TimeMarkerTexts()
 	, bIsCollapsed(false)
diff --git a/Engine/Source/Developer/TraceInsights/Private/Insights/ViewModels/MarkersTimingTrack.h b/Engine/Source/Developer/TraceInsights/Private/Insights/ViewModels/MarkersTimingTrack.h
--- a/Engine/Source/Developer/TraceInsights/Private/Insights/ViewModels/MarkersTimingTrack.h
+++ b/Engine/Source/Developer/TraceInsights/Private/Insights/ViewModels/MarkersTimingTrack.h
@@ -43,6 +43,8 @@ class FMarkersTimingTrack : public FBaseTimingTrack
 {
 	friend class FTimeMarkerTrackBuilder;
 
+	INSIGHTS_DECLARE_RTTI(FMarkersTimingTrack, FBaseTimingTrack)
+
 public:
 	FMarkersTimingTrack();
 	virtual ~FMarkersTimingTrack();
diff --git a/Engine/Source/Developer/TraceInsights/Private/Insights/ViewModels/ThreadTimingTrack.cpp b/Engine/Source/Developer/TraceInsights/Private/Insights/ViewModels/ThreadTimingTrack.cpp
--- a/Engine/Source/Developer/TraceInsights/Private/Insights/ViewModels/ThreadTimingTrack.cpp
+++ b/Engine/Source/Developer/TraceInsights/Private/Insights/ViewModels/ThreadTimingTrack.cpp
@@ -156,14 +156,13 @@ void FThreadTimingSharedState::Tick(Insights::ITimingViewSession& InSession, con
 			uint32 CpuTimelineIndex;
 			if (TimingProfilerProvider->GetCpuThreadTimelineIndex(ThreadInfo.Id, CpuTimelineIndex))
 			{
-				TSharedPtr<FCpuTimingTrack> Track;
-
-				if (!CpuTracks.Contains(ThreadInfo.Id))
+				TSharedPtr<FCpuTimingTrack>* TrackPtrPtr = CpuTracks.Find(ThreadInfo.Id);
+				if (TrackPtrPtr == nullptr)
 				{
 					FString TrackName(ThreadInfo.Name && *ThreadInfo.Name ? ThreadInfo.Name : FString::Printf(TEXT("Thread %u"), ThreadInfo.Id));
 
 					// Create new Timing Events track for the CPU thread.
-					Track = MakeShared<FCpuTimingTrack>(*this, TrackName, GroupName, CpuTimelineIndex, ThreadInfo.Id);
+					TSharedPtr<FCpuTimingTrack> Track = MakeShared<FCpuTimingTrack>(*this, TrackName, GroupName, CpuTimelineIndex, ThreadInfo.Id);
 					Track->SetOrder(Order);
 					CpuTracks.Add(ThreadInfo.Id, Track);
 
@@ -184,8 +183,7 @@ void FThreadTimingSharedState::Tick(Insights::ITimingViewSession& InSession, con
 				}
 				else
 				{
-					Track = CpuTracks[ThreadInfo.Id];
-
+					TSharedPtr<FCpuTimingTrack> Track = *TrackPtrPtr;
 					if (Track->GetOrder() != Order)
 					{
 						Track->SetOrder(Order);
@@ -374,6 +372,10 @@ void FThreadTimingSharedState::ToggleTrackVisibilityByGroup_Execute(const TCHAR*
 // FThreadTimingTrack
 ////////////////////////////////////////////////////////////////////////////////////////////////////
 
+INSIGHTS_IMPLEMENT_RTTI(FThreadTimingTrack)
+
+////////////////////////////////////////////////////////////////////////////////////////////////////
+
 void FThreadTimingTrack::BuildDrawState(ITimingEventsTrackDrawStateBuilder& Builder, const ITimingTrackUpdateContext& Context)
 {
 	TSharedPtr<const Trace::IAnalysisSession> Session = FInsightsManager::Get()->GetSession();
@@ -422,10 +424,16 @@ void FThreadTimingTrack::BuildDrawState(ITimingEventsTrackDrawStateBuilder& Buil
 void FThreadTimingTrack::BuildFilteredDrawState(ITimingEventsTrackDrawStateBuilder& Builder, const ITimingTrackUpdateContext& Context)
 {
 	const TSharedPtr<ITimingEventFilter> EventFilterPtr = Context.GetEventFilter();
-	if (EventFilterPtr.IsValid() && FTimingEventFilter::CheckTypeName(*EventFilterPtr))
+	if (EventFilterPtr.IsValid() && EventFilterPtr->FilterTrack(*this))
 	{
-		const FTimingEventFilter& EventFilter = static_cast<const FTimingEventFilter&>(*EventFilterPtr);
-		const uint64 FilterEventType = EventFilter.GetEventType();
+		bool bFilterOnlyByEventType = false; // this is the most often use case, so the below code tries to optimize it
+		uint64 FilterEventType = 0;
+		if (EventFilterPtr->Is<FTimingEventFilterByEventType>())
+		{
+			bFilterOnlyByEventType = true;
+			const FTimingEventFilterByEventType& EventFilter = EventFilterPtr->As<FTimingEventFilterByEventType>();
+			FilterEventType = EventFilter.GetEventType();
+		}
 
 		TSharedPtr<const Trace::IAnalysisSession> Session = FInsightsManager::Get()->GetSession();
 		if (Session.IsValid() && Trace::ReadTimingProfilerProvider(*Session.Get()))
@@ -436,26 +444,55 @@ void FThreadTimingTrack::BuildFilteredDrawState(ITimingEventsTrackDrawStateBuild
 
 			const FTimingTrackViewport& Viewport = Context.GetViewport();
 
-			TimingProfilerProvider.ReadTimers([this, &Builder, &Viewport, &TimingProfilerProvider, FilterEventType](const Trace::FTimingProfilerTimer* Timers, uint64 TimersCount)
+			if (bFilterOnlyByEventType)
 			{
-				TimingProfilerProvider.ReadTimeline(TimelineIndex, [this, &Builder, &Viewport, Timers, FilterEventType](const Trace::ITimingProfilerProvider::Timeline& Timeline)
+				TimingProfilerProvider.ReadTimers([this, &Builder, &Viewport, &TimingProfilerProvider, FilterEventType](const Trace::FTimingProfilerTimer* Timers, uint64 TimersCount)
 				{
-					// Note: Enumerating events for filtering should not use downsampling.
-					//const double SecondsPerPixel = 1.0 / Viewport.GetScaleX();
-					//Timeline.EnumerateEventsDownSampled(Viewport.GetStartTime(), Viewport.GetEndTime(), SecondsPerPixel,
-					Timeline.EnumerateEvents(Viewport.GetStartTime(), Viewport.GetEndTime(),
-						[this, &Builder, Timers, FilterEventType](double StartTime, double EndTime, uint32 Depth, const Trace::FTimingProfilerEvent& Event)
+					TimingProfilerProvider.ReadTimeline(TimelineIndex, [this, &Builder, &Viewport, Timers, FilterEventType](const Trace::ITimingProfilerProvider::Timeline& Timeline)
 					{
-						//ensure(Timers[Event.TimerIndex].Id == Event.TimerIndex);
-						const uint64 Type = Event.TimerIndex;
-						if (Type == FilterEventType)
+						// Note: Enumerating events for filtering should not use downsampling.
+						//const double SecondsPerPixel = 1.0 / Viewport.GetScaleX();
+						//Timeline.EnumerateEventsDownSampled(Viewport.GetStartTime(), Viewport.GetEndTime(), SecondsPerPixel,
+						Timeline.EnumerateEvents(Viewport.GetStartTime(), Viewport.GetEndTime(),
+							[this, &Builder, Timers, FilterEventType](double StartTime, double EndTime, uint32 Depth, const Trace::FTimingProfilerEvent& Event)
 						{
-							const uint32 Color = 0;
-							Builder.AddEvent(StartTime, EndTime, Depth, Timers[Event.TimerIndex].Name, Type, Color);
-						}
+							//ensure(Timers[Event.TimerIndex].Id == Event.TimerIndex);
+							const uint64 Type = Event.TimerIndex;
+							if (Type == FilterEventType)
+							{
+								constexpr uint32 Color = 0;
+								Builder.AddEvent(StartTime, EndTime, Depth, Timers[Event.TimerIndex].Name, Type, Color);
+							}
+						});
 					});
 				});
-			});
+			}
+			else // generic filter
+			{
+				TimingProfilerProvider.ReadTimers([this, &Builder, &Viewport, &TimingProfilerProvider, &EventFilterPtr](const Trace::FTimingProfilerTimer* Timers, uint64 TimersCount)
+				{
+					TimingProfilerProvider.ReadTimeline(TimelineIndex, [this, &Builder, &Viewport, Timers, &EventFilterPtr](const Trace::ITimingProfilerProvider::Timeline& Timeline)
+					{
+						// Note: Enumerating events for filtering should not use downsampling.
+						//const double SecondsPerPixel = 1.0 / Viewport.GetScaleX();
+						//Timeline.EnumerateEventsDownSampled(Viewport.GetStartTime(), Viewport.GetEndTime(), SecondsPerPixel,
+						Timeline.EnumerateEvents(Viewport.GetStartTime(), Viewport.GetEndTime(),
+							[this, &Builder, Timers, &EventFilterPtr](double StartTime, double EndTime, uint32 Depth, const Trace::FTimingProfilerEvent& Event)
+						{
+							//ensure(Timers[Event.TimerIndex].Id == Event.TimerIndex);
+							const uint64 Type = Event.TimerIndex;
+
+							FTimingEvent TimingEvent(SharedThis(this), StartTime, EndTime, Depth, Type);
+
+							if (EventFilterPtr->FilterEvent(TimingEvent))
+							{
+								constexpr uint32 Color = 0;
+								Builder.AddEvent(StartTime, EndTime, Depth, Timers[Event.TimerIndex].Name, Type, Color);
+							}
+						});
+					});
+				});
+			}
 		}
 	}
 }
@@ -467,9 +504,9 @@ void FThreadTimingTrack::PostDraw(const ITimingTrackDrawContext& Context) const
 	const TSharedPtr<const ITimingEvent> SelectedEventPtr = Context.GetSelectedEvent();
 	if (SelectedEventPtr.IsValid() &&
 		SelectedEventPtr->CheckTrack(this) &&
-		FTimingEvent::CheckTypeName(*SelectedEventPtr))
+		SelectedEventPtr->Is<FTimingEvent>())
 	{
-		const FTimingEvent& SelectedEvent = static_cast<const FTimingEvent&>(*SelectedEventPtr);
+		const FTimingEvent& SelectedEvent = SelectedEventPtr->As<FTimingEvent>();
 		const ITimingViewDrawHelper& Helper = Context.GetHelper();
 		DrawSelectedEventInfo(SelectedEvent, Context.GetViewport(), Context.GetDrawContext(), Helper.GetWhiteBrush(), Helper.GetEventFont());
 	}
@@ -510,9 +547,9 @@ void FThreadTimingTrack::DrawSelectedEventInfo(const FTimingEvent& SelectedEvent
 
 void FThreadTimingTrack::InitTooltip(FTooltipDrawState& InOutTooltip, const ITimingEvent& InTooltipEvent) const
 {
-	if (InTooltipEvent.CheckTrack(this) && FTimingEvent::CheckTypeName(InTooltipEvent))
+	if (InTooltipEvent.CheckTrack(this) && InTooltipEvent.Is<FTimingEvent>())
 	{
-		const FTimingEvent& TooltipEvent = static_cast<const FTimingEvent&>(InTooltipEvent);
+		const FTimingEvent& TooltipEvent = InTooltipEvent.As<FTimingEvent>();
 
 		FindTimingProfilerEvent(TooltipEvent, [this, &InOutTooltip, &TooltipEvent](double InFoundStartTime, double InFoundEndTime, uint32 InFoundDepth, const Trace::FTimingProfilerEvent& InFoundEvent)
 		{
@@ -621,9 +658,9 @@ const TSharedPtr<const ITimingEvent> FThreadTimingTrack::SearchEvent(const FTimi
 
 void FThreadTimingTrack::UpdateEventStats(ITimingEvent& InOutEvent) const
 {
-	if (InOutEvent.CheckTrack(this) && FTimingEvent::CheckTypeName(InOutEvent))
+	if (InOutEvent.CheckTrack(this) && InOutEvent.Is<FTimingEvent>())
 	{
-		FTimingEvent& TrackEvent = static_cast<FTimingEvent&>(InOutEvent);
+		FTimingEvent& TrackEvent = InOutEvent.As<FTimingEvent>();
 
 		TSharedPtr<const Trace::IAnalysisSession> Session = FInsightsManager::Get()->GetSession();
 		if (Session.IsValid())
@@ -696,9 +733,9 @@ void FThreadTimingTrack::UpdateEventStats(ITimingEvent& InOutEvent) const
 
 void FThreadTimingTrack::OnEventSelected(const ITimingEvent& InSelectedEvent) const
 {
-	if (InSelectedEvent.CheckTrack(this) && FTimingEvent::CheckTypeName(InSelectedEvent))
+	if (InSelectedEvent.CheckTrack(this) && InSelectedEvent.Is<FTimingEvent>())
 	{
-		const FTimingEvent& TrackEvent = static_cast<const FTimingEvent&>(InSelectedEvent);
+		const FTimingEvent& TrackEvent = InSelectedEvent.As<FTimingEvent>();
 
 		//uint64 TimerType = uint64(-1);
 		//FindTimingProfilerEvent(TrackEvent, [&TimerType](double InFoundStartTime, double InFoundEndTime, uint32 InFoundDepth, const Trace::FTimingProfilerEvent& InFoundEvent)
@@ -717,9 +754,9 @@ void FThreadTimingTrack::OnEventSelected(const ITimingEvent& InSelectedEvent) co
 
 void FThreadTimingTrack::OnClipboardCopyEvent(const ITimingEvent& InSelectedEvent) const
 {
-	if (InSelectedEvent.CheckTrack(this) && FTimingEvent::CheckTypeName(InSelectedEvent))
+	if (InSelectedEvent.CheckTrack(this) && InSelectedEvent.Is<FTimingEvent>())
 	{
-		const FTimingEvent& TrackEvent = static_cast<const FTimingEvent&>(InSelectedEvent);
+		const FTimingEvent& TrackEvent = InSelectedEvent.As<FTimingEvent>();
 
 		//uint64 TimerType = uint64(-1);
 		//FindTimingProfilerEvent(TrackEvent, [&TimerType](double InFoundStartTime, double InFoundEndTime, uint32 InFoundDepth, const Trace::FTimingProfilerEvent& InFoundEvent)
diff --git a/Engine/Source/Developer/TraceInsights/Private/Insights/ViewModels/ThreadTimingTrack.h b/Engine/Source/Developer/TraceInsights/Private/Insights/ViewModels/ThreadTimingTrack.h
--- a/Engine/Source/Developer/TraceInsights/Private/Insights/ViewModels/ThreadTimingTrack.h
+++ b/Engine/Source/Developer/TraceInsights/Private/Insights/ViewModels/ThreadTimingTrack.h
@@ -18,10 +18,6 @@ class FCpuTimingTrack;
 class STimingView;
 struct FSlateBrush;
 
-////////////////////////////////////////////////////////////////////////////////////////////////////
-
-
-
 ////////////////////////////////////////////////////////////////////////////////////////////////////
 
 class FThreadTimingSharedState : public Insights::ITimingViewExtender, public TSharedFromThis<FThreadTimingSharedState>
@@ -91,9 +87,11 @@ private:
 
 class FThreadTimingTrack : public FTimingEventsTrack
 {
+	INSIGHTS_DECLARE_RTTI(FThreadTimingTrack, FTimingEventsTrack)
+
 public:
-	explicit FThreadTimingTrack(FThreadTimingSharedState& InSharedState, const FName& InSubType, const FString& InName, const TCHAR* InGroupName, uint32 InTimelineIndex, uint32 InThreadId)
-		: FTimingEventsTrack(FName(TEXT("Thread")), InSubType, InName)
+	explicit FThreadTimingTrack(FThreadTimingSharedState& InSharedState, const FString& InName, const TCHAR* InGroupName, uint32 InTimelineIndex, uint32 InThreadId)
+		: FTimingEventsTrack(InName)
 		, SharedState(InSharedState)
 		, GroupName(InGroupName)
 		, TimelineIndex(InTimelineIndex)
@@ -154,7 +152,7 @@ class FCpuTimingTrack : public FThreadTimingTrack
 {
 public:
 	explicit FCpuTimingTrack(FThreadTimingSharedState& InSharedState, const FString& InName, const TCHAR* InGroupName, uint32 InTimelineIndex, uint32 InThreadId)
-		: FThreadTimingTrack(InSharedState, FName(TEXT("CPU")), InName, InGroupName, InTimelineIndex, InThreadId)
+		: FThreadTimingTrack(InSharedState, InName, InGroupName, InTimelineIndex, InThreadId)
 	{
 	}
 };
@@ -165,7 +163,7 @@ class FGpuTimingTrack : public FThreadTimingTrack
 {
 public:
 	explicit FGpuTimingTrack(FThreadTimingSharedState& InSharedState, const FString& InName, const TCHAR* InGroupName, uint32 InTimelineIndex, uint32 InThreadId)
-		: FThreadTimingTrack(InSharedState, FName(TEXT("GPU")), InName, InGroupName, InTimelineIndex, InThreadId)
+		: FThreadTimingTrack(InSharedState, InName, InGroupName, InTimelineIndex, InThreadId)
 	{
 	}
 };
diff --git a/Engine/Source/Developer/TraceInsights/Private/Insights/ViewModels/TimeRulerTrack.cpp b/Engine/Source/Developer/TraceInsights/Private/Insights/ViewModels/TimeRulerTrack.cpp
--- a/Engine/Source/Developer/TraceInsights/Private/Insights/ViewModels/TimeRulerTrack.cpp
+++ b/Engine/Source/Developer/TraceInsights/Private/Insights/ViewModels/TimeRulerTrack.cpp
@@ -20,8 +20,12 @@
 
 ////////////////////////////////////////////////////////////////////////////////////////////////////
 
+INSIGHTS_IMPLEMENT_RTTI(FTimeRulerTrack)
+
+////////////////////////////////////////////////////////////////////////////////////////////////////
+
 FTimeRulerTrack::FTimeRulerTrack()
-	: FBaseTimingTrack(FName(TEXT("TimeRuler")))
+	: FBaseTimingTrack()
 	, WhiteBrush(FInsightsStyle::Get().GetBrush("WhiteBrush"))
 	, Font(FCoreStyle::GetDefaultFontStyle("Regular", 8))
 	, CrtMousePosTextWidth(0.0f)
diff --git a/Engine/Source/Developer/TraceInsights/Private/Insights/ViewModels/TimeRulerTrack.h b/Engine/Source/Developer/TraceInsights/Private/Insights/ViewModels/TimeRulerTrack.h
--- a/Engine/Source/Developer/TraceInsights/Private/Insights/ViewModels/TimeRulerTrack.h
+++ b/Engine/Source/Developer/TraceInsights/Private/Insights/ViewModels/TimeRulerTrack.h
@@ -16,6 +16,8 @@ class FTimingTrackViewport;
 
 class FTimeRulerTrack : public FBaseTimingTrack
 {
+	INSIGHTS_DECLARE_RTTI(FTimeRulerTrack, FBaseTimingTrack)
+
 public:
 	FTimeRulerTrack();
 	virtual ~FTimeRulerTrack();
diff --git a/Engine/Source/Developer/TraceInsights/Private/Insights/ViewModels/TimingEventsTrack.cpp b/Engine/Source/Developer/TraceInsights/Private/Insights/ViewModels/TimingEventsTrack.cpp
--- a/Engine/Source/Developer/TraceInsights/Private/Insights/ViewModels/TimingEventsTrack.cpp
+++ b/Engine/Source/Developer/TraceInsights/Private/Insights/ViewModels/TimingEventsTrack.cpp
@@ -10,28 +10,26 @@
 
 #define LOCTEXT_NAMESPACE "TimingEventsTrack"
 
-////////////////////////////////////////////////////////////////////////////////////////////////////
-// FTimingEvent, FTimingEventFilter
 ////////////////////////////////////////////////////////////////////////////////////////////////////
 
-const FName FTimingEvent::TypeName = FName(TEXT("FTimingEvent"));
-const FName FTimingEventFilter::TypeName = FName(TEXT("FTimingEventFilter"));
+INSIGHTS_IMPLEMENT_RTTI(FTimingEventsTrack)
 
-bool FTimingEventFilter::FilterTrack(const FBaseTimingTrack& InTrack) const
-{
-	return !bFilterByTrackType || InTrack.GetType() == TrackType;
-}
+bool FTimingEventsTrack::bUseDownSampling = true;
 
-////////////////////////////////////////////////////////////////////////////////////////////////////
-// FTimingEventsTrack
 ////////////////////////////////////////////////////////////////////////////////////////////////////
 
-bool FTimingEventsTrack::bUseDownSampling = true;
+FTimingEventsTrack::FTimingEventsTrack()
+	: FBaseTimingTrack()
+	, NumLanes(0)
+	, DrawState(MakeShared<FTimingEventsTrackDrawState>())
+	, FilteredDrawState(MakeShared<FTimingEventsTrackDrawState>())
+{
+}
 
 ////////////////////////////////////////////////////////////////////////////////////////////////////
 
-FTimingEventsTrack::FTimingEventsTrack(const FName& InType, const FName& InSubType, const FString& InName)
-	: FBaseTimingTrack(InType, InSubType, InName)
+FTimingEventsTrack::FTimingEventsTrack(const FString& InName)
+	: FBaseTimingTrack(InName)
 	, NumLanes(0)
 	, DrawState(MakeShared<FTimingEventsTrackDrawState>())
 	, FilteredDrawState(MakeShared<FTimingEventsTrackDrawState>())
@@ -273,9 +271,9 @@ void FTimingEventsTrack::DrawHeader(const ITimingTrackDrawContext& Context) cons
 
 void FTimingEventsTrack::DrawEvent(const ITimingTrackDrawContext& Context, const ITimingEvent& InTimingEvent, EDrawEventMode InDrawMode) const
 {
-	if (InTimingEvent.CheckTrack(this) && FTimingEvent::CheckTypeName(InTimingEvent))
+	if (InTimingEvent.CheckTrack(this) && InTimingEvent.Is<FTimingEvent>())
 	{
-		const FTimingEvent& TrackEvent = static_cast<const FTimingEvent&>(InTimingEvent);
+		const FTimingEvent& TrackEvent = InTimingEvent.As<FTimingEvent>();
 		const FTimingViewLayout& Layout = Context.GetViewport().GetLayout();
 		const float Y = TrackEvent.GetTrack()->GetPosY() + Layout.GetLaneY(TrackEvent.GetDepth());
 
@@ -316,14 +314,10 @@ const TSharedPtr<const ITimingEvent> FTimingEventsTrack::GetEvent(float InPosX,
 
 TSharedPtr<ITimingEventFilter> FTimingEventsTrack::GetFilterByEvent(const TSharedPtr<const ITimingEvent> InTimingEvent) const
 {
-	if (InTimingEvent.IsValid() && FTimingEvent::CheckTypeName(*InTimingEvent))
+	if (InTimingEvent.IsValid() && InTimingEvent->Is<FTimingEvent>())
 	{
-		const FTimingEvent& Event = *StaticCastSharedPtr<const FTimingEvent, const ITimingEvent>(InTimingEvent);
-
-		TSharedRef<FTimingEventFilter> EventFilterRef = MakeShared<FTimingEventFilter>();
-		EventFilterRef->SetFilterByEventType(true);
-		EventFilterRef->SetEventType(Event.GetType());
-
+		const FTimingEvent& Event = InTimingEvent->As<FTimingEvent>();
+		TSharedRef<FTimingEventFilter> EventFilterRef = MakeShared<FTimingEventFilterByEventType>(Event.GetType());
 		return EventFilterRef;
 	}
 	return nullptr;
diff --git a/Engine/Source/Developer/TraceInsights/Private/Insights/ViewModels/TimingGraphTrack.cpp b/Engine/Source/Developer/TraceInsights/Private/Insights/ViewModels/TimingGraphTrack.cpp
--- a/Engine/Source/Developer/TraceInsights/Private/Insights/ViewModels/TimingGraphTrack.cpp
+++ b/Engine/Source/Developer/TraceInsights/Private/Insights/ViewModels/TimingGraphTrack.cpp
@@ -55,8 +55,12 @@ FString FTimingGraphSeries::FormatValue(double Value) const
 // FTimingGraphTrack
 ////////////////////////////////////////////////////////////////////////////////////////////////////
 
+INSIGHTS_IMPLEMENT_RTTI(FTimingGraphTrack)
+
+////////////////////////////////////////////////////////////////////////////////////////////////////
+
 FTimingGraphTrack::FTimingGraphTrack()
-	: FGraphTrack(FName(TEXT("Random")))
+	: FGraphTrack()
 {
 	bDrawPoints = true;
 	bDrawPointsWithBorder = true;
diff --git a/Engine/Source/Developer/TraceInsights/Private/Insights/ViewModels/TimingGraphTrack.h b/Engine/Source/Developer/TraceInsights/Private/Insights/ViewModels/TimingGraphTrack.h
--- a/Engine/Source/Developer/TraceInsights/Private/Insights/ViewModels/TimingGraphTrack.h
+++ b/Engine/Source/Developer/TraceInsights/Private/Insights/ViewModels/TimingGraphTrack.h
@@ -34,6 +34,8 @@ public:
 
 class FTimingGraphTrack : public FGraphTrack
 {
+	INSIGHTS_DECLARE_RTTI(FTimingGraphTrack, FGraphTrack)
+
 public:
 	FTimingGraphTrack();
 	virtual ~FTimingGraphTrack();
diff --git a/Engine/Source/Developer/TraceInsights/Private/Insights/ViewModels/TimingViewDrawHelper.cpp b/Engine/Source/Developer/TraceInsights/Private/Insights/ViewModels/TimingViewDrawHelper.cpp
--- a/Engine/Source/Developer/TraceInsights/Private/Insights/ViewModels/TimingViewDrawHelper.cpp
+++ b/Engine/Source/Developer/TraceInsights/Private/Insights/ViewModels/TimingViewDrawHelper.cpp
@@ -504,28 +504,38 @@ void FTimingViewDrawHelper::DrawTrackHeader(const FTimingEventsTrack& Track) con
 		if (TrackH > 4.0f)
 		{
 			const TSharedRef<FSlateFontMeasure> FontMeasureService = FSlateApplication::Get().GetRenderer()->GetFontMeasureService();
-			float NameWidth = FontMeasureService->Measure(Track.GetName(), EventFont).X;
+			float TextWidth = FontMeasureService->Measure(Track.GetName(), EventFont).X;
 
-			float TextX = 2.0f;
+			constexpr float PinWidth = 8.0f;
 			if (Track.IsSelected())
 			{
-				NameWidth += 8.0f;
-				TextX += 8.0f;
+				TextWidth += PinWidth;
 			}
 
+			const float HeaderX = 0.0f;
+			const float HeaderW = TextWidth + 4.0f;
+
+			const float HeaderY = TrackY + 1.0f;
 			const float HeaderH = FMath::Min(12.0f, Track.GetHeight() - 1.0f);
-			DrawContext.DrawBox(HeaderLayerId, 0.0f, TrackY + 1.0f, NameWidth + 4.0f, HeaderH, WhiteBrush, EdgeColor);
 
-			const FLinearColor TextColor = GetTrackNameTextColor(Track);
+			if (HeaderH > 0)
+			{
+				DrawContext.DrawBox(HeaderLayerId, HeaderX, HeaderY, HeaderW, HeaderH, WhiteBrush, EdgeColor);
 
-			const float TextY = TrackY + FMath::Min(0.0f, (Track.GetHeight() - 13.0f) / 2.0f);
-			const int32 HeaderTextLayerId = ReservedLayerId + ToInt32(EDrawLayer::HeaderText);
-			DrawContext.DrawText(HeaderTextLayerId, TextX, TextY, Track.GetName(), EventFont, TextColor);
+				const FLinearColor TextColor = GetTrackNameTextColor(Track);
 
-			if (Track.IsSelected())
-			{
-				// TODO: Use a "pin" image brush instead.
-				DrawContext.DrawText(HeaderTextLayerId, 2.0f, TextY, TEXT(">"), EventFont, TextColor);
+				float TextX = HeaderX + 2.0f;
+				const float TextY = HeaderY + HeaderH / 2.0f - 7.0f;
+				const int32 HeaderTextLayerId = ReservedLayerId + ToInt32(EDrawLayer::HeaderText);
+
+				if (Track.IsSelected())
+				{
+					// TODO: Use a "pin" image brush instead.
+					DrawContext.DrawText(HeaderTextLayerId, TextX, TextY, TEXT(">"), EventFont, TextColor);
+					TextX += PinWidth;
+				}
+
+				DrawContext.DrawText(HeaderTextLayerId, TextX, TextY, Track.GetName(), EventFont, TextColor);
 			}
 		}
 	}
diff --git a/Engine/Source/Developer/TraceInsights/Private/Insights/Widgets/SFrameTrack.cpp b/Engine/Source/Developer/TraceInsights/Private/Insights/Widgets/SFrameTrack.cpp
--- a/Engine/Source/Developer/TraceInsights/Private/Insights/Widgets/SFrameTrack.cpp
+++ b/Engine/Source/Developer/TraceInsights/Private/Insights/Widgets/SFrameTrack.cpp
@@ -378,18 +378,6 @@ void SFrameTrack::SelectFrameAtMousePosition(float X, float Y)
 
 ////////////////////////////////////////////////////////////////////////////////////////////////////
 
-const TCHAR* FrameTypeToString(int32 FrameType)
-{
-	switch (FrameType)
-	{
-		case TraceFrameType_Game:      return TEXT("Game");
-		case TraceFrameType_Rendering: return TEXT("Rendering");
-		default:                       return TEXT("Misc");
-	}
-}
-
-////////////////////////////////////////////////////////////////////////////////////////////////////
-
 int32 SFrameTrack::OnPaint(const FPaintArgs& Args, const FGeometry& AllottedGeometry, const FSlateRect& MyCullingRect, FSlateWindowElementList& OutDrawElements, int32 LayerId, const FWidgetStyle& InWidgetStyle, bool bParentEnabled) const
 {
 	const bool bEnabled = ShouldBeEnabled(bParentEnabled);
@@ -461,7 +449,7 @@ int32 SFrameTrack::OnPaint(const FPaintArgs& Args, const FGeometry& AllottedGeom
 
 			const FString Text = FString::Format(TEXT("{0} frame {1} ({2})"),
 				{
-					::FrameTypeToString(HoveredSample.Series->FrameType),
+					FFrameTrackDrawHelper::FrameTypeToString(HoveredSample.Series->FrameType),
 					FText::AsNumber(HoveredSample.Sample->LargestFrameIndex).ToString(),
 					TimeUtils::FormatTimeAuto(HoveredSample.Sample->LargestFrameDuration)
 				});
diff --git a/Engine/Source/Developer/TraceInsights/Private/Insights/Widgets/SStatsView.cpp b/Engine/Source/Developer/TraceInsights/Private/Insights/Widgets/SStatsView.cpp
--- a/Engine/Source/Developer/TraceInsights/Private/Insights/Widgets/SStatsView.cpp
+++ b/Engine/Source/Developer/TraceInsights/Private/Insights/Widgets/SStatsView.cpp
@@ -16,6 +16,7 @@
 #include "Widgets/Views/STableViewBase.h"
 
 // Insights
+#include "Insights/Common/Stopwatch.h"
 #include "Insights/TimingProfilerCommon.h"
 #include "Insights/TimingProfilerManager.h"
 #include "Insights/ViewModels/TimingGraphTrack.h"
@@ -1863,6 +1864,10 @@ void TTimeCalculationHelper<int64>::PostProcess(TMap<uint64, FStatsNodePtr>& Sta
 
 void SStatsView::UpdateStats(double StartTime, double EndTime)
 {
+	FStopwatch AggregationStopwatch;
+	FStopwatch Stopwatch;
+	Stopwatch.Start();
+
 	StatsStartTime = StartTime;
 	StatsEndTime = EndTime;
 
@@ -1884,6 +1889,7 @@ void SStatsView::UpdateStats(double StartTime, double EndTime)
 		TTimeCalculationHelper<double> CalculationHelperDbl(StartTime, EndTime);
 		TTimeCalculationHelper<int64>  CalculationHelperInt(StartTime, EndTime);
 
+		AggregationStopwatch.Start();
 		{
 			Trace::FAnalysisSessionReadScope SessionReadScope(*Session.Get());
 
@@ -1925,6 +1931,7 @@ void SStatsView::UpdateStats(double StartTime, double EndTime)
 				});
 			}
 		}
+		AggregationStopwatch.Stop();
 
 		// Compute average and median inclusive/exclusive times.
 		CalculationHelperDbl.PostProcess(StatsNodesIdMap, bComputeMedian);
@@ -1943,6 +1950,9 @@ void SStatsView::UpdateStats(double StartTime, double EndTime)
 	}
 
 	UpdateTree();
+
+	Stopwatch.Stop();
+	UE_LOG(TimingProfiler, Log, TEXT("Counters updated in %.3fs (%.3fs)"), Stopwatch.GetAccumulatedTime(), AggregationStopwatch.GetAccumulatedTime());
 }
 
 ////////////////////////////////////////////////////////////////////////////////////////////////////
diff --git a/Engine/Source/Developer/TraceInsights/Private/Insights/Widgets/STimersView.cpp b/Engine/Source/Developer/TraceInsights/Private/Insights/Widgets/STimersView.cpp
--- a/Engine/Source/Developer/TraceInsights/Private/Insights/Widgets/STimersView.cpp
+++ b/Engine/Source/Developer/TraceInsights/Private/Insights/Widgets/STimersView.cpp
@@ -17,6 +17,7 @@
 #include "Widgets/Views/STableViewBase.h"
 
 // Insights
+#include "Insights/Common/Stopwatch.h"
 #include "Insights/Table/ViewModels/Table.h"
 #include "Insights/Table/ViewModels/TableColumn.h"
 #include "Insights/TimingProfilerCommon.h"
@@ -825,11 +826,11 @@ void STimersView::TreeView_OnMouseButtonDoubleClick(FTimerNodePtr TimerNodePtr)
 			bool bSameFilter = false;
 
 			const TSharedPtr<ITimingEventFilter> EventFilterPtr = TimingView->GetEventFilter();
-			if (EventFilterPtr.IsValid())
+			if (EventFilterPtr.IsValid() &&
+				EventFilterPtr->Is<FTimingEventFilterByEventType>())
 			{
-				const FTimingEventFilter& EventFilter = static_cast<const FTimingEventFilter&>(*EventFilterPtr);
-				if (EventFilter.IsFilteringByEventType() &&
-					EventFilter.GetEventType() == TimerNodePtr->GetId())
+				const FTimingEventFilterByEventType& EventFilter = EventFilterPtr->As<FTimingEventFilterByEventType>();
+				if (EventFilter.GetEventType() == TimerNodePtr->GetId())
 				{
 					bSameFilter = true;
 					TimingView->SetEventFilter(nullptr); // reset filter
@@ -838,9 +839,7 @@ void STimersView::TreeView_OnMouseButtonDoubleClick(FTimerNodePtr TimerNodePtr)
 
 			if (!bSameFilter)
 			{
-				TSharedRef<FTimingEventFilter> EventFilterRef = MakeShared<FTimingEventFilter>();
-				EventFilterRef->SetFilterByEventType(true);
-				EventFilterRef->SetEventType(TimerNodePtr->GetId());
+				TSharedRef<FTimingEventFilterByEventType> EventFilterRef = MakeShared<FTimingEventFilterByEventType>(TimerNodePtr->GetId());
 				TimingView->SetEventFilter(EventFilterRef); // set new filter
 			}
 		}
@@ -1601,6 +1600,10 @@ void STimersView::RebuildTree(bool bResync)
 
 void STimersView::UpdateStats(double StartTime, double EndTime)
 {
+	FStopwatch AggregationStopwatch;
+	FStopwatch Stopwatch;
+	Stopwatch.Start();
+
 	StatsStartTime = StartTime;
 	StatsEndTime = EndTime;
 
@@ -1629,11 +1632,13 @@ void STimersView::UpdateStats(double StartTime, double EndTime)
 
 		TUniquePtr<Trace::ITable<Trace::FTimingProfilerAggregatedStats>> AggregationResultTable;
 
+		AggregationStopwatch.Start();
 		{
 			Trace::FAnalysisSessionReadScope SessionReadScope(*Session.Get());
 			const Trace::ITimingProfilerProvider& TimingProfilerProvider = *Trace::ReadTimingProfilerProvider(*Session.Get());
 			AggregationResultTable.Reset(TimingProfilerProvider.CreateAggregation(StartTime, EndTime, ThreadFilter, bIsGpuTrackVisible));
 		}
+		AggregationStopwatch.Stop();
 
 		if (AggregationResultTable.IsValid())
 		{
@@ -1666,6 +1671,9 @@ void STimersView::UpdateStats(double StartTime, double EndTime)
 	{
 		TreeView->RequestScrollIntoView(SelectedNodes[0]);
 	}
+
+	Stopwatch.Stop();
+	UE_LOG(TimingProfiler, Log, TEXT("Timers updated in %.3fs (%.3fs)"), Stopwatch.GetAccumulatedTime(), AggregationStopwatch.GetAccumulatedTime());
 }
 
 ////////////////////////////////////////////////////////////////////////////////////////////////////
diff --git a/Engine/Source/Developer/TraceInsights/Private/Insights/Widgets/STimingView.cpp b/Engine/Source/Developer/TraceInsights/Private/Insights/Widgets/STimingView.cpp
--- a/Engine/Source/Developer/TraceInsights/Private/Insights/Widgets/STimingView.cpp
+++ b/Engine/Source/Developer/TraceInsights/Private/Insights/Widgets/STimingView.cpp
@@ -42,6 +42,7 @@
 #include "Insights/ViewModels/BaseTimingTrack.h"
 #include "Insights/ViewModels/DrawHelpers.h"
 #include "Insights/ViewModels/FileActivityTimingTrack.h"
+#include "Insights/ViewModels/FrameTimingTrack.h"
 #include "Insights/ViewModels/GraphSeries.h"
 #include "Insights/ViewModels/GraphTrack.h"
 #include "Insights/ViewModels/LoadingTimingTrack.h"
@@ -73,6 +74,7 @@ namespace Insights { const FName TimingViewExtenderFeatureName(TEXT("TimingViewE
 
 STimingView::STimingView()
 	: bScrollableTracksOrderIsDirty(false)
+	, FrameSharedState(MakeShared<FFrameSharedState>(this))
 	, ThreadTimingSharedState(MakeShared<FThreadTimingSharedState>(this))
 	, LoadingSharedState(MakeShared<FLoadingSharedState>(this))
 	, bAssetLoadingMode(false)
@@ -83,6 +85,7 @@ STimingView::STimingView()
 	, WhiteBrush(FInsightsStyle::Get().GetBrush("WhiteBrush"))
 	, MainFont(FCoreStyle::GetDefaultFontStyle("Regular", 8))
 {
+	IModularFeatures::Get().RegisterModularFeature(Insights::TimingViewExtenderFeatureName, FrameSharedState.Get());
 	IModularFeatures::Get().RegisterModularFeature(Insights::TimingViewExtenderFeatureName, ThreadTimingSharedState.Get());
 	IModularFeatures::Get().RegisterModularFeature(Insights::TimingViewExtenderFeatureName, LoadingSharedState.Get());
 	IModularFeatures::Get().RegisterModularFeature(Insights::TimingViewExtenderFeatureName, FileActivitySharedState.Get());
@@ -105,9 +108,10 @@ STimingView::~STimingView()
 		Extender->OnEndSession(*this);
 	}
 
-	IModularFeatures::Get().UnregisterModularFeature(Insights::TimingViewExtenderFeatureName, ThreadTimingSharedState.Get());
-	IModularFeatures::Get().UnregisterModularFeature(Insights::TimingViewExtenderFeatureName, LoadingSharedState.Get());
 	IModularFeatures::Get().UnregisterModularFeature(Insights::TimingViewExtenderFeatureName, FileActivitySharedState.Get());
+	IModularFeatures::Get().UnregisterModularFeature(Insights::TimingViewExtenderFeatureName, LoadingSharedState.Get());
+	IModularFeatures::Get().UnregisterModularFeature(Insights::TimingViewExtenderFeatureName, ThreadTimingSharedState.Get());
+	IModularFeatures::Get().UnregisterModularFeature(Insights::TimingViewExtenderFeatureName, FrameSharedState.Get());
 }
 
 ////////////////////////////////////////////////////////////////////////////////////////////////////
@@ -1353,10 +1357,9 @@ int32 STimingView::OnPaint(const FPaintArgs& Args, const FGeometry& AllottedGeom
 void STimingView::AddTopDockedTrack(TSharedPtr<FBaseTimingTrack> Track)
 {
 	check(Track.IsValid());
-	UE_LOG(TimingProfiler, Log, TEXT("New Top Docked Track (%d) : %s - %s (\"%s\")"),
+	UE_LOG(TimingProfiler, Log, TEXT("New Top Docked Track (%d) : %s (\"%s\")"),
 		TopDockedTracks.Num() + 1,
-		*Track->GetType().ToString(),
-		*Track->GetSubType().ToString(),
+		*Track->GetTypeName().ToString(),
 		*Track->GetName());
 	check(!AllTracks.Contains(Track->GetId()));
 	AllTracks.Add(Track->GetId(), Track);
@@ -1369,10 +1372,9 @@ void STimingView::AddTopDockedTrack(TSharedPtr<FBaseTimingTrack> Track)
 void STimingView::AddBottomDockedTrack(TSharedPtr<FBaseTimingTrack> Track)
 {
 	check(Track.IsValid());
-	UE_LOG(TimingProfiler, Log, TEXT("New Bottom Docked Track (%d) : %s - %s (\"%s\")"),
+	UE_LOG(TimingProfiler, Log, TEXT("New Bottom Docked Track (%d) : %s (\"%s\")"),
 		BottomDockedTracks.Num() + 1,
-		*Track->GetType().ToString(),
-		*Track->GetSubType().ToString(),
+		*Track->GetTypeName().ToString(),
 		*Track->GetName());
 	check(!AllTracks.Contains(Track->GetId()));
 	AllTracks.Add(Track->GetId(), Track);
@@ -1385,10 +1387,9 @@ void STimingView::AddBottomDockedTrack(TSharedPtr<FBaseTimingTrack> Track)
 void STimingView::AddScrollableTrack(TSharedPtr<FBaseTimingTrack> Track)
 {
 	check(Track.IsValid());
-	UE_LOG(TimingProfiler, Log, TEXT("New Scrollable Track (%d) : %s - %s (\"%s\")"),
+	UE_LOG(TimingProfiler, Log, TEXT("New Scrollable Track (%d) : %s (\"%s\")"),
 		ScrollableTracks.Num() + 1,
-		*Track->GetType().ToString(),
-		*Track->GetSubType().ToString(),
+		*Track->GetTypeName().ToString(),
 		*Track->GetName());
 	check(!AllTracks.Contains(Track->GetId()));
 	AllTracks.Add(Track->GetId(), Track);
@@ -1419,10 +1420,9 @@ void STimingView::UpdateScrollableTracksOrder()
 void STimingView::AddForegroundTrack(TSharedPtr<FBaseTimingTrack> Track)
 {
 	check(Track.IsValid());
-	UE_LOG(TimingProfiler, Log, TEXT("New Foreground Track (%d) : %s - %s (\"%s\")"),
+	UE_LOG(TimingProfiler, Log, TEXT("New Foreground Track (%d) : %s (\"%s\")"),
 		ForegroundTracks.Num() + 1,
-		*Track->GetType().ToString(),
-		*Track->GetSubType().ToString(),
+		*Track->GetTypeName().ToString(),
 		*Track->GetName());
 	check(!AllTracks.Contains(Track->GetId()));
 	AllTracks.Add(Track->GetId(), Track);
@@ -2093,6 +2093,11 @@ FReply STimingView::OnKeyDown(const FGeometry& MyGeometry, const FKeyEvent& InKe
 		ShowHideGraphTrack_Execute();
 		return FReply::Handled();
 	}
+	else if (InKeyEvent.GetKey() == EKeys::R) // debug: toggles Frame tracks on/off
+	{
+		FrameSharedState->ShowHideAllFrameTracks();
+		return FReply::Handled();
+	}
 	else if (InKeyEvent.GetKey() == EKeys::Y) // debug: toggles GPU track on/off
 	{
 		ThreadTimingSharedState->ShowHideAllGpuTracks();
diff --git a/Engine/Source/Developer/TraceInsights/Private/Insights/Widgets/STimingView.h b/Engine/Source/Developer/TraceInsights/Private/Insights/Widgets/STimingView.h
--- a/Engine/Source/Developer/TraceInsights/Private/Insights/Widgets/STimingView.h
+++ b/Engine/Source/Developer/TraceInsights/Private/Insights/Widgets/STimingView.h
@@ -23,6 +23,7 @@
 #include "Insights/ViewModels/TooltipDrawState.h"
 
 class FFileActivitySharedState;
+class FFrameSharedState;
 class FLoadingSharedState;
 class FMarkersTimingTrack;
 class FMenuBuilder;
@@ -384,6 +385,9 @@ protected:
 
 	////////////////////////////////////////////////////////////
 
+	// Shared state for Frame Thread tracks
+	TSharedPtr<FFrameSharedState> FrameSharedState;
+
 	// Shared state for Cpu/Gpu Thread tracks
 	TSharedPtr<FThreadTimingSharedState> ThreadTimingSharedState;
 
diff --git a/Engine/Source/Developer/TraceInsights/Public/Insights/Common/SimpleRtti.h b/Engine/Source/Developer/TraceInsights/Public/Insights/Common/SimpleRtti.h
new file mode 100644
--- /dev/null
+++ b/Engine/Source/Developer/TraceInsights/Public/Insights/Common/SimpleRtti.h
@@ -0,0 +1,33 @@
+// Copyright 1998-2019 Epic Games, Inc. All Rights Reserved.
+
+#pragma once
+
+// Simple RTTI implementation.
+
+// Note: The SimpleRTTI class hierarchy may be different than the actual class hierarchy.
+//       Ex.: Some classes in the actual hierarchy can be hidden (do not use INSIGHTS_DECLARE_RTTI, so just inherits type info from parent class).
+
+// Macro to be used in declaration of the base class of the SimpleRTTI class hierarchy.
+#define INSIGHTS_DECLARE_RTTI_BASE(BaseClassName) \
+	public: \
+		static const FName& GetStaticTypeName() { return BaseClassName::TypeName; } \
+		virtual const FName& GetTypeName() const { return BaseClassName::GetStaticTypeName(); } \
+		virtual bool IsKindOf(const FName& InTypeName) const { return InTypeName == BaseClassName::GetStaticTypeName(); } \
+		template<typename Type> bool Is() const { return IsKindOf(Type::GetStaticTypeName()); } \
+		template<typename Type> const Type& As() const { return *static_cast<const Type*>(this); } \
+		template<typename Type> Type& As() { return *static_cast<Type*>(this); } \
+	private: \
+		static const FName TypeName;
+
+// Macro to be used in declaration of a derived class in a SimpleRTTI hierarchy.
+#define INSIGHTS_DECLARE_RTTI(ClassName, BaseClassName) \
+	public: \
+		static const FName& GetStaticTypeName() { return ClassName::TypeName; } \
+		virtual const FName& GetTypeName() const override { return ClassName::GetStaticTypeName(); } \
+		virtual bool IsKindOf(const FName& InTypeName) const override { return InTypeName == ClassName::GetStaticTypeName() || BaseClassName::IsKindOf(InTypeName); } \
+	private: \
+		static const FName TypeName;
+
+// Macro to be used in implementation of a class in a SimpleRTTI hierarchy.
+#define INSIGHTS_IMPLEMENT_RTTI(ClassName) \
+	const FName ClassName::TypeName = TEXT(#ClassName);
diff --git a/Engine/Source/Developer/TraceInsights/Public/Insights/ViewModels/BaseTimingTrack.h b/Engine/Source/Developer/TraceInsights/Public/Insights/ViewModels/BaseTimingTrack.h
--- a/Engine/Source/Developer/TraceInsights/Public/Insights/ViewModels/BaseTimingTrack.h
+++ b/Engine/Source/Developer/TraceInsights/Public/Insights/ViewModels/BaseTimingTrack.h
@@ -5,6 +5,8 @@
 #include "CoreMinimal.h"
 #include "Misc/EnumClassFlags.h"
 
+#include "Insights/Common/SimpleRtti.h"
+
 ////////////////////////////////////////////////////////////////////////////////////////////////////
 
 class FMenuBuilder;
@@ -86,11 +88,11 @@ class TRACEINSIGHTS_API FBaseTimingTrack : public TSharedFromThis<FBaseTimingTra
 {
 	friend class FTimingViewDrawHelper;
 
+	INSIGHTS_DECLARE_RTTI_BASE(FBaseTimingTrack)
+
 protected:
-	explicit FBaseTimingTrack(const FName& InType, const FName& InSubType = NAME_None, const FString& InName = FString())
+	explicit FBaseTimingTrack(const FString& InName = FString())
 		: Id(GenerateId())
-		, Type(InType)
-		, SubType(InSubType)
 		, Name(InName)
 		, Order(0)
 		, PosY(0.0f)
@@ -113,9 +115,6 @@ public:
 
 	uint64 GetId() const { return Id; }
 
-	const FName& GetType() const { return Type; }
-	const FName& GetSubType() const { return SubType; }
-
 	const FString& GetName() const { return Name; }
 	void SetName(const FString& InName) { Name = InName; }
 
@@ -226,8 +225,6 @@ protected:
 
 private:
 	const uint64 Id;
-	const FName Type; // Thread, Loading, FileActivity
-	const FName SubType; // Cpu, Gpu, MainThread, AsyncThread, Overview, Detailed
 	FString Name;
 	int32 Order;
 	float PosY; // y position, in Slate units
diff --git a/Engine/Source/Developer/TraceInsights/Public/Insights/ViewModels/GraphTrack.h b/Engine/Source/Developer/TraceInsights/Public/Insights/ViewModels/GraphTrack.h
--- a/Engine/Source/Developer/TraceInsights/Public/Insights/ViewModels/GraphTrack.h
+++ b/Engine/Source/Developer/TraceInsights/Public/Insights/ViewModels/GraphTrack.h
@@ -20,6 +20,8 @@ class TRACEINSIGHTS_API FGraphTrack : public FBaseTimingTrack
 {
 	friend class FGraphTrackBuilder;
 
+	INSIGHTS_DECLARE_RTTI(FGraphTrack, FBaseTimingTrack)
+
 private:
 	// Visual size of points (in pixels).
 	static constexpr float PointVisualSize = 5.0f;
@@ -29,8 +31,8 @@ private:
 	static constexpr float PointSizeY = 3.0f;
 
 public:
-	explicit FGraphTrack(const FName& InSubType = NAME_None);
-	explicit FGraphTrack(const FName& InSubType, const FString& InName);
+	explicit FGraphTrack();
+	explicit FGraphTrack(const FString& InName);
 	virtual ~FGraphTrack();
 
 	//////////////////////////////////////////////////
@@ -133,6 +135,8 @@ protected:
 
 class TRACEINSIGHTS_API FRandomGraphTrack : public FGraphTrack
 {
+	INSIGHTS_DECLARE_RTTI(FRandomGraphTrack, FGraphTrack)
+
 public:
 	FRandomGraphTrack();
 	virtual ~FRandomGraphTrack();
diff --git a/Engine/Source/Developer/TraceInsights/Public/Insights/ViewModels/GraphTrackEvent.h b/Engine/Source/Developer/TraceInsights/Public/Insights/ViewModels/GraphTrackEvent.h
--- a/Engine/Source/Developer/TraceInsights/Public/Insights/ViewModels/GraphTrackEvent.h
+++ b/Engine/Source/Developer/TraceInsights/Public/Insights/ViewModels/GraphTrackEvent.h
@@ -28,6 +28,8 @@ struct TRACEINSIGHTS_API FGraphSeriesEvent
 
 class TRACEINSIGHTS_API FGraphTrackEvent : public ITimingEvent
 {
+	INSIGHTS_DECLARE_RTTI(FGraphTrackEvent, ITimingEvent)
+
 public:
 	explicit FGraphTrackEvent(const TSharedRef<const FGraphTrack> InTrack, const TSharedRef<const FGraphSeries> InSeries, const FGraphSeriesEvent& InSeriesEvent)
 		: Track(InTrack), Series(InSeries), SeriesEvent(InSeriesEvent)
@@ -38,8 +40,6 @@ public:
 	//////////////////////////////////////////////////
 	// ITimingEvent interface
 
-	virtual const FName& GetTypeName() const override { return FGraphTrackEvent::TypeName; }
-
 	virtual const TSharedRef<const FBaseTimingTrack> GetTrack() const override;
 
 	virtual uint32 GetDepth() const override { return 0; }
@@ -58,7 +58,7 @@ public:
 			return false;
 		}
 
-		const FGraphTrackEvent& OtherGraphEvent = static_cast<const FGraphTrackEvent&>(Other);
+		const FGraphTrackEvent& OtherGraphEvent = Other.As<FGraphTrackEvent>();
 		return GetTrack() == Other.GetTrack()
 			&& Series == OtherGraphEvent.GetSeries()
 			&& SeriesEvent.Equals(OtherGraphEvent.GetSeriesEvent());
@@ -70,15 +70,10 @@ public:
 	const FGraphSeriesEvent& GetSeriesEvent() const { return SeriesEvent; }
 	double GetValue() const { return SeriesEvent.Value; }
 
-	static const FName& GetStaticTypeName() { return FGraphTrackEvent::TypeName; }
-	static bool CheckTypeName(const ITimingEvent& Event) { return Event.GetTypeName() == FGraphTrackEvent::TypeName; }
-
 private:
 	const TSharedRef<const FGraphTrack> Track;
 	const TSharedRef<const FGraphSeries> Series;
 	const FGraphSeriesEvent SeriesEvent;
-
-	static const FName TypeName;
 };
 
 ////////////////////////////////////////////////////////////////////////////////////////////////////
diff --git a/Engine/Source/Developer/TraceInsights/Public/Insights/ViewModels/ITimingEvent.h b/Engine/Source/Developer/TraceInsights/Public/Insights/ViewModels/ITimingEvent.h
--- a/Engine/Source/Developer/TraceInsights/Public/Insights/ViewModels/ITimingEvent.h
+++ b/Engine/Source/Developer/TraceInsights/Public/Insights/ViewModels/ITimingEvent.h
@@ -4,13 +4,17 @@
 
 #include "CoreMinimal.h"
 
+#include "Insights/Common/SimpleRtti.h"
+
 class FBaseTimingTrack;
 
+////////////////////////////////////////////////////////////////////////////////////////////////////
+
 class TRACEINSIGHTS_API ITimingEvent
 {
-public:
-	virtual const FName& GetTypeName() const = 0;
+	INSIGHTS_DECLARE_RTTI_BASE(ITimingEvent)
 
+public:
 	virtual const TSharedRef<const FBaseTimingTrack> GetTrack() const = 0;
 
 	bool CheckTrack(const FBaseTimingTrack* TrackPtr) const { return &GetTrack().Get() == TrackPtr; }
@@ -34,13 +38,28 @@ public:
 	}
 };
 
+////////////////////////////////////////////////////////////////////////////////////////////////////
+
 class TRACEINSIGHTS_API ITimingEventFilter
 {
-public:
-	virtual const FName& GetTypeName() const = 0;
+	INSIGHTS_DECLARE_RTTI_BASE(ITimingEventFilter)
 
+public:
+	/**
+	 * Returns true if the track passes the filter.
+	 */
 	virtual bool FilterTrack(const FBaseTimingTrack& InTrack) const = 0;
-	virtual bool FilterEvent(double InEventStartTime, double InEventEndTime, uint32 InEventDepth, const TCHAR* InEventName = nullptr, uint64 InEventType = 0, uint32 InEventColor = 0) const = 0;
 
+	/**
+	 * Returns true if the timing event passes the filter.
+	 */
+	virtual bool FilterEvent(const ITimingEvent& InEvent) const = 0;
+	virtual bool FilterEvent(double InEventStartTime, double InEventEndTime, uint32 InEventDepth, const TCHAR* InEventName, uint64 InEventType = 0, uint32 InEventColor = 0) const = 0;
+
+	/**
+	 * Returns a number that changes each time an attribute of this filter changes.
+	 */
 	virtual uint32 GetChangeNumber() const = 0;
 };
+
+////////////////////////////////////////////////////////////////////////////////////////////////////
diff --git a/Engine/Source/Developer/TraceInsights/Public/Insights/ViewModels/TimingEvent.h b/Engine/Source/Developer/TraceInsights/Public/Insights/ViewModels/TimingEvent.h
--- a/Engine/Source/Developer/TraceInsights/Public/Insights/ViewModels/TimingEvent.h
+++ b/Engine/Source/Developer/TraceInsights/Public/Insights/ViewModels/TimingEvent.h
@@ -14,6 +14,8 @@ class FBaseTimingTrack;
 
 class TRACEINSIGHTS_API FTimingEvent : public ITimingEvent
 {
+	INSIGHTS_DECLARE_RTTI(FTimingEvent, ITimingEvent)
+
 public:
 	FTimingEvent(const TSharedRef<const FBaseTimingTrack> InTrack, double InStartTime, double InEndTime, uint32 InDepth, uint64 InType = uint64(-1))
 		: Track(InTrack)
@@ -35,8 +37,6 @@ public:
 	//////////////////////////////////////////////////
 	// ITimingEvent interface
 
-	virtual const FName& GetTypeName() const override { return FTimingEvent::TypeName; }
-
 	virtual const TSharedRef<const FBaseTimingTrack> GetTrack() const override { return Track; }
 
 	virtual uint32 GetDepth() const override { return Depth; }
@@ -52,7 +52,7 @@ public:
 			return false;
 		}
 
-		const FTimingEvent& OtherTimingEvent = static_cast<const FTimingEvent&>(Other);
+		const FTimingEvent& OtherTimingEvent = Other.As<FTimingEvent>();
 		return Track == Other.GetTrack()
 			&& Depth == OtherTimingEvent.GetDepth()
 			&& StartTime == OtherTimingEvent.GetStartTime()
@@ -68,9 +68,6 @@ public:
 
 	uint64 GetType() const { return Type; }
 
-	static const FName& GetStaticTypeName() { return FTimingEvent::TypeName; }
-	static bool CheckTypeName(const ITimingEvent& Event) { return Event.GetTypeName() == FTimingEvent::TypeName; }
-
 private:
 	// The track this timing event is contained within
 	TSharedRef<const FBaseTimingTrack> Track;
@@ -92,22 +89,185 @@ private:
 
 	// The event type (category, group id, etc.).
 	uint64 Type;
+};
+
+////////////////////////////////////////////////////////////////////////////////////////////////////
+
+class FAcceptNoneTimingEventFilter : public ITimingEventFilter
+{
+	INSIGHTS_DECLARE_RTTI(FAcceptNoneTimingEventFilter, ITimingEventFilter)
+
+public:
+	FAcceptNoneTimingEventFilter() {}
+	virtual ~FAcceptNoneTimingEventFilter() {}
+
+	virtual bool FilterTrack(const FBaseTimingTrack& InTrack) const override { return false; }
+	virtual bool FilterEvent(const ITimingEvent& InEvent) const override { return false; }
+	virtual bool FilterEvent(double InEventStartTime, double InEventEndTime, uint32 InEventDepth, const TCHAR* InEventName, uint64 InEventType = 0, uint32 InEventColor = 0) const override { return false; }
+	virtual uint32 GetChangeNumber() const override { return 0; }
+};
+
+////////////////////////////////////////////////////////////////////////////////////////////////////
+
+class FAcceptAllTimingEventFilter : public ITimingEventFilter
+{
+	INSIGHTS_DECLARE_RTTI(FAcceptAllTimingEventFilter, ITimingEventFilter)
+
+public:
+	FAcceptAllTimingEventFilter() {}
+	virtual ~FAcceptAllTimingEventFilter() {}
+
+	virtual bool FilterTrack(const FBaseTimingTrack& InTrack) const override { return true; }
+	virtual bool FilterEvent(const ITimingEvent& InEvent) const override { return true; }
+	virtual bool FilterEvent(double InEventStartTime, double InEventEndTime, uint32 InEventDepth, const TCHAR* InEventName, uint64 InEventType = 0, uint32 InEventColor = 0) const override { return true; }
+	virtual uint32 GetChangeNumber() const override { return 0; }
+};
+
+////////////////////////////////////////////////////////////////////////////////////////////////////
+
+class FAggregatedTimingEventFilter : public ITimingEventFilter
+{
+	INSIGHTS_DECLARE_RTTI(FAggregatedTimingEventFilter, ITimingEventFilter)
+
+public:
+	FAggregatedTimingEventFilter() : ChangeNumber(0) {}
+	virtual ~FAggregatedTimingEventFilter() {}
+
+	//////////////////////////////////////////////////
+	// ITimingEventFilter interface
+
+	//virtual bool FilterTrack(const FBaseTimingTrack& InTrack) const = 0;
+	//virtual bool FilterEvent(const ITimingEvent& InEvent) const = 0;
+	//virtual bool FilterEvent(double InEventStartTime, double InEventEndTime, uint32 InEventDepth, const TCHAR* InEventName = nullptr, uint64 InEventType = 0, uint32 InEventColor = 0) const = 0;
+	virtual uint32 GetChangeNumber() const override { return ChangeNumber; }
+
+	//////////////////////////////////////////////////
+
+	const TArray<TSharedPtr<ITimingEventFilter>> GetChildren() const { return Children; }
+	void AddChild(TSharedPtr<ITimingEventFilter> InFilter) { Children.Add(InFilter); ChangeNumber++; }
+
+protected:
+	uint32 ChangeNumber;
+	TArray<TSharedPtr<ITimingEventFilter>> Children;
+};
+
+////////////////////////////////////////////////////////////////////////////////////////////////////
+
+class FAllAggregatedTimingEventFilter : public FAggregatedTimingEventFilter
+{
+	INSIGHTS_DECLARE_RTTI(FAllAggregatedTimingEventFilter, FAggregatedTimingEventFilter)
+
+public:
+	FAllAggregatedTimingEventFilter() {}
+	virtual ~FAllAggregatedTimingEventFilter() {}
+
+	//////////////////////////////////////////////////
+	// ITimingEventFilter interface
+
+	virtual bool FilterTrack(const FBaseTimingTrack& InTrack) const override
+	{
+		for (TSharedPtr<ITimingEventFilter> Filter : Children)
+		{
+			if (!FilterTrack(InTrack))
+			{
+				return false;
+			}
+		}
+		return true;
+	}
+
+	virtual bool FilterEvent(const ITimingEvent& InEvent) const override
+	{
+		for (TSharedPtr<ITimingEventFilter> Filter : Children)
+		{
+			if (FilterEvent(InEvent))
+			{
+				return false;
+			}
+		}
+		return true;
+	}
+
+	virtual bool FilterEvent(double InEventStartTime, double InEventEndTime, uint32 InEventDepth, const TCHAR* InEventName, uint64 InEventType = 0, uint32 InEventColor = 0) const override
+	{
+		for (TSharedPtr<ITimingEventFilter> Filter : Children)
+		{
+			if (!FilterEvent(InEventStartTime, InEventEndTime, InEventDepth, InEventName, InEventType, InEventColor))
+			{
+				return false;
+			}
+		}
+		return true;
+	}
+
+	//////////////////////////////////////////////////
+};
+
+////////////////////////////////////////////////////////////////////////////////////////////////////
+
+class FAnyAggregatedTimingEventFilter : public FAggregatedTimingEventFilter
+{
+	INSIGHTS_DECLARE_RTTI(FAnyAggregatedTimingEventFilter, FAggregatedTimingEventFilter)
+
+public:
+	FAnyAggregatedTimingEventFilter() {}
+	virtual ~FAnyAggregatedTimingEventFilter() {}
+
+	//////////////////////////////////////////////////
+	// ITimingEventFilter interface
+
+	virtual bool FilterTrack(const FBaseTimingTrack& InTrack) const override
+	{
+		for (TSharedPtr<ITimingEventFilter> Filter : Children)
+		{
+			if (FilterTrack(InTrack))
+			{
+				return true;
+			}
+		}
+		return false;
+	}
+
+	virtual bool FilterEvent(const ITimingEvent& InEvent) const override
+	{
+		for (TSharedPtr<ITimingEventFilter> Filter : Children)
+		{
+			if (FilterEvent(InEvent))
+			{
+				return true;
+			}
+		}
+		return false;
+	}
 
-	static const FName TypeName;
+	virtual bool FilterEvent(double InEventStartTime, double InEventEndTime, uint32 InEventDepth, const TCHAR* InEventName, uint64 InEventType = 0, uint32 InEventColor = 0) const override
+	{
+		for (TSharedPtr<ITimingEventFilter> Filter : Children)
+		{
+			if (FilterEvent(InEventStartTime, InEventEndTime, InEventDepth, InEventName, InEventType, InEventColor))
+			{
+				return true;
+			}
+		}
+		return false;
+	}
+
+	//////////////////////////////////////////////////
 };
 
 ////////////////////////////////////////////////////////////////////////////////////////////////////
 
 class TRACEINSIGHTS_API FTimingEventFilter : public ITimingEventFilter
 {
+	INSIGHTS_DECLARE_RTTI(FTimingEventFilter, ITimingEventFilter)
+
 public:
 	FTimingEventFilter()
-		: bFilterByTrackType(false)
-		, TrackType(NAME_None)
-		, bFilterByEventType(false)
-		, EventType(uint64(-1))
-		, bFilterByMinDuration(false)
-		, MinDuration(0.0)
+		: bFilterByTrackTypeName(false)
+		, TrackTypeName(NAME_None)
+		, bFilterByTrackInstance(false)
+		, TrackInstance(nullptr)
+		, ChangeNumber(0)
 	{}
 
 	virtual ~FTimingEventFilter() {}
@@ -121,70 +281,171 @@ public:
 	//////////////////////////////////////////////////
 	// ITimingEventFilter interface
 
-	virtual const FName& GetTypeName() const override { return FTimingEventFilter::TypeName; }
-
 	virtual bool FilterTrack(const FBaseTimingTrack& InTrack) const override;
+	virtual bool FilterEvent(const ITimingEvent& InEvent) const override { return true; }
+	virtual bool FilterEvent(double InEventStartTime, double InEventEndTime, uint32 InEventDepth, const TCHAR* InEventName, uint64 InEventType = 0, uint32 InEventColor = 0) const override { return true; }
+	virtual uint32 GetChangeNumber() const override { return ChangeNumber; }
 
-	virtual bool FilterEvent(double InEventStartTime, double InEventEndTime, uint32 InEventDepth, const TCHAR* InEventName, uint64 InEventType = 0, uint32 InEventColor = 0) const override
-	{
-		return (!bFilterByEventType || InEventType == EventType)
-			&& (!bFilterByMinDuration || InEventEndTime - InEventStartTime >= MinDuration)
-			&& (!bFilterByMaxDuration || InEventEndTime - InEventStartTime <= MaxDuration);
-	}
+	//////////////////////////////////////////////////
 
-	virtual uint32 GetChangeNumber() const override { return ChangeNumber; }
+	bool IsFilteringByTrackTypeName() const { return bFilterByTrackTypeName; }
+	const FName& GetTrackTypeName() const { return TrackTypeName; }
+	void SetFilterByTrackTypeName(bool bInFilterByTrackTypeName) { if (bFilterByTrackTypeName != bInFilterByTrackTypeName) { bFilterByTrackTypeName = bInFilterByTrackTypeName; ChangeNumber++; } }
+	void SetTrackTypeName(const FName InTrackTypeName) { if (TrackTypeName != InTrackTypeName) { TrackTypeName = InTrackTypeName; ChangeNumber++; } }
+
+	bool IsFilteringByTrackInstance() const { return bFilterByTrackInstance; }
+	TSharedPtr<FBaseTimingTrack> GetTrackInstance() const { return TrackInstance; }
+	void SetFilterByTrackInstance(bool bInFilterByTrackInstance) { if (bFilterByTrackInstance != bInFilterByTrackInstance) { bFilterByTrackInstance = bInFilterByTrackInstance; ChangeNumber++; } }
+	void SetTrackInstance(TSharedPtr<FBaseTimingTrack> InTrackInstance) { if (TrackInstance != InTrackInstance) { TrackInstance = InTrackInstance; ChangeNumber++; } }
 
 	//////////////////////////////////////////////////
 
-	bool IsFilteringByTrackType() const { return bFilterByTrackType; }
-	const FName& GetTrackType() const { return TrackType; }
-	void SetFilterByTrackType(bool bInFilterByTrackType) { if (bFilterByTrackType != bInFilterByTrackType) { bFilterByTrackType = bInFilterByTrackType; ChangeNumber++; } }
-	void SetTrackType(const FName InTrackType) { if (TrackType != InTrackType) { TrackType = InTrackType; ChangeNumber++; } }
+protected:
+	bool bFilterByTrackTypeName;
+	FName TrackTypeName;
 
-	bool IsFilteringByEventType() const { return bFilterByEventType; }
-	uint64 GetEventType() const { return EventType; }
-	void SetFilterByEventType(bool bInFilterByEventType) { if (bFilterByEventType != bInFilterByEventType) { bFilterByEventType = bInFilterByEventType; ChangeNumber++; } }
-	void SetEventType(uint64 InEventType) { if (EventType != InEventType) { EventType = InEventType; ChangeNumber++; } }
+	bool bFilterByTrackInstance;
+	TSharedPtr<FBaseTimingTrack> TrackInstance;
 
-	bool IsFilteringByMinDuration() const { return bFilterByMinDuration; }
-	double GetMinDuration() const { return MinDuration; }
+	uint32 ChangeNumber;
+};
 
-	bool IsFilteringByMaxDuration() const { return bFilterByMaxDuration; }
-	double GetMaxDuration() const { return MaxDuration; }
+////////////////////////////////////////////////////////////////////////////////////////////////////
+
+class FTimingEventFilterByMinDuration : public FTimingEventFilter
+{
+	INSIGHTS_DECLARE_RTTI(FTimingEventFilterByMinDuration, FTimingEventFilter)
 
-	static const FName& GetStaticTypeName() { return FTimingEventFilter::TypeName; }
-	static bool CheckTypeName(const ITimingEventFilter& EventFilter) { return EventFilter.GetTypeName() == FTimingEventFilter::TypeName; }
+public:
+	FTimingEventFilterByMinDuration(double InMinDuration) : MinDuration(InMinDuration) {}
+	virtual ~FTimingEventFilterByMinDuration() {}
 
-private:
-	bool FilterEventByType(double InEventStartTime, double InEventEndTime, uint32 InEventDepth, const TCHAR* InEventName, uint64 InEventType, uint32 InEventColor) const
+	//////////////////////////////////////////////////
+	// ITimingEventFilter interface
+
+	virtual bool FilterEvent(const ITimingEvent& InEvent) const override
 	{
-		return InEventType == EventType;
+		return InEvent.GetDuration() >= MinDuration;
 	}
-	bool FilterEventByMinDuration(double InEventStartTime, double InEventEndTime, uint32 InEventDepth, const TCHAR* InEventName, uint64 InEventType, uint32 InEventColor) const
+
+	virtual bool FilterEvent(double InEventStartTime, double InEventEndTime, uint32 InEventDepth, const TCHAR* InEventName, uint64 InEventType = 0, uint32 InEventColor = 0) const override
 	{
 		return InEventEndTime - InEventStartTime >= MinDuration;
 	}
-	bool FilterEventByMaxDuration(double InEventStartTime, double InEventEndTime, uint32 InEventDepth, const TCHAR* InEventName, uint64 InEventType, uint32 InEventColor) const
+
+	//////////////////////////////////////////////////
+
+	double GetMinDuration() const { return MinDuration; }
+	void SetMinDuration(double InMinDuration) { if (MinDuration != InMinDuration) { MinDuration = InMinDuration; ChangeNumber++; } }
+
+private:
+	double MinDuration;
+};
+
+////////////////////////////////////////////////////////////////////////////////////////////////////
+
+class FTimingEventFilterByMaxDuration : public FTimingEventFilter
+{
+	INSIGHTS_DECLARE_RTTI(FTimingEventFilterByMaxDuration, FTimingEventFilter)
+
+public:
+	FTimingEventFilterByMaxDuration(double InMaxDuration) : MaxDuration(InMaxDuration) {}
+	virtual ~FTimingEventFilterByMaxDuration() {}
+
+	//////////////////////////////////////////////////
+	// ITimingEventFilter interface
+
+	virtual bool FilterEvent(const ITimingEvent& InEvent) const override
+	{
+		return InEvent.GetDuration() <= MaxDuration;
+	}
+
+	virtual bool FilterEvent(double InEventStartTime, double InEventEndTime, uint32 InEventDepth, const TCHAR* InEventName, uint64 InEventType = 0, uint32 InEventColor = 0) const override
 	{
 		return InEventEndTime - InEventStartTime <= MaxDuration;
 	}
 
+	//////////////////////////////////////////////////
+
+	double GetMaxDuration() const { return MaxDuration; }
+	void SetMaxDuration(double InMaxDuration) { if (MaxDuration != InMaxDuration) { MaxDuration = InMaxDuration; ChangeNumber++; } }
+
 private:
-	bool bFilterByTrackType;
-	FName TrackType;
+	double MaxDuration;
+};
+
+////////////////////////////////////////////////////////////////////////////////////////////////////
+
+class TRACEINSIGHTS_API FTimingEventFilterByEventType : public FTimingEventFilter
+{
+	INSIGHTS_DECLARE_RTTI(FTimingEventFilterByEventType, FTimingEventFilter)
+
+public:
+	FTimingEventFilterByEventType(uint64 InEventType) : EventType(InEventType) {}
+	virtual ~FTimingEventFilterByEventType() {}
+
+	//////////////////////////////////////////////////
+	// ITimingEventFilter interface
+
+	virtual bool FilterEvent(const ITimingEvent& InEvent) const override
+	{
+		if (InEvent.Is<FTimingEvent>())
+		{
+			const FTimingEvent& Event = InEvent.As<FTimingEvent>();
+			return Event.GetType() == EventType;
+		}
+		return false;
+	}
+
+	virtual bool FilterEvent(double InEventStartTime, double InEventEndTime, uint32 InEventDepth, const TCHAR* InEventName, uint64 InEventType = 0, uint32 InEventColor = 0) const override
+	{
+		return InEventType == EventType;
+	}
+
+	//////////////////////////////////////////////////
+
+	uint64 GetEventType() const { return EventType; }
+	void SetEventType(uint64 InEventType) { if (EventType != InEventType) { EventType = InEventType; ChangeNumber++; } }
 
-	bool bFilterByEventType;
+private:
 	uint64 EventType;
+};
 
-	bool bFilterByMinDuration;
-	double MinDuration;
+////////////////////////////////////////////////////////////////////////////////////////////////////
 
-	bool bFilterByMaxDuration;
-	double MaxDuration;
+class FTimingEventFilterByFrameIndex : public FTimingEventFilter
+{
+	INSIGHTS_DECLARE_RTTI(FTimingEventFilterByFrameIndex, FTimingEventFilter)
 
-	uint32 ChangeNumber;
+public:
+	FTimingEventFilterByFrameIndex(uint64 InFrameIndex) : FrameIndex(InFrameIndex) {}
+	virtual ~FTimingEventFilterByFrameIndex() {}
+
+	//////////////////////////////////////////////////
+	// ITimingEventFilter interface
 
-	static const FName TypeName;
+	virtual bool FilterEvent(const ITimingEvent& InEvent) const override
+	{
+		if (InEvent.Is<FTimingEvent>())
+		{
+			const FTimingEvent& Event = InEvent.As<FTimingEvent>();
+			return Event.GetType() == FrameIndex;
+		}
+		return false;
+	}
+
+	virtual bool FilterEvent(double InEventStartTime, double InEventEndTime, uint32 InEventDepth, const TCHAR* InEventName, uint64 InFrameIndex = 0, uint32 InEventColor = 0) const override
+	{
+		return InFrameIndex == FrameIndex;
+	}
+
+	//////////////////////////////////////////////////
+
+	uint64 GetFrameIndex() const { return FrameIndex; }
+	void SetFrameIndex(uint64 InFrameIndex) { if (FrameIndex != InFrameIndex) { FrameIndex = InFrameIndex; ChangeNumber++; } }
+
+private:
+	uint64 FrameIndex;
 };
 
 ////////////////////////////////////////////////////////////////////////////////////////////////////
diff --git a/Engine/Source/Developer/TraceInsights/Public/Insights/ViewModels/TimingEventsTrack.h b/Engine/Source/Developer/TraceInsights/Public/Insights/ViewModels/TimingEventsTrack.h
--- a/Engine/Source/Developer/TraceInsights/Public/Insights/ViewModels/TimingEventsTrack.h
+++ b/Engine/Source/Developer/TraceInsights/Public/Insights/ViewModels/TimingEventsTrack.h
@@ -21,8 +21,11 @@ public:
 
 class TRACEINSIGHTS_API FTimingEventsTrack : public FBaseTimingTrack
 {
+	INSIGHTS_DECLARE_RTTI(FTimingEventsTrack, FBaseTimingTrack)
+
 public:
-	explicit FTimingEventsTrack(const FName& InType, const FName& InSubType, const FString& InName);
+	explicit FTimingEventsTrack();
+	explicit FTimingEventsTrack(const FString& InTrackName);
 	virtual ~FTimingEventsTrack();
 
 	//////////////////////////////////////////////////
@@ -75,7 +78,7 @@ private:
 	FFilteredDrawStateInfo FilteredDrawStateInfo;
 
 public:
-	static bool bUseDownSampling; // toggle to enable/disbale downsampling, for debugging purposes only
+	static bool bUseDownSampling; // toggle to enable/disable downsampling, for debugging purposes only
 };
 
 ////////////////////////////////////////////////////////////////////////////////////////////////////
