Fixed linker error.

#rnx

diff --git a/Engine/Source/Developer/TraceInsights/Public/Insights/ViewModels/TimingEvent.h b/Engine/Source/Developer/TraceInsights/Public/Insights/ViewModels/TimingEvent.h
--- a/Engine/Source/Developer/TraceInsights/Public/Insights/ViewModels/TimingEvent.h
+++ b/Engine/Source/Developer/TraceInsights/Public/Insights/ViewModels/TimingEvent.h
@@ -281,7 +281,7 @@ public:
 	//////////////////////////////////////////////////
 	// ITimingEventFilter interface
 
-	virtual bool FilterTrack(const FBaseTimingTrack& InTrack) const override;
+	virtual bool FilterTrack(const FBaseTimingTrack& InTrack) const override { return true; }
 	virtual bool FilterEvent(const ITimingEvent& InEvent) const override { return true; }
 	virtual bool FilterEvent(double InEventStartTime, double InEventEndTime, uint32 InEventDepth, const TCHAR* InEventName, uint64 InEventType = 0, uint32 InEventColor = 0) const override { return true; }
 	virtual uint32 GetChangeNumber() const override { return ChangeNumber; }
