[Insights] Missing file.

diff --git a/Engine/Source/Developer/TraceInsights/Private/Insights/ViewModels/TimingEvent.cpp b/Engine/Source/Developer/TraceInsights/Private/Insights/ViewModels/TimingEvent.cpp
new file mode 100644
--- /dev/null
+++ b/Engine/Source/Developer/TraceInsights/Private/Insights/ViewModels/TimingEvent.cpp
@@ -0,0 +1,42 @@
+// Copyright 1998-2019 Epic Games, Inc. All Rights Reserved.
+
+#include "Insights/ViewModels/TimingEvent.h"
+
+#include "Insights/ViewModels/BaseTimingTrack.h"
+
+#define LOCTEXT_NAMESPACE "TimingEvent"
+
+////////////////////////////////////////////////////////////////////////////////////////////////////
+// FTimingEvent
+////////////////////////////////////////////////////////////////////////////////////////////////////
+
+INSIGHTS_IMPLEMENT_RTTI(ITimingEvent)
+INSIGHTS_IMPLEMENT_RTTI(FTimingEvent)
+
+////////////////////////////////////////////////////////////////////////////////////////////////////
+// ITimingEventFilter implementations
+////////////////////////////////////////////////////////////////////////////////////////////////////
+
+INSIGHTS_IMPLEMENT_RTTI(ITimingEventFilter)
+INSIGHTS_IMPLEMENT_RTTI(FTimingEventFilter)
+INSIGHTS_IMPLEMENT_RTTI(FAcceptNoneTimingEventFilter)
+INSIGHTS_IMPLEMENT_RTTI(FAcceptAllTimingEventFilter)
+INSIGHTS_IMPLEMENT_RTTI(FAggregatedTimingEventFilter)
+INSIGHTS_IMPLEMENT_RTTI(FAllAggregatedTimingEventFilter)
+INSIGHTS_IMPLEMENT_RTTI(FAnyAggregatedTimingEventFilter)
+INSIGHTS_IMPLEMENT_RTTI(FTimingEventFilterByMinDuration)
+INSIGHTS_IMPLEMENT_RTTI(FTimingEventFilterByMaxDuration)
+INSIGHTS_IMPLEMENT_RTTI(FTimingEventFilterByEventType)
+INSIGHTS_IMPLEMENT_RTTI(FTimingEventFilterByFrameIndex)
+
+////////////////////////////////////////////////////////////////////////////////////////////////////
+
+bool FTimingEventFilter::FilterTrack(const FBaseTimingTrack& InTrack) const
+{
+	return !bFilterByTrackTypeName || InTrack.IsKindOf(TrackTypeName) &&
+	       !bFilterByTrackInstance || &InTrack == TrackInstance.Get();
+}
+
+////////////////////////////////////////////////////////////////////////////////////////////////////
+
+#undef LOCTEXT_NAMESPACE
