[Insights] Fixed FTimingEventFilter::FilterTrack.

diff --git a/Engine/Source/Developer/TraceInsights/Private/Insights/ViewModels/TimingEvent.cpp b/Engine/Source/Developer/TraceInsights/Private/Insights/ViewModels/TimingEvent.cpp
--- a/Engine/Source/Developer/TraceInsights/Private/Insights/ViewModels/TimingEvent.cpp
+++ b/Engine/Source/Developer/TraceInsights/Private/Insights/ViewModels/TimingEvent.cpp
@@ -7,14 +7,14 @@
 #define LOCTEXT_NAMESPACE "TimingEvent"
 
 ////////////////////////////////////////////////////////////////////////////////////////////////////
-// FTimingEvent
+// ITimingEvent, FTimingEvent
 ////////////////////////////////////////////////////////////////////////////////////////////////////
 
 INSIGHTS_IMPLEMENT_RTTI(ITimingEvent)
 INSIGHTS_IMPLEMENT_RTTI(FTimingEvent)
 
 ////////////////////////////////////////////////////////////////////////////////////////////////////
-// ITimingEventFilter implementations
+// ITimingEventFilter, FTimingEventFilter, ...
 ////////////////////////////////////////////////////////////////////////////////////////////////////
 
 INSIGHTS_IMPLEMENT_RTTI(ITimingEventFilter)
@@ -33,8 +33,8 @@ INSIGHTS_IMPLEMENT_RTTI(FTimingEventFilterByFrameIndex)
 
 bool FTimingEventFilter::FilterTrack(const FBaseTimingTrack& InTrack) const
 {
-	return !bFilterByTrackTypeName || InTrack.IsKindOf(TrackTypeName) &&
-	       !bFilterByTrackInstance || &InTrack == TrackInstance.Get();
+	return (!bFilterByTrackTypeName || InTrack.IsKindOf(TrackTypeName)) &&
+	       (!bFilterByTrackInstance || &InTrack == TrackInstance.Get());
 }
 
 ////////////////////////////////////////////////////////////////////////////////////////////////////
