Misc. graph track fixes

Added bitfield defining how the context menu gets built, so menu items are hidden rather than disabled in tracks that override default behaviour.
Fixed series list not dirtying the underlying track when series are made visible/invisible.
Fixed tick record track tooltips to display only the hovered series event info, rather than all series.

diff --git a/Engine/Plugins/Animation/GameplayInsights/Source/GameplayInsights/Private/AnimationTickRecordsTrack.cpp b/Engine/Plugins/Animation/GameplayInsights/Source/GameplayInsights/Private/AnimationTickRecordsTrack.cpp
--- a/Engine/Plugins/Animation/GameplayInsights/Source/GameplayInsights/Private/AnimationTickRecordsTrack.cpp
+++ b/Engine/Plugins/Animation/GameplayInsights/Source/GameplayInsights/Private/AnimationTickRecordsTrack.cpp
@@ -72,6 +72,7 @@ FAnimationTickRecordsTrack::FAnimationTickRecordsTrack(const FAnimationSharedDat
 	, AssetId(InAssetId)
 	, NodeId(InNodeId)
 	, RequestedTrackSizeScale(1.0f)
+	, BorderY(0.0f)
 {
 	uint32 Hash = GetTypeHash(GetName());
 	MainSeriesLineColor = MakeSeriesColor(Hash, true);
@@ -83,6 +84,7 @@ FAnimationTickRecordsTrack::FAnimationTickRecordsTrack(const FAnimationSharedDat
 	bDrawBoxes = false;
 	bDrawBaseline = false;
 	bUseEventDuration = false;
+	VisibleOptions &= ~(EGraphOptions::ShowBars | EGraphOptions::UseEventDuration);
 
 #if WITH_EDITOR
 	const FGameplayProvider* GameplayProvider = SharedData.GetAnalysisSession().ReadProvider<FGameplayProvider>(FGameplayProvider::ProviderName);
@@ -367,25 +369,18 @@ void FAnimationTickRecordsTrack::Draw(const ITimingTrackDrawContext& Context) co
 
 void FAnimationTickRecordsTrack::InitTooltip(FTooltipDrawState& Tooltip, const ITimingEvent& HoveredTimingEvent) const
 {
+	const FGraphTrackEvent& GraphTrackEvent = *static_cast<const FGraphTrackEvent*>(&HoveredTimingEvent);
+
 	FTimingEventSearchParameters SearchParameters(HoveredTimingEvent.GetStartTime(), HoveredTimingEvent.GetEndTime(), ETimingEventSearchFlags::StopAtFirstMatch);
 
-	FindTickRecordMessage(SearchParameters, [this, &Tooltip](double InFoundStartTime, double InFoundEndTime, uint32 InFoundDepth, const FTickRecordMessage& InMessage)
+	FindTickRecordMessage(SearchParameters, [this, &GraphTrackEvent, &Tooltip](double InFoundStartTime, double InFoundEndTime, uint32 InFoundDepth, const FTickRecordMessage& InMessage)
 	{
 		Tooltip.ResetContent();
 
 		Tooltip.AddTitle(GetName());
 
 		Tooltip.AddNameValueTextLine(LOCTEXT("EventTime", "Time").ToString(), FText::AsNumber(InFoundStartTime).ToString());
-		Tooltip.AddNameValueTextLine(LOCTEXT("BlendWeight", "Blend Weight").ToString(), FText::AsNumber(InMessage.BlendWeight).ToString());
-		if(InMessage.bIsBlendSpace)
-		{
-			Tooltip.AddNameValueTextLine(LOCTEXT("BlendSpacePositionX", "X").ToString(), FText::AsNumber(InMessage.BlendSpacePositionX).ToString());
-			Tooltip.AddNameValueTextLine(LOCTEXT("BlendSpacePositionX", "Y").ToString(), FText::AsNumber(InMessage.BlendSpacePositionY).ToString());
-		}
-		Tooltip.AddNameValueTextLine(LOCTEXT("PlaybackTime", "Playback Time").ToString(), FText::AsNumber(InMessage.PlaybackTime).ToString());
-		Tooltip.AddNameValueTextLine(LOCTEXT("RootMotionWeight", "Root Motion Weight").ToString(), FText::AsNumber(InMessage.RootMotionWeight).ToString());
-		Tooltip.AddNameValueTextLine(LOCTEXT("PlayRate", "Play Rate").ToString(), FText::AsNumber(InMessage.PlayRate).ToString());
-		Tooltip.AddNameValueTextLine(LOCTEXT("Looping", "Looping").ToString(), InMessage.bLooping ? LOCTEXT("True", "True").ToString() : LOCTEXT("False", "False").ToString());
+		Tooltip.AddNameValueTextLine(GraphTrackEvent.GetSeries()->GetName().ToString(), FText::AsNumber(GraphTrackEvent.GetValue()).ToString());
 
 		Tooltip.UpdateLayout();
 	});
diff --git a/Engine/Source/Developer/TraceInsights/Private/Insights/ViewModels/GraphTrack.cpp b/Engine/Source/Developer/TraceInsights/Private/Insights/ViewModels/GraphTrack.cpp
--- a/Engine/Source/Developer/TraceInsights/Private/Insights/ViewModels/GraphTrack.cpp
+++ b/Engine/Source/Developer/TraceInsights/Private/Insights/ViewModels/GraphTrack.cpp
@@ -46,6 +46,8 @@ FGraphTrack::FGraphTrack()
 	, bUseEventDuration(true)
 	, bDrawBoxes(true)
 	, bDrawBaseline(true)
+	, VisibleOptions(EGraphOptions::All)
+	, EditableOptions(EGraphOptions::All)
 	, NumAddedEvents(0)
 	, NumDrawPoints(0)
 	, NumDrawLines(0)
@@ -70,6 +72,8 @@ FGraphTrack::FGraphTrack(const FString& InName)
 	, bUseEventDuration(true)
 	, bDrawBoxes(true)
 	, bDrawBaseline(true)
+	, VisibleOptions(EGraphOptions::All)
+	, EditableOptions(EGraphOptions::All)
 	, NumAddedEvents(0)
 	, NumDrawPoints(0)
 	, NumDrawLines(0)
@@ -622,7 +626,7 @@ void FGraphTrack::BuildContextMenu(FMenuBuilder& MenuBuilder)
 {
 	MenuBuilder.BeginSection("Misc");
 	{
-		if (false) // debug functionality
+		if (EnumHasAnyFlags(VisibleOptions, EGraphOptions::ShowDebugInfo)) // debug functionality
 		{
 			FUIAction Action_ShowDebugInfo
 			(
@@ -641,101 +645,119 @@ void FGraphTrack::BuildContextMenu(FMenuBuilder& MenuBuilder)
 			);
 		}
 
-		FUIAction Action_ShowPoints
-		(
-			FExecuteAction::CreateSP(this, &FGraphTrack::ContextMenu_ShowPoints_Execute),
-			FCanExecuteAction::CreateSP(this, &FGraphTrack::ContextMenu_ShowPoints_CanExecute),
-			FIsActionChecked::CreateSP(this, &FGraphTrack::ContextMenu_ShowPoints_IsChecked)
-		);
-		MenuBuilder.AddMenuEntry
-		(
-			LOCTEXT("ContextMenu_ShowPoints", "Show Points"),
-			LOCTEXT("ContextMenu_ShowPoints_Desc", "Show points."),
-			FSlateIcon(),
-			Action_ShowPoints,
-			NAME_None,
-			EUserInterfaceActionType::ToggleButton
-		);
-
-		FUIAction Action_ShowPointsWithBorder
-		(
-			FExecuteAction::CreateSP(this, &FGraphTrack::ContextMenu_ShowPointsWithBorder_Execute),
-			FCanExecuteAction::CreateSP(this, &FGraphTrack::ContextMenu_ShowPointsWithBorder_CanExecute),
-			FIsActionChecked::CreateSP(this, &FGraphTrack::ContextMenu_ShowPointsWithBorder_IsChecked)
-		);
-		MenuBuilder.AddMenuEntry
-		(
-			LOCTEXT("ContextMenu_ShowPointsWithBorder", "Show Points with Border"),
-			LOCTEXT("ContextMenu_ShowPointsWithBorder_Desc", "Show border around points."),
-			FSlateIcon(),
-			Action_ShowPointsWithBorder,
-			NAME_None,
-			EUserInterfaceActionType::ToggleButton
-		);
-
-		FUIAction Action_ShowLines
-		(
-			FExecuteAction::CreateSP(this, &FGraphTrack::ContextMenu_ShowLines_Execute),
-			FCanExecuteAction::CreateSP(this, &FGraphTrack::ContextMenu_ShowLines_CanExecute),
-			FIsActionChecked::CreateSP(this, &FGraphTrack::ContextMenu_ShowLines_IsChecked)
-		);
-		MenuBuilder.AddMenuEntry
-		(
-			LOCTEXT("ContextMenu_ShowLines", "Show Connected Lines"),
-			LOCTEXT("ContextMenu_ShowLines_Desc", "Show connected lines. Each event is a single point in time."),
-			FSlateIcon(),
-			Action_ShowLines,
-			NAME_None,
-			EUserInterfaceActionType::ToggleButton
-		);
-
-		FUIAction Action_ShowPolygon
-		(
-			FExecuteAction::CreateSP(this, &FGraphTrack::ContextMenu_ShowPolygon_Execute),
-			FCanExecuteAction::CreateSP(this, &FGraphTrack::ContextMenu_ShowPolygon_CanExecute),
-			FIsActionChecked::CreateSP(this, &FGraphTrack::ContextMenu_ShowPolygon_IsChecked)
-		);
-		MenuBuilder.AddMenuEntry
-		(
-			LOCTEXT("ContextMenu_ShowPolygon", "Show Polygon"),
-			LOCTEXT("ContextMenu_ShowPolygon_Desc", "Show filled polygon under the graph series."),
-			FSlateIcon(),
-			Action_ShowPolygon,
-			NAME_None,
-			EUserInterfaceActionType::ToggleButton
-		);
-
-		FUIAction Action_UseEventDuration
-		(
-			FExecuteAction::CreateSP(this, &FGraphTrack::ContextMenu_UseEventDuration_Execute),
-			FCanExecuteAction::CreateSP(this, &FGraphTrack::ContextMenu_UseEventDuration_CanExecute),
-			FIsActionChecked::CreateSP(this, &FGraphTrack::ContextMenu_UseEventDuration_IsChecked)
-		);
-		MenuBuilder.AddMenuEntry
-		(
-			LOCTEXT("ContextMenu_UseEventDuration", "Use Event Duration"),
-			LOCTEXT("ContextMenu_UseEventDuration_Desc", "Use duration of timing events (for Connected Lines and Polygon)."),
-			FSlateIcon(),
-			Action_UseEventDuration,
-			NAME_None,
-			EUserInterfaceActionType::ToggleButton
-		);
-
-		FUIAction Action_ShowBars
-		(
-			FExecuteAction::CreateSP(this, &FGraphTrack::ContextMenu_ShowBars_Execute),
-			FCanExecuteAction::CreateSP(this, &FGraphTrack::ContextMenu_ShowBars_CanExecute),
-			FIsActionChecked::CreateSP(this, &FGraphTrack::ContextMenu_ShowBars_IsChecked)
-		);
-		MenuBuilder.AddMenuEntry
-		(
-			LOCTEXT("ContextMenu_ShowBars", "Show Bars"),
-			LOCTEXT("ContextMenu_ShowBars_Desc", "Show bars. Width of bars corresponds to duration of timing events."),
-			FSlateIcon(),
-			Action_ShowBars,
-			NAME_None,
-			EUserInterfaceActionType::ToggleButton
-		);
+		if(EnumHasAnyFlags(VisibleOptions, EGraphOptions::ShowPoints))
+		{
+			FUIAction Action_ShowPoints
+			(
+				FExecuteAction::CreateSP(this, &FGraphTrack::ContextMenu_ShowPoints_Execute),
+				FCanExecuteAction::CreateSP(this, &FGraphTrack::ContextMenu_ShowPoints_CanExecute),
+				FIsActionChecked::CreateSP(this, &FGraphTrack::ContextMenu_ShowPoints_IsChecked)
+			);
+			MenuBuilder.AddMenuEntry
+			(
+				LOCTEXT("ContextMenu_ShowPoints", "Show Points"),
+				LOCTEXT("ContextMenu_ShowPoints_Desc", "Show points."),
+				FSlateIcon(),
+				Action_ShowPoints,
+				NAME_None,
+				EUserInterfaceActionType::ToggleButton
+			);
+		}
+
+		if(EnumHasAnyFlags(VisibleOptions, EGraphOptions::ShowPointsWithBorder))
+		{
+			FUIAction Action_ShowPointsWithBorder
+			(
+				FExecuteAction::CreateSP(this, &FGraphTrack::ContextMenu_ShowPointsWithBorder_Execute),
+				FCanExecuteAction::CreateSP(this, &FGraphTrack::ContextMenu_ShowPointsWithBorder_CanExecute),
+				FIsActionChecked::CreateSP(this, &FGraphTrack::ContextMenu_ShowPointsWithBorder_IsChecked)
+			);
+			MenuBuilder.AddMenuEntry
+			(
+				LOCTEXT("ContextMenu_ShowPointsWithBorder", "Show Points with Border"),
+				LOCTEXT("ContextMenu_ShowPointsWithBorder_Desc", "Show border around points."),
+				FSlateIcon(),
+				Action_ShowPointsWithBorder,
+				NAME_None,
+				EUserInterfaceActionType::ToggleButton
+			);
+		}
+
+		if(EnumHasAnyFlags(VisibleOptions, EGraphOptions::ShowLines))
+		{
+			FUIAction Action_ShowLines
+			(
+				FExecuteAction::CreateSP(this, &FGraphTrack::ContextMenu_ShowLines_Execute),
+				FCanExecuteAction::CreateSP(this, &FGraphTrack::ContextMenu_ShowLines_CanExecute),
+				FIsActionChecked::CreateSP(this, &FGraphTrack::ContextMenu_ShowLines_IsChecked)
+			);
+			MenuBuilder.AddMenuEntry
+			(
+				LOCTEXT("ContextMenu_ShowLines", "Show Connected Lines"),
+				LOCTEXT("ContextMenu_ShowLines_Desc", "Show connected lines. Each event is a single point in time."),
+				FSlateIcon(),
+				Action_ShowLines,
+				NAME_None,
+				EUserInterfaceActionType::ToggleButton
+			);
+		}
+
+		if(EnumHasAnyFlags(VisibleOptions, EGraphOptions::ShowPolygon))
+		{
+			FUIAction Action_ShowPolygon
+			(
+				FExecuteAction::CreateSP(this, &FGraphTrack::ContextMenu_ShowPolygon_Execute),
+				FCanExecuteAction::CreateSP(this, &FGraphTrack::ContextMenu_ShowPolygon_CanExecute),
+				FIsActionChecked::CreateSP(this, &FGraphTrack::ContextMenu_ShowPolygon_IsChecked)
+			);
+			MenuBuilder.AddMenuEntry
+			(
+				LOCTEXT("ContextMenu_ShowPolygon", "Show Polygon"),
+				LOCTEXT("ContextMenu_ShowPolygon_Desc", "Show filled polygon under the graph series."),
+				FSlateIcon(),
+				Action_ShowPolygon,
+				NAME_None,
+				EUserInterfaceActionType::ToggleButton
+			);
+		}
+
+		if(EnumHasAnyFlags(VisibleOptions, EGraphOptions::UseEventDuration))
+		{
+			FUIAction Action_UseEventDuration
+			(
+				FExecuteAction::CreateSP(this, &FGraphTrack::ContextMenu_UseEventDuration_Execute),
+				FCanExecuteAction::CreateSP(this, &FGraphTrack::ContextMenu_UseEventDuration_CanExecute),
+				FIsActionChecked::CreateSP(this, &FGraphTrack::ContextMenu_UseEventDuration_IsChecked)
+			);
+			MenuBuilder.AddMenuEntry
+			(
+				LOCTEXT("ContextMenu_UseEventDuration", "Use Event Duration"),
+				LOCTEXT("ContextMenu_UseEventDuration_Desc", "Use duration of timing events (for Connected Lines and Polygon)."),
+				FSlateIcon(),
+				Action_UseEventDuration,
+				NAME_None,
+				EUserInterfaceActionType::ToggleButton
+			);
+		}
+
+		if(EnumHasAnyFlags(VisibleOptions, EGraphOptions::ShowBars))
+		{
+			FUIAction Action_ShowBars
+			(
+				FExecuteAction::CreateSP(this, &FGraphTrack::ContextMenu_ShowBars_Execute),
+				FCanExecuteAction::CreateSP(this, &FGraphTrack::ContextMenu_ShowBars_CanExecute),
+				FIsActionChecked::CreateSP(this, &FGraphTrack::ContextMenu_ShowBars_IsChecked)
+			);
+			MenuBuilder.AddMenuEntry
+			(
+				LOCTEXT("ContextMenu_ShowBars", "Show Bars"),
+				LOCTEXT("ContextMenu_ShowBars_Desc", "Show bars. Width of bars corresponds to duration of timing events."),
+				FSlateIcon(),
+				Action_ShowBars,
+				NAME_None,
+				EUserInterfaceActionType::ToggleButton
+			);
+		}
 	}
 	MenuBuilder.EndSection();
 
@@ -764,7 +786,7 @@ void FGraphTrack::ContextMenu_ShowDebugInfo_Execute()
 
 bool FGraphTrack::ContextMenu_ShowDebugInfo_CanExecute()
 {
-	return true;
+	return  EnumHasAnyFlags(EditableOptions, EGraphOptions::ShowDebugInfo);
 }
 
 ////////////////////////////////////////////////////////////////////////////////////////////////////
@@ -786,7 +808,7 @@ void FGraphTrack::ContextMenu_ShowPoints_Execute()
 
 bool FGraphTrack::ContextMenu_ShowPoints_CanExecute()
 {
-	return true;
+	return  EnumHasAnyFlags(EditableOptions, EGraphOptions::ShowPoints);
 }
 
 ////////////////////////////////////////////////////////////////////////////////////////////////////
@@ -808,7 +830,7 @@ void FGraphTrack::ContextMenu_ShowPointsWithBorder_Execute()
 
 bool FGraphTrack::ContextMenu_ShowPointsWithBorder_CanExecute()
 {
-	return bDrawPoints;
+	return bDrawPoints && EnumHasAnyFlags(EditableOptions, EGraphOptions::ShowPointsWithBorder);
 }
 
 ////////////////////////////////////////////////////////////////////////////////////////////////////
@@ -830,7 +852,7 @@ void FGraphTrack::ContextMenu_ShowLines_Execute()
 
 bool FGraphTrack::ContextMenu_ShowLines_CanExecute()
 {
-	return true;
+	return EnumHasAnyFlags(EditableOptions, EGraphOptions::ShowLines);
 }
 
 ////////////////////////////////////////////////////////////////////////////////////////////////////
@@ -852,7 +874,7 @@ void FGraphTrack::ContextMenu_ShowPolygon_Execute()
 
 bool FGraphTrack::ContextMenu_ShowPolygon_CanExecute()
 {
-	return true;
+	return EnumHasAnyFlags(EditableOptions, EGraphOptions::ShowPolygon);
 }
 
 ////////////////////////////////////////////////////////////////////////////////////////////////////
@@ -874,7 +896,7 @@ void FGraphTrack::ContextMenu_UseEventDuration_Execute()
 
 bool FGraphTrack::ContextMenu_UseEventDuration_CanExecute()
 {
-	return bDrawLines || bDrawPolygon;
+	return (bDrawLines || bDrawPolygon) && EnumHasAnyFlags(EditableOptions, EGraphOptions::UseEventDuration);
 }
 
 ////////////////////////////////////////////////////////////////////////////////////////////////////
@@ -896,7 +918,7 @@ void FGraphTrack::ContextMenu_ShowBars_Execute()
 
 bool FGraphTrack::ContextMenu_ShowBars_CanExecute()
 {
-	return true;
+	return EnumHasAnyFlags(EditableOptions, EGraphOptions::ShowBars);
 }
 
 ////////////////////////////////////////////////////////////////////////////////////////////////////
@@ -906,28 +928,6 @@ bool FGraphTrack::ContextMenu_ShowBars_IsChecked()
 	return bDrawBoxes;
 }
 
-////////////////////////////////////////////////////////////////////////////////////////////////////
-
-void FGraphTrack::ContextMenu_ShowSeries_Execute(FGraphSeries* Series)
-{
-	Series->SetVisibility(!Series->IsVisible());
-	SetDirtyFlag();
-}
-
-////////////////////////////////////////////////////////////////////////////////////////////////////
-
-bool FGraphTrack::ContextMenu_ShowSeries_CanExecute(FGraphSeries* Series)
-{
-	return true;
-}
-
-////////////////////////////////////////////////////////////////////////////////////////////////////
-
-bool FGraphTrack::ContextMenu_ShowSeries_IsChecked(FGraphSeries* Series)
-{
-	return Series->IsVisible();
-}
-
 ////////////////////////////////////////////////////////////////////////////////////////////////////
 // FRandomGraphTrack
 ////////////////////////////////////////////////////////////////////////////////////////////////////
diff --git a/Engine/Source/Developer/TraceInsights/Private/Insights/Widgets/SGraphSeriesList.cpp b/Engine/Source/Developer/TraceInsights/Private/Insights/Widgets/SGraphSeriesList.cpp
--- a/Engine/Source/Developer/TraceInsights/Private/Insights/Widgets/SGraphSeriesList.cpp
+++ b/Engine/Source/Developer/TraceInsights/Private/Insights/Widgets/SGraphSeriesList.cpp
@@ -24,6 +24,8 @@ public:
 
 	SLATE_ARGUMENT(TSharedPtr<FGraphSeries>, GraphSeries)
 
+	SLATE_ARGUMENT(TSharedPtr<FGraphTrack>, GraphTrack)
+
 	SLATE_ATTRIBUTE(FText, SearchText)
 
 	SLATE_END_ARGS()
@@ -31,6 +33,7 @@ public:
 	void Construct(const FArguments& InArgs, const TSharedRef<STableViewBase>& InOwnerTable)
 	{
 		GraphSeries = InArgs._GraphSeries;
+		GraphTrack = InArgs._GraphTrack;
 		SearchText = InArgs._SearchText;
 
 		STableRow<TSharedPtr<FGraphSeries>>::Construct(STableRow<TSharedPtr<FGraphSeries>>::FArguments(), InOwnerTable);
@@ -49,7 +52,7 @@ public:
 			[
 				SNew(SCheckBox)
 				.IsChecked_Lambda([this](){ return GraphSeries.Pin()->IsVisible() ? ECheckBoxState::Checked : ECheckBoxState::Unchecked; })
-				.OnCheckStateChanged_Lambda([this](ECheckBoxState InCheckBoxState) { GraphSeries.Pin()->SetVisibility(InCheckBoxState == ECheckBoxState::Checked); })
+				.OnCheckStateChanged_Lambda([this](ECheckBoxState InCheckBoxState) { GraphSeries.Pin()->SetVisibility(InCheckBoxState == ECheckBoxState::Checked); GraphTrack.Pin()->SetDirtyFlag(); })
 				.Content()
 				[
 					SNew(STextBlock)
@@ -71,6 +74,9 @@ public:
 		];
 	}
 
+	// The track containing the series we represent
+	TWeakPtr<FGraphTrack> GraphTrack;
+
 	// The series we represent
 	TWeakPtr<FGraphSeries> GraphSeries;
 
@@ -161,6 +167,7 @@ void SGraphSeriesList::Construct(const FArguments& InArgs, const TSharedRef<FGra
 TSharedRef<ITableRow> SGraphSeriesList::OnGenerateRow(TSharedPtr<FGraphSeries> Item, const TSharedRef<STableViewBase>& OwnerTable)
 {
 	return SNew(SGraphSeriesListEntry, OwnerTable)
+		.GraphTrack(GraphTrack.Pin())
 		.GraphSeries(Item)
 		.SearchText_Lambda([this](){ return SearchText; });
 }
diff --git a/Engine/Source/Developer/TraceInsights/Public/Insights/ViewModels/GraphTrack.h b/Engine/Source/Developer/TraceInsights/Public/Insights/ViewModels/GraphTrack.h
--- a/Engine/Source/Developer/TraceInsights/Public/Insights/ViewModels/GraphTrack.h
+++ b/Engine/Source/Developer/TraceInsights/Public/Insights/ViewModels/GraphTrack.h
@@ -16,6 +16,26 @@ class FTimingTrackViewport;
 
 ////////////////////////////////////////////////////////////////////////////////////////////////////
 
+// Various available options for display
+enum class EGraphOptions
+{
+	None					= 0,
+
+	ShowDebugInfo			= (1 << 0),
+	ShowPoints				= (1 << 1),
+	ShowPointsWithBorder	= (1 << 2),
+	ShowLines				= (1 << 3),
+	ShowPolygon				= (1 << 4),
+	UseEventDuration		= (1 << 5),
+	ShowBars				= (1 << 6),
+
+	All						= ShowPoints | ShowPointsWithBorder | ShowLines | ShowPolygon | UseEventDuration | ShowBars,
+};
+
+ENUM_CLASS_FLAGS(EGraphOptions);
+
+////////////////////////////////////////////////////////////////////////////////////////////////////
+
 class TRACEINSIGHTS_API FGraphTrack : public FBaseTimingTrack
 {
 	friend class FGraphTrackBuilder;
@@ -69,14 +89,6 @@ protected:
 
 	void DrawSeries(const FGraphSeries& Series, FDrawContext& DrawContext, const FTimingTrackViewport& Viewport) const;
 
-	virtual bool ContextMenu_ShowPoints_CanExecute();
-	virtual bool ContextMenu_ShowPointsWithBorder_CanExecute();
-	virtual bool ContextMenu_ShowLines_CanExecute();
-	virtual bool ContextMenu_ShowPolygon_CanExecute();
-	virtual bool ContextMenu_UseEventDuration_CanExecute();
-	virtual bool ContextMenu_ShowBars_CanExecute();
-	virtual bool ContextMenu_ShowSeries_CanExecute(FGraphSeries* Series);
-
 	// Get the Y value that is used to provide a clipping border between adjacent graph tracks.
 	virtual float GetBorderY() const { return 0.0f; }
 
@@ -85,27 +97,30 @@ private:
 	void ContextMenu_ShowDebugInfo_Execute();
 	bool ContextMenu_ShowDebugInfo_IsChecked();
 
+	bool ContextMenu_ShowPoints_CanExecute();
 	void ContextMenu_ShowPoints_Execute();
 	bool ContextMenu_ShowPoints_IsChecked();
 
+	bool ContextMenu_ShowPointsWithBorder_CanExecute();
 	void ContextMenu_ShowPointsWithBorder_Execute();
 	bool ContextMenu_ShowPointsWithBorder_IsChecked();
 
+	bool ContextMenu_ShowLines_CanExecute();
 	void ContextMenu_ShowLines_Execute();
 	bool ContextMenu_ShowLines_IsChecked();
 
+	bool ContextMenu_ShowPolygon_CanExecute();
 	void ContextMenu_ShowPolygon_Execute();
 	bool ContextMenu_ShowPolygon_IsChecked();
 
+	bool ContextMenu_UseEventDuration_CanExecute();
 	void ContextMenu_UseEventDuration_Execute();
 	bool ContextMenu_UseEventDuration_IsChecked();
 
+	bool ContextMenu_ShowBars_CanExecute();
 	void ContextMenu_ShowBars_Execute();
 	bool ContextMenu_ShowBars_IsChecked();
 
-	void ContextMenu_ShowSeries_Execute(FGraphSeries* Series);
-	bool ContextMenu_ShowSeries_IsChecked(FGraphSeries* Series);
-
 protected:
 	TArray<TSharedPtr<FGraphSeries>> AllSeries;
 
@@ -124,6 +139,10 @@ protected:
 	bool bDrawBoxes;
 	bool bDrawBaseline;
 
+	// Flags controlling whether menu items are available
+	EGraphOptions VisibleOptions;
+	EGraphOptions EditableOptions;
+
 	// Stats
 	int32 NumAddedEvents; // total event count
 	int32 NumDrawPoints;
