Minor fix to dirty track and series appropriately when editing visible graph series

diff --git a/Engine/Source/Developer/TraceInsights/Private/Insights/Widgets/SGraphSeriesList.cpp b/Engine/Source/Developer/TraceInsights/Private/Insights/Widgets/SGraphSeriesList.cpp
--- a/Engine/Source/Developer/TraceInsights/Private/Insights/Widgets/SGraphSeriesList.cpp
+++ b/Engine/Source/Developer/TraceInsights/Private/Insights/Widgets/SGraphSeriesList.cpp
@@ -52,7 +52,12 @@ public:
 			[
 				SNew(SCheckBox)
 				.IsChecked_Lambda([this](){ return GraphSeries.Pin()->IsVisible() ? ECheckBoxState::Checked : ECheckBoxState::Unchecked; })
-				.OnCheckStateChanged_Lambda([this](ECheckBoxState InCheckBoxState) { GraphSeries.Pin()->SetVisibility(InCheckBoxState == ECheckBoxState::Checked); GraphTrack.Pin()->SetDirtyFlag(); })
+				.OnCheckStateChanged_Lambda([this](ECheckBoxState InCheckBoxState) 
+				{
+					GraphSeries.Pin()->SetVisibility(InCheckBoxState == ECheckBoxState::Checked);
+					GraphSeries.Pin()->SetDirtyFlag();
+					GraphTrack.Pin()->SetDirtyFlag(); 
+				})
 				.Content()
 				[
 					SNew(STextBlock)
@@ -138,7 +143,10 @@ void SGraphSeriesList::Construct(const FArguments& InArgs, const TSharedRef<FGra
 					for(TSharedPtr<FGraphSeries> Series : FilteredSeries)
 					{
 						Series->SetVisibility(bVisible);
+						Series->SetDirtyFlag();
 					}
+
+					GraphTrack.Pin()->SetDirtyFlag();
 				})
 			]
 			+SHorizontalBox::Slot()
