[Insights]
  - Added FTrackHeader class. Implements the track header with collapsable state, shared between Markers track and Frame tracks.
  - Timing view: Added toggle header state to Frame timing tracks. When expanded, the track will draw vertical marker lines on edges of each frame.
  - Added support to allow tracks to handle mouse down, up and double click events.
  - minor code style fixes

diff --git a/Engine/Source/Developer/TraceInsights/Private/Insights/ViewModels/FrameTimingTrack.cpp b/Engine/Source/Developer/TraceInsights/Private/Insights/ViewModels/FrameTimingTrack.cpp
--- a/Engine/Source/Developer/TraceInsights/Private/Insights/ViewModels/FrameTimingTrack.cpp
+++ b/Engine/Source/Developer/TraceInsights/Private/Insights/ViewModels/FrameTimingTrack.cpp
@@ -171,6 +171,18 @@ INSIGHTS_IMPLEMENT_RTTI(FFrameTimingTrack)
 
 ////////////////////////////////////////////////////////////////////////////////////////////////////
 
+void FFrameTimingTrack::Reset()
+{
+	FTimingEventsTrack::Reset();
+
+	Header.Reset();
+	Header.SetCanBeCollapsed(true);
+	Header.SetIsCollapsed(true);
+	Header.UpdateSize();
+}
+
+////////////////////////////////////////////////////////////////////////////////////////////////////
+
 void FFrameTimingTrack::BuildDrawState(ITimingEventsTrackDrawStateBuilder& Builder, const ITimingTrackUpdateContext& Context)
 {
 	TSharedPtr<const Trace::IAnalysisSession> Session = FInsightsManager::Get()->GetSession();
@@ -233,8 +245,41 @@ void FFrameTimingTrack::BuildFilteredDrawState(ITimingEventsTrackDrawStateBuilde
 
 ////////////////////////////////////////////////////////////////////////////////////////////////////
 
+void FFrameTimingTrack::Update(const ITimingTrackUpdateContext& Context)
+{
+	FTimingEventsTrack::Update(Context);
+
+	Header.UpdateSize();
+	Header.Update(Context);
+}
+
+////////////////////////////////////////////////////////////////////////////////////////////////////
+
+void FFrameTimingTrack::PostUpdate(const ITimingTrackUpdateContext& Context)
+{
+	FTimingEventsTrack::PostUpdate(Context);
+
+	Header.PostUpdate(Context);
+}
+
+////////////////////////////////////////////////////////////////////////////////////////////////////
+
+void FFrameTimingTrack::Draw(const ITimingTrackDrawContext& Context) const
+{
+	DrawEvents(Context, 1.0f);
+	Header.Draw(Context);
+}
+
+////////////////////////////////////////////////////////////////////////////////////////////////////
+
 void FFrameTimingTrack::PostDraw(const ITimingTrackDrawContext& Context) const
 {
+	if (!Header.IsCollapsed())
+	{
+		const FTimingTrackViewport& Viewport = Context.GetViewport();
+		DrawMarkers(Context, 0.0f, Viewport.GetHeight());
+	}
+
 	const TSharedPtr<const ITimingEvent> SelectedEventPtr = Context.GetSelectedEvent();
 	if (SelectedEventPtr.IsValid() &&
 		SelectedEventPtr->CheckTrack(this) &&
@@ -244,6 +289,8 @@ void FFrameTimingTrack::PostDraw(const ITimingTrackDrawContext& Context) const
 		const ITimingViewDrawHelper& Helper = Context.GetHelper();
 		DrawSelectedEventInfo(SelectedEvent, Context.GetViewport(), Context.GetDrawContext(), Helper.GetWhiteBrush(), Helper.GetEventFont());
 	}
+
+	Header.PostDraw(Context);
 }
 
 ////////////////////////////////////////////////////////////////////////////////////////////////////
@@ -274,6 +321,31 @@ void FFrameTimingTrack::DrawSelectedEventInfo(const FTimingEvent& SelectedEvent,
 
 ////////////////////////////////////////////////////////////////////////////////////////////////////
 
+FReply FFrameTimingTrack::OnMouseButtonDown(const FGeometry& MyGeometry, const FPointerEvent& MouseEvent)
+{
+	FReply Reply = FReply::Unhandled();
+
+	if (MouseEvent.GetEffectingButton() == EKeys::LeftMouseButton)
+	{
+		if (IsVisible() && IsHeaderHovered())
+		{
+			ToggleCollapsed();
+			Reply = FReply::Handled();
+		}
+	}
+
+	return Reply;
+}
+
+////////////////////////////////////////////////////////////////////////////////////////////////////
+
+FReply FFrameTimingTrack::OnMouseButtonDoubleClick(const FGeometry& MyGeometry, const FPointerEvent& MouseEvent)
+{
+	return OnMouseButtonDown(MyGeometry, MouseEvent);
+}
+
+////////////////////////////////////////////////////////////////////////////////////////////////////
+
 void FFrameTimingTrack::InitTooltip(FTooltipDrawState& InOutTooltip, const ITimingEvent& InTooltipEvent) const
 {
 	if (InTooltipEvent.CheckTrack(this) && InTooltipEvent.Is<FTimingEvent>())
@@ -362,7 +434,7 @@ void FFrameTimingTrack::BuildContextMenu(FMenuBuilder& MenuBuilder)
 bool FFrameTimingTrack::FindFrame(const FTimingEvent& InTimingEvent, TFunctionRef<void(double, double, uint32, const Trace::FFrame&)> InFoundPredicate) const
 {
 	auto MatchEvent = [&InTimingEvent](double InStartTime, double InEndTime, uint32 InDepth)
-	{ 
+	{
 		return InDepth == InTimingEvent.GetDepth()
 			&& InStartTime == InTimingEvent.GetStartTime()
 			&& InEndTime == InTimingEvent.GetEndTime();
@@ -407,7 +479,7 @@ bool FFrameTimingTrack::FindFrame(const FTimingEventSearchParameters& InParamete
 		{
 			InFoundPredicate(InFoundStartTime, InFoundEndTime, InFoundDepth, InEvent);
 		},
-		
+
 		SearchCache);
 }
 
@@ -431,7 +503,7 @@ const FString FFrameTimingTrack::GetCompleteFrameName(const uint64 InFrameIndex,
 {
 	return FString::Printf(TEXT("%s Frame %d (%s)"), FFrameTrackDrawHelper::FrameTypeToString(FrameType), InFrameIndex, *TimeUtils::FormatTimeAuto(InFrameDuration));
 }
- 
+
 ////////////////////////////////////////////////////////////////////////////////////////////////////
 
 #undef LOCTEXT_NAMESPACE
diff --git a/Engine/Source/Developer/TraceInsights/Private/Insights/ViewModels/FrameTimingTrack.h b/Engine/Source/Developer/TraceInsights/Private/Insights/ViewModels/FrameTimingTrack.h
--- a/Engine/Source/Developer/TraceInsights/Private/Insights/ViewModels/FrameTimingTrack.h
+++ b/Engine/Source/Developer/TraceInsights/Private/Insights/ViewModels/FrameTimingTrack.h
@@ -10,6 +10,7 @@
 #include "Insights/ITimingViewExtender.h"
 #include "Insights/ViewModels/TimingEventSearch.h" // for TTimingEventSearchCache
 #include "Insights/ViewModels/TimingEventsTrack.h"
+#include "Insights/ViewModels/TrackHeader.h"
 
 class FTimingEvent;
 class FTimingEventSearchParameters;
@@ -65,7 +66,9 @@ public:
 		: FTimingEventsTrack(InName)
 		, SharedState(InSharedState)
 		, FrameType(InFrameType)
+		, Header(*this)
 	{
+		Reset();
 	}
 
 	virtual ~FFrameTimingTrack() {}
@@ -75,11 +78,30 @@ public:
 	uint32 GetFrameType() const { return FrameType; }
 	//void SetFrameType(uint32 InFrameType) { FrameType = InFrameType; }
 
+	bool IsCollapsed() const { return Header.IsCollapsed(); }
+	void ToggleCollapsed() { Header.ToggleCollapsed(); }
+
+	const FString GetShortFrameName(const uint64 FrameIndex) const;
+	const FString GetFrameName(const uint64 FrameIndex) const;
+	const FString GetCompleteFrameName(const uint64 FrameIndex, const double FrameDuration) const;
+
+	//////////////////////////////////////////////////
+	// FBaseTimingTrack/FTimingEventsTrack overrides
+
+	virtual void Reset() override;
+
 	virtual void BuildDrawState(ITimingEventsTrackDrawStateBuilder& Builder, const ITimingTrackUpdateContext& Context) override;
 	virtual void BuildFilteredDrawState(ITimingEventsTrackDrawStateBuilder& Builder, const ITimingTrackUpdateContext& Context) override;
 
+	virtual void Update(const ITimingTrackUpdateContext& Context) override;
+	virtual void PostUpdate(const ITimingTrackUpdateContext& Context) override;
+
+	virtual void Draw(const ITimingTrackDrawContext& Context) const override;
 	virtual void PostDraw(const ITimingTrackDrawContext& Context) const override;
 
+	virtual FReply OnMouseButtonDown(const FGeometry& MyGeometry, const FPointerEvent& MouseEvent) override;
+	virtual FReply OnMouseButtonDoubleClick(const FGeometry& MyGeometry, const FPointerEvent& MouseEvent) override;
+
 	virtual void InitTooltip(FTooltipDrawState& InOutTooltip, const ITimingEvent& InTooltipEvent) const override;
 
 	virtual const TSharedPtr<const ITimingEvent> SearchEvent(const FTimingEventSearchParameters& InSearchParameters) const override;
@@ -88,9 +110,7 @@ public:
 	virtual void OnClipboardCopyEvent(const ITimingEvent& InSelectedEvent) const override;
 	virtual void BuildContextMenu(FMenuBuilder& MenuBuilder) override;
 
-	const FString GetShortFrameName(const uint64 FrameIndex) const;
-	const FString GetFrameName(const uint64 FrameIndex) const;
-	const FString GetCompleteFrameName(const uint64 FrameIndex, const double FrameDuration) const;
+	//////////////////////////////////////////////////
 
 private:
 	void DrawSelectedEventInfo(const FTimingEvent& SelectedEvent, const FTimingTrackViewport& Viewport, const FDrawContext& DrawContext, const FSlateBrush* WhiteBrush, const FSlateFontInfo& Font) const;
@@ -103,6 +123,8 @@ private:
 
 	uint32 FrameType; // ETraceFrameType
 
+	FTrackHeader Header;
+
 	// Search cache
 	mutable TTimingEventSearchCache<Trace::FFrame> SearchCache;
 };
diff --git a/Engine/Source/Developer/TraceInsights/Private/Insights/ViewModels/MarkersTimingTrack.cpp b/Engine/Source/Developer/TraceInsights/Private/Insights/ViewModels/MarkersTimingTrack.cpp
--- a/Engine/Source/Developer/TraceInsights/Private/Insights/ViewModels/MarkersTimingTrack.cpp
+++ b/Engine/Source/Developer/TraceInsights/Private/Insights/ViewModels/MarkersTimingTrack.cpp
@@ -18,8 +18,6 @@
 
 #define LOCTEXT_NAMESPACE "MarkersTimingTrack"
 
-////////////////////////////////////////////////////////////////////////////////////////////////////
-// FMarkersTimingTrack
 ////////////////////////////////////////////////////////////////////////////////////////////////////
 
 INSIGHTS_IMPLEMENT_RTTI(FMarkersTimingTrack)
@@ -30,10 +28,8 @@ FMarkersTimingTrack::FMarkersTimingTrack()
 	: FBaseTimingTrack()
 	//, TimeMarkerBoxes()
 	//, TimeMarkerTexts()
-	, bIsCollapsed(false)
 	, bUseOnlyBookmarks(true)
-	, TargetHoveredAnimPercent(0.0f)
-	, CurrentHoveredAnimPercent(0.0f)
+	, Header(*this)
 	, NumLogMessages(0)
 	, NumDrawBoxes(0)
 	, NumDrawTexts(0)
@@ -57,91 +53,62 @@ void FMarkersTimingTrack::Reset()
 	TimeMarkerBoxes.Reset();
 	TimeMarkerTexts.Reset();
 
-	bIsCollapsed = false;
 	bUseOnlyBookmarks = true;
 
-	TargetHoveredAnimPercent = 0.0f;
-	CurrentHoveredAnimPercent = 0.0f;
+	Header.Reset();
+	Header.SetIsInBackground(true);
+	Header.SetCanBeCollapsed(true);
 
 	NumLogMessages = 0;
 	NumDrawBoxes = 0;
 	NumDrawTexts = 0;
 
-	UpdateHeight();
+	UpdateTrackNameAndHeight();
 }
 
 ////////////////////////////////////////////////////////////////////////////////////////////////////
 
-void FMarkersTimingTrack::UpdateHeight()
+void FMarkersTimingTrack::UpdateTrackNameAndHeight()
 {
-	constexpr float BookmarksTrackHeight = 14.0f;
-	constexpr float TimeMarkersTrackHeight = 28.0f;
-
 	if (bUseOnlyBookmarks)
 	{
-		SetHeight(BookmarksTrackHeight);
+		const FString NameString = TEXT("Bookmarks");
+		SetName(NameString);
+		SetHeight(14.0f);
 	}
 	else
 	{
-		SetHeight(TimeMarkersTrackHeight);
+		const FString NameString = TEXT("Logs");
+		SetName(NameString);
+		SetHeight(28.0f);
 	}
+
+	Header.UpdateSize();
 }
 
 ////////////////////////////////////////////////////////////////////////////////////////////////////
 
-void FMarkersTimingTrack::PostUpdate(const ITimingTrackUpdateContext& Context)
+void FMarkersTimingTrack::Update(const ITimingTrackUpdateContext& Context)
 {
-	constexpr float HeaderWidth = 80.0f;
-	constexpr float HeaderHeight = 14.0f;
+	Header.Update(Context);
 
-	const float MouseY = Context.GetMousePosition().Y;
-	if (MouseY >= GetPosY() && MouseY < GetPosY() + GetHeight())
-	{
-		SetHoveredState(true);
-		const float MouseX = Context.GetMousePosition().X;
-		SetHeaderHoveredState(MouseX < HeaderWidth && MouseY < GetPosY() + HeaderHeight);
-		TargetHoveredAnimPercent = 1.0f;
-	}
-	else
+	const FTimingTrackViewport& Viewport = Context.GetViewport();
+	if (IsDirty() || Viewport.IsHorizontalViewportDirty())
 	{
-		SetHoveredState(false);
-		TargetHoveredAnimPercent = 0.0f;
+		ClearDirtyFlag();
+
+		UpdateDrawState(Viewport);
 	}
 }
 
 ////////////////////////////////////////////////////////////////////////////////////////////////////
 
-void FMarkersTimingTrack::Update(const ITimingTrackUpdateContext& InContext)
+void FMarkersTimingTrack::PostUpdate(const ITimingTrackUpdateContext& Context)
 {
-	if (CurrentHoveredAnimPercent != TargetHoveredAnimPercent)
-	{
-		if (CurrentHoveredAnimPercent < TargetHoveredAnimPercent)
-		{
-			const float ShowAnimSpeed = 2 * InContext.GetDeltaTime();
-			CurrentHoveredAnimPercent += ShowAnimSpeed;
-			if (CurrentHoveredAnimPercent > TargetHoveredAnimPercent)
-			{
-				CurrentHoveredAnimPercent = TargetHoveredAnimPercent;
-			}
-		}
-		else
-		{
-			const float HideAnimSpeed = 3 * InContext.GetDeltaTime();
-			CurrentHoveredAnimPercent -= HideAnimSpeed;
-			if (CurrentHoveredAnimPercent < TargetHoveredAnimPercent)
-			{
-				CurrentHoveredAnimPercent = TargetHoveredAnimPercent;
-			}
-		}
-	}
-
-	const FTimingTrackViewport& Viewport = InContext.GetViewport();
-	if (IsDirty() || Viewport.IsHorizontalViewportDirty())
-	{
-		ClearDirtyFlag();
+	const float MouseY = Context.GetMousePosition().Y;
+	SetHoveredState(MouseY >= GetPosY() && MouseY < GetPosY() + GetHeight());
 
-		UpdateDrawState(Viewport);
-	}
+	Header.PostUpdate(Context);
 }
 
 ////////////////////////////////////////////////////////////////////////////////////////////////////
@@ -179,11 +146,7 @@ void FMarkersTimingTrack::Draw(const ITimingTrackDrawContext& Context) const
 	DrawContext.DrawBox(0.0f, GetPosY(), Viewport.GetWidth(), GetHeight(), WhiteBrush, BackgroundColor);
 	DrawContext.LayerId++;
 
-	// Draw the track's header, in background.
-	if (!IsHovered() || CurrentHoveredAnimPercent < 1.0f)
-	{
-		DrawHeader(DrawContext, true);
-	}
+	Header.Draw(Context);
 }
 
 ////////////////////////////////////////////////////////////////////////////////////////////////////
@@ -245,71 +208,32 @@ void FMarkersTimingTrack::PostDraw(const ITimingTrackDrawContext& Context) const
 
 	//////////////////////////////////////////////////
 
-	// When hovered, the track's header is draw on top.
-	if (IsHovered() || CurrentHoveredAnimPercent > 0.0f)
-	{
-		DrawHeader(DrawContext, false);
-	}
+	Header.PostDraw(Context);
 }
 
 ////////////////////////////////////////////////////////////////////////////////////////////////////
 
-void FMarkersTimingTrack::DrawHeader(FDrawContext& DrawContext, bool bFirstDraw) const
+FReply FMarkersTimingTrack::OnMouseButtonDown(const FGeometry& MyGeometry, const FPointerEvent& MouseEvent)
 {
-	const FString NameString = IsBookmarksTrack() ? TEXT("Bookmarks") : TEXT("Logs");
-	//const float ArrowSize = 6.0f;
-	const float ArrowSizeX = 4.0f;
-	const float ArrowSizeY = 8.0f;
-	const float ArrowX = IsBookmarksTrack() ? 64.0f : 31.0f;
-	const float ArrowY = GetPosY() + 3.0f;
-	const float HeaderW = ArrowX + ArrowSizeX + 4.0f;
-	const float HeaderH = 14.0f;
-
-	if (!bFirstDraw)
-	{
-		DrawContext.DrawBox(0.0f, GetPosY(), HeaderW, HeaderH, WhiteBrush, FLinearColor(0.04f, 0.04f, 0.04f, CurrentHoveredAnimPercent));
-		DrawContext.LayerId++;
-	}
+	FReply Reply = FReply::Unhandled();
 
-	FLinearColor Color;
-	if (bFirstDraw || CurrentHoveredAnimPercent == 0.0)
-	{
-		Color = FLinearColor(0.07f, 0.07f, 0.07f, 1.0f);
-	}
-	else if (IsHeaderHovered())
+	if (MouseEvent.GetEffectingButton() == EKeys::LeftMouseButton)
 	{
-		Color = FLinearColor(1.0f, 1.0f, 0.0f, CurrentHoveredAnimPercent);
-	}
-	else
-	{
-		Color = FLinearColor(1.0f, 1.0f, 1.0f, CurrentHoveredAnimPercent);
-	}
-
-	// Draw "Bookmarks" or "Logs" text.
-	DrawContext.DrawText(2.0f, GetPosY() + 1.0f, NameString, Font, Color);
-
-	if (IsCollapsed())
-	{
-		// Draw "right empty arrow".
-		TArray<FVector2D> Points =
-		{
-			FVector2D(ArrowX, ArrowY),
-			FVector2D(ArrowX, ArrowY + ArrowSizeY),
-			FVector2D(ArrowX + ArrowSizeX, ArrowY + ArrowSizeY / 2.0f),
-			FVector2D(ArrowX, ArrowY)
-		};
-		FSlateDrawElement::MakeLines(DrawContext.ElementList, DrawContext.LayerId, DrawContext.Geometry.ToPaintGeometry(), Points, DrawContext.DrawEffects, Color, false, 1.0f);
-	}
-	else
-	{
-		// Draw "down-right filled arrow".
-		for (float A = 1.0f; A < ArrowSizeY; A += 1.0f)
+		if (IsVisible() && IsHeaderHovered())
 		{
-			DrawContext.DrawBox(ArrowX - 3.0 + ArrowSizeY - A, ArrowY + A - 1.0f, A, 1.0f, WhiteBrush, Color);
+			ToggleCollapsed();
+			Reply = FReply::Handled();
 		}
 	}
 
-	DrawContext.LayerId++;
+	return Reply;
+}
+
+////////////////////////////////////////////////////////////////////////////////////////////////////
+
+FReply FMarkersTimingTrack::OnMouseButtonDoubleClick(const FGeometry& MyGeometry, const FPointerEvent& MouseEvent)
+{
+	return OnMouseButtonDown(MyGeometry, MouseEvent);
 }
 
 ////////////////////////////////////////////////////////////////////////////////////////////////////
diff --git a/Engine/Source/Developer/TraceInsights/Private/Insights/ViewModels/MarkersTimingTrack.h b/Engine/Source/Developer/TraceInsights/Private/Insights/ViewModels/MarkersTimingTrack.h
--- a/Engine/Source/Developer/TraceInsights/Private/Insights/ViewModels/MarkersTimingTrack.h
+++ b/Engine/Source/Developer/TraceInsights/Private/Insights/ViewModels/MarkersTimingTrack.h
@@ -7,6 +7,7 @@
 
 // Insights
 #include "Insights/ViewModels/BaseTimingTrack.h"
+#include "Insights/ViewModels/TrackHeader.h"
 
 namespace Trace
 {
@@ -51,14 +52,14 @@ public:
 
 	virtual void Reset() override;
 
-	bool IsCollapsed() const { return bIsCollapsed; }
-	void ToggleCollapsed() { bIsCollapsed = !bIsCollapsed; }
+	bool IsCollapsed() const { return Header.IsCollapsed(); }
+	void ToggleCollapsed() { Header.ToggleCollapsed(); }
 
 	bool IsBookmarksTrack() const { return bUseOnlyBookmarks; }
 	void SetBookmarksTrackFlag(bool bInUseOnlyBookmarks)
 	{
 		bUseOnlyBookmarks = bInUseOnlyBookmarks;
-		UpdateHeight();
+		UpdateTrackNameAndHeight();
 	}
 
 	// Stats
@@ -71,6 +72,8 @@ public:
 	virtual void PostUpdate(const ITimingTrackUpdateContext& Context) override;
 	virtual void Draw(const ITimingTrackDrawContext& Context) const override;
 	virtual void PostDraw(const ITimingTrackDrawContext& Context) const override;
+	virtual FReply OnMouseButtonDown(const FGeometry& MyGeometry, const FPointerEvent& MouseEvent) override;
+	virtual FReply OnMouseButtonDoubleClick(const FGeometry& MyGeometry, const FPointerEvent& MouseEvent) override;
 
 private:
 	void ResetCache()
@@ -79,20 +82,16 @@ private:
 		TimeMarkerTexts.Reset();
 	}
 
-	void UpdateHeight();
+	void UpdateTrackNameAndHeight();
 	void UpdateDrawState(const FTimingTrackViewport& Viewport);
 
-	void DrawHeader(FDrawContext& DrawContext, bool bFirstDraw) const;
-
 private:
 	TArray<FTimeMarkerBoxInfo> TimeMarkerBoxes;
 	TArray<FTimeMarkerTextInfo> TimeMarkerTexts;
 
-	bool bIsCollapsed; // If false, the vertical lines will extend to entire viewport height; otherwise will be limited to this track's height.
 	bool bUseOnlyBookmarks; // If true, uses only bookmarks; otherwise it uses all log messages.
 
-	float TargetHoveredAnimPercent; // [0.0 .. 1.0], 0.0 = hidden, 1.0 = visible
-	float CurrentHoveredAnimPercent;
+	FTrackHeader Header;
 
 	// Stats
 	int32 NumLogMessages;
diff --git a/Engine/Source/Developer/TraceInsights/Private/Insights/ViewModels/TimingEventsTrack.cpp b/Engine/Source/Developer/TraceInsights/Private/Insights/ViewModels/TimingEventsTrack.cpp
--- a/Engine/Source/Developer/TraceInsights/Private/Insights/ViewModels/TimingEventsTrack.cpp
+++ b/Engine/Source/Developer/TraceInsights/Private/Insights/ViewModels/TimingEventsTrack.cpp
@@ -245,6 +245,23 @@ void FTimingEventsTrack::DrawEvents(const ITimingTrackDrawContext& Context, cons
 
 ////////////////////////////////////////////////////////////////////////////////////////////////////
 
+void FTimingEventsTrack::DrawMarkers(const ITimingTrackDrawContext& Context, float LineY, float LineH) const
+{
+	const FTimingViewDrawHelper& Helper = *static_cast<const FTimingViewDrawHelper*>(&Context.GetHelper());
+
+	if (Context.GetEventFilter().IsValid())
+	{
+		Helper.DrawMarkers(*DrawState, LineY, LineH, 0.2f);
+		Helper.DrawMarkers(*FilteredDrawState, LineY, LineH, 0.75f * FilteredDrawStateInfo.Opacity);
+	}
+	else
+	{
+		Helper.DrawMarkers(*DrawState, LineY, LineH, 0.2f);
+	}
+}
+
+////////////////////////////////////////////////////////////////////////////////////////////////////
+
 int32 FTimingEventsTrack::GetHeaderBackgroundLayerId(const ITimingTrackDrawContext& Context) const
 {
 	const FTimingViewDrawHelper& Helper = *static_cast<const FTimingViewDrawHelper*>(&Context.GetHelper());
diff --git a/Engine/Source/Developer/TraceInsights/Private/Insights/ViewModels/TimingViewDrawHelper.cpp b/Engine/Source/Developer/TraceInsights/Private/Insights/ViewModels/TimingViewDrawHelper.cpp
--- a/Engine/Source/Developer/TraceInsights/Private/Insights/ViewModels/TimingViewDrawHelper.cpp
+++ b/Engine/Source/Developer/TraceInsights/Private/Insights/ViewModels/TimingViewDrawHelper.cpp
@@ -487,6 +487,34 @@ void FTimingViewDrawHelper::DrawFadedEvents(const FTimingEventsTrackDrawState& D
 
 ////////////////////////////////////////////////////////////////////////////////////////////////////
 
+void FTimingViewDrawHelper::DrawMarkers(const FTimingEventsTrackDrawState& DrawState, float LineY, float LineH, float Opacity) const
+{
+	if (LineH > 0.0f)
+	{
+		const FTimingViewLayout& Layout = Viewport.GetLayout();
+
+		// Draw markers from filled boxes (merged borders).
+		for (const FTimingEventsTrackDrawState::FBoxPrimitive& Box : DrawState.Boxes)
+		{
+			DrawContext.DrawBox(Box.X, LineY, Box.W, LineH, WhiteBrush, Box.Color.CopyWithNewOpacity(Opacity));
+		}
+
+		// Draw markers from borders.
+		for (const FTimingEventsTrackDrawState::FBoxPrimitive& Box : DrawState.Borders)
+		{
+			DrawContext.DrawBox(Box.X, LineY, 1.0f, LineH, WhiteBrush, Box.Color.CopyWithNewOpacity(Opacity));
+			if (Box.W > 1.0f)
+			{
+				DrawContext.DrawBox(Box.X + Box.W - 1.0f, LineY, 1.0f, LineH, WhiteBrush, Box.Color.CopyWithNewOpacity(Opacity));
+			}
+		}
+
+		DrawContext.LayerId++;
+	}
+}
+
+////////////////////////////////////////////////////////////////////////////////////////////////////
+
 void FTimingViewDrawHelper::DrawTrackHeader(const FTimingEventsTrack& Track) const
 {
 	const float TrackY = Track.GetPosY();
diff --git a/Engine/Source/Developer/TraceInsights/Private/Insights/ViewModels/TimingViewDrawHelper.h b/Engine/Source/Developer/TraceInsights/Private/Insights/ViewModels/TimingViewDrawHelper.h
--- a/Engine/Source/Developer/TraceInsights/Private/Insights/ViewModels/TimingViewDrawHelper.h
+++ b/Engine/Source/Developer/TraceInsights/Private/Insights/ViewModels/TimingViewDrawHelper.h
@@ -181,6 +181,8 @@ public:
 
 	void DrawFadedEvents(const FTimingEventsTrackDrawState& DrawState, const FTimingEventsTrack& Track, const float OffsetY = 1.0f, const float Opacity = 0.1f) const;
 
+	void DrawMarkers(const FTimingEventsTrackDrawState& DrawState, float LineY, float LineH, float Opacity) const;
+
 	void DrawTrackHeader(const FTimingEventsTrack& Track) const;
 
 	void EndDrawTracks() const;
diff --git a/Engine/Source/Developer/TraceInsights/Private/Insights/ViewModels/TrackHeader.cpp b/Engine/Source/Developer/TraceInsights/Private/Insights/ViewModels/TrackHeader.cpp
new file mode 100644
--- /dev/null
+++ b/Engine/Source/Developer/TraceInsights/Private/Insights/ViewModels/TrackHeader.cpp
@@ -0,0 +1,214 @@
+// Copyright 1998-2019 Epic Games, Inc. All Rights Reserved.
+
+#include "TrackHeader.h"
+
+#include "Fonts/FontMeasure.h"
+#include "Framework/Application/SlateApplication.h"
+#include "Styling/CoreStyle.h"
+
+// Insights
+#include "Insights/Common/PaintUtils.h"
+#include "Insights/InsightsStyle.h"
+#include "Insights/ViewModels/BaseTimingTrack.h"
+#include "Insights/ViewModels/TimingTrackViewport.h"
+
+////////////////////////////////////////////////////////////////////////////////////////////////////
+
+#define LOCTEXT_NAMESPACE "TrackHeader"
+
+////////////////////////////////////////////////////////////////////////////////////////////////////
+// FTrackHeader
+////////////////////////////////////////////////////////////////////////////////////////////////////
+
+FTrackHeader::FTrackHeader(FBaseTimingTrack& InParentTrack)
+	: ParentTrack(InParentTrack)
+	, WhiteBrush(FInsightsStyle::Get().GetBrush("WhiteBrush"))
+	, Font(FCoreStyle::GetDefaultFontStyle("Regular", 8))
+{
+	Reset();
+}
+
+////////////////////////////////////////////////////////////////////////////////////////////////////
+
+void FTrackHeader::Reset()
+{
+	Width = 80.0f;
+	Height = 14.0f;
+	bIsInBackground = false;
+	bCanBeCollapsed = false;
+	bIsCollapsed = false;
+	TargetHoveredAnimPercent = 0.0f;
+	CurrentHoveredAnimPercent = 0.0f;
+}
+
+////////////////////////////////////////////////////////////////////////////////////////////////////
+
+void FTrackHeader::UpdateSize()
+{
+	//ParentTrack.GetName()
+	const TSharedRef<FSlateFontMeasure> FontMeasureService = FSlateApplication::Get().GetRenderer()->GetFontMeasureService();
+	const float TextWidth = FontMeasureService->Measure(ParentTrack.GetName(), Font).X;
+
+	Width = TextWidth + 4.0f;
+	if (bCanBeCollapsed)
+	{
+		Width += 9.0f;
+	}
+	//if (ParentTrack.IsSelected())
+	//{
+	//	Width += 9.0f;
+	//}
+}
+
+////////////////////////////////////////////////////////////////////////////////////////////////////
+
+void FTrackHeader::Update(const ITimingTrackUpdateContext& InContext)
+{
+	if (CurrentHoveredAnimPercent != TargetHoveredAnimPercent)
+	{
+		if (CurrentHoveredAnimPercent < TargetHoveredAnimPercent)
+		{
+			const float ShowAnimSpeed = 2 * InContext.GetDeltaTime();
+			CurrentHoveredAnimPercent += ShowAnimSpeed;
+			if (CurrentHoveredAnimPercent > TargetHoveredAnimPercent)
+			{
+				CurrentHoveredAnimPercent = TargetHoveredAnimPercent;
+			}
+		}
+		else
+		{
+			const float HideAnimSpeed = 3 * InContext.GetDeltaTime();
+			CurrentHoveredAnimPercent -= HideAnimSpeed;
+			if (CurrentHoveredAnimPercent < TargetHoveredAnimPercent)
+			{
+				CurrentHoveredAnimPercent = TargetHoveredAnimPercent;
+			}
+		}
+	}
+}
+
+////////////////////////////////////////////////////////////////////////////////////////////////////
+
+void FTrackHeader::PostUpdate(const ITimingTrackUpdateContext& Context)
+{
+	const float MouseX = Context.GetMousePosition().X;
+	const float MouseY = Context.GetMousePosition().Y;
+
+	ParentTrack.SetHeaderHoveredState(MouseX < Width && MouseY < ParentTrack.GetPosY() + Height);
+
+	if (ParentTrack.IsHovered())
+	{
+		TargetHoveredAnimPercent = 1.0f;
+	}
+	else
+	{
+		TargetHoveredAnimPercent = 0.0f;
+	}
+}
+
+////////////////////////////////////////////////////////////////////////////////////////////////////
+
+void FTrackHeader::Draw(const ITimingTrackDrawContext& Context) const
+{
+	if (bIsInBackground)
+	{
+		// Draw the track's header, in background.
+		if (!ParentTrack.IsHovered() || CurrentHoveredAnimPercent < 1.0f)
+		{
+			DrawInternal(Context, true);
+		}
+	}
+	else
+	{
+		FDrawContext& DrawContext = Context.GetDrawContext();
+
+		// Draw a horizontal line between timelines (top line of a track).
+		DrawContext.DrawBox(0.0f, ParentTrack.GetPosY(), Context.GetViewport().GetWidth(), 1.0f, WhiteBrush, FLinearColor(0.05f, 0.05f, 0.05f, 1.0f));
+
+		DrawInternal(Context, false);
+	}
+}
+
+////////////////////////////////////////////////////////////////////////////////////////////////////
+
+void FTrackHeader::PostDraw(const ITimingTrackDrawContext& Context) const
+{
+	if (bIsInBackground)
+	{
+		// When track is hovered, the track's header is draw on top.
+		if (ParentTrack.IsHovered() || CurrentHoveredAnimPercent > 0.0f)
+		{
+			DrawInternal(Context, false);
+		}
+	}
+}
+
+////////////////////////////////////////////////////////////////////////////////////////////////////
+
+void FTrackHeader::DrawInternal(const ITimingTrackDrawContext& Context, bool bDrawBackgroundLayer) const
+{
+	FDrawContext& DrawContext = Context.GetDrawContext();
+
+	const float Opacity = bIsInBackground ? CurrentHoveredAnimPercent : 1.0f;
+
+	if (!bDrawBackgroundLayer)
+	{
+		DrawContext.DrawBox(0.0f, ParentTrack.GetPosY(), Width, Height, WhiteBrush, FLinearColor(0.04f, 0.04f, 0.04f, Opacity));
+		DrawContext.LayerId++;
+	}
+
+	FLinearColor Color;
+	if (bIsInBackground && (bDrawBackgroundLayer || CurrentHoveredAnimPercent == 0.0))
+	{
+		Color = FLinearColor(0.07f, 0.07f, 0.07f, 1.0f);
+	}
+	else if (ParentTrack.IsHeaderHovered())
+	{
+		Color = FLinearColor(1.0f, 1.0f, 0.0f, Opacity);
+	}
+	else
+	{
+		Color = FLinearColor(1.0f, 1.0f, 1.0f, Opacity);
+	}
+
+	// Draw text.
+	DrawContext.DrawText(2.0f, ParentTrack.GetPosY() + 1.0f, ParentTrack.GetName(), Font, Color);
+
+	if (bCanBeCollapsed)
+	{
+		const float ArrowSizeX = 4.0f;
+		const float ArrowSizeY = 8.0f;
+		const float ArrowX = Width - ArrowSizeX - 4.0f;
+		const float ArrowY = ParentTrack.GetPosY() + 3.0f;
+
+		// Draw arrow indicating "collapsed" state.
+		if (bIsCollapsed)
+		{
+			// Draw "right empty arrow".
+			//TODO: use a brush/image instead
+			TArray<FVector2D> Points =
+			{
+				FVector2D(ArrowX, ArrowY),
+				FVector2D(ArrowX, ArrowY + ArrowSizeY),
+				FVector2D(ArrowX + ArrowSizeX, ArrowY + ArrowSizeY / 2.0f),
+				FVector2D(ArrowX, ArrowY)
+			};
+			FSlateDrawElement::MakeLines(DrawContext.ElementList, DrawContext.LayerId, DrawContext.Geometry.ToPaintGeometry(), Points, DrawContext.DrawEffects, Color, false, 1.0f);
+		}
+		else
+		{
+			// Draw "down-right filled arrow".
+			//TODO: use a brush/image instead
+			for (float A = 1.0f; A < ArrowSizeY; A += 1.0f)
+			{
+				DrawContext.DrawBox(ArrowX - 3.0 + ArrowSizeY - A, ArrowY + A - 1.0f, A, 1.0f, WhiteBrush, Color);
+			}
+		}
+	}
+
+	DrawContext.LayerId++;
+}
+
+////////////////////////////////////////////////////////////////////////////////////////////////////
+
+#undef LOCTEXT_NAMESPACE
diff --git a/Engine/Source/Developer/TraceInsights/Private/Insights/ViewModels/TrackHeader.h b/Engine/Source/Developer/TraceInsights/Private/Insights/ViewModels/TrackHeader.h
new file mode 100644
--- /dev/null
+++ b/Engine/Source/Developer/TraceInsights/Private/Insights/ViewModels/TrackHeader.h
@@ -0,0 +1,59 @@
+// Copyright 1998-2019 Epic Games, Inc. All Rights Reserved.
+
+#pragma once
+
+#include "CoreMinimal.h"
+#include "Fonts/SlateFontInfo.h"
+
+class ITimingTrackUpdateContext;
+class ITimingTrackDrawContext;
+class FBaseTimingTrack;
+struct FSlateBrush;
+
+class FTrackHeader
+{
+public:
+	FTrackHeader(FBaseTimingTrack& InParentTrack);
+	~FTrackHeader() {}
+
+	void Reset();
+
+	bool IsInBackground() const { return bIsInBackground; }
+	void SetIsInBackground(bool bOnOff) { bIsInBackground = bOnOff; }
+
+	bool CanBeCollapsed() const { return bCanBeCollapsed; }
+	void SetCanBeCollapsed(bool bOnOff) { bCanBeCollapsed = bOnOff; }
+
+	bool IsCollapsed() const { return bIsCollapsed; }
+	void SetIsCollapsed(bool bOnOff) { bIsCollapsed = bOnOff; }
+	void ToggleCollapsed() { bIsCollapsed = !bIsCollapsed; }
+
+	void UpdateSize();
+
+	void Update(const ITimingTrackUpdateContext& Context);
+	void PostUpdate(const ITimingTrackUpdateContext& Context);
+
+	void Draw(const ITimingTrackDrawContext& Context) const;
+	void PostDraw(const ITimingTrackDrawContext& Context) const;
+
+private:
+	void DrawInternal(const ITimingTrackDrawContext& Context, bool bDrawBackgroundLayer) const;
+
+private:
+	FBaseTimingTrack& ParentTrack;
+
+	float Width;
+	float Height;
+
+	bool bIsInBackground;
+
+	bool bCanBeCollapsed;
+	bool bIsCollapsed;
+
+	float TargetHoveredAnimPercent; // [0.0 .. 1.0], 0.0 = hidden, 1.0 = visible
+	float CurrentHoveredAnimPercent;
+
+	// Slate resources
+	const FSlateBrush* WhiteBrush;
+	const FSlateFontInfo Font;
+};
diff --git a/Engine/Source/Developer/TraceInsights/Private/Insights/Widgets/SGraphSeriesList.cpp b/Engine/Source/Developer/TraceInsights/Private/Insights/Widgets/SGraphSeriesList.cpp
--- a/Engine/Source/Developer/TraceInsights/Private/Insights/Widgets/SGraphSeriesList.cpp
+++ b/Engine/Source/Developer/TraceInsights/Private/Insights/Widgets/SGraphSeriesList.cpp
@@ -1,21 +1,27 @@
 // Copyright 1998-2019 Epic Games, Inc. All Rights Reserved.
 
 #include "SGraphSeriesList.h"
-#include "Widgets/SBoxPanel.h"
+
+#include "EditorFontGlyphs.h"
+#include "EditorStyleSet.h"
+#include "Framework/Application/SlateApplication.h"
+#include "Widgets/Input/SCheckBox.h"
 #include "Widgets/Input/SSearchBox.h"
 #include "Widgets/Text/STextBlock.h"
-#include "Widgets/Input/SCheckBox.h"
-#include "Widgets/Views/SHeaderRow.h"
 #include "Widgets/Layout/SScrollBorder.h"
-#include "Framework/Application/SlateApplication.h"
-#include "EditorStyleSet.h"
-#include "EditorFontGlyphs.h"
+#include "Widgets/SBoxPanel.h"
+#include "Widgets/Views/SHeaderRow.h"
+
 
 #include "Insights/ViewModels/GraphSeries.h"
 #include "Insights/ViewModels/GraphTrack.h"
 
 #define LOCTEXT_NAMESPACE "SGraphSeriesList"
 
+////////////////////////////////////////////////////////////////////////////////////////////////////
+// SGraphSeriesListEntry
+////////////////////////////////////////////////////////////////////////////////////////////////////
+
 // A list entry widget for a graph series
 class SGraphSeriesListEntry : public STableRow<TSharedPtr<FGraphSeries>>
 {
@@ -32,8 +38,8 @@ public:
 
 	void Construct(const FArguments& InArgs, const TSharedRef<STableViewBase>& InOwnerTable)
 	{
-		GraphSeries = InArgs._GraphSeries;
 		GraphTrack = InArgs._GraphTrack;
+		GraphSeries = InArgs._GraphSeries;
 		SearchText = InArgs._SearchText;
 
 		STableRow<TSharedPtr<FGraphSeries>>::Construct(STableRow<TSharedPtr<FGraphSeries>>::FArguments(), InOwnerTable);
@@ -44,7 +50,7 @@ public:
 		ChildSlot
 		[
 			SNew(SHorizontalBox)
-			+SHorizontalBox::Slot()
+			+ SHorizontalBox::Slot()
 			.AutoWidth()
 			.Padding(2.0f)
 			.HAlign(HAlign_Center)
@@ -52,11 +58,11 @@ public:
 			[
 				SNew(SCheckBox)
 				.IsChecked_Lambda([this](){ return GraphSeries.Pin()->IsVisible() ? ECheckBoxState::Checked : ECheckBoxState::Unchecked; })
-				.OnCheckStateChanged_Lambda([this](ECheckBoxState InCheckBoxState) 
+				.OnCheckStateChanged_Lambda([this](ECheckBoxState InCheckBoxState)
 				{
 					GraphSeries.Pin()->SetVisibility(InCheckBoxState == ECheckBoxState::Checked);
 					GraphSeries.Pin()->SetDirtyFlag();
-					GraphTrack.Pin()->SetDirtyFlag(); 
+					GraphTrack.Pin()->SetDirtyFlag();
 				})
 				.Content()
 				[
@@ -65,7 +71,7 @@ public:
 					.HighlightText(SearchText)
 				]
 			]
-			+SHorizontalBox::Slot()
+			+ SHorizontalBox::Slot()
 			.FillWidth(1.0f)
 			.Padding(2.0f)
 			.HAlign(HAlign_Right)
@@ -89,6 +95,10 @@ public:
 	TAttribute<FText> SearchText;
 };
 
+////////////////////////////////////////////////////////////////////////////////////////////////////
+// SGraphSeriesList
+////////////////////////////////////////////////////////////////////////////////////////////////////
+
 void SGraphSeriesList::Construct(const FArguments& InArgs, const TSharedRef<FGraphTrack>& InGraphTrack)
 {
 	GraphTrack = InGraphTrack;
@@ -118,29 +128,29 @@ void SGraphSeriesList::Construct(const FArguments& InArgs, const TSharedRef<FGra
 				.IsChecked_Lambda([this]()
 				{
 					int32 NumVisible = 0;
-					for(TSharedPtr<FGraphSeries> Series : FilteredSeries)
+					for (TSharedPtr<FGraphSeries> Series : FilteredSeries)
 					{
-						if(Series->IsVisible())
+						if (Series->IsVisible())
 						{
 							NumVisible++;
 						}
 					}
 
-					if(NumVisible == 0)
+					if (NumVisible == 0)
 					{
 						return ECheckBoxState::Unchecked;
 					}
-					else if(NumVisible == FilteredSeries.Num())
+					else if (NumVisible == FilteredSeries.Num())
 					{
 						return ECheckBoxState::Checked;
 					}
 
 					return ECheckBoxState::Undetermined;
 				})
-				.OnCheckStateChanged_Lambda([this](ECheckBoxState InCheckBoxState) 
+				.OnCheckStateChanged_Lambda([this](ECheckBoxState InCheckBoxState)
 				{
 					const bool bVisible = InCheckBoxState != ECheckBoxState::Unchecked;
-					for(TSharedPtr<FGraphSeries> Series : FilteredSeries)
+					for (TSharedPtr<FGraphSeries> Series : FilteredSeries)
 					{
 						Series->SetVisibility(bVisible);
 						Series->SetDirtyFlag();
@@ -149,7 +159,7 @@ void SGraphSeriesList::Construct(const FArguments& InArgs, const TSharedRef<FGra
 					GraphTrack.Pin()->SetDirtyFlag();
 				})
 			]
-			+SHorizontalBox::Slot()
+			+ SHorizontalBox::Slot()
 			.FillWidth(1.0f)
 			[
 				// Search box allows for filtering
@@ -157,7 +167,7 @@ void SGraphSeriesList::Construct(const FArguments& InArgs, const TSharedRef<FGra
 				.OnTextChanged_Lambda([this](const FText& InText){ SearchText = InText; RefreshFilter(); })
 			]
 		]
-		+SVerticalBox::Slot()
+		+ SVerticalBox::Slot()
 		.FillHeight(1.0f)
 		[
 			SNew(SScrollBorder, ListView.ToSharedRef())
@@ -172,6 +182,8 @@ void SGraphSeriesList::Construct(const FArguments& InArgs, const TSharedRef<FGra
 	FSlateApplication::Get().SetUserFocus(0, SearchBox);
 }
 
+////////////////////////////////////////////////////////////////////////////////////////////////////
+
 TSharedRef<ITableRow> SGraphSeriesList::OnGenerateRow(TSharedPtr<FGraphSeries> Item, const TSharedRef<STableViewBase>& OwnerTable)
 {
 	return SNew(SGraphSeriesListEntry, OwnerTable)
@@ -180,13 +192,15 @@ TSharedRef<ITableRow> SGraphSeriesList::OnGenerateRow(TSharedPtr<FGraphSeries> I
 		.SearchText_Lambda([this](){ return SearchText; });
 }
 
+////////////////////////////////////////////////////////////////////////////////////////////////////
+
 void SGraphSeriesList::RefreshFilter()
 {
 	FilteredSeries.Reset();
 
-	for(TSharedPtr<FGraphSeries> Series : GraphTrack.Pin()->GetSeries())
+	for (TSharedPtr<FGraphSeries> Series : GraphTrack.Pin()->GetSeries())
 	{
-		if(SearchText.IsEmpty() || Series->GetName().ToString().Contains(SearchText.ToString()))
+		if (SearchText.IsEmpty() || Series->GetName().ToString().Contains(SearchText.ToString()))
 		{
 			FilteredSeries.Add(Series);
 		}
@@ -195,4 +209,6 @@ void SGraphSeriesList::RefreshFilter()
 	ListView->RequestListRefresh();
 }
 
-#undef LOCTEXT_NAMESPACE
\ No newline at end of file
+////////////////////////////////////////////////////////////////////////////////////////////////////
+
+#undef LOCTEXT_NAMESPACE
diff --git a/Engine/Source/Developer/TraceInsights/Private/Insights/Widgets/STimingView.cpp b/Engine/Source/Developer/TraceInsights/Private/Insights/Widgets/STimingView.cpp
--- a/Engine/Source/Developer/TraceInsights/Private/Insights/Widgets/STimingView.cpp
+++ b/Engine/Source/Developer/TraceInsights/Private/Insights/Widgets/STimingView.cpp
@@ -1432,9 +1432,68 @@ void STimingView::AddForegroundTrack(TSharedPtr<FBaseTimingTrack> Track)
 
 ////////////////////////////////////////////////////////////////////////////////////////////////////
 
+FReply STimingView::AllowTracksToProcessOnMouseButtonDown(const FGeometry& MyGeometry, const FPointerEvent& MouseEvent)
+{
+	for (TSharedPtr<FBaseTimingTrack>& TrackPtr : TopDockedTracks)
+	{
+		if (TrackPtr->IsVisible())
+		{
+			FReply Reply = TrackPtr->OnMouseButtonDown(MyGeometry, MouseEvent);
+			if (Reply.IsEventHandled())
+			{
+				return Reply;
+			}
+		}
+	}
+
+	for (TSharedPtr<FBaseTimingTrack>& TrackPtr : BottomDockedTracks)
+	{
+		if (TrackPtr->IsVisible())
+		{
+			FReply Reply = TrackPtr->OnMouseButtonDown(MyGeometry, MouseEvent);
+			if (Reply.IsEventHandled())
+			{
+				return Reply;
+			}
+		}
+	}
+
+	for (TSharedPtr<FBaseTimingTrack>& TrackPtr : ScrollableTracks)
+	{
+		if (TrackPtr->IsVisible())
+		{
+			FReply Reply = TrackPtr->OnMouseButtonDown(MyGeometry, MouseEvent);
+			if (Reply.IsEventHandled())
+			{
+				return Reply;
+			}
+		}
+	}
+
+	for (TSharedPtr<FBaseTimingTrack>& TrackPtr : ForegroundTracks)
+	{
+		if (TrackPtr->IsVisible())
+		{
+			FReply Reply = TrackPtr->OnMouseButtonDown(MyGeometry, MouseEvent);
+			if (Reply.IsEventHandled())
+			{
+				return Reply;
+			}
+		}
+	}
+
+	return FReply::Unhandled();
+}
+
+////////////////////////////////////////////////////////////////////////////////////////////////////
+
 FReply STimingView::OnMouseButtonDown(const FGeometry& MyGeometry, const FPointerEvent& MouseEvent)
 {
-	FReply Reply = FReply::Unhandled();
+	FReply Reply = AllowTracksToProcessOnMouseButtonDown(MyGeometry, MouseEvent);
+	if (Reply.IsEventHandled())
+	{
+		return Reply;
+	}
 
 	MousePositionOnButtonDown = MyGeometry.AbsoluteToLocal(MouseEvent.GetScreenSpacePosition());
 	MousePosition = MousePositionOnButtonDown;
@@ -1446,11 +1505,6 @@ FReply STimingView::OnMouseButtonDown(const FGeometry& MyGeometry, const FPointe
 
 	if (MouseEvent.GetEffectingButton() == EKeys::LeftMouseButton)
 	{
-		if (MarkersTrack->IsVisible() && MarkersTrack->IsHeaderHovered())
-		{
-			MarkersTrack->ToggleCollapsed();
-		}
-
 		if (!bIsRMB_Pressed)
 		{
 			bIsLMB_Pressed = true;
@@ -1546,9 +1600,68 @@ FReply STimingView::OnMouseButtonDown(const FGeometry& MyGeometry, const FPointe
 
 ////////////////////////////////////////////////////////////////////////////////////////////////////
 
+FReply STimingView::AllowTracksToProcessOnMouseButtonUp(const FGeometry& MyGeometry, const FPointerEvent& MouseEvent)
+{
+	for (TSharedPtr<FBaseTimingTrack>& TrackPtr : TopDockedTracks)
+	{
+		if (TrackPtr->IsVisible())
+		{
+			FReply Reply = TrackPtr->OnMouseButtonUp(MyGeometry, MouseEvent);
+			if (Reply.IsEventHandled())
+			{
+				return Reply;
+			}
+		}
+	}
+
+	for (TSharedPtr<FBaseTimingTrack>& TrackPtr : BottomDockedTracks)
+	{
+		if (TrackPtr->IsVisible())
+		{
+			FReply Reply = TrackPtr->OnMouseButtonUp(MyGeometry, MouseEvent);
+			if (Reply.IsEventHandled())
+			{
+				return Reply;
+			}
+		}
+	}
+
+	for (TSharedPtr<FBaseTimingTrack>& TrackPtr : ScrollableTracks)
+	{
+		if (TrackPtr->IsVisible())
+		{
+			FReply Reply = TrackPtr->OnMouseButtonUp(MyGeometry, MouseEvent);
+			if (Reply.IsEventHandled())
+			{
+				return Reply;
+			}
+		}
+	}
+
+	for (TSharedPtr<FBaseTimingTrack>& TrackPtr : ForegroundTracks)
+	{
+		if (TrackPtr->IsVisible())
+		{
+			FReply Reply = TrackPtr->OnMouseButtonUp(MyGeometry, MouseEvent);
+			if (Reply.IsEventHandled())
+			{
+				return Reply;
+			}
+		}
+	}
+
+	return FReply::Unhandled();
+}
+
+////////////////////////////////////////////////////////////////////////////////////////////////////
+
 FReply STimingView::OnMouseButtonUp(const FGeometry& MyGeometry, const FPointerEvent& MouseEvent)
 {
-	FReply Reply = FReply::Unhandled();
+	FReply Reply = AllowTracksToProcessOnMouseButtonUp(MyGeometry, MouseEvent);
+	if (Reply.IsEventHandled())
+	{
+		return Reply;
+	}
 
 	MousePositionOnButtonUp = MyGeometry.AbsoluteToLocal(MouseEvent.GetScreenSpacePosition());
 	MousePosition = MousePositionOnButtonUp;
@@ -1640,17 +1753,72 @@ FReply STimingView::OnMouseButtonUp(const FGeometry& MyGeometry, const FPointerE
 
 ////////////////////////////////////////////////////////////////////////////////////////////////////
 
-FReply STimingView::OnMouseButtonDoubleClick(const FGeometry& MyGeometry, const FPointerEvent& MouseEvent)
+FReply STimingView::AllowTracksToProcessOnMouseButtonDoubleClick(const FGeometry& MyGeometry, const FPointerEvent& MouseEvent)
 {
-	FReply Reply = FReply::Unhandled();
+	for (TSharedPtr<FBaseTimingTrack>& TrackPtr : TopDockedTracks)
+	{
+		if (TrackPtr->IsVisible())
+		{
+			FReply Reply = TrackPtr->OnMouseButtonDoubleClick(MyGeometry, MouseEvent);
+			if (Reply.IsEventHandled())
+			{
+				return Reply;
+			}
+		}
+	}
 
-	if (MouseEvent.GetEffectingButton() == EKeys::LeftMouseButton)
+	for (TSharedPtr<FBaseTimingTrack>& TrackPtr : BottomDockedTracks)
 	{
-		if (MarkersTrack->IsVisible() && MarkersTrack->IsHovered())
+		if (TrackPtr->IsVisible())
 		{
-			MarkersTrack->ToggleCollapsed();
+			FReply Reply = TrackPtr->OnMouseButtonDoubleClick(MyGeometry, MouseEvent);
+			if (Reply.IsEventHandled())
+			{
+				return Reply;
+			}
 		}
-		else if (HoveredEvent.IsValid())
+	}
+
+	for (TSharedPtr<FBaseTimingTrack>& TrackPtr : ScrollableTracks)
+	{
+		if (TrackPtr->IsVisible())
+		{
+			FReply Reply = TrackPtr->OnMouseButtonDoubleClick(MyGeometry, MouseEvent);
+			if (Reply.IsEventHandled())
+			{
+				return Reply;
+			}
+		}
+	}
+
+	for (TSharedPtr<FBaseTimingTrack>& TrackPtr : ForegroundTracks)
+	{
+		if (TrackPtr->IsVisible())
+		{
+			FReply Reply = TrackPtr->OnMouseButtonDoubleClick(MyGeometry, MouseEvent);
+			if (Reply.IsEventHandled())
+			{
+				return Reply;
+			}
+		}
+	}
+
+	return FReply::Unhandled();
+}
+
+////////////////////////////////////////////////////////////////////////////////////////////////////
+
+FReply STimingView::OnMouseButtonDoubleClick(const FGeometry& MyGeometry, const FPointerEvent& MouseEvent)
+{
+	FReply Reply = AllowTracksToProcessOnMouseButtonDoubleClick(MyGeometry, MouseEvent);
+	if (Reply.IsEventHandled())
+	{
+		return Reply;
+	}
+
+	if (MouseEvent.GetEffectingButton() == EKeys::LeftMouseButton)
+	{
+		if (HoveredEvent.IsValid())
 		{
 			if (MouseEvent.GetModifierKeys().IsControlDown())
 			{
diff --git a/Engine/Source/Developer/TraceInsights/Private/Insights/Widgets/STimingView.h b/Engine/Source/Developer/TraceInsights/Private/Insights/Widgets/STimingView.h
--- a/Engine/Source/Developer/TraceInsights/Private/Insights/Widgets/STimingView.h
+++ b/Engine/Source/Developer/TraceInsights/Private/Insights/Widgets/STimingView.h
@@ -341,28 +341,9 @@ protected:
 	// Get all the plugin extenders we care about
 	TArray<Insights::ITimingViewExtender*> GetExtenders() const;
 
-	////////////////////////////////////////////////////////////////////////////////////////////////////
-	// FrameSelectionChanged Event
-
-public:
-	/**
-	 * The event to execute when the selected frames have been changed.
-	 *
-	 * @param FrameStartIndex	- The index of the first frame selected.
-	 * @param FrameEndIndex		- The index of the last frame selected.
-	 *
-	 */
-	DECLARE_EVENT_TwoParams(STimingView, FFrameSelectionChangedEvent, int32 /*FrameStartIndex*/, int32 /*FrameEndIndex*/);
-	FFrameSelectionChangedEvent& OnFrameSelectionChanged()
-	{
-		return FrameSelectionChangedEvent;
-	}
-
-protected:
-	/** The event to execute when the selected frames have been changed. */
-	FFrameSelectionChangedEvent FrameSelectionChangedEvent;
-
-	////////////////////////////////////////////////////////////////////////////////////////////////////
+	FReply AllowTracksToProcessOnMouseButtonDown(const FGeometry& MyGeometry, const FPointerEvent& MouseEvent);
+	FReply AllowTracksToProcessOnMouseButtonUp(const FGeometry& MyGeometry, const FPointerEvent& MouseEvent);
+	FReply AllowTracksToProcessOnMouseButtonDoubleClick(const FGeometry& MyGeometry, const FPointerEvent& MouseEvent);
 
 protected:
 	/** The track's viewport. Encapsulates info about position and scale. */
diff --git a/Engine/Source/Developer/TraceInsights/Public/Insights/ViewModels/BaseTimingTrack.h b/Engine/Source/Developer/TraceInsights/Public/Insights/ViewModels/BaseTimingTrack.h
--- a/Engine/Source/Developer/TraceInsights/Public/Insights/ViewModels/BaseTimingTrack.h
+++ b/Engine/Source/Developer/TraceInsights/Public/Insights/ViewModels/BaseTimingTrack.h
@@ -3,6 +3,8 @@
 #pragma once
 
 #include "CoreMinimal.h"
+#include "Input/Events.h"
+#include "Input/Reply.h"
 #include "Misc/EnumClassFlags.h"
 
 #include "Insights/Common/SimpleRtti.h"
@@ -182,6 +184,10 @@ public:
 
 	//////////////////////////////////////////////////
 
+	virtual FReply OnMouseButtonDown(const FGeometry& MyGeometry, const FPointerEvent& MouseEvent) { return FReply::Unhandled(); }
+	virtual FReply OnMouseButtonUp(const FGeometry& MyGeometry, const FPointerEvent& MouseEvent) { return FReply::Unhandled(); }
+	virtual FReply OnMouseButtonDoubleClick(const FGeometry& MyGeometry, const FPointerEvent& MouseEvent) { return FReply::Unhandled(); }
+
 	virtual void BuildContextMenu(FMenuBuilder& MenuBuilder);
 
 	/**
diff --git a/Engine/Source/Developer/TraceInsights/Public/Insights/ViewModels/TimingEventsTrack.h b/Engine/Source/Developer/TraceInsights/Public/Insights/ViewModels/TimingEventsTrack.h
--- a/Engine/Source/Developer/TraceInsights/Public/Insights/ViewModels/TimingEventsTrack.h
+++ b/Engine/Source/Developer/TraceInsights/Public/Insights/ViewModels/TimingEventsTrack.h
@@ -57,6 +57,8 @@ protected:
 	void DrawEvents(const ITimingTrackDrawContext& Context, const float OffsetY = 1.0f) const;
 	void DrawHeader(const ITimingTrackDrawContext& Context) const;
 
+	void DrawMarkers(const ITimingTrackDrawContext& Context, float LineY, float LineH) const;
+
 	int32 GetHeaderBackgroundLayerId(const ITimingTrackDrawContext& Context) const;
 	int32 GetHeaderTextLayerId(const ITimingTrackDrawContext& Context) const;
 
