Some small graph track tweaks

Expose animation parameter of auto-zoom to allow derived tracks/series to optionally animate auto-zoom.
Fix an issue with the builder where empty line batches could be added by repeatedly re-starting connected lines.

diff --git a/Engine/Source/Developer/TraceInsights/Private/Insights/ViewModels/GraphSeries.cpp b/Engine/Source/Developer/TraceInsights/Private/Insights/ViewModels/GraphSeries.cpp
--- a/Engine/Source/Developer/TraceInsights/Private/Insights/ViewModels/GraphSeries.cpp
+++ b/Engine/Source/Developer/TraceInsights/Private/Insights/ViewModels/GraphSeries.cpp
@@ -101,7 +101,7 @@ FString FGraphSeries::FormatValue(double Value) const
 
 ////////////////////////////////////////////////////////////////////////////////////////////////////
 
-void FGraphSeries::UpdateAutoZoom(const float InTopY, const float InBottomY, const double InMinEventValue, const double InMaxEventValue)
+void FGraphSeries::UpdateAutoZoom(const float InTopY, const float InBottomY, const double InMinEventValue, const double InMaxEventValue, const bool bIsAutoZoomAnimated)
 {
 	if (IsAutoZoomEnabled())
 	{
@@ -120,7 +120,6 @@ void FGraphSeries::UpdateAutoZoom(const float InTopY, const float InBottomY, con
 
 		if (MinValue < MaxValue)
 		{
-			constexpr bool bIsAutoZoomAnimated = true;
 			if (bIsAutoZoomAnimated)
 			{
 				// Interpolate the min-max interval (animating the vertical position and scale of the graph series).
diff --git a/Engine/Source/Developer/TraceInsights/Private/Insights/ViewModels/GraphTrackBuilder.cpp b/Engine/Source/Developer/TraceInsights/Private/Insights/ViewModels/GraphTrackBuilder.cpp
--- a/Engine/Source/Developer/TraceInsights/Private/Insights/ViewModels/GraphTrackBuilder.cpp
+++ b/Engine/Source/Developer/TraceInsights/Private/Insights/ViewModels/GraphTrackBuilder.cpp
@@ -176,7 +176,7 @@ void FGraphTrackBuilder::AddConnectedLine(double Time, double Value, bool bNewBa
 	// Start a new batch
 	if (bNewBatch)
 	{
-		if (Series.LinePoints.Num() > 0)
+		if (Series.LinePoints.Num() > 0 && Series.LinePoints.Last().Num() > 0)
 		{
 			Series.LinePoints.AddDefaulted();
 		}
diff --git a/Engine/Source/Developer/TraceInsights/Public/Insights/ViewModels/GraphSeries.h b/Engine/Source/Developer/TraceInsights/Public/Insights/ViewModels/GraphSeries.h
--- a/Engine/Source/Developer/TraceInsights/Public/Insights/ViewModels/GraphSeries.h
+++ b/Engine/Source/Developer/TraceInsights/Public/Insights/ViewModels/GraphSeries.h
@@ -129,7 +129,7 @@ public:
 	const FGraphSeriesEvent* GetEvent(const float PosX, const float PosY, const FTimingTrackViewport& Viewport, bool bCheckLine, bool bCheckBox) const;
 
 	/** Update the track's auto-zoom. Does nothing if IsAutoZoomEnabled() is false. */
-	void UpdateAutoZoom(const float InTopY, const float InBottomY, const double InMinEventValue, const double InMaxEventValue);
+	void UpdateAutoZoom(const float InTopY, const float InBottomY, const double InMinEventValue, const double InMaxEventValue, const bool bIsAutoZoomAnimated = true);
 
 	virtual FString FormatValue(double Value) const;
 
