[Insights] Updated more copyrights to remove years.

diff --git a/Engine/Source/Developer/TraceInsights/Private/Insights/ViewModels/FrameTimingTrack.cpp b/Engine/Source/Developer/TraceInsights/Private/Insights/ViewModels/FrameTimingTrack.cpp
--- a/Engine/Source/Developer/TraceInsights/Private/Insights/ViewModels/FrameTimingTrack.cpp
+++ b/Engine/Source/Developer/TraceInsights/Private/Insights/ViewModels/FrameTimingTrack.cpp
@@ -1,4 +1,4 @@
-// Copyright 1998-2019 Epic Games, Inc. All Rights Reserved.
+// Copyright Epic Games, Inc. All Rights Reserved.
 
 #include "FrameTimingTrack.h"
 
diff --git a/Engine/Source/Developer/TraceInsights/Private/Insights/ViewModels/FrameTimingTrack.h b/Engine/Source/Developer/TraceInsights/Private/Insights/ViewModels/FrameTimingTrack.h
--- a/Engine/Source/Developer/TraceInsights/Private/Insights/ViewModels/FrameTimingTrack.h
+++ b/Engine/Source/Developer/TraceInsights/Private/Insights/ViewModels/FrameTimingTrack.h
@@ -1,4 +1,4 @@
-// Copyright 1998-2019 Epic Games, Inc. All Rights Reserved.
+// Copyright Epic Games, Inc. All Rights Reserved.
 
 #pragma once
 
diff --git a/Engine/Source/Developer/TraceInsights/Private/Insights/ViewModels/TimingEvent.cpp b/Engine/Source/Developer/TraceInsights/Private/Insights/ViewModels/TimingEvent.cpp
--- a/Engine/Source/Developer/TraceInsights/Private/Insights/ViewModels/TimingEvent.cpp
+++ b/Engine/Source/Developer/TraceInsights/Private/Insights/ViewModels/TimingEvent.cpp
@@ -1,4 +1,4 @@
-// Copyright 1998-2019 Epic Games, Inc. All Rights Reserved.
+// Copyright Epic Games, Inc. All Rights Reserved.
 
 #include "Insights/ViewModels/TimingEvent.h"
 
diff --git a/Engine/Source/Developer/TraceInsights/Private/Insights/Widgets/SGraphSeriesList.cpp b/Engine/Source/Developer/TraceInsights/Private/Insights/Widgets/SGraphSeriesList.cpp
--- a/Engine/Source/Developer/TraceInsights/Private/Insights/Widgets/SGraphSeriesList.cpp
+++ b/Engine/Source/Developer/TraceInsights/Private/Insights/Widgets/SGraphSeriesList.cpp
@@ -12,7 +12,7 @@
 #include "Widgets/SBoxPanel.h"
 #include "Widgets/Views/SHeaderRow.h"
 
-
+// Insights
 #include "Insights/ViewModels/GraphSeries.h"
 #include "Insights/ViewModels/GraphTrack.h"
 
