[Insights] More copyright updates.

diff --git a/Engine/Source/Developer/TraceInsights/Private/Insights/ViewModels/TrackHeader.cpp b/Engine/Source/Developer/TraceInsights/Private/Insights/ViewModels/TrackHeader.cpp
--- a/Engine/Source/Developer/TraceInsights/Private/Insights/ViewModels/TrackHeader.cpp
+++ b/Engine/Source/Developer/TraceInsights/Private/Insights/ViewModels/TrackHeader.cpp
@@ -1,4 +1,4 @@
-// Copyright 1998-2019 Epic Games, Inc. All Rights Reserved.
+// Copyright Epic Games, Inc. All Rights Reserved.
 
 #include "TrackHeader.h"
 
diff --git a/Engine/Source/Developer/TraceInsights/Private/Insights/ViewModels/TrackHeader.h b/Engine/Source/Developer/TraceInsights/Private/Insights/ViewModels/TrackHeader.h
--- a/Engine/Source/Developer/TraceInsights/Private/Insights/ViewModels/TrackHeader.h
+++ b/Engine/Source/Developer/TraceInsights/Private/Insights/ViewModels/TrackHeader.h
@@ -1,4 +1,4 @@
-// Copyright 1998-2019 Epic Games, Inc. All Rights Reserved.
+// Copyright Epic Games, Inc. All Rights Reserved.
 
 #pragma once
 
diff --git a/Engine/Source/Developer/TraceInsights/Public/Insights/Common/SimpleRtti.h b/Engine/Source/Developer/TraceInsights/Public/Insights/Common/SimpleRtti.h
--- a/Engine/Source/Developer/TraceInsights/Public/Insights/Common/SimpleRtti.h
+++ b/Engine/Source/Developer/TraceInsights/Public/Insights/Common/SimpleRtti.h
@@ -1,4 +1,4 @@
-// Copyright 1998-2019 Epic Games, Inc. All Rights Reserved.
+// Copyright Epic Games, Inc. All Rights Reserved.
 
 #pragma once
 
