[Insights] Table Tree View: The widget can now rebuild the columns if a new source table is provided with a new table layout. Fixes UE-86573.

diff --git a/Engine/Source/Developer/TraceInsights/Private/Insights/LoadingProfiler/Widgets/SLoadingProfilerWindow.cpp b/Engine/Source/Developer/TraceInsights/Private/Insights/LoadingProfiler/Widgets/SLoadingProfilerWindow.cpp
--- a/Engine/Source/Developer/TraceInsights/Private/Insights/LoadingProfiler/Widgets/SLoadingProfilerWindow.cpp
+++ b/Engine/Source/Developer/TraceInsights/Private/Insights/LoadingProfiler/Widgets/SLoadingProfilerWindow.cpp
@@ -99,41 +99,108 @@ void SLoadingProfilerWindow::Reset()
 
 void SLoadingProfilerWindow::UpdateTableTreeViews()
 {
-	TSharedPtr<const Trace::IAnalysisSession> Session = FInsightsManager::Get()->GetSession();
-	if (Session.IsValid() && Trace::ReadLoadTimeProfilerProvider(*Session.Get()))
-	{
-		Trace::FAnalysisSessionReadScope SessionReadScope(*Session.Get());
-		const Trace::ILoadTimeProfilerProvider& LoadTimeProfilerProvider = *Trace::ReadLoadTimeProfilerProvider(*Session.Get());
+	UpdateEventAggregationTreeView();
+	UpdateObjectTypeAggregationTreeView();
+	UpdatePackageDetailsTreeView();
+	UpdateExportDetailsTreeView();
+}
 
-		const double SelectionStartTime = TimingView ? TimingView->GetSelectionStartTime() : 0.0;
-		const double SelectionEndTime = TimingView ? TimingView->GetSelectionEndTime() : 0.0;
+////////////////////////////////////////////////////////////////////////////////////////////////////
 
-		if (EventAggregationTreeView)
+void SLoadingProfilerWindow::UpdateEventAggregationTreeView()
+{
+	if (EventAggregationTreeView)
+	{
+		TSharedPtr<const Trace::IAnalysisSession> Session = FInsightsManager::Get()->GetSession();
+		if (Session.IsValid() && Trace::ReadLoadTimeProfilerProvider(*Session.Get()))
 		{
+			Trace::FAnalysisSessionReadScope SessionReadScope(*Session.Get());
+			const Trace::ILoadTimeProfilerProvider& LoadTimeProfilerProvider = *Trace::ReadLoadTimeProfilerProvider(*Session.Get());
+
+			const double SelectionStartTime = TimingView ? TimingView->GetSelectionStartTime() : 0.0;
+			const double SelectionEndTime = TimingView ? TimingView->GetSelectionEndTime() : 0.0;
+
 			Trace::ITable<Trace::FLoadTimeProfilerAggregatedStats>* EventAggregationTable = LoadTimeProfilerProvider.CreateEventAggregation(SelectionStartTime, SelectionEndTime);
-			EventAggregationTreeView->GetTable()->UpdateSourceTable(MakeShareable(EventAggregationTable));
-			EventAggregationTreeView->RebuildTree(true);
+			EventAggregationTreeView->UpdateSourceTable(MakeShareable(EventAggregationTable));
+		}
+		else
+		{
+			EventAggregationTreeView->UpdateSourceTable(nullptr);
 		}
+	}
+}
+
+////////////////////////////////////////////////////////////////////////////////////////////////////
 
-		if (ObjectTypeAggregationTreeView)
+void SLoadingProfilerWindow::UpdateObjectTypeAggregationTreeView()
+{
+	if (ObjectTypeAggregationTreeView)
+	{
+		TSharedPtr<const Trace::IAnalysisSession> Session = FInsightsManager::Get()->GetSession();
+		if (Session.IsValid() && Trace::ReadLoadTimeProfilerProvider(*Session.Get()))
 		{
+			Trace::FAnalysisSessionReadScope SessionReadScope(*Session.Get());
+			const Trace::ILoadTimeProfilerProvider& LoadTimeProfilerProvider = *Trace::ReadLoadTimeProfilerProvider(*Session.Get());
+
+			const double SelectionStartTime = TimingView ? TimingView->GetSelectionStartTime() : 0.0;
+			const double SelectionEndTime = TimingView ? TimingView->GetSelectionEndTime() : 0.0;
+
 			Trace::ITable<Trace::FLoadTimeProfilerAggregatedStats>* ObjectTypeAggregationTable = LoadTimeProfilerProvider.CreateObjectTypeAggregation(SelectionStartTime, SelectionEndTime);
-			ObjectTypeAggregationTreeView->GetTable()->UpdateSourceTable(MakeShareable(ObjectTypeAggregationTable));
-			ObjectTypeAggregationTreeView->RebuildTree(true);
+			ObjectTypeAggregationTreeView->UpdateSourceTable(MakeShareable(ObjectTypeAggregationTable));
+		}
+		else
+		{
+			ObjectTypeAggregationTreeView->UpdateSourceTable(nullptr);
 		}
+	}
+}
+
+////////////////////////////////////////////////////////////////////////////////////////////////////
 
-		if (PackageDetailsTreeView)
+void SLoadingProfilerWindow::UpdatePackageDetailsTreeView()
+{
+	if (PackageDetailsTreeView)
+	{
+		TSharedPtr<const Trace::IAnalysisSession> Session = FInsightsManager::Get()->GetSession();
+		if (Session.IsValid() && Trace::ReadLoadTimeProfilerProvider(*Session.Get()))
 		{
+			Trace::FAnalysisSessionReadScope SessionReadScope(*Session.Get());
+			const Trace::ILoadTimeProfilerProvider& LoadTimeProfilerProvider = *Trace::ReadLoadTimeProfilerProvider(*Session.Get());
+
+			const double SelectionStartTime = TimingView ? TimingView->GetSelectionStartTime() : 0.0;
+			const double SelectionEndTime = TimingView ? TimingView->GetSelectionEndTime() : 0.0;
+
 			Trace::ITable<Trace::FPackagesTableRow>* PackageDetailsTable = LoadTimeProfilerProvider.CreatePackageDetailsTable(SelectionStartTime, SelectionEndTime);
-			PackageDetailsTreeView->GetTable()->UpdateSourceTable(MakeShareable(PackageDetailsTable));
-			PackageDetailsTreeView->RebuildTree(true);
+			PackageDetailsTreeView->UpdateSourceTable(MakeShareable(PackageDetailsTable));
+		}
+		else
+		{
+			PackageDetailsTreeView->UpdateSourceTable(nullptr);
 		}
+	}
+}
+
+////////////////////////////////////////////////////////////////////////////////////////////////////
 
-		if (ExportDetailsTreeView)
+void SLoadingProfilerWindow::UpdateExportDetailsTreeView()
+{
+	if (ExportDetailsTreeView)
+	{
+		TSharedPtr<const Trace::IAnalysisSession> Session = FInsightsManager::Get()->GetSession();
+		if (Session.IsValid() && Trace::ReadLoadTimeProfilerProvider(*Session.Get()))
 		{
+			Trace::FAnalysisSessionReadScope SessionReadScope(*Session.Get());
+			const Trace::ILoadTimeProfilerProvider& LoadTimeProfilerProvider = *Trace::ReadLoadTimeProfilerProvider(*Session.Get());
+
+			const double SelectionStartTime = TimingView ? TimingView->GetSelectionStartTime() : 0.0;
+			const double SelectionEndTime = TimingView ? TimingView->GetSelectionEndTime() : 0.0;
+
 			Trace::ITable<Trace::FExportsTableRow>* ExportDetailsTable = LoadTimeProfilerProvider.CreateExportDetailsTable(SelectionStartTime, SelectionEndTime);
-			ExportDetailsTreeView->GetTable()->UpdateSourceTable(MakeShareable(ExportDetailsTable));
-			ExportDetailsTreeView->RebuildTree(true);
+			ExportDetailsTreeView->UpdateSourceTable(MakeShareable(ExportDetailsTable));
+		}
+		else
+		{
+			ExportDetailsTreeView->UpdateSourceTable(nullptr);
 		}
 	}
 }
@@ -198,17 +265,6 @@ TSharedRef<SDockTab> SLoadingProfilerWindow::SpawnTab_EventAggregationTreeView(c
 
 	TSharedRef<Insights::FTable> Table = MakeShared<Insights::FTable>();
 
-	TSharedPtr<const Trace::IAnalysisSession> Session = FInsightsManager::Get()->GetSession();
-	if (Session.IsValid() && Trace::ReadLoadTimeProfilerProvider(*Session.Get()))
-	{
-		Trace::FAnalysisSessionReadScope SessionReadScope(*Session.Get());
-		const Trace::ILoadTimeProfilerProvider& LoadTimeProfilerProvider = *Trace::ReadLoadTimeProfilerProvider(*Session.Get());
-		const double SelectionStartTime = TimingView ? TimingView->GetSelectionStartTime() : 0.0;
-		const double SelectionEndTime = TimingView ? TimingView->GetSelectionEndTime() : 0.0;
-		TSharedPtr<Trace::ITable<Trace::FLoadTimeProfilerAggregatedStats>> SourceTable = MakeShareable(LoadTimeProfilerProvider.CreateEventAggregation(SelectionStartTime, SelectionEndTime));
-		Table->Init(SourceTable);
-	}
-
 	const TSharedRef<SDockTab> DockTab = SNew(SDockTab)
 		.ShouldAutosize(false)
 		.TabRole(ETabRole::PanelTab)
@@ -216,7 +272,7 @@ TSharedRef<SDockTab> SLoadingProfilerWindow::SpawnTab_EventAggregationTreeView(c
 			SAssignNew(EventAggregationTreeView, Insights::STableTreeView, Table)
 		];
 
-	EventAggregationTreeView->Reset();
+	UpdateEventAggregationTreeView();
 
 	DockTab->SetOnTabClosed(SDockTab::FOnTabClosedCallback::CreateRaw(this, &SLoadingProfilerWindow::OnEventAggregationTreeViewTabClosed));
 
@@ -239,17 +295,6 @@ TSharedRef<SDockTab> SLoadingProfilerWindow::SpawnTab_ObjectTypeAggregationTreeV
 
 	TSharedRef<Insights::FTable> Table = MakeShared<Insights::FTable>();
 
-	TSharedPtr<const Trace::IAnalysisSession> Session = FInsightsManager::Get()->GetSession();
-	if (Session.IsValid() && Trace::ReadLoadTimeProfilerProvider(*Session.Get()))
-	{
-		Trace::FAnalysisSessionReadScope SessionReadScope(*Session.Get());
-		const Trace::ILoadTimeProfilerProvider& LoadTimeProfilerProvider = *Trace::ReadLoadTimeProfilerProvider(*Session.Get());
-		const double SelectionStartTime = TimingView ? TimingView->GetSelectionStartTime() : 0.0;
-		const double SelectionEndTime = TimingView ? TimingView->GetSelectionEndTime() : 0.0;
-		TSharedPtr<Trace::ITable<Trace::FLoadTimeProfilerAggregatedStats>> SourceTable = MakeShareable(LoadTimeProfilerProvider.CreateObjectTypeAggregation(SelectionStartTime, SelectionEndTime));
-		Table->Init(SourceTable);
-	}
-
 	const TSharedRef<SDockTab> DockTab = SNew(SDockTab)
 		.ShouldAutosize(false)
 		.TabRole(ETabRole::PanelTab)
@@ -257,7 +302,7 @@ TSharedRef<SDockTab> SLoadingProfilerWindow::SpawnTab_ObjectTypeAggregationTreeV
 			SAssignNew(ObjectTypeAggregationTreeView, Insights::STableTreeView, Table)
 		];
 
-	ObjectTypeAggregationTreeView->Reset();
+	UpdateObjectTypeAggregationTreeView();
 
 	DockTab->SetOnTabClosed(SDockTab::FOnTabClosedCallback::CreateRaw(this, &SLoadingProfilerWindow::OnObjectTypeAggregationTreeViewTabClosed));
 
@@ -280,17 +325,6 @@ TSharedRef<SDockTab> SLoadingProfilerWindow::SpawnTab_PackageDetailsTreeView(con
 
 	TSharedRef<Insights::FTable> Table = MakeShared<Insights::FTable>();
 
-	TSharedPtr<const Trace::IAnalysisSession> Session = FInsightsManager::Get()->GetSession();
-	if (Session.IsValid() && Trace::ReadLoadTimeProfilerProvider(*Session.Get()))
-	{
-		Trace::FAnalysisSessionReadScope SessionReadScope(*Session.Get());
-		const Trace::ILoadTimeProfilerProvider& LoadTimeProfilerProvider = *Trace::ReadLoadTimeProfilerProvider(*Session.Get());
-		const double SelectionStartTime = TimingView ? TimingView->GetSelectionStartTime() : 0.0;
-		const double SelectionEndTime = TimingView ? TimingView->GetSelectionEndTime() : 0.0;
-		TSharedPtr<Trace::ITable<Trace::FPackagesTableRow>> SourceTable = MakeShareable(LoadTimeProfilerProvider.CreatePackageDetailsTable(SelectionStartTime, SelectionEndTime));
-		Table->Init(SourceTable);
-	}
-
 	const TSharedRef<SDockTab> DockTab = SNew(SDockTab)
 		.ShouldAutosize(false)
 		.TabRole(ETabRole::PanelTab)
@@ -298,7 +332,7 @@ TSharedRef<SDockTab> SLoadingProfilerWindow::SpawnTab_PackageDetailsTreeView(con
 			SAssignNew(PackageDetailsTreeView, Insights::STableTreeView, Table)
 		];
 
-	PackageDetailsTreeView->Reset();
+	UpdatePackageDetailsTreeView();
 
 	DockTab->SetOnTabClosed(SDockTab::FOnTabClosedCallback::CreateRaw(this, &SLoadingProfilerWindow::OnPackageDetailsTreeViewTabClosed));
 
@@ -321,17 +355,6 @@ TSharedRef<SDockTab> SLoadingProfilerWindow::SpawnTab_ExportDetailsTreeView(cons
 
 	TSharedRef<Insights::FTable> Table = MakeShared<Insights::FTable>();
 
-	TSharedPtr<const Trace::IAnalysisSession> Session = FInsightsManager::Get()->GetSession();
-	if (Session.IsValid() && Trace::ReadLoadTimeProfilerProvider(*Session.Get()))
-	{
-		Trace::FAnalysisSessionReadScope SessionReadScope(*Session.Get());
-		const Trace::ILoadTimeProfilerProvider& LoadTimeProfilerProvider = *Trace::ReadLoadTimeProfilerProvider(*Session.Get());
-		const double SelectionStartTime = TimingView ? TimingView->GetSelectionStartTime() : 0.0;
-		const double SelectionEndTime = TimingView ? TimingView->GetSelectionEndTime() : 0.0;
-		TSharedPtr<Trace::ITable<Trace::FExportsTableRow>> SourceTable = MakeShareable(LoadTimeProfilerProvider.CreateExportDetailsTable(SelectionStartTime, SelectionEndTime));
-		Table->Init(SourceTable);
-	}
-
 	const TSharedRef<SDockTab> DockTab = SNew(SDockTab)
 		.ShouldAutosize(false)
 		.TabRole(ETabRole::PanelTab)
@@ -339,7 +362,7 @@ TSharedRef<SDockTab> SLoadingProfilerWindow::SpawnTab_ExportDetailsTreeView(cons
 			SAssignNew(ExportDetailsTreeView, Insights::STableTreeView, Table)
 		];
 
-	ExportDetailsTreeView->Reset();
+	UpdateExportDetailsTreeView();
 
 	DockTab->SetOnTabClosed(SDockTab::FOnTabClosedCallback::CreateRaw(this, &SLoadingProfilerWindow::OnExportDetailsTreeViewTabClosed));
 
diff --git a/Engine/Source/Developer/TraceInsights/Private/Insights/LoadingProfiler/Widgets/SLoadingProfilerWindow.h b/Engine/Source/Developer/TraceInsights/Private/Insights/LoadingProfiler/Widgets/SLoadingProfilerWindow.h
--- a/Engine/Source/Developer/TraceInsights/Private/Insights/LoadingProfiler/Widgets/SLoadingProfilerWindow.h
+++ b/Engine/Source/Developer/TraceInsights/Private/Insights/LoadingProfiler/Widgets/SLoadingProfilerWindow.h
@@ -57,6 +57,10 @@ public:
 
 	void Reset();
 	void UpdateTableTreeViews();
+	void UpdateEventAggregationTreeView();
+	void UpdateObjectTypeAggregationTreeView();
+	void UpdatePackageDetailsTreeView();
+	void UpdateExportDetailsTreeView();
 
 	/** Constructs this widget. */
 	void Construct(const FArguments& InArgs, const TSharedRef<SDockTab>& ConstructUnderMajorTab, const TSharedPtr<SWindow>& ConstructUnderWindow);
diff --git a/Engine/Source/Developer/TraceInsights/Private/Insights/Table/ViewModels/Table.cpp b/Engine/Source/Developer/TraceInsights/Private/Insights/Table/ViewModels/Table.cpp
--- a/Engine/Source/Developer/TraceInsights/Private/Insights/Table/ViewModels/Table.cpp
+++ b/Engine/Source/Developer/TraceInsights/Private/Insights/Table/ViewModels/Table.cpp
@@ -112,21 +112,6 @@ int32 FTable::GetColumnPositionIndex(const FName& ColumnId) const
 
 ////////////////////////////////////////////////////////////////////////////////////////////////////
 
-void FTable::Init(TSharedPtr<Trace::IUntypedTable> InSourceTable)
-{
-	Reset();
-
-	SourceTable = InSourceTable;
-
-	if (SourceTable)
-	{
-		TableReader = MakeShareable(SourceTable->CreateReader());
-		CreateColumnsFromTableLayout();
-	}
-}
-
-////////////////////////////////////////////////////////////////////////////////////////////////////
-
 bool AreTableLayoutsEqual(const Trace::ITableLayout& TableLayoutA, const Trace::ITableLayout& TableLayoutB)
 {
 	if (TableLayoutA.GetColumnCount() != TableLayoutB.GetColumnCount())
@@ -152,12 +137,33 @@ bool AreTableLayoutsEqual(const Trace::ITableLayout& TableLayoutA, const Trace::
 
 ////////////////////////////////////////////////////////////////////////////////////////////////////
 
-void FTable::UpdateSourceTable(TSharedPtr<Trace::IUntypedTable> InSourceTable)
+bool FTable::UpdateSourceTable(TSharedPtr<Trace::IUntypedTable> InSourceTable)
 {
-	check(InSourceTable.IsValid() && SourceTable.IsValid());
-	check(AreTableLayoutsEqual(InSourceTable->GetLayout(), SourceTable->GetLayout()));
-	SourceTable = InSourceTable;
-	TableReader = MakeShareable(SourceTable->CreateReader());
+	bool bTableLayoutChanged;
+
+	if (InSourceTable.IsValid())
+	{
+		bTableLayoutChanged = !SourceTable.IsValid() || !AreTableLayoutsEqual(InSourceTable->GetLayout(), SourceTable->GetLayout());
+		SourceTable = InSourceTable;
+		TableReader = MakeShareable(SourceTable->CreateReader());
+	}
+	else
+	{
+		bTableLayoutChanged = SourceTable.IsValid();
+		SourceTable.Reset();
+		TableReader.Reset();
+	}
+
+	if (bTableLayoutChanged)
+	{
+		Columns.Reset();
+		if (SourceTable.IsValid())
+		{
+			CreateColumns(SourceTable->GetLayout());
+		}
+	}
+
+	return bTableLayoutChanged;
 }
 
 ////////////////////////////////////////////////////////////////////////////////////////////////////
@@ -220,12 +226,9 @@ void FTable::CreateHierarchyColumn(int32 ColumnIndex, const TCHAR* ColumnName)
 
 ////////////////////////////////////////////////////////////////////////////////////////////////////
 
-void FTable::CreateColumnsFromTableLayout()
+void FTable::CreateColumns(const Trace::ITableLayout& TableLayout)
 {
-	ensure(TableReader.IsValid());
 	ensure(Columns.Num() == 0);
-
-	const Trace::ITableLayout& TableLayout = SourceTable->GetLayout();
 	const int32 ColumnCount = TableLayout.GetColumnCount();
 
 	//////////////////////////////////////////////////
diff --git a/Engine/Source/Developer/TraceInsights/Private/Insights/Table/ViewModels/Table.h b/Engine/Source/Developer/TraceInsights/Private/Insights/Table/ViewModels/Table.h
--- a/Engine/Source/Developer/TraceInsights/Private/Insights/Table/ViewModels/Table.h
+++ b/Engine/Source/Developer/TraceInsights/Private/Insights/Table/ViewModels/Table.h
@@ -6,6 +6,7 @@
 
 namespace Trace
 {
+	class ITableLayout;
 	class IUntypedTable;
 	class IUntypedTableReader;
 }
@@ -79,14 +80,13 @@ public:
 	TSharedPtr<Trace::IUntypedTable> GetSourceTable() const { return SourceTable; }
 	TSharedPtr<Trace::IUntypedTableReader> GetTableReader() const { return TableReader; }
 
-	/* Init this table to use the IUntypedTable source table. It will create new Columns array based on table layout. */
-	void Init(TSharedPtr<Trace::IUntypedTable> InSourceTable);
-	void UpdateSourceTable(TSharedPtr<Trace::IUntypedTable> InSourceTable);
+	/* Update table content. Returns true if the table layout has changed. */
+	bool UpdateSourceTable(TSharedPtr<Trace::IUntypedTable> InSourceTable);
 
 private:
 	void AddColumn(TSharedRef<FTableColumn> Column);
 	void CreateHierarchyColumn(int32 ColumnIndex, const TCHAR* ColumnName);
-	void CreateColumnsFromTableLayout();
+	void CreateColumns(const Trace::ITableLayout& TableLayout);
 
 private:
 	FName Name;
diff --git a/Engine/Source/Developer/TraceInsights/Private/Insights/Table/Widgets/STableTreeView.cpp b/Engine/Source/Developer/TraceInsights/Private/Insights/Table/Widgets/STableTreeView.cpp
--- a/Engine/Source/Developer/TraceInsights/Private/Insights/Table/Widgets/STableTreeView.cpp
+++ b/Engine/Source/Developer/TraceInsights/Private/Insights/Table/Widgets/STableTreeView.cpp
@@ -202,14 +202,14 @@ void STableTreeView::Construct(const FArguments& InArgs, TSharedPtr<FTable> InTa
 		]
 	];
 
-	InitializeAndShowHeaderColumns();
-	//BindCommands();
-
 	// Create the search filters: text based, type based etc.
 	TextFilter = MakeShared<FTableTreeNodeTextFilter>(FTableTreeNodeTextFilter::FItemToStringArray::CreateSP(this, &STableTreeView::HandleItemToStringArray));
 	Filters = MakeShared<FTableTreeNodeFilterCollection>();
 	Filters->Add(TextFilter);
 
+	//BindCommands();
+
+	InitializeAndShowHeaderColumns();
 	CreateGroupings();
 	CreateSortings();
 
@@ -449,6 +449,20 @@ void STableTreeView::InitializeAndShowHeaderColumns()
 
 ////////////////////////////////////////////////////////////////////////////////////////////////////
 
+void STableTreeView::RebuildColumns()
+{
+	TreeViewHeaderRow->ClearColumns();
+	InitializeAndShowHeaderColumns();
+
+	PreChangeGroupings();
+	CreateGroupings();
+	PostChangeGroupings();
+
+	CreateSortings();
+}
+
+////////////////////////////////////////////////////////////////////////////////////////////////////
+
 FText STableTreeView::GetColumnHeaderText(const FName ColumnId) const
 {
 	const FTableColumn& Column = *Table->FindColumnChecked(ColumnId);
@@ -984,7 +998,8 @@ void STableTreeView::UpdateDoubleSumAggregationRec(FTableColumn& Column, FTableT
 
 void STableTreeView::CreateGroupings()
 {
-	AvailableGroupings.Reset(3);
+	AvailableGroupings.Reset();
+	CurrentGroupings.Reset();
 
 	AvailableGroupings.Add(MakeShared<FTreeNodeGroupingFlat>());
 	//AvailableGroupings.Add(MakeShared<FTreeNodeGroupingByNameFirstLetter>());
@@ -1068,22 +1083,25 @@ void STableTreeView::PostChangeGroupings()
 
 	//////////////////////////////////////////////////
 
-	// Set with for the Hierarchy column.
-	SHeaderRow::FColumn& HierarchyColumn = const_cast<SHeaderRow::FColumn&>(TreeViewHeaderRow->GetColumns()[0]);
-	HierarchyColumn.SetWidth(HierarchyColumnWidth);
-
-	// Set name for the Hierarchy column.
-	//FTableColumn& HierarchyTableColumn = *Table->FindColumnChecked(HierarchyColumn.ColumnId);
-	//if (!GroupingStr.IsEmpty())
-	//{
-	//	const FText HierarchyColumnName = FText::Format(LOCTEXT("HierarchyShortNameFmt", "Hierarchy ({0})"), FText::FromString(GroupingStr));
-	//	HierarchyTableColumn.SetShortName(HierarchyColumnName);
-	//}
-	//else
-	//{
-	//	const FText HierarchyColumnName(LOCTEXT("HierarchyShortName", "Hierarchy"));
-	//	HierarchyTableColumn.SetShortName(HierarchyColumnName);
-	//}
+	if (TreeViewHeaderRow->GetColumns().Num() > 0)
+	{
+		// Set width for the Hierarchy column.
+		SHeaderRow::FColumn& HierarchyColumn = const_cast<SHeaderRow::FColumn&>(TreeViewHeaderRow->GetColumns()[0]);
+		HierarchyColumn.SetWidth(HierarchyColumnWidth);
+
+		// Set name for the Hierarchy column.
+		//FTableColumn& HierarchyTableColumn = *Table->FindColumnChecked(HierarchyColumn.ColumnId);
+		//if (!GroupingStr.IsEmpty())
+		//{
+		//	const FText HierarchyColumnName = FText::Format(LOCTEXT("HierarchyShortNameFmt", "Hierarchy ({0})"), FText::FromString(GroupingStr));
+		//	HierarchyTableColumn.SetShortName(HierarchyColumnName);
+		//}
+		//else
+		//{
+		//	const FText HierarchyColumnName(LOCTEXT("HierarchyShortName", "Hierarchy"));
+		//	HierarchyTableColumn.SetShortName(HierarchyColumnName);
+		//}
+	}
 
 	//////////////////////////////////////////////////
 
@@ -1822,6 +1840,17 @@ void STableTreeView::Reset()
 
 ////////////////////////////////////////////////////////////////////////////////////////////////////
 
+void STableTreeView::UpdateSourceTable(TSharedPtr<Trace::IUntypedTable> SourceTable)
+{
+	if (Table->UpdateSourceTable(SourceTable))
+	{
+		RebuildColumns();
+	}
+	RebuildTree(true);
+}
+
+////////////////////////////////////////////////////////////////////////////////////////////////////
+
 void STableTreeView::RebuildTree(bool bResync)
 {
 	bool bListHasChanged = false;
diff --git a/Engine/Source/Developer/TraceInsights/Private/Insights/Table/Widgets/STableTreeView.h b/Engine/Source/Developer/TraceInsights/Private/Insights/Table/Widgets/STableTreeView.h
--- a/Engine/Source/Developer/TraceInsights/Private/Insights/Table/Widgets/STableTreeView.h
+++ b/Engine/Source/Developer/TraceInsights/Private/Insights/Table/Widgets/STableTreeView.h
@@ -66,8 +66,12 @@ public:
 
 	TSharedPtr<Insights::FTable> GetTable() const { return Table; }
 
+	void UpdateSourceTable(TSharedPtr<Trace::IUntypedTable> SourceTable);
+
 	virtual void Reset();
 
+	void RebuildColumns();
+
 	/**
 	 * Rebuilds the tree (if necessary).
 	 * @param bResync - If true, it forces a resync even if the list did not changed since last sync.
