[Insights]
  - Changed UnrealInsights to start in "Browser mode". When opening a trace file it will spawn a new UnrealInsights process.
  - Added "-Trace=.." command line parameter to force UnrealInsights to start analysis of the specified trace file (Start Page will not be displayed).
  - Added "-SingleProcess" command line parameter to force UnrealInsights to start analysis in same process (same as old behaviour).
  - Added  new "Session Info" tab with few metadata about session (currently only info about URI, analysis status and session duration are available).
  - Start Page: Added a fake splash screen to appear each time session analysis is started in separate process.
  - Start Page: Added auto-start analysis for live sessions (with possibility to set filter by platform and/or app name)
 -  Start Page: Added vertical scrollbar (visible when window height is reduced). Fixes UE-86471.
  - Increased version to 0.12.

diff --git a/Engine/Source/Developer/TraceInsights/Private/Insights/InsightsManager.cpp b/Engine/Source/Developer/TraceInsights/Private/Insights/InsightsManager.cpp
--- a/Engine/Source/Developer/TraceInsights/Private/Insights/InsightsManager.cpp
+++ b/Engine/Source/Developer/TraceInsights/Private/Insights/InsightsManager.cpp
@@ -19,6 +19,7 @@
 #define LOCTEXT_NAMESPACE "InsightsManager"
 
 const FName FInsightsManagerTabs::StartPageTabId(TEXT("StartPage"));
+const FName FInsightsManagerTabs::SessionInfoTabId(TEXT("SessionInfo"));
 const FName FInsightsManagerTabs::TimingProfilerTabId(TEXT("TimingProfiler"));
 const FName FInsightsManagerTabs::LoadingProfilerTabId(TEXT("LoadingProfiler"));
 const FName FInsightsManagerTabs::NetworkingProfilerTabId(TEXT("NetworkingProfiler"));
@@ -38,6 +39,11 @@ FInsightsManager::FInsightsManager(TSharedRef<Trace::IAnalysisService> InTraceAn
 	, Settings()
 	, bIsDebugInfoEnabled(false)
 	, bIsNetworkingProfilerAvailable(false)
+#if WITH_EDITOR
+	, bShouldOpenAnalysisInSeparateProcess(false)
+#else
+	, bShouldOpenAnalysisInSeparateProcess(true)
+#endif
 {
 }
 
@@ -195,6 +201,12 @@ void FInsightsManager::OnSessionChanged()
 
 void FInsightsManager::SpawnAndActivateTabs()
 {
+	// Open Session Info tab.
+	if (FGlobalTabmanager::Get()->HasTabSpawner(FInsightsManagerTabs::SessionInfoTabId))
+	{
+		FGlobalTabmanager::Get()->InvokeTab(FInsightsManagerTabs::SessionInfoTabId);
+	}
+
 	// Open Timing Insights tab.
 	if (FGlobalTabmanager::Get()->HasTabSpawner(FInsightsManagerTabs::TimingProfilerTabId))
 	{
diff --git a/Engine/Source/Developer/TraceInsights/Private/Insights/InsightsManager.h b/Engine/Source/Developer/TraceInsights/Private/Insights/InsightsManager.h
--- a/Engine/Source/Developer/TraceInsights/Private/Insights/InsightsManager.h
+++ b/Engine/Source/Developer/TraceInsights/Private/Insights/InsightsManager.h
@@ -22,6 +22,7 @@ class SStartPageWindow;
 struct FInsightsManagerTabs
 {
 	static const FName StartPageTabId;
+	static const FName SessionInfoTabId;
 	static const FName TimingProfilerTabId;
 	static const FName LoadingProfilerTabId;
 	static const FName NetworkingProfilerTabId;
@@ -123,6 +124,9 @@ public:
 
 	////////////////////////////////////////////////////////////////////////////////////////////////////
 
+	bool ShouldOpenAnalysisInSeparateProcess() const { return bShouldOpenAnalysisInSeparateProcess; }
+	void SetOpenAnalysisInSeparateProcess(bool bOnOff) { bShouldOpenAnalysisInSeparateProcess = bOnOff; }
+
 	bool IsAnyLiveSessionAvailable(Trace::FSessionHandle& OutLastLiveSessionHandle) const;
 	bool IsAnySessionAvailable(Trace::FSessionHandle& OutLastSessionHandle) const;
 
@@ -215,4 +219,6 @@ private:
 	static TSharedPtr<FInsightsManager> Instance;
 
 	bool bIsNetworkingProfilerAvailable;
+
+	bool bShouldOpenAnalysisInSeparateProcess;
 };
diff --git a/Engine/Source/Developer/TraceInsights/Private/Insights/InsightsStyle.cpp b/Engine/Source/Developer/TraceInsights/Private/Insights/InsightsStyle.cpp
--- a/Engine/Source/Developer/TraceInsights/Private/Insights/InsightsStyle.cpp
+++ b/Engine/Source/Developer/TraceInsights/Private/Insights/InsightsStyle.cpp
@@ -39,6 +39,8 @@ TSharedRef<FSlateStyleSet> FInsightsStyle::Create()
 
 	FSlateStyleSet& Style = StyleRef.Get();
 
+	//Style.Set("AppIcon", new IMAGE_BRUSH("Icons/Insights/AppIcon_24x", FVector2D(24.0f, 24.0f)));
+
 	//////////////////////////////////////////////////
 
 	Style.Set("WhiteBrush", new FSlateColorBrush(FLinearColor::White));
@@ -55,6 +57,9 @@ TSharedRef<FSlateStyleSet> FInsightsStyle::Create()
 	Style.Set("StartPage.Icon.Large", new IMAGE_BRUSH("Icons/icon_tab_Tools_16x", FVector2D(32.0f, 32.0f)));
 	Style.Set("StartPage.Icon.Small", new IMAGE_BRUSH("Icons/icon_tab_Tools_16x", FVector2D(16.0f, 16.0f)));
 
+	Style.Set("SessionInfo.Icon.Large", new IMAGE_BRUSH("Icons/icon_tab_Tools_16x", FVector2D(32.0f, 32.0f)));
+	Style.Set("SessionInfo.Icon.Small", new IMAGE_BRUSH("Icons/icon_tab_Tools_16x", FVector2D(16.0f, 16.0f)));
+
 	Style.Set("TimingProfiler.Icon.Large", new IMAGE_BRUSH("Icons/Profiler/profiler_stats_40x", FVector2D(32.0f, 32.0f)));
 	Style.Set("TimingProfiler.Icon.Small", new IMAGE_BRUSH("Icons/Profiler/profiler_stats_40x", FVector2D(16.0f, 16.0f)));
 
diff --git a/Engine/Source/Developer/TraceInsights/Private/Insights/TraceInsightsModule.cpp b/Engine/Source/Developer/TraceInsights/Private/Insights/TraceInsightsModule.cpp
--- a/Engine/Source/Developer/TraceInsights/Private/Insights/TraceInsightsModule.cpp
+++ b/Engine/Source/Developer/TraceInsights/Private/Insights/TraceInsightsModule.cpp
@@ -12,6 +12,8 @@
 #include "WorkspaceMenuStructure.h"
 #include "WorkspaceMenuStructureModule.h"
 
+#include "Widgets/Text/STextBlock.h"
+
 // Insights
 #include "Insights/InsightsManager.h"
 #include "Insights/InsightsStyle.h"
@@ -21,12 +23,13 @@
 #include "Insights/NetworkingProfiler/NetworkingProfilerManager.h"
 #include "Insights/NetworkingProfiler/Widgets/SNetworkingProfilerWindow.h"
 #include "Insights/TimingProfilerManager.h"
+#include "Insights/Widgets/SSessionInfoWindow.h"
 #include "Insights/Widgets/SStartPageWindow.h"
 #include "Insights/Widgets/STimingProfilerWindow.h"
 
 ////////////////////////////////////////////////////////////////////////////////////////////////////
 /**
- * Implements the Trace Insights module (TimingProfiler, LoadingProfiler, etc.).
+ * Implements the Trace Insights module.
  */
 class FTraceInsightsModule : public IUnrealInsightsModule
 {
@@ -39,10 +42,16 @@ public:
 		return false;
 	}
 
-	virtual void OnNewLayout(TSharedRef<FTabManager::FLayout> NewLayout) override;
-	virtual void OnLayoutRestored(TSharedPtr<FTabManager> TabManager) override;
+	virtual void CreateSessionBrowser(bool bAllowDebugTools, bool bSingleProcess) override;
+	virtual void CreateSessionViewer(bool bAllowDebugTools) override;
+	virtual void StartAnalysisForTraceFile(const TCHAR* InTraceFile) override;
+	virtual void StartAnalysisForSession(const TCHAR* InSessionId) override;
+	virtual void ShutdownUserInterface() override;
 
 protected:
+	void AddAreaForSessionViewer(TSharedRef<FTabManager::FLayout> Layout);
+	void AddAreaForWidgetReflector(TSharedRef<FTabManager::FLayout> Layout, bool bAllowDebugTools);
+
 	/** Callback called when a major tab is closed. */
 	void OnTabBeingClosed(TSharedRef<SDockTab> TabBeingClosed);
 
@@ -58,6 +67,9 @@ protected:
 	/** Start Page */
 	TSharedRef<SDockTab> SpawnStartPageTab(const FSpawnTabArgs& Args);
 
+	/** Session Info */
+	TSharedRef<SDockTab> SpawnSessionInfoTab(const FSpawnTabArgs& Args);
+
 	/** Timing Profiler */
 	TSharedRef<SDockTab> SpawnTimingProfilerTab(const FSpawnTabArgs& Args);
 
@@ -73,13 +85,13 @@ protected:
 #endif
 
 protected:
+	TSharedPtr<FTabManager::FLayout> PersistentLayout;
+	bool bBrowserMode;
+	FString UnrealInsightsLayoutIni;
+
 	TSharedPtr<Trace::IAnalysisService> TraceAnalysisService;
 	TSharedPtr<Trace::ISessionService> TraceSessionService;
 	TSharedPtr<Trace::IModuleService> TraceModuleService;
-
-#if WITH_EDITOR
-	TSharedPtr<FTabManager::FLayout> PersistentLayout;
-#endif
 };
 
 ////////////////////////////////////////////////////////////////////////////////////////////////////
@@ -104,9 +116,6 @@ void FTraceInsightsModule::StartupModule()
 	TraceSessionService = TraceServicesModule.GetSessionService();
 	TraceModuleService = TraceServicesModule.GetModuleService();
 
-	// Starts the Trace Recorder service.
-	TraceSessionService->StartRecorderServer();
-
 	FInsightsStyle::Initialize();
 
 	FInsightsManager::Initialize(TraceAnalysisService.ToSharedRef(), TraceSessionService.ToSharedRef(), TraceModuleService.ToSharedRef());
@@ -114,6 +123,8 @@ void FTraceInsightsModule::StartupModule()
 	FLoadingProfilerManager::Initialize();
 	FNetworkingProfilerManager::Initialize();
 
+	//////////////////////////////////////////////////
+
 	// Register tab spawner for the Start Page.
 	auto& StartPageTabSpawnerEntry = FGlobalTabmanager::Get()->RegisterNomadTabSpawner(FInsightsManagerTabs::StartPageTabId,
 		FOnSpawnTab::CreateRaw(this, &FTraceInsightsModule::SpawnStartPageTab))
@@ -127,6 +138,23 @@ void FTraceInsightsModule::StartupModule()
 	StartPageTabSpawnerEntry.SetGroup(WorkspaceMenu::GetMenuStructure().GetToolsCategory());
 #endif
 
+	//////////////////////////////////////////////////
+
+	// Register tab spawner for the Session Info.
+	auto& SessionInfoTabSpawnerEntry = FGlobalTabmanager::Get()->RegisterNomadTabSpawner(FInsightsManagerTabs::SessionInfoTabId,
+		FOnSpawnTab::CreateRaw(this, &FTraceInsightsModule::SpawnSessionInfoTab))
+		.SetDisplayName(NSLOCTEXT("FTraceInsightsModule", "SessionInfoTabTitle", "Session Info"))
+		.SetTooltipText(NSLOCTEXT("FTraceInsightsModule", "SessionInfoTooltipText", "Open the Session Info tab."))
+		.SetIcon(FSlateIcon(FInsightsStyle::GetStyleSetName(), "SessionInfo.Icon.Small"));
+
+#if WITH_EDITOR
+	SessionInfoTabSpawnerEntry.SetGroup(WorkspaceMenu::GetMenuStructure().GetDeveloperToolsProfilingCategory());
+#else
+	SessionInfoTabSpawnerEntry.SetGroup(WorkspaceMenu::GetMenuStructure().GetToolsCategory());
+#endif
+
+	//////////////////////////////////////////////////
+
 	// Register tab spawner for the Timing Insights.
 	auto& TimingProfilerTabSpawnerEntry = FGlobalTabmanager::Get()->RegisterNomadTabSpawner(FInsightsManagerTabs::TimingProfilerTabId,
 		FOnSpawnTab::CreateRaw(this, &FTraceInsightsModule::SpawnTimingProfilerTab))
@@ -140,6 +168,8 @@ void FTraceInsightsModule::StartupModule()
 	TimingProfilerTabSpawnerEntry.SetGroup(WorkspaceMenu::GetMenuStructure().GetToolsCategory());
 #endif
 
+	//////////////////////////////////////////////////
+
 	// Register tab spawner for the Asset Loading Insights.
 	auto& LoadingProfilerTabSpawnerEntry = FGlobalTabmanager::Get()->RegisterNomadTabSpawner(FInsightsManagerTabs::LoadingProfilerTabId,
 		FOnSpawnTab::CreateRaw(this, &FTraceInsightsModule::SpawnLoadingProfilerTab))
@@ -153,6 +183,8 @@ void FTraceInsightsModule::StartupModule()
 	LoadingProfilerTabSpawnerEntry.SetGroup(WorkspaceMenu::GetMenuStructure().GetToolsCategory());
 #endif
 
+	//////////////////////////////////////////////////
+
 	// Register tab spawner for the Networking Insights.
 	auto& NetworkingProfilerTabSpawnerEntry = FGlobalTabmanager::Get()->RegisterNomadTabSpawner(FInsightsManagerTabs::NetworkingProfilerTabId,
 		FOnSpawnTab::CreateRaw(this, &FTraceInsightsModule::SpawnNetworkingProfilerTab))
@@ -167,10 +199,17 @@ void FTraceInsightsModule::StartupModule()
 	NetworkingProfilerTabSpawnerEntry.SetGroup(WorkspaceMenu::GetMenuStructure().GetToolsCategory());
 #endif
 
+	//////////////////////////////////////////////////
+
 #if WITH_EDITOR
-	// In editor, load the layout and display it here. In the standalone tool this is done in UserInterfaceCommand.
+	if (TraceSessionService.IsValid())
+	{
+		TraceSessionService->StartRecorderServer();
+	}
+
+	// In editor, load the layout and display it here. In the standalone tool this is done through CreateSessionBrowser or CreateSessionViewer.
 	TSharedRef<FTabManager::FLayout> DefaultLayout = FTabManager::NewLayout("UnrealInsightsLayout_v1.0");
-	OnNewLayout(DefaultLayout);
+	AddAreaForSessionViewer(DefaultLayout);
 	PersistentLayout = FLayoutSaveRestore::LoadFromConfig(GEditorLayoutIni, DefaultLayout);
 	FGlobalTabmanager::Get()->RestoreFrom(PersistentLayout.ToSharedRef(), TSharedPtr<SWindow>());
 
@@ -189,12 +228,16 @@ void FTraceInsightsModule::ShutdownModule()
 #endif
 
 #if !WITH_EDITOR
-	TraceSessionService->StopRecorderServer();
+	if (TraceSessionService.IsValid())
+	{
+		TraceSessionService->StopRecorderServer();
+	}
 #endif
 
 	FGlobalTabmanager::Get()->UnregisterNomadTabSpawner(FInsightsManagerTabs::NetworkingProfilerTabId);
 	FGlobalTabmanager::Get()->UnregisterNomadTabSpawner(FInsightsManagerTabs::LoadingProfilerTabId);
 	FGlobalTabmanager::Get()->UnregisterNomadTabSpawner(FInsightsManagerTabs::TimingProfilerTabId);
+	FGlobalTabmanager::Get()->UnregisterNomadTabSpawner(FInsightsManagerTabs::SessionInfoTabId);
 	FGlobalTabmanager::Get()->UnregisterNomadTabSpawner(FInsightsManagerTabs::StartPageTabId);
 
 	if (FNetworkingProfilerManager::Get().IsValid())
@@ -238,33 +281,140 @@ void FTraceInsightsModule::HandleExit()
 
 ////////////////////////////////////////////////////////////////////////////////////////////////////
 
-void FTraceInsightsModule::OnNewLayout(TSharedRef<FTabManager::FLayout> NewLayout)
+void FTraceInsightsModule::CreateSessionBrowser(bool bAllowDebugTools, bool bSingleProcess)
 {
+	bBrowserMode = true;
+	FInsightsManager::Get()->SetOpenAnalysisInSeparateProcess(!bSingleProcess);
+
+	if (TraceSessionService.IsValid())
+	{
+		TraceSessionService->StartRecorderServer();
+	}
+
 	const float DPIScaleFactor = FPlatformApplicationMisc::GetDPIScaleFactorAtPoint(10.0f, 10.0f);
-	NewLayout->AddArea
+
+	TSharedRef<FTabManager::FLayout> DefaultLayout = FTabManager::NewLayout("TraceSessionBrowserLayout_v1.0");
+
+	const float WindowWidth  = bSingleProcess ? 1280.0f : 920.0f;
+	const float WindowHeight = bSingleProcess ? 720.0f : 664.0f;
+
+	DefaultLayout->AddArea
+	(
+		FTabManager::NewArea(WindowWidth * DPIScaleFactor, WindowHeight * DPIScaleFactor)
+		->Split
+		(
+			FTabManager::NewStack()
+			->AddTab(FInsightsManagerTabs::StartPageTabId, ETabState::OpenedTab)
+			->AddTab(FInsightsManagerTabs::SessionInfoTabId, ETabState::ClosedTab)
+			->AddTab(FInsightsManagerTabs::TimingProfilerTabId, ETabState::ClosedTab)
+			->AddTab(FInsightsManagerTabs::LoadingProfilerTabId, ETabState::ClosedTab)
+			->AddTab(FInsightsManagerTabs::NetworkingProfilerTabId, ETabState::ClosedTab)
+			->SetForegroundTab(FTabId(FInsightsManagerTabs::StartPageTabId))
+			//->SetHideTabWell(true)
+		)
+	);
+
+	AddAreaForWidgetReflector(DefaultLayout, bAllowDebugTools);
+
+	// Restore application layout.
+	UnrealInsightsLayoutIni = FPaths::GetPath(GEngineIni) + "/UnrealInsightsLayout.ini";
+	PersistentLayout = FLayoutSaveRestore::LoadFromConfig(UnrealInsightsLayoutIni, DefaultLayout);
+	FGlobalTabmanager::Get()->RestoreFrom(PersistentLayout.ToSharedRef(), TSharedPtr<SWindow>());
+}
+
+////////////////////////////////////////////////////////////////////////////////////////////////////
+
+void FTraceInsightsModule::CreateSessionViewer(bool bAllowDebugTools)
+{
+	bBrowserMode = false;
+
+	TSharedRef<FTabManager::FLayout> DefaultLayout = FTabManager::NewLayout("UnrealInsightsLayout_v1.0");
+
+	AddAreaForSessionViewer(DefaultLayout);
+	AddAreaForWidgetReflector(DefaultLayout, bAllowDebugTools);
+
+	// Restore application layout.
+	UnrealInsightsLayoutIni = FPaths::GetPath(GEngineIni) + "/UnrealInsightsLayout.ini";
+	PersistentLayout = FLayoutSaveRestore::LoadFromConfig(UnrealInsightsLayoutIni, DefaultLayout);
+	FGlobalTabmanager::Get()->RestoreFrom(PersistentLayout.ToSharedRef(), TSharedPtr<SWindow>());
+}
+
+////////////////////////////////////////////////////////////////////////////////////////////////////
+
+void FTraceInsightsModule::AddAreaForSessionViewer(TSharedRef<FTabManager::FLayout> Layout)
+{
+	const float DPIScaleFactor = FPlatformApplicationMisc::GetDPIScaleFactorAtPoint(10.0f, 10.0f);
+
+	// Create area for the main window.
+	Layout->AddArea
 	(
 		FTabManager::NewArea(1280.f * DPIScaleFactor, 720.0f * DPIScaleFactor)
 		->Split
 		(
 			FTabManager::NewStack()
+			->AddTab(FInsightsManagerTabs::StartPageTabId, ETabState::ClosedTab)
 #if WITH_EDITOR
 			// In editor, we default to all tabs closed.
-			->AddTab(FInsightsManagerTabs::StartPageTabId, ETabState::ClosedTab)
+			->AddTab(FInsightsManagerTabs::SessionInfoTabId, ETabState::ClosedTab)
+			->AddTab(FInsightsManagerTabs::TimingProfilerTabId, ETabState::ClosedTab)
 #else
-			->AddTab(FInsightsManagerTabs::StartPageTabId, ETabState::OpenedTab)
+			->AddTab(FInsightsManagerTabs::SessionInfoTabId, ETabState::OpenedTab)
+			->AddTab(FInsightsManagerTabs::TimingProfilerTabId, ETabState::OpenedTab)
 #endif
-			->AddTab(FInsightsManagerTabs::TimingProfilerTabId, ETabState::ClosedTab)
 			->AddTab(FInsightsManagerTabs::LoadingProfilerTabId, ETabState::ClosedTab)
 			->AddTab(FInsightsManagerTabs::NetworkingProfilerTabId, ETabState::ClosedTab)
-			->SetForegroundTab(FTabId(FInsightsManagerTabs::StartPageTabId))
+			->SetForegroundTab(FTabId(FInsightsManagerTabs::TimingProfilerTabId))
+		)
+	);
+}
+
+////////////////////////////////////////////////////////////////////////////////////////////////////
+
+void FTraceInsightsModule::AddAreaForWidgetReflector(TSharedRef<FTabManager::FLayout> Layout, bool bAllowDebugTools)
+{
+	const float DPIScaleFactor = FPlatformApplicationMisc::GetDPIScaleFactorAtPoint(10.0f, 10.0f);
+
+	// Create area and tab for Slate's WidgetReflector.
+	Layout->AddArea
+	(
+		FTabManager::NewArea(600.0f * DPIScaleFactor, 600.0f * DPIScaleFactor)
+		->SetWindow(FVector2D(10.0f * DPIScaleFactor, 10.0f * DPIScaleFactor), false)
+		->Split
+		(
+			FTabManager::NewStack()->AddTab("WidgetReflector", bAllowDebugTools ? ETabState::OpenedTab : ETabState::ClosedTab)
 		)
 	);
 }
 
 ////////////////////////////////////////////////////////////////////////////////////////////////////
 
-void FTraceInsightsModule::OnLayoutRestored(TSharedPtr<FTabManager> TabManager)
+void FTraceInsightsModule::ShutdownUserInterface()
+{
+	check(PersistentLayout.IsValid());
+
+	// Save application layout.
+	FLayoutSaveRestore::SaveToConfig(UnrealInsightsLayoutIni, PersistentLayout.ToSharedRef());
+	GConfig->Flush(false, UnrealInsightsLayoutIni);
+}
+
+////////////////////////////////////////////////////////////////////////////////////////////////////
+
+void FTraceInsightsModule::StartAnalysisForTraceFile(const TCHAR* InTraceFile)
+{
+	if (InTraceFile != nullptr)
+	{
+		FInsightsManager::Get()->LoadTraceFile(FString(InTraceFile));
+	}
+}
+
+////////////////////////////////////////////////////////////////////////////////////////////////////
+
+void FTraceInsightsModule::StartAnalysisForSession(const TCHAR* InSessionId)
 {
+	if (InSessionId != nullptr)
+	{
+		//TODO: FInsightsManager::Get()->LoadSession(FString(InSessionId));
+	}
 }
 
 ////////////////////////////////////////////////////////////////////////////////////////////////////
@@ -273,6 +423,8 @@ TSharedRef<SDockTab> FTraceInsightsModule::SpawnStartPageTab(const FSpawnTabArgs
 {
 	const TSharedRef<SDockTab> DockTab = SNew(SDockTab)
 		.TabRole(ETabRole::NomadTab);
+		//.OnCanCloseTab_Lambda([]() { return false; })
+		//.ContentPadding(FMargin(2.0f, 20.0f, 2.0f, 2.0f));
 
 	// Create the Start Page widget.
 	TSharedRef<SStartPageWindow> Window = SNew(SStartPageWindow);
@@ -283,6 +435,20 @@ TSharedRef<SDockTab> FTraceInsightsModule::SpawnStartPageTab(const FSpawnTabArgs
 
 ////////////////////////////////////////////////////////////////////////////////////////////////////
 
+TSharedRef<SDockTab> FTraceInsightsModule::SpawnSessionInfoTab(const FSpawnTabArgs& Args)
+{
+	const TSharedRef<SDockTab> DockTab = SNew(SDockTab)
+		.TabRole(ETabRole::NomadTab);
+
+	// Create the Session Info widget.
+	TSharedRef<SSessionInfoWindow> Window = SNew(SSessionInfoWindow);
+	DockTab->SetContent(Window);
+
+	return DockTab;
+}
+
+////////////////////////////////////////////////////////////////////////////////////////////////////
+
 TSharedRef<SDockTab> FTraceInsightsModule::SpawnTimingProfilerTab(const FSpawnTabArgs& Args)
 {
 	const TSharedRef<SDockTab> DockTab = SNew(SDockTab)
diff --git a/Engine/Source/Developer/TraceInsights/Private/Insights/Widgets/SSessionInfoWindow.cpp b/Engine/Source/Developer/TraceInsights/Private/Insights/Widgets/SSessionInfoWindow.cpp
new file mode 100644
--- /dev/null
+++ b/Engine/Source/Developer/TraceInsights/Private/Insights/Widgets/SSessionInfoWindow.cpp
@@ -0,0 +1,322 @@
+// Copyright Epic Games, Inc. All Rights Reserved.
+
+#include "SSessionInfoWindow.h"
+
+#include "SlateOptMacros.h"
+#include "TraceServices/ModuleService.h"
+#include "Misc/Paths.h"
+#include "Widgets/Input/SEditableTextBox.h"
+#include "Widgets/Layout/SScrollBox.h"
+#include "Widgets/SOverlay.h"
+#include "Widgets/Text/STextBlock.h"
+
+#if WITH_EDITOR
+	#include "EngineAnalytics.h"
+	#include "Runtime/Analytics/Analytics/Public/AnalyticsEventAttribute.h"
+	#include "Runtime/Analytics/Analytics/Public/Interfaces/IAnalyticsProvider.h"
+#endif // WITH_EDITOR
+
+// Insights
+#include "Insights/Common/TimeUtils.h"
+#include "Insights/InsightsManager.h"
+#include "Insights/Version.h"
+
+////////////////////////////////////////////////////////////////////////////////////////////////////
+
+#define LOCTEXT_NAMESPACE "SSessionInfoWindow"
+
+////////////////////////////////////////////////////////////////////////////////////////////////////
+
+SSessionInfoWindow::SSessionInfoWindow()
+	: DurationActive(0.0f)
+{
+}
+
+////////////////////////////////////////////////////////////////////////////////////////////////////
+
+SSessionInfoWindow::~SSessionInfoWindow()
+{
+#if WITH_EDITOR
+	if (DurationActive > 0.0f && FEngineAnalytics::IsAvailable())
+	{
+		FEngineAnalytics::GetProvider().RecordEvent(TEXT("Editor.Usage.Insights.SessionInfo"), FAnalyticsEventAttribute(TEXT("Duration"), DurationActive));
+	}
+#endif // WITH_EDITOR
+}
+
+////////////////////////////////////////////////////////////////////////////////////////////////////
+
+BEGIN_SLATE_FUNCTION_BUILD_OPTIMIZATION
+
+void SSessionInfoWindow::Construct(const FArguments& InArgs)
+{
+	TSharedPtr<SVerticalBox> VerticalBox;
+
+	TSharedPtr<SScrollBar> VScrollbar;
+	SAssignNew(VScrollbar, SScrollBar)
+	.Orientation(Orient_Vertical);
+
+	ChildSlot
+		[
+			SNew(SOverlay)
+
+			// Version
+			+ SOverlay::Slot()
+			.HAlign(HAlign_Right)
+			.VAlign(VAlign_Top)
+			.Padding(0.0f, -16.0f, 0.0f, 0.0f)
+			[
+				SNew(STextBlock)
+				.Clipping(EWidgetClipping::ClipToBoundsWithoutIntersecting)
+				.Text(LOCTEXT("UnrealInsightsVersion", UNREAL_INSIGHTS_VERSION_STRING_EX))
+				.ColorAndOpacity(FLinearColor(0.15f, 0.15f, 0.15f, 1.0f))
+			]
+
+			// Overlay slot for the main window area
+			+ SOverlay::Slot()
+			.HAlign(HAlign_Fill)
+			.VAlign(VAlign_Top)
+			[
+				SNew(SScrollBox)
+				.Orientation(Orient_Vertical)
+				.ExternalScrollbar(VScrollbar)
+
+				+ SScrollBox::Slot()
+				[
+					SAssignNew(VerticalBox, SVerticalBox)
+				]
+			]
+
+			// Overlay slot for the vertical scrollbar
+			+ SOverlay::Slot()
+			.HAlign(HAlign_Right)
+			.VAlign(VAlign_Fill)
+			[
+				SNew(SBox)
+				.WidthOverride(FOptionalSize(13.0f))
+				[
+					VScrollbar.ToSharedRef()
+				]
+			]
+		];
+
+	AddInfoLine(VerticalBox, LOCTEXT("SessionName_HeaderText",	"Session Name:"),		TAttribute<FText>(this, &SSessionInfoWindow::GetSessionNameText));
+	AddInfoLine(VerticalBox, LOCTEXT("Uri_HeaderText",			"URI:"),				TAttribute<FText>(this, &SSessionInfoWindow::GetUriText));
+	AddInfoLine(VerticalBox, LOCTEXT("Platform_HeaderText",		"Platform:"),			TAttribute<FText>(this, &SSessionInfoWindow::GetPlatformText));
+	AddInfoLine(VerticalBox, LOCTEXT("AppName_HeaderText",		"Application Name:"),	TAttribute<FText>(this, &SSessionInfoWindow::GetAppNameText));
+	AddInfoLine(VerticalBox, LOCTEXT("BuildConfig_HeaderText",	"Build Config:"),		TAttribute<FText>(this, &SSessionInfoWindow::GetBuildConfigText));
+	AddInfoLine(VerticalBox, LOCTEXT("BuildTarget_HeaderText",	"Build Target:"),		TAttribute<FText>(this, &SSessionInfoWindow::GetBuildTargetText));
+	AddInfoLine(VerticalBox, LOCTEXT("CommandLine_HeaderText",	"Command Line:"),		TAttribute<FText>(this, &SSessionInfoWindow::GetCommandLineText));
+	AddInfoLine(VerticalBox, LOCTEXT("FileSize_HeaderText",		"File Size:"),			TAttribute<FText>(this, &SSessionInfoWindow::GetFileSizeText));
+	AddInfoLine(VerticalBox, LOCTEXT("Status_HeaderText",		"Status:"),				TAttribute<FText>(this, &SSessionInfoWindow::GetStatusText));
+	AddInfoLine(VerticalBox, LOCTEXT("Modules_HeaderText",		"Modules:"),			TAttribute<FText>(this, &SSessionInfoWindow::GetModulesText));
+}
+
+////////////////////////////////////////////////////////////////////////////////////////////////////
+
+void SSessionInfoWindow::AddInfoLine(TSharedPtr<SVerticalBox> InVerticalBox, const FText& InHeader, const TAttribute<FText>& InValue) const
+{
+	InVerticalBox->AddSlot()
+		.AutoHeight()
+		.Padding(8.0f, 4.0f, 8.0f, 0.0f)
+		[
+			SNew(STextBlock)
+			.Text(InHeader)
+			.Font(FCoreStyle::GetDefaultFontStyle("Bold", 11))
+		];
+
+	InVerticalBox->AddSlot()
+		.AutoHeight()
+		.Padding(8.0f, 0.0f, 8.0f, 4.0f)
+		[
+			SNew(SEditableTextBox)
+			.Text(InValue)
+			.HintText(LOCTEXT("HintText", "N/A"))
+			.BackgroundColor(FLinearColor(0.1f, 0.1f, 0.1f, 1.0f))
+			.IsReadOnly(true)
+		];
+}
+
+END_SLATE_FUNCTION_BUILD_OPTIMIZATION
+
+////////////////////////////////////////////////////////////////////////////////////////////////////
+
+void SSessionInfoWindow::Tick(const FGeometry& AllottedGeometry, const double InCurrentTime, const float InDeltaTime)
+{
+}
+
+////////////////////////////////////////////////////////////////////////////////////////////////////
+
+EActiveTimerReturnType SSessionInfoWindow::UpdateActiveDuration(double InCurrentTime, float InDeltaTime)
+{
+	DurationActive += InDeltaTime;
+
+	// The window will explicitly unregister this active timer when the mouse leaves.
+	return EActiveTimerReturnType::Continue;
+}
+
+////////////////////////////////////////////////////////////////////////////////////////////////////
+
+void SSessionInfoWindow::OnMouseEnter(const FGeometry& MyGeometry, const FPointerEvent& MouseEvent)
+{
+	SCompoundWidget::OnMouseEnter(MyGeometry, MouseEvent);
+
+	if (!ActiveTimerHandle.IsValid())
+	{
+		ActiveTimerHandle = RegisterActiveTimer(0.f, FWidgetActiveTimerDelegate::CreateSP(this, &SSessionInfoWindow::UpdateActiveDuration));
+	}
+}
+
+////////////////////////////////////////////////////////////////////////////////////////////////////
+
+void SSessionInfoWindow::OnMouseLeave(const FPointerEvent& MouseEvent)
+{
+	SCompoundWidget::OnMouseLeave(MouseEvent);
+
+	auto PinnedActiveTimerHandle = ActiveTimerHandle.Pin();
+	if (PinnedActiveTimerHandle.IsValid())
+	{
+		UnRegisterActiveTimer(PinnedActiveTimerHandle.ToSharedRef());
+	}
+}
+
+////////////////////////////////////////////////////////////////////////////////////////////////////
+
+FReply SSessionInfoWindow::OnKeyDown(const FGeometry& MyGeometry, const FKeyEvent& InKeyEvent)
+{
+	return FReply::Unhandled();
+}
+
+////////////////////////////////////////////////////////////////////////////////////////////////////
+
+FReply SSessionInfoWindow::OnDragOver(const FGeometry& MyGeometry, const FDragDropEvent& DragDropEvent)
+{
+	return SCompoundWidget::OnDragOver(MyGeometry, DragDropEvent);
+}
+
+////////////////////////////////////////////////////////////////////////////////////////////////////
+
+FReply SSessionInfoWindow::OnDrop(const FGeometry& MyGeometry, const FDragDropEvent& DragDropEvent)
+{
+	return SCompoundWidget::OnDrop(MyGeometry, DragDropEvent);
+}
+
+////////////////////////////////////////////////////////////////////////////////////////////////////
+
+FText SSessionInfoWindow::GetSessionNameText() const
+{
+	FText SessionName;
+	TSharedPtr<const Trace::IAnalysisSession> Session = FInsightsManager::Get()->GetSession();
+	if (Session.IsValid())
+	{
+		Trace::FAnalysisSessionReadScope SessionReadScope(*Session.Get());
+		SessionName = FText::FromString(FPaths::GetBaseFilename(Session->GetName()));
+	}
+	return SessionName;
+}
+
+////////////////////////////////////////////////////////////////////////////////////////////////////
+
+FText SSessionInfoWindow::GetUriText() const
+{
+	//TODO: update code to use a SessionInfo provider instead
+	FText LocalUri;
+	TSharedPtr<const Trace::IAnalysisSession> Session = FInsightsManager::Get()->GetSession();
+	if (Session.IsValid())
+	{
+		Trace::FAnalysisSessionReadScope SessionReadScope(*Session.Get());
+		LocalUri = FText::FromString(FString(Session->GetName()));
+	}
+	return LocalUri;
+}
+
+////////////////////////////////////////////////////////////////////////////////////////////////////
+
+FText SSessionInfoWindow::GetPlatformText() const
+{
+	//TODO
+	return FText::GetEmpty();
+}
+
+////////////////////////////////////////////////////////////////////////////////////////////////////
+
+FText SSessionInfoWindow::GetAppNameText() const
+{
+	//TODO
+	return FText::GetEmpty();
+}
+
+////////////////////////////////////////////////////////////////////////////////////////////////////
+
+FText SSessionInfoWindow::GetBuildConfigText() const
+{
+	//TODO
+	return FText::GetEmpty();
+}
+
+////////////////////////////////////////////////////////////////////////////////////////////////////
+
+FText SSessionInfoWindow::GetBuildTargetText() const
+{
+	//TODO
+	return FText::GetEmpty();
+}
+
+////////////////////////////////////////////////////////////////////////////////////////////////////
+
+FText SSessionInfoWindow::GetCommandLineText() const
+{
+	//TODO
+	return FText::GetEmpty();
+}
+
+////////////////////////////////////////////////////////////////////////////////////////////////////
+
+FText SSessionInfoWindow::GetFileSizeText() const
+{
+	//TODO
+	return FText::GetEmpty();
+}
+
+////////////////////////////////////////////////////////////////////////////////////////////////////
+
+FText SSessionInfoWindow::GetStatusText() const
+{
+	//TODO: add also info from a SessionInfo provider
+	FText Status;
+	TSharedPtr<const Trace::IAnalysisSession> Session = FInsightsManager::Get()->GetSession();
+	if (Session.IsValid())
+	{
+		Trace::FAnalysisSessionReadScope SessionReadScope(*Session.Get());
+		Status = FText::Format(LOCTEXT("StatusFmt", "{0} Session Duration: {1}"),
+			Session->IsAnalysisComplete() ? FText::FromString(FString(TEXT("ANALYSIS COMPLETED."))) : FText::FromString(FString(TEXT("ANALYZING..."))),
+			FText::FromString(TimeUtils::FormatTimeAuto(Session->GetDurationSeconds(), 2)));
+	}
+	return Status;
+}
+
+////////////////////////////////////////////////////////////////////////////////////////////////////
+
+FText SSessionInfoWindow::GetModulesText() const
+{
+	FString ModulesStr;
+	TArray<Trace::FModuleInfo> Modules;
+
+	TSharedPtr<Trace::IModuleService> ModuleService = FInsightsManager::Get()->GetModuleService();
+	if (ModuleService)
+	{
+		ModuleService->GetAvailableModules(Modules);
+	}
+
+	for (const Trace::FModuleInfo& Module : Modules)
+	{
+		ModulesStr += Module.DisplayName;
+		ModulesStr += TEXT(", ");
+	}
+
+	return FText::FromString(ModulesStr);
+}
+
+////////////////////////////////////////////////////////////////////////////////////////////////////
+
+#undef LOCTEXT_NAMESPACE
diff --git a/Engine/Source/Developer/TraceInsights/Private/Insights/Widgets/SSessionInfoWindow.h b/Engine/Source/Developer/TraceInsights/Private/Insights/Widgets/SSessionInfoWindow.h
new file mode 100644
--- /dev/null
+++ b/Engine/Source/Developer/TraceInsights/Private/Insights/Widgets/SSessionInfoWindow.h
@@ -0,0 +1,139 @@
+// Copyright Epic Games, Inc. All Rights Reserved.
+
+#pragma once
+
+#include "CoreMinimal.h"
+#include "Framework/Docking/TabManager.h"
+#include "GenericPlatform/GenericPlatformMisc.h"
+#include "Input/Reply.h"
+#include "Layout/Visibility.h"
+#include "Misc/Guid.h"
+#include "SlateFwd.h"
+#include "TraceServices/ModuleService.h"
+#include "TraceServices/SessionService.h"
+#include "Widgets/DeclarativeSyntaxSupport.h"
+#include "Widgets/Docking/SDockTab.h"
+#include "Widgets/Layout/SSplitter.h"
+#include "Widgets/SCompoundWidget.h"
+#include "Widgets/SOverlay.h"
+#include "Widgets/Views/SListView.h"
+
+// Insights
+#include "Insights/InsightsManager.h"
+
+////////////////////////////////////////////////////////////////////////////////////////////////////
+
+class FActiveTimerHandle;
+class SVerticalBox;
+class SEditableTextBox;
+
+////////////////////////////////////////////////////////////////////////////////////////////////////
+
+/** Implements the Start Page window. */
+class SSessionInfoWindow : public SCompoundWidget
+{
+public:
+	/** Default constructor. */
+	SSessionInfoWindow();
+
+	/** Virtual destructor. */
+	virtual ~SSessionInfoWindow();
+
+	SLATE_BEGIN_ARGS(SSessionInfoWindow) {}
+	SLATE_END_ARGS()
+
+	/** Constructs this widget. */
+	void Construct(const FArguments& InArgs);
+
+private:
+	/** Updates the amount of time the profiler has been active. */
+	EActiveTimerReturnType UpdateActiveDuration(double InCurrentTime, float InDeltaTime);
+
+	/**
+	 * Ticks this widget.  Override in derived classes, but always call the parent implementation.
+	 *
+	 * @param  AllottedGeometry The space allotted for this widget
+	 * @param  InCurrentTime  Current absolute real time
+	 * @param  InDeltaTime  Real time passed since last tick
+	 */
+	virtual void Tick(const FGeometry& AllottedGeometry, const double InCurrentTime, const float InDeltaTime) override;
+
+	/**
+	 * The system will use this event to notify a widget that the cursor has entered it. This event is NOT bubbled.
+	 *
+	 * @param MyGeometry The Geometry of the widget receiving the event
+	 * @param MouseEvent Information about the input event
+	 */
+	virtual void OnMouseEnter(const FGeometry& MyGeometry, const FPointerEvent& MouseEvent) override;
+
+	/**
+	 * The system will use this event to notify a widget that the cursor has left it. This event is NOT bubbled.
+	 *
+	 * @param MouseEvent Information about the input event
+	 */
+	virtual void OnMouseLeave(const FPointerEvent& MouseEvent) override;
+
+	/**
+	 * Called after a key is pressed when this widget has focus
+	 *
+	 * @param MyGeometry The Geometry of the widget receiving the event
+	 * @param  InKeyEvent  Key event
+	 *
+	 * @return  Returns whether the event was handled, along with other possible actions
+	 */
+	virtual FReply OnKeyDown(const FGeometry& MyGeometry, const FKeyEvent& InKeyEvent) override;
+
+	/**
+	 * Called when the user is dropping something onto a widget; terminates drag and drop.
+	 *
+	 * @param MyGeometry      The geometry of the widget receiving the event.
+	 * @param DragDropEvent   The drag and drop event.
+	 *
+	 * @return A reply that indicated whether this event was handled.
+	 */
+	virtual FReply OnDrop(const FGeometry& MyGeometry, const FDragDropEvent& DragDropEvent) override;
+
+	/**
+	 * Called during drag and drop when the the mouse is being dragged over a widget.
+	 *
+	 * @param MyGeometry      The geometry of the widget receiving the event.
+	 * @param DragDropEvent   The drag and drop event.
+	 *
+	 * @return A reply that indicated whether this event was handled.
+	 */
+	virtual FReply OnDragOver(const FGeometry& MyGeometry, const FDragDropEvent& DragDropEvent)  override;
+
+private:
+	void AddInfoLine(TSharedPtr<SVerticalBox> InVerticalBox, const FText& InHeader, const TAttribute<FText>& InValue) const;
+
+	FText GetSessionNameText() const;
+	FText GetUriText() const;
+	FText GetPlatformText() const;
+	FText GetAppNameText() const;
+	FText GetBuildConfigText() const;
+	FText GetBuildTargetText() const;
+	FText GetCommandLineText() const;
+	FText GetFileSizeText() const;
+	FText GetStatusText() const;
+	FText GetModulesText() const;
+
+public:
+	/** The number of seconds the profiler has been active */
+	float DurationActive;
+
+private:
+	/** The handle to the active update duration tick */
+	TWeakPtr<FActiveTimerHandle> ActiveTimerHandle;
+
+	//Trace::FSessionHandle SessionHandle;
+	//FText Name;
+	//FText Uri;
+	//FText Platform;
+	//FText AppName;
+	//FText CommandLine;
+	//EBuildConfiguration ConfigurationType;
+	//EBuildTargetType TargetType;
+	//FDateTime TimeStamp;
+	//uint64 Size;
+	//bool bIsLive;
+};
diff --git a/Engine/Source/Developer/TraceInsights/Private/Insights/Widgets/SStartPageWindow.cpp b/Engine/Source/Developer/TraceInsights/Private/Insights/Widgets/SStartPageWindow.cpp
--- a/Engine/Source/Developer/TraceInsights/Private/Insights/Widgets/SStartPageWindow.cpp
+++ b/Engine/Source/Developer/TraceInsights/Private/Insights/Widgets/SStartPageWindow.cpp
@@ -14,8 +14,11 @@
 #include "Styling/CoreStyle.h"
 #include "Widgets/Docking/SDockTab.h"
 #include "Widgets/Images/SImage.h"
+#include "Widgets/Input/SCheckBox.h"
+#include "Widgets/Input/SSearchBox.h"
 #include "Widgets/Layout/SBorder.h"
 #include "Widgets/Layout/SBox.h"
+#include "Widgets/Layout/SScrollBox.h"
 #include "Widgets/Layout/SSpacer.h"
 #include "Widgets/Notifications/SNotificationList.h"
 #include "Widgets/SBoxPanel.h"
@@ -361,9 +364,30 @@ private:
 ////////////////////////////////////////////////////////////////////////////////////////////////////
 
 SStartPageWindow::SStartPageWindow()
-	: DurationActive(0.0f)
+	: NotificationList()
+	, ActiveNotifications()
+	, OverlaySettingsSlot(nullptr)
+	, DurationActive(0.0f)
+	, ActiveTimerHandle()
+	, MainContentPanel()
 	, AvailableSessionCount(0)
 	, LiveSessionCount(0)
+#if WITH_EDITOR
+	, bAutoStartAnalysisForLiveSessions(false)
+#else
+	, bAutoStartAnalysisForLiveSessions(true)
+#endif
+	, AutoStartedSessions()
+	, AutoStartPlatformFilter()
+	, AutoStartAppNameFilter()
+	, AutoStartConfigurationTypeFilter(EBuildConfiguration::Unknown)
+	, AutoStartTargetTypeFilter(EBuildTargetType::Unknown)
+	, TraceSessionsListView()
+	, TraceSessions()
+	, TraceSessionsMap()
+	, HostTextBox()
+	, SelectedTraceSession()
+	, SplashScreenOverlayFadeTime(0.0f)
 {
 }
 
@@ -391,24 +415,47 @@ void SStartPageWindow::Construct(const FArguments& InArgs)
 
 			// Version
 			+ SOverlay::Slot()
-				.HAlign(HAlign_Right)
-				.VAlign(VAlign_Top)
-				.Padding(0.0f, -16.0f, 0.0f, 0.0f)
-				[
-					SNew(STextBlock)
-						.Clipping(EWidgetClipping::ClipToBoundsWithoutIntersecting)
-						.Text(LOCTEXT("UnrealInsightsVersion", UNREAL_INSIGHTS_VERSION_STRING_EX))
-						.ColorAndOpacity(FLinearColor(0.15f, 0.15f, 0.15f, 1.0f))
-				]
+			.HAlign(HAlign_Right)
+			.VAlign(VAlign_Top)
+			.Padding(0.0f, -16.0f, 0.0f, 0.0f)
+			[
+				SNew(STextBlock)
+				.Clipping(EWidgetClipping::ClipToBoundsWithoutIntersecting)
+				.Text(LOCTEXT("UnrealInsightsVersion", UNREAL_INSIGHTS_VERSION_STRING_EX))
+				.ColorAndOpacity(FLinearColor(0.15f, 0.15f, 0.15f, 1.0f))
+			]
 
 			// Overlay slot for the main window area
 			+ SOverlay::Slot()
-				.HAlign(HAlign_Fill)
-				.VAlign(VAlign_Center)
+			.HAlign(HAlign_Fill)
+			.VAlign(VAlign_Center)
+			[
+				SNew(SScrollBox)
+				.Orientation(Orient_Vertical)
+
+				+ SScrollBox::Slot()
 				[
-					SNew(SVerticalBox)
+					SAssignNew(MainContentPanel, SVerticalBox)
 
 					+ SVerticalBox::Slot()
+					.AutoHeight()
+					.HAlign(HAlign_Center)
+					.Padding(3.0f, 3.0f)
+					[
+						SNew(SBox)
+						.WidthOverride(1024.0f)
+						[
+							SNew(SBorder)
+							.BorderImage(FEditorStyle::GetBrush("NotificationList.ItemBackground"))
+							.Padding(8.0f)
+							.HAlign(HAlign_Fill)
+							[
+								ConstructSessionsPanel()
+							]
+						]
+					]
+
+						+ SVerticalBox::Slot()
 						.AutoHeight()
 						.HAlign(HAlign_Center)
 						.Padding(3.0f, 3.0f)
@@ -421,63 +468,93 @@ void SStartPageWindow::Construct(const FArguments& InArgs)
 								.Padding(8.0f)
 								.HAlign(HAlign_Fill)
 								[
-									ConstructSessionsPanel()
+									ConstructAutoStartPanel()
 								]
 							]
 						]
 
 					+ SVerticalBox::Slot()
-						.AutoHeight()
-						.HAlign(HAlign_Center)
-						.Padding(3.0f, 3.0f)
+					.AutoHeight()
+					.HAlign(HAlign_Center)
+					.Padding(3.0f, 3.0f)
+					[
+						SNew(SBox)
+						.WidthOverride(256.0f)
 						[
-							SNew(SBox)
-							.WidthOverride(256.0f)
+							SNew(SBorder)
+							.BorderImage(FEditorStyle::GetBrush("NotificationList.ItemBackground"))
+							.Padding(8.0f)
+							.HAlign(HAlign_Fill)
 							[
-								SNew(SBorder)
-								.BorderImage(FEditorStyle::GetBrush("NotificationList.ItemBackground"))
-								.Padding(8.0f)
-								.HAlign(HAlign_Fill)
-								[
-									ConstructRecorderPanel()
-								]
+								ConstructRecorderPanel()
 							]
 						]
+					]
 
 					+ SVerticalBox::Slot()
-						.AutoHeight()
-						.HAlign(HAlign_Center)
-						.Padding(3.0f, 3.0f)
+					.AutoHeight()
+					.HAlign(HAlign_Center)
+					.Padding(3.0f, 3.0f)
+					[
+						SNew(SBox)
+						.WidthOverride(512.0f)
+						.Visibility(this, &SStartPageWindow::StopTraceRecorder_Visibility)
 						[
-							SNew(SBox)
-							.WidthOverride(512.0f)
-							.Visibility(this, &SStartPageWindow::StopTraceRecorder_Visibility)
+							SNew(SBorder)
+							.BorderImage(FEditorStyle::GetBrush("NotificationList.ItemBackground"))
+							.Padding(8.0f)
+							.HAlign(HAlign_Fill)
 							[
-								SNew(SBorder)
-								.BorderImage(FEditorStyle::GetBrush("NotificationList.ItemBackground"))
-								.Padding(8.0f)
-								.HAlign(HAlign_Fill)
-								[
-									ConstructConnectPanel()
-								]
+								ConstructConnectPanel()
 							]
 						]
+					]
 				]
+			]
 
-			// Notification area overlay
+			// Overlay for fake splashscreen.
 			+ SOverlay::Slot()
-				.HAlign(HAlign_Right)
-				.VAlign(VAlign_Bottom)
-				.Padding(16.0f)
+			.HAlign(HAlign_Fill)
+			.VAlign(VAlign_Fill)
+			.Padding(0.0f)
+			[
+				SNew(SBox)
+				.Visibility(this, &SStartPageWindow::SplashScreenOverlay_Visibility)
 				[
-					SAssignNew(NotificationList, SNotificationList)
+					SNew(SBorder)
+					.BorderImage(FEditorStyle::GetBrush("NotificationList.ItemBackground"))
+					.BorderBackgroundColor(this, &SStartPageWindow::SplashScreenOverlay_ColorAndOpacity)
+					.Padding(0.0f)
+					.HAlign(HAlign_Fill)
+					.VAlign(VAlign_Fill)
+					[
+						SNew(SBox)
+						.HAlign(HAlign_Center)
+						.VAlign(VAlign_Center)
+						[
+							SNew(STextBlock)
+							.Text(this, &SStartPageWindow::GetSplashScreenOverlayText)
+							.Font(FCoreStyle::GetDefaultFontStyle("Bold", 11))
+							.ColorAndOpacity(this, &SStartPageWindow::SplashScreenOverlay_TextColorAndOpacity)
+						]
+					]
 				]
+			]
+
+			// Notification area overlay
+			+ SOverlay::Slot()
+			.HAlign(HAlign_Right)
+			.VAlign(VAlign_Bottom)
+			.Padding(16.0f)
+			[
+				SAssignNew(NotificationList, SNotificationList)
+			]
 
 			// Settings dialog overlay
 			+ SOverlay::Slot()
-				.HAlign(HAlign_Center)
-				.VAlign(VAlign_Center)
-				.Expose(OverlaySettingsSlot)
+			.HAlign(HAlign_Center)
+			.VAlign(VAlign_Center)
+			.Expose(OverlaySettingsSlot)
 		];
 
 #if WITH_EDITOR
@@ -681,6 +758,51 @@ TSharedRef<SWidget> SStartPageWindow::ConstructLocalSessionDirectoryPanel()
 
 ////////////////////////////////////////////////////////////////////////////////////////////////////
 
+TSharedRef<SWidget> SStartPageWindow::ConstructAutoStartPanel()
+{
+	TSharedRef<SWidget> Widget = SNew(SHorizontalBox)
+
+	+ SHorizontalBox::Slot()
+		.AutoWidth()
+		.Padding(0.0f, 0.0f, 0.0f, 0.0f)
+		.VAlign(VAlign_Center)
+		[
+			SNew(SCheckBox)
+			.ToolTipText(LOCTEXT("AutoStart_Tooltip", "Enable auto-start analysis for LIVE sessions."))
+			.IsChecked(this, &SStartPageWindow::AutoStart_IsChecked)
+			.OnCheckStateChanged(this, &SStartPageWindow::AutoStart_OnCheckStateChanged)
+			[
+				SNew(STextBlock)
+				.Text(LOCTEXT("AutoStart_Text", "Auto Start Analysis for LIVE Sessions"))
+			]
+		]
+
+	+ SHorizontalBox::Slot()
+		.AutoWidth()
+		.Padding(4.0f, 0.0f, 0.0f, 0.0f)
+		.VAlign(VAlign_Center)
+		[
+			SAssignNew(AutoStartPlatformFilter, SSearchBox)
+			.HintText(LOCTEXT("AutoStartPlatformFilter_Hint", "Platform"))
+			.ToolTipText(LOCTEXT("AutoStartPlatformFilter_Tooltip", "Type here to specify the Platform filter.\nAuto-start analysis will be enabled only for live sessions with this specified Platform."))
+		]
+
+	+ SHorizontalBox::Slot()
+		.AutoWidth()
+		.Padding(4.0f, 0.0f, 0.0f, 0.0f)
+		.VAlign(VAlign_Center)
+		[
+			SAssignNew(AutoStartAppNameFilter, SSearchBox)
+			.HintText(LOCTEXT("AutoStartAppNameFilter_Hint", "AppName"))
+			.ToolTipText(LOCTEXT("AutoStartAppNameFilter_Tooltip", "Type here to specify the AppName filter.\nAuto-start analysis will be enabled only for live sessions with this specified AppName."))
+		]
+	;
+
+	return Widget;
+}
+
+////////////////////////////////////////////////////////////////////////////////////////////////////
+
 TSharedRef<SWidget> SStartPageWindow::ConstructRecorderPanel()
 {
 	TSharedRef<SWidget> Widget = SNew(SVerticalBox)
@@ -768,7 +890,6 @@ TSharedRef<SWidget> SStartPageWindow::ConstructRecorderPanel()
 					.ColorAndOpacity(FLinearColor::Gray)
 				]
 		]
-
 	;
 
 	return Widget;
@@ -842,6 +963,69 @@ END_SLATE_FUNCTION_BUILD_OPTIMIZATION
 
 ////////////////////////////////////////////////////////////////////////////////////////////////////
 
+void SStartPageWindow::ShowSplashScreenOverlay()
+{
+	SplashScreenOverlayFadeTime = 3.5f;
+}
+
+////////////////////////////////////////////////////////////////////////////////////////////////////
+
+void SStartPageWindow::TickSplashScreenOverlay(const float InDeltaTime)
+{
+	if (SplashScreenOverlayFadeTime > 0.0f)
+	{
+		SplashScreenOverlayFadeTime = FMath::Max(0.0f, SplashScreenOverlayFadeTime - InDeltaTime);
+	}
+}
+
+////////////////////////////////////////////////////////////////////////////////////////////////////
+
+float SStartPageWindow::SplashScreenOverlayOpacity() const
+{
+	constexpr float FadeInStartTime = 3.5f;
+	constexpr float FadeInEndTime = 3.0f;
+	constexpr float FadeOutStartTime = 1.0f;
+	constexpr float FadeOutEndTime = 0.0f;
+
+	const float Opacity =
+		SplashScreenOverlayFadeTime > FadeInStartTime ? 0.0f :
+		SplashScreenOverlayFadeTime > FadeInEndTime ? 1.0f - (SplashScreenOverlayFadeTime - FadeInEndTime) / (FadeInStartTime - FadeInEndTime) :
+		SplashScreenOverlayFadeTime > FadeOutStartTime ? 1.0f :
+		SplashScreenOverlayFadeTime > FadeOutEndTime ? (SplashScreenOverlayFadeTime - FadeOutEndTime) / (FadeOutStartTime - FadeOutEndTime) : 0.0f;
+
+	return Opacity;
+}
+
+////////////////////////////////////////////////////////////////////////////////////////////////////
+
+EVisibility SStartPageWindow::SplashScreenOverlay_Visibility() const
+{
+	return SplashScreenOverlayFadeTime > 0.0f ? EVisibility::Visible : EVisibility::Collapsed;
+}
+
+////////////////////////////////////////////////////////////////////////////////////////////////////
+
+FSlateColor SStartPageWindow::SplashScreenOverlay_ColorAndOpacity() const
+{
+	return FSlateColor(FLinearColor(0.7f, 0.7f, 0.7f, SplashScreenOverlayOpacity()));
+}
+
+////////////////////////////////////////////////////////////////////////////////////////////////////
+
+FSlateColor SStartPageWindow::SplashScreenOverlay_TextColorAndOpacity() const
+{
+	return FSlateColor(FLinearColor(0.8f, 0.8f, 0.8f, SplashScreenOverlayOpacity()));
+}
+
+////////////////////////////////////////////////////////////////////////////////////////////////////
+
+FText SStartPageWindow::GetSplashScreenOverlayText() const
+{
+	return FText::Format(LOCTEXT("StartAnalysis", "Starting analysis...\n{0}"), FText::FromString(SplashScreenOverlayTraceFile));
+}
+
+////////////////////////////////////////////////////////////////////////////////////////////////////
+
 FReply SStartPageWindow::RefreshTraceSessions_OnClicked()
 {
 	RefreshTraceSessionList();
@@ -901,6 +1085,9 @@ void SStartPageWindow::RefreshTraceSessionList()
 
 	bool bSessionChanged = false;
 
+	FString AutoStartPlatformFilterStr = AutoStartPlatformFilter->GetText().ToString();
+	FString AutoStartAppNameFilterStr = AutoStartAppNameFilter->GetText().ToString();
+
 	// Count number of live sessions and update file sizes.
 	int32 OldLiveSessionCount = LiveSessionCount;
 	LiveSessionCount = 0;
@@ -911,6 +1098,19 @@ void SStartPageWindow::RefreshTraceSessionList()
 		if (SessionInfo.bIsLive)
 		{
 			++LiveSessionCount;
+
+			// Auto start.
+			if (bAutoStartAnalysisForLiveSessions && !AutoStartedSessions.Contains(SessionHandle))
+			{
+				if ((AutoStartPlatformFilterStr.IsEmpty() || FCString::Strcmp(*AutoStartPlatformFilterStr, SessionInfo.Platform) == 0) &&
+					(AutoStartAppNameFilterStr.IsEmpty() || FCString::Strcmp(*AutoStartAppNameFilterStr, SessionInfo.AppName) == 0) &&
+					(AutoStartConfigurationTypeFilter == EBuildConfiguration::Unknown || AutoStartConfigurationTypeFilter == SessionInfo.ConfigurationType) &&
+					(AutoStartTargetTypeFilter == EBuildTargetType::Unknown || AutoStartTargetTypeFilter == SessionInfo.TargetType))
+				{
+					AutoStartedSessions.Add(SessionHandle);
+					LoadTraceFile(SessionInfo.Uri);
+				}
+			}
 		}
 
 		TSharedPtr<FTraceSession>* TraceSessionPtrPtr = TraceSessionsMap.Find(SessionHandle);
@@ -982,10 +1182,7 @@ void SStartPageWindow::TraceSessions_OnSelectionChanged(TSharedPtr<FTraceSession
 
 void SStartPageWindow::TraceSessions_OnMouseButtonDoubleClick(TSharedPtr<FTraceSession> TraceSession)
 {
-	if (TraceSession.IsValid())
-	{
-		FInsightsManager::Get()->LoadSession(TraceSession->SessionHandle);
-	}
+	LoadTraceSession(TraceSession);
 }
 
 ////////////////////////////////////////////////////////////////////////////////////////////////////
@@ -997,6 +1194,20 @@ EVisibility SStartPageWindow::TraceSessions_Visibility() const
 
 ////////////////////////////////////////////////////////////////////////////////////////////////////
 
+ECheckBoxState SStartPageWindow::AutoStart_IsChecked() const
+{
+	return bAutoStartAnalysisForLiveSessions ? ECheckBoxState::Checked : ECheckBoxState::Unchecked;
+}
+
+////////////////////////////////////////////////////////////////////////////////////////////////////
+
+void SStartPageWindow::AutoStart_OnCheckStateChanged(ECheckBoxState NewState)
+{
+	bAutoStartAnalysisForLiveSessions = (NewState == ECheckBoxState::Checked);
+}
+
+////////////////////////////////////////////////////////////////////////////////////////////////////
+
 void SStartPageWindow::Tick(const FGeometry& AllottedGeometry, const double InCurrentTime, const float InDeltaTime)
 {
 	// We need to check if there is any available session, but not too often.
@@ -1008,6 +1219,8 @@ void SStartPageWindow::Tick(const FGeometry& AllottedGeometry, const double InCu
 		NextTimestamp = Time + WaitTime;
 		RefreshTraceSessionList();
 	}
+
+	TickSplashScreenOverlay(InDeltaTime);
 }
 
 ////////////////////////////////////////////////////////////////////////////////////////////////////
@@ -1111,7 +1324,7 @@ FReply SStartPageWindow::OnDrop(const FGeometry& MyGeometry, const FDragDropEven
 				const FString DraggedFileExtension = FPaths::GetExtension(Files[0], true);
 				if (DraggedFileExtension == TEXT(".utrace"))
 				{
-					FInsightsManager::Get()->LoadTraceFile(Files[0]);
+					LoadTraceFile(Files[0]);
 					return FReply::Handled();
 				}
 			}
@@ -1132,11 +1345,7 @@ bool SStartPageWindow::Open_IsEnabled() const
 
 FReply SStartPageWindow::Open_OnClicked()
 {
-	if (SelectedTraceSession.IsValid())
-	{
-		Trace::FSessionHandle SessionHandle = SelectedTraceSession->SessionHandle;
-		FInsightsManager::Get()->LoadSession(SessionHandle);
-	}
+	LoadTraceSession(SelectedTraceSession);
 	return FReply::Handled();
 }
 
@@ -1171,23 +1380,76 @@ void SStartPageWindow::OpenFileDialog()
 	{
 		if (OutFiles.Num() == 1)
 		{
-			FInsightsManager::Get()->LoadTraceFile(OutFiles[0]);
+			LoadTraceFile(OutFiles[0]);
 		}
 	}
 }
 
 ////////////////////////////////////////////////////////////////////////////////////////////////////
 
-void SStartPageWindow::LoadTraceFile(const TCHAR* TraceFile)
+void SStartPageWindow::LoadTraceSession(TSharedPtr<FTraceSession> InTraceSession)
 {
-	FInsightsManager::Get()->LoadTraceFile(FString(TraceFile));
+	if (InTraceSession.IsValid())
+	{
+		if (FInsightsManager::Get()->ShouldOpenAnalysisInSeparateProcess())
+		{
+			LoadTraceFile(InTraceSession->Uri.ToString());
+		}
+		else
+		{
+			LoadSession(InTraceSession->SessionHandle);
+		}
+	}
+}
+
+////////////////////////////////////////////////////////////////////////////////////////////////////
+
+void SStartPageWindow::LoadTraceFile(const FString& InTraceFile)
+{
+	if (FInsightsManager::Get()->ShouldOpenAnalysisInSeparateProcess())
+	{
+		const TCHAR* ExecutablePath = FPlatformProcess::ExecutablePath();
+
+		FString CmdLine = TEXT("-Trace=\"") + InTraceFile + TEXT("\"");
+
+		constexpr bool bLaunchDetached = false;
+		constexpr bool bLaunchHidden = false;
+		constexpr bool bLaunchReallyHidden = false;
+
+		uint32 ProcessID = 0;
+		const int32 PriorityModifier = 0;
+		const TCHAR* OptionalWorkingDirectory = nullptr;
+
+		void* PipeWriteChild = nullptr;
+		void* PipeReadChild = nullptr;
+
+		FProcHandle Handle = FPlatformProcess::CreateProc(ExecutablePath, *CmdLine, bLaunchDetached, bLaunchHidden, bLaunchReallyHidden, &ProcessID, PriorityModifier, OptionalWorkingDirectory, PipeWriteChild, PipeReadChild);
+		FPlatformProcess::CloseProc(Handle);
+
+		SplashScreenOverlayTraceFile = FPaths::GetBaseFilename(InTraceFile);
+		ShowSplashScreenOverlay();
+	}
+	else
+	{
+		FInsightsManager::Get()->LoadTraceFile(InTraceFile);
+	}
 }
 
 ////////////////////////////////////////////////////////////////////////////////////////////////////
 
 void SStartPageWindow::LoadSession(Trace::FSessionHandle SessionHandle)
 {
-	FInsightsManager::Get()->LoadSession(SessionHandle);
+	if (FInsightsManager::Get()->ShouldOpenAnalysisInSeparateProcess())
+	{
+		TSharedRef<Trace::ISessionService> SessionService = FInsightsManager::Get()->GetSessionService();
+		Trace::FSessionInfo SessionInfo;
+		SessionService->GetSessionInfo(SessionHandle, SessionInfo);
+		LoadTraceFile(SessionInfo.Uri);
+	}
+	else
+	{
+		FInsightsManager::Get()->LoadSession(SessionHandle);
+	}
 }
 
 ////////////////////////////////////////////////////////////////////////////////////////////////////
diff --git a/Engine/Source/Developer/TraceInsights/Private/Insights/Widgets/SStartPageWindow.h b/Engine/Source/Developer/TraceInsights/Private/Insights/Widgets/SStartPageWindow.h
--- a/Engine/Source/Developer/TraceInsights/Private/Insights/Widgets/SStartPageWindow.h
+++ b/Engine/Source/Developer/TraceInsights/Private/Insights/Widgets/SStartPageWindow.h
@@ -92,6 +92,7 @@ public:
 
 private:
 	TSharedRef<SWidget> ConstructSessionsPanel();
+	TSharedRef<SWidget> ConstructAutoStartPanel();
 	TSharedRef<SWidget> ConstructLoadPanel();
 	TSharedRef<SWidget> ConstructLocalSessionDirectoryPanel();
 	TSharedRef<SWidget> ConstructRecorderPanel();
@@ -100,6 +101,15 @@ private:
 	/** Generate a new row for the TraceSessions list view. */
 	TSharedRef<ITableRow> TraceSessions_OnGenerateRow(TSharedPtr<FTraceSession> InConnection, const TSharedRef<STableViewBase>& OwnerTable);
 
+	void ShowSplashScreenOverlay();
+	void TickSplashScreenOverlay(const float InDeltaTime);
+	float SplashScreenOverlayOpacity() const;
+
+	EVisibility SplashScreenOverlay_Visibility() const;
+	FSlateColor SplashScreenOverlay_ColorAndOpacity() const;
+	FSlateColor SplashScreenOverlay_TextColorAndOpacity() const;
+	FText GetSplashScreenOverlayText() const;
+
 	/** Callback for determining the visibility of the 'Please select a trace' overlay. */
 	EVisibility IsSessionOverlayVisible() const;
 
@@ -109,8 +119,10 @@ private:
 	FReply Open_OnClicked();
 
 	void OpenFileDialog();
-	void LoadTraceFile(const TCHAR* TraceFile);
-	void LoadSession(Trace::FSessionHandle SessionHandle);
+
+	void LoadTraceSession(TSharedPtr<FTraceSession> InTraceSession);
+	void LoadTraceFile(const FString& InTraceFile);
+	void LoadSession(Trace::FSessionHandle InSessionHandle);
 
 	TSharedRef<SWidget> MakeSessionListMenu();
 
@@ -120,6 +132,9 @@ private:
 	EVisibility TraceSessions_Visibility() const;
 	FReply RefreshTraceSessions_OnClicked();
 
+	ECheckBoxState AutoStart_IsChecked() const;
+	void AutoStart_OnCheckStateChanged(ECheckBoxState NewState);
+
 	FText GetLocalSessionDirectory() const;
 	FReply ExploreLocalSessionDirectory_OnClicked();
 
@@ -215,12 +230,20 @@ private:
 	int32 AvailableSessionCount;
 	int32 LiveSessionCount;
 
+	bool bAutoStartAnalysisForLiveSessions;
+	TArray<Trace::FSessionHandle> AutoStartedSessions; // tracks sessions that were auto started (in order to not start them again)
+
+	TSharedPtr<SSearchBox> AutoStartPlatformFilter;
+	TSharedPtr<SSearchBox> AutoStartAppNameFilter;
+	EBuildConfiguration AutoStartConfigurationTypeFilter;
+	EBuildTargetType AutoStartTargetTypeFilter;
+
 	TSharedPtr<SListView<TSharedPtr<FTraceSession>>> TraceSessionsListView;
 	TArray<TSharedPtr<FTraceSession>> TraceSessions;
 	TMap<Trace::FSessionHandle, TSharedPtr<FTraceSession>> TraceSessionsMap;
 	TSharedPtr<SEditableTextBox> HostTextBox;
 	TSharedPtr<FTraceSession> SelectedTraceSession;
 
-	//TArray<Trace::FModuleInfo> AvailableModules;
-	//TArray<bool> AvailableModulesEnabledState;
+	FString SplashScreenOverlayTraceFile;
+	float SplashScreenOverlayFadeTime;
 };
diff --git a/Engine/Source/Developer/TraceInsights/Private/Insights/Widgets/STimingProfilerToolbar.cpp b/Engine/Source/Developer/TraceInsights/Private/Insights/Widgets/STimingProfilerToolbar.cpp
--- a/Engine/Source/Developer/TraceInsights/Private/Insights/Widgets/STimingProfilerToolbar.cpp
+++ b/Engine/Source/Developer/TraceInsights/Private/Insights/Widgets/STimingProfilerToolbar.cpp
@@ -8,10 +8,7 @@
 #include "Framework/MultiBox/MultiBoxDefs.h"
 #include "Framework/MultiBox/MultiBoxBuilder.h"
 #include "Widgets/Docking/SDockTab.h"
-#include "Widgets/Input/SCheckBox.h"
 #include "Widgets/Layout/SBorder.h"
-#include "Widgets/SBoxPanel.h"
-#include "Widgets/SToolTip.h"
 
 // Insights
 #include "Insights/InsightsCommands.h"
@@ -70,7 +67,6 @@ void STimingProfilerToolbar::Construct(const FArguments& InArgs)
 
 	FToolBarBuilder ToolbarBuilder(CommandList.ToSharedRef(), FMultiBoxCustomization::None);
 	Local::FillViewToolbar(ToolbarBuilder);
-	//FillModulesToolbar(ToolbarBuilder);
 
 	FToolBarBuilder RightSideToolbarBuilder(CommandList.ToSharedRef(), FMultiBoxCustomization::None);
 	Local::FillRightSideToolbar(RightSideToolbarBuilder);
@@ -80,7 +76,7 @@ void STimingProfilerToolbar::Construct(const FArguments& InArgs)
 	[
 		SNew(SHorizontalBox)
 
-		+SHorizontalBox::Slot()
+		+ SHorizontalBox::Slot()
 		.HAlign(HAlign_Left)
 		.VAlign(VAlign_Center)
 		.FillWidth(1.0)
@@ -95,7 +91,7 @@ void STimingProfilerToolbar::Construct(const FArguments& InArgs)
 			]
 		]
 
-		+SHorizontalBox::Slot()
+		+ SHorizontalBox::Slot()
 		.HAlign(HAlign_Right)
 		.VAlign(VAlign_Center)
 		.AutoWidth()
@@ -114,126 +110,4 @@ void STimingProfilerToolbar::Construct(const FArguments& InArgs)
 
 ////////////////////////////////////////////////////////////////////////////////////////////////////
 
-bool STimingProfilerToolbar::ToggleModule_CanExecute(FName ModuleName) const
-{
-	TSharedRef<Trace::ISessionService> SessionService = FInsightsManager::Get()->GetSessionService();
-	if (!SessionService->IsRecorderServerRunning())
-	{
-		return false;
-	}
-
-	TSharedPtr<const Trace::IAnalysisSession> Session = FInsightsManager::Get()->GetSession();
-	if (!Session.IsValid())
-	{
-		return false;
-	}
-
-	Trace::FSessionHandle SessionHandle = FInsightsManager::Get()->GetSessionHandle();
-	if (SessionHandle == 0)
-	{
-		return false;
-	}
-
-	Trace::FSessionInfo SessionInfo;
-	SessionService->GetSessionInfo(SessionHandle, SessionInfo);
-	return SessionInfo.bIsLive;
-}
-
-////////////////////////////////////////////////////////////////////////////////////////////////////
-
-void STimingProfilerToolbar::ToggleModule_Execute(FName ModuleName)
-{
-	TSharedRef<Trace::ISessionService> SessionService = FInsightsManager::Get()->GetSessionService();
-	TSharedPtr<const Trace::IAnalysisSession> Session = FInsightsManager::Get()->GetSession();
-	Trace::FSessionHandle SessionHandle = FInsightsManager::Get()->GetSessionHandle();
-
-	bool bState = SessionService->IsModuleEnabled(SessionHandle, ModuleName);
-	return SessionService->SetModuleEnabled(SessionHandle, ModuleName, !bState);
-}
-
-////////////////////////////////////////////////////////////////////////////////////////////////////
-
-bool STimingProfilerToolbar::ToggleModule_IsChecked(FName ModuleName) const
-{
-	TSharedRef<Trace::ISessionService> SessionService = FInsightsManager::Get()->GetSessionService();
-	TSharedPtr<const Trace::IAnalysisSession> Session = FInsightsManager::Get()->GetSession();
-	Trace::FSessionHandle SessionHandle = FInsightsManager::Get()->GetSessionHandle();
-
-	return SessionService->IsModuleEnabled(SessionHandle, ModuleName);
-}
-
-////////////////////////////////////////////////////////////////////////////////////////////////////
-
-ECheckBoxState STimingProfilerToolbar::ToggleModule_IsChecked2(FName ModuleName) const
-{
-	TSharedRef<Trace::ISessionService> SessionService = FInsightsManager::Get()->GetSessionService();
-	TSharedPtr<const Trace::IAnalysisSession> Session = FInsightsManager::Get()->GetSession();
-	Trace::FSessionHandle SessionHandle = FInsightsManager::Get()->GetSessionHandle();
-
-	return SessionService->IsModuleEnabled(SessionHandle, ModuleName) ? ECheckBoxState::Checked : ECheckBoxState::Unchecked;
-}
-
-////////////////////////////////////////////////////////////////////////////////////////////////////
-
-void STimingProfilerToolbar::ToggleModule_OnCheckStateChanged(ECheckBoxState NewRadioState, FName ModuleName)
-{
-	TSharedRef<Trace::ISessionService> SessionService = FInsightsManager::Get()->GetSessionService();
-	TSharedPtr<const Trace::IAnalysisSession> Session = FInsightsManager::Get()->GetSession();
-	Trace::FSessionHandle SessionHandle = FInsightsManager::Get()->GetSessionHandle();
-
-	bool bState = (NewRadioState == ECheckBoxState::Checked);
-	return SessionService->SetModuleEnabled(SessionHandle, ModuleName, bState);
-}
-
-////////////////////////////////////////////////////////////////////////////////////////////////////
-
-void STimingProfilerToolbar::FillModulesToolbar(FToolBarBuilder& ToolbarBuilder)
-{
-	ToolbarBuilder.BeginSection("Modules");
-	{
-		TSharedRef<Trace::IModuleService> ModuleService = FInsightsManager::Get()->GetModuleService();
-
-		TArray<Trace::FModuleInfo> AvailableModules;
-		ModuleService->GetAvailableModules(AvailableModules);
-
-		for (int32 ModuleIndex = 0; ModuleIndex < AvailableModules.Num(); ++ModuleIndex)
-		{
-			const Trace::FModuleInfo& Module = AvailableModules[ModuleIndex];
-
-			FText Label = FText::FromName(Module.DisplayName);
-			FText ToggleModuleToolTip = FText::Format(LOCTEXT("ToggleModuleToolTip", "Enable/disable {0} trace module (only for live sessions)."), FText::FromName(Module.DisplayName));
-
-			// TODO: Uncomment this when adding icons to toolbar.
-			//ToolbarBuilder.AddToolBarButton(
-			//	FUIAction(
-			//		FExecuteAction::CreateRaw(this, &STimingProfilerToolbar::ToggleModule_Execute, Module.Name),
-			//		FCanExecuteAction::CreateRaw(this, &STimingProfilerToolbar::ToggleModule_CanExecute, Module.Name),
-			//		FIsActionChecked::CreateRaw(this, &STimingProfilerToolbar::ToggleModule_IsChecked, Module.Name)
-			//	),
-			//	NAME_None, // ExtensionHook
-			//	Label,
-			//	ToggleModuleToolTip,
-			//	TAttribute<FSlateIcon>(), // Icon -- empty icon is not really empty :(
-			//	EUserInterfaceActionType::ToggleButton,
-			//	NAME_None); // TutorialHighlightName
-
-			// Workaround for having toggle buttons without icons.
-			ToolbarBuilder.AddWidget(SNew(SCheckBox)
-				.IsEnabled(this, &STimingProfilerToolbar::ToggleModule_CanExecute, Module.Name)
-				.IsChecked(this, &STimingProfilerToolbar::ToggleModule_IsChecked2, Module.Name)
-				.OnCheckStateChanged(this, &STimingProfilerToolbar::ToggleModule_OnCheckStateChanged, Module.Name)
-				.Content()
-				[
-					SNew(STextBlock)
-					.Text(Label)
-					.ToolTip(SNew(SToolTip).Text(ToggleModuleToolTip))
-				]
-			);
-		}
-	}
-	ToolbarBuilder.EndSection();
-}
-
-////////////////////////////////////////////////////////////////////////////////////////////////////
-
 #undef LOCTEXT_NAMESPACE
diff --git a/Engine/Source/Developer/TraceInsights/Private/Insights/Widgets/STimingProfilerToolbar.h b/Engine/Source/Developer/TraceInsights/Private/Insights/Widgets/STimingProfilerToolbar.h
--- a/Engine/Source/Developer/TraceInsights/Private/Insights/Widgets/STimingProfilerToolbar.h
+++ b/Engine/Source/Developer/TraceInsights/Private/Insights/Widgets/STimingProfilerToolbar.h
@@ -1,6 +1,5 @@
 // Copyright Epic Games, Inc. All Rights Reserved.
 
-
 #pragma once
 
 #include "CoreMinimal.h"
@@ -8,9 +7,6 @@
 #include "Widgets/DeclarativeSyntaxSupport.h"
 #include "Widgets/SCompoundWidget.h"
 
-class SBorder;
-class FToolBarBuilder;
-
 /** Ribbon based toolbar used as a main menu in the Profiler window. */
 class STimingProfilerToolbar : public SCompoundWidget
 {
@@ -27,18 +23,11 @@ public:
 	/**
 	 * Construct this widget
 	 *
-	 * @param	InArgs	The declaration data for this widget
+	 * @param InArgs The declaration data for this widget
 	 */
-	void Construct( const FArguments& InArgs );
+	void Construct(const FArguments& InArgs);
 
 private:
 	/** Create the UI commands for the toolbar */
 	void CreateCommands();
-
-	bool ToggleModule_CanExecute(FName ModuleName) const;
-	void ToggleModule_Execute(FName ModuleName);
-	bool ToggleModule_IsChecked(FName ModuleName) const;
-	ECheckBoxState ToggleModule_IsChecked2(FName ModuleName) const;
-	void ToggleModule_OnCheckStateChanged(ECheckBoxState NewRadioState, FName ModuleName);
-	void FillModulesToolbar(FToolBarBuilder& ToolbarBuilder);
 };
diff --git a/Engine/Source/Developer/TraceInsights/Public/Insights/IUnrealInsightsModule.h b/Engine/Source/Developer/TraceInsights/Public/Insights/IUnrealInsightsModule.h
--- a/Engine/Source/Developer/TraceInsights/Public/Insights/IUnrealInsightsModule.h
+++ b/Engine/Source/Developer/TraceInsights/Public/Insights/IUnrealInsightsModule.h
@@ -11,16 +11,31 @@ class IUnrealInsightsModule : public IModuleInterface
 {
 public:
 	/**
-	 * Called when application starts and a new layout is created. It allows the module to create its own areas.
+	 * Called when the application starts in "Browser" mode.
+	 */
+	virtual void CreateSessionBrowser(bool bAllowDebugTools, bool bSingleProcess) = 0;
+
+	/**
+	 * Called when the application starts in "Viewer" mode.
+	 */
+	virtual void CreateSessionViewer(bool bAllowDebugTools) = 0;
+
+	/**
+	 * Starts analysis of the specified *.utrace file. Called when the application starts in "Viewer" mode.
 	 *
-	 * @param NewLayout The newly created layout.
+	 * @param InTraceFile The file path to the *.utrace file to analyze.
 	 */
-	virtual void OnNewLayout(TSharedRef<FTabManager::FLayout> NewLayout) = 0;
+	virtual void StartAnalysisForTraceFile(const TCHAR* InTraceFile) = 0;
 
 	/**
-	 * Called after application layout was restored.
+	 * Starts analysis of the specified session. Called when the application starts in "Viewer" mode.
 	 *
-	 * @param TabManager The tab manager used by Unreal Insights tools.
+	 * @param InSessionId The id of the session to analyze. If nullptr, the app will wait for a live session.
+	 */
+	virtual void StartAnalysisForSession(const TCHAR* InSessionId) = 0;
+
+	/**
+	 * Called when the application shutsdown.
 	 */
-	virtual void OnLayoutRestored(TSharedPtr<FTabManager> TabManager) = 0;
+	virtual void ShutdownUserInterface() = 0;
 };
diff --git a/Engine/Source/Developer/TraceInsights/Public/Insights/Version.h b/Engine/Source/Developer/TraceInsights/Public/Insights/Version.h
--- a/Engine/Source/Developer/TraceInsights/Public/Insights/Version.h
+++ b/Engine/Source/Developer/TraceInsights/Public/Insights/Version.h
@@ -5,10 +5,10 @@
 // See https://semver.org/
 
 #define UNREAL_INSIGHTS_VERSION_MAJOR 0
-#define UNREAL_INSIGHTS_VERSION_MINOR 11
+#define UNREAL_INSIGHTS_VERSION_MINOR 12
 #define UNREAL_INSIGHTS_VERSION_PATCH 0
 
-#define UNREAL_INSIGHTS_VERSION_STRING "0.11"
+#define UNREAL_INSIGHTS_VERSION_STRING "0.12"
 
 /** Extra identifier string (like: "alpha", "beta.1", "beta.2", "test", "featureX", etc.). */
 //#define UNREAL_INSIGHTS_VERSION_ID_STRING "beta"
diff --git a/Engine/Source/Programs/UnrealInsights/Private/UserInterfaceCommand.cpp b/Engine/Source/Programs/UnrealInsights/Private/UserInterfaceCommand.cpp
--- a/Engine/Source/Programs/UnrealInsights/Private/UserInterfaceCommand.cpp
+++ b/Engine/Source/Programs/UnrealInsights/Private/UserInterfaceCommand.cpp
@@ -3,6 +3,7 @@
 #include "UserInterfaceCommand.h"
 
 #include "Async/TaskGraphInterfaces.h"
+//#include "Brushes/SlateImageBrush.h"
 #include "Containers/Ticker.h"
 #include "EditorStyleSet.h"
 #include "Framework/Application/SlateApplication.h"
@@ -15,12 +16,14 @@
 #include "ISourceCodeAccessModule.h"
 #include "Misc/CommandLine.h"
 #include "Misc/ConfigCacheIni.h"
+#include "Misc/Parse.h"
 #include "Modules/ModuleManager.h"
 #include "StandaloneRenderer.h"
 #include "Widgets/Docking/SDockTab.h"
 
 // Insights
 #include "Insights/IUnrealInsightsModule.h"
+#include "Insights/Version.h"
 
 ////////////////////////////////////////////////////////////////////////////////////////////////////
 
@@ -30,7 +33,6 @@
 
 namespace UserInterfaceCommand
 {
-	TSharedPtr<FTabManager::FLayout> ApplicationLayout;
 	TSharedRef<FWorkspaceItem> DeveloperTools = FWorkspaceItem::NewGroup(NSLOCTEXT("UnrealInsights", "DeveloperToolsMenu", "Developer Tools"));
 }
 
@@ -38,8 +40,6 @@ namespace UserInterfaceCommand
 
 void FUserInterfaceCommand::Run()
 {
-	FString UnrealInsightsLayoutIni = FPaths::GetPath(GEngineIni) + "/UnrealInsightsLayout.ini";
-
 	FCoreStyle::ResetToDefault();
 
 	// Load required modules.
@@ -53,7 +53,7 @@ void FUserInterfaceCommand::Run()
 	// Load optional modules.
 	FModuleManager::Get().LoadModule("SettingsEditor");
 
-	InitializeSlateApplication(UnrealInsightsLayoutIni);
+	InitializeSlateApplication();
 
 	// Initialize source code access.
 	// Load the source code access module.
@@ -103,23 +103,27 @@ void FUserInterfaceCommand::Run()
 
 	//im: ??? FCoreDelegates::OnExit.Broadcast();
 
-	ShutdownSlateApplication(UnrealInsightsLayoutIni);
+	ShutdownSlateApplication();
 }
 
 ////////////////////////////////////////////////////////////////////////////////////////////////////
 
-void FUserInterfaceCommand::InitializeSlateApplication(const FString& LayoutIni)
+void FUserInterfaceCommand::InitializeSlateApplication()
 {
 	//TODO: FSlateApplication::InitHighDPI(true);
 
 	// Crank up a normal Slate application using the platform's standalone renderer.
 	FSlateApplication::InitializeAsStandaloneApplication(GetStandardStandaloneRenderer());
 
+	//const FSlateBrush* AppIcon = new FSlateImageBrush(FPaths::EngineContentDir() / "Editor/Slate/Icons/Insights/AppIcon_24x.png", FVector2D(24.0f, 24.0f));
+	//FSlateApplication::Get().SetAppIcon(AppIcon);
+
 	// Menu anims aren't supported. See Runtime\Slate\Private\Framework\Application\MenuStack.cpp.
 	FSlateApplication::Get().EnableMenuAnimations(false);
 
 	// Set the application name.
-	FGlobalTabmanager::Get()->SetApplicationTitle(NSLOCTEXT("UnrealInsights", "AppTitle", "Unreal Insights"));
+	const FText ApplicationTitle = FText::Format(NSLOCTEXT("UnrealInsights", "AppTitle", "Unreal Insights {0}"), FText::FromString(TEXT(UNREAL_INSIGHTS_VERSION_STRING_EX)));
+	FGlobalTabmanager::Get()->SetApplicationTitle(ApplicationTitle);
 
 	// Load widget reflector.
 	const bool bAllowDebugTools = FParse::Param(FCommandLine::Get(), TEXT("DebugTools"));
@@ -128,40 +132,33 @@ void FUserInterfaceCommand::InitializeSlateApplication(const FString& LayoutIni)
 		FModuleManager::LoadModuleChecked<ISlateReflectorModule>("SlateReflector").RegisterTabSpawner(UserInterfaceCommand::DeveloperTools);
 	}
 
-	TSharedRef<FTabManager::FLayout> NewLayout = FTabManager::NewLayout("UnrealInsightsLayout_v1.0");
-
-	// Allow TraceInsights module to update the layout.
 	IUnrealInsightsModule& TraceInsightsModule = FModuleManager::LoadModuleChecked<IUnrealInsightsModule>("TraceInsights");
-	TraceInsightsModule.OnNewLayout(NewLayout);
-
-	// Create area and tab for Slate's WidgetReflector.
-	const float DPIScaleFactor = FPlatformApplicationMisc::GetDPIScaleFactorAtPoint(10.0f, 10.0f);
-	NewLayout->AddArea
-	(
-		FTabManager::NewArea(600.0f * DPIScaleFactor, 600.0f * DPIScaleFactor)
-			->SetWindow(FVector2D(10.0f * DPIScaleFactor, 10.0f * DPIScaleFactor), false)
-			->Split
-			(
-				FTabManager::NewStack()->AddTab("WidgetReflector", bAllowDebugTools ? ETabState::OpenedTab : ETabState::ClosedTab)
-			)
-	);
-
-	// Restore application layout.
-	UserInterfaceCommand::ApplicationLayout = FLayoutSaveRestore::LoadFromConfig(LayoutIni, NewLayout);
-	FGlobalTabmanager::Get()->RestoreFrom(UserInterfaceCommand::ApplicationLayout.ToSharedRef(), TSharedPtr<SWindow>());
-
-	TraceInsightsModule.OnLayoutRestored(FGlobalTabmanager::Get());
+
+	const uint32 MaxPath = FPlatformMisc::GetMaxPathLength();
+	TCHAR* TraceFile = new TCHAR[MaxPath + 1];
+	TraceFile[0] = 0;
+	bool bUseTraceFile = FParse::Value(FCommandLine::Get(), TEXT("-Trace="), TraceFile, MaxPath, true);
+
+	if (bUseTraceFile)
+	{
+		TraceInsightsModule.CreateSessionViewer(bAllowDebugTools);
+		TraceInsightsModule.StartAnalysisForTraceFile(TraceFile);
+	}
+	else
+	{
+		const bool bSingleProcess = FParse::Param(FCommandLine::Get(), TEXT("SingleProcess"));
+		TraceInsightsModule.CreateSessionBrowser(bAllowDebugTools, bSingleProcess);
+	}
+
+	delete[] TraceFile;
 }
 
 ////////////////////////////////////////////////////////////////////////////////////////////////////
 
-void FUserInterfaceCommand::ShutdownSlateApplication(const FString& LayoutIni)
+void FUserInterfaceCommand::ShutdownSlateApplication()
 {
-	check(UserInterfaceCommand::ApplicationLayout.IsValid());
-
-	// Save application layout.
-	FLayoutSaveRestore::SaveToConfig(LayoutIni, UserInterfaceCommand::ApplicationLayout.ToSharedRef());
-	GConfig->Flush(false, LayoutIni);
+	IUnrealInsightsModule& TraceInsightsModule = FModuleManager::LoadModuleChecked<IUnrealInsightsModule>("TraceInsights");
+	TraceInsightsModule.ShutdownUserInterface();
 
 	// Shut down application.
 	FSlateApplication::Shutdown();
diff --git a/Engine/Source/Programs/UnrealInsights/Private/UserInterfaceCommand.h b/Engine/Source/Programs/UnrealInsights/Private/UserInterfaceCommand.h
--- a/Engine/Source/Programs/UnrealInsights/Private/UserInterfaceCommand.h
+++ b/Engine/Source/Programs/UnrealInsights/Private/UserInterfaceCommand.h
@@ -4,34 +4,20 @@
 
 #include "CoreMinimal.h"
 
-class SDockTab;
-class FSpawnTabArgs;
-
 class FUserInterfaceCommand
 {
 public:
-
 	/** Executes the command. */
 	static void Run();
 
 protected:
-
 	/**
 	 * Initializes the Slate application.
-	 *
-	 * @param LayoutInit The path to the layout configuration file.
 	 */
-	static void InitializeSlateApplication(const FString& LayoutIni);
+	static void InitializeSlateApplication();
 
 	/**
 	 * Shuts down the Slate application.
-	 *
-	 * @param LayoutInit The path to the layout configuration file.
 	 */
-	static void ShutdownSlateApplication(const FString& LayoutIni);
-
-private:
-
-	/** Callback for spawning tabs. */
-	TSharedRef<SDockTab> HandleTabManagerSpawnTab(const FSpawnTabArgs& Args, FName TabIdentifier) const;
+	static void ShutdownSlateApplication();
 };
