[Insights] SessionInfo: Removed Platform, App Name, Build Config, Build Traget, Cmd Line and File Size info fields until proper metadata will be exposed for each session.

diff --git a/Engine/Source/Developer/TraceInsights/Private/Insights/Widgets/SSessionInfoWindow.cpp b/Engine/Source/Developer/TraceInsights/Private/Insights/Widgets/SSessionInfoWindow.cpp
--- a/Engine/Source/Developer/TraceInsights/Private/Insights/Widgets/SSessionInfoWindow.cpp
+++ b/Engine/Source/Developer/TraceInsights/Private/Insights/Widgets/SSessionInfoWindow.cpp
@@ -102,12 +102,12 @@ void SSessionInfoWindow::Construct(const FArguments& InArgs)
 
 	AddInfoLine(VerticalBox, LOCTEXT("SessionName_HeaderText",	"Session Name:"),		TAttribute<FText>(this, &SSessionInfoWindow::GetSessionNameText));
 	AddInfoLine(VerticalBox, LOCTEXT("Uri_HeaderText",			"URI:"),				TAttribute<FText>(this, &SSessionInfoWindow::GetUriText));
-	AddInfoLine(VerticalBox, LOCTEXT("Platform_HeaderText",		"Platform:"),			TAttribute<FText>(this, &SSessionInfoWindow::GetPlatformText));
-	AddInfoLine(VerticalBox, LOCTEXT("AppName_HeaderText",		"Application Name:"),	TAttribute<FText>(this, &SSessionInfoWindow::GetAppNameText));
-	AddInfoLine(VerticalBox, LOCTEXT("BuildConfig_HeaderText",	"Build Config:"),		TAttribute<FText>(this, &SSessionInfoWindow::GetBuildConfigText));
-	AddInfoLine(VerticalBox, LOCTEXT("BuildTarget_HeaderText",	"Build Target:"),		TAttribute<FText>(this, &SSessionInfoWindow::GetBuildTargetText));
-	AddInfoLine(VerticalBox, LOCTEXT("CommandLine_HeaderText",	"Command Line:"),		TAttribute<FText>(this, &SSessionInfoWindow::GetCommandLineText));
-	AddInfoLine(VerticalBox, LOCTEXT("FileSize_HeaderText",		"File Size:"),			TAttribute<FText>(this, &SSessionInfoWindow::GetFileSizeText));
+	//AddInfoLine(VerticalBox, LOCTEXT("Platform_HeaderText",		"Platform:"),			TAttribute<FText>(this, &SSessionInfoWindow::GetPlatformText));
+	//AddInfoLine(VerticalBox, LOCTEXT("AppName_HeaderText",		"Application Name:"),	TAttribute<FText>(this, &SSessionInfoWindow::GetAppNameText));
+	//AddInfoLine(VerticalBox, LOCTEXT("BuildConfig_HeaderText",	"Build Config:"),		TAttribute<FText>(this, &SSessionInfoWindow::GetBuildConfigText));
+	//AddInfoLine(VerticalBox, LOCTEXT("BuildTarget_HeaderText",	"Build Target:"),		TAttribute<FText>(this, &SSessionInfoWindow::GetBuildTargetText));
+	//AddInfoLine(VerticalBox, LOCTEXT("CommandLine_HeaderText",	"Command Line:"),		TAttribute<FText>(this, &SSessionInfoWindow::GetCommandLineText));
+	//AddInfoLine(VerticalBox, LOCTEXT("FileSize_HeaderText",		"File Size:"),			TAttribute<FText>(this, &SSessionInfoWindow::GetFileSizeText));
 	AddInfoLine(VerticalBox, LOCTEXT("Status_HeaderText",		"Status:"),				TAttribute<FText>(this, &SSessionInfoWindow::GetStatusText));
 	AddInfoLine(VerticalBox, LOCTEXT("Modules_HeaderText",		"Modules:"),			TAttribute<FText>(this, &SSessionInfoWindow::GetModulesText));
 }
