[Insights] Graph: Fixed auto vertical zoom to allow bigger value ranges (was initially limited to around 10^6, now should be around 10^18 -- for a 100px height graph).

diff --git a/Engine/Source/Developer/TraceInsights/Public/Insights/ViewModels/GraphSeries.h b/Engine/Source/Developer/TraceInsights/Public/Insights/ViewModels/GraphSeries.h
--- a/Engine/Source/Developer/TraceInsights/Public/Insights/ViewModels/GraphSeries.h
+++ b/Engine/Source/Developer/TraceInsights/Public/Insights/ViewModels/GraphSeries.h
@@ -73,7 +73,7 @@ public:
 	 * @return The scale between Value units and viewport units; in pixels (Slate units) / Value unit.
 	 */
 	double GetScaleY() const { return ScaleY; }
-	void SetScaleY(const double InScaleY) { ScaleY = FMath::Max(InScaleY, (double)KINDA_SMALL_NUMBER); }
+	void SetScaleY(const double InScaleY) { ScaleY = FMath::Max(InScaleY, DBL_EPSILON); }
 
 	/**
 	 * @param Value a value; in Value units
