Added methods to extend the layout of Unreal Insights

Added configuration structures for major & minor tabs. Currently minor tabs are only customized in the timing view.
Removed built-in behaviour for session connection. In editor connection is now initialized in the same manner as in the standalone tool (from the GameplayInsights plugin).
Fixed issue with trace services where the analysis service would not be created.

diff --git a/Engine/Source/Developer/TraceInsights/Private/Insights/InsightsManager.cpp b/Engine/Source/Developer/TraceInsights/Private/Insights/InsightsManager.cpp
--- a/Engine/Source/Developer/TraceInsights/Private/Insights/InsightsManager.cpp
+++ b/Engine/Source/Developer/TraceInsights/Private/Insights/InsightsManager.cpp
@@ -13,6 +13,7 @@
 #include "Insights/TimingProfilerManager.h"
 #include "Insights/Widgets/SStartPageWindow.h"
 #include "Insights/Widgets/STimingProfilerWindow.h"
+#include "Insights/IUnrealInsightsModule.h"
 
 ////////////////////////////////////////////////////////////////////////////////////////////////////
 
diff --git a/Engine/Source/Developer/TraceInsights/Private/Insights/InsightsManager.h b/Engine/Source/Developer/TraceInsights/Private/Insights/InsightsManager.h
--- a/Engine/Source/Developer/TraceInsights/Private/Insights/InsightsManager.h
+++ b/Engine/Source/Developer/TraceInsights/Private/Insights/InsightsManager.h
@@ -17,17 +17,6 @@
 
 class SStartPageWindow;
 
-////////////////////////////////////////////////////////////////////////////////////////////////////
-
-struct FInsightsManagerTabs
-{
-	static const FName StartPageTabId;
-	static const FName SessionInfoTabId;
-	static const FName TimingProfilerTabId;
-	static const FName LoadingProfilerTabId;
-	static const FName NetworkingProfilerTabId;
-};
-
 ////////////////////////////////////////////////////////////////////////////////////////////////////
 /**
  * This class manages following areas:
diff --git a/Engine/Source/Developer/TraceInsights/Private/Insights/TimingProfilerManager.cpp b/Engine/Source/Developer/TraceInsights/Private/Insights/TimingProfilerManager.cpp
--- a/Engine/Source/Developer/TraceInsights/Private/Insights/TimingProfilerManager.cpp
+++ b/Engine/Source/Developer/TraceInsights/Private/Insights/TimingProfilerManager.cpp
@@ -16,6 +16,7 @@
 #include "Insights/Widgets/STimerTreeView.h"
 #include "Insights/Widgets/STimingView.h"
 #include "Insights/Widgets/STimingProfilerWindow.h"
+#include "Insights/IUnrealInsightsModule.h"
 
 ////////////////////////////////////////////////////////////////////////////////////////////////////
 
diff --git a/Engine/Source/Developer/TraceInsights/Private/Insights/TraceInsightsModule.cpp b/Engine/Source/Developer/TraceInsights/Private/Insights/TraceInsightsModule.cpp
--- a/Engine/Source/Developer/TraceInsights/Private/Insights/TraceInsightsModule.cpp
+++ b/Engine/Source/Developer/TraceInsights/Private/Insights/TraceInsightsModule.cpp
@@ -1,6 +1,6 @@
 // Copyright Epic Games, Inc. All Rights Reserved.
 
-#include "CoreMinimal.h"
+#include "TraceInsightsModule.h"
 #include "HAL/PlatformApplicationMisc.h"
 #include "Modules/ModuleManager.h"
 #include "TraceServices/ITraceServicesModule.h"
@@ -27,73 +27,6 @@
 #include "Insights/Widgets/SStartPageWindow.h"
 #include "Insights/Widgets/STimingProfilerWindow.h"
 
-////////////////////////////////////////////////////////////////////////////////////////////////////
-/**
- * Implements the Trace Insights module.
- */
-class FTraceInsightsModule : public IUnrealInsightsModule
-{
-public:
-	virtual void StartupModule() override;
-	virtual void ShutdownModule() override;
-
-	virtual bool SupportsDynamicReloading() override
-	{
-		return false;
-	}
-
-	virtual void CreateSessionBrowser(bool bAllowDebugTools, bool bSingleProcess) override;
-	virtual void CreateSessionViewer(bool bAllowDebugTools) override;
-	virtual void StartAnalysisForTraceFile(const TCHAR* InTraceFile) override;
-	virtual void StartAnalysisForSession(const TCHAR* InSessionId) override;
-	virtual void ShutdownUserInterface() override;
-
-protected:
-	void AddAreaForSessionViewer(TSharedRef<FTabManager::FLayout> Layout);
-	void AddAreaForWidgetReflector(TSharedRef<FTabManager::FLayout> Layout, bool bAllowDebugTools);
-
-	/** Callback called when a major tab is closed. */
-	void OnTabBeingClosed(TSharedRef<SDockTab> TabBeingClosed);
-
-	/** Callback called when the Timing Profiler major tab is closed. */
-	void OnTimingProfilerTabBeingClosed(TSharedRef<SDockTab> TabBeingClosed);
-
-	/** Callback called when the Loading Profiler major tab is closed. */
-	void OnLoadingProfilerTabBeingClosed(TSharedRef<SDockTab> TabBeingClosed);
-
-	/** Callback called when the Networking Profiler major tab is closed. */
-	void OnNetworkingProfilerTabBeingClosed(TSharedRef<SDockTab> TabBeingClosed);
-
-	/** Start Page */
-	TSharedRef<SDockTab> SpawnStartPageTab(const FSpawnTabArgs& Args);
-
-	/** Session Info */
-	TSharedRef<SDockTab> SpawnSessionInfoTab(const FSpawnTabArgs& Args);
-
-	/** Timing Profiler */
-	TSharedRef<SDockTab> SpawnTimingProfilerTab(const FSpawnTabArgs& Args);
-
-	/** Loading Profiler */
-	TSharedRef<SDockTab> SpawnLoadingProfilerTab(const FSpawnTabArgs& Args);
-
-	/** Networking Profiler */
-	TSharedRef<SDockTab> SpawnNetworkingProfilerTab(const FSpawnTabArgs& Args);
-
-#if WITH_EDITOR
-	/** Handle exit */
-	void HandleExit();
-#endif
-
-protected:
-	TSharedPtr<FTabManager::FLayout> PersistentLayout;
-	bool bBrowserMode;
-	FString UnrealInsightsLayoutIni;
-
-	TSharedPtr<Trace::IAnalysisService> TraceAnalysisService;
-	TSharedPtr<Trace::ISessionService> TraceSessionService;
-	TSharedPtr<Trace::IModuleService> TraceModuleService;
-};
-
 ////////////////////////////////////////////////////////////////////////////////////////////////////
 
 IMPLEMENT_MODULE(FTraceInsightsModule, TraceInsights);
@@ -125,94 +58,12 @@ void FTraceInsightsModule::StartupModule()
 
 	//////////////////////////////////////////////////
 
-	// Register tab spawner for the Start Page.
-	auto& StartPageTabSpawnerEntry = FGlobalTabmanager::Get()->RegisterNomadTabSpawner(FInsightsManagerTabs::StartPageTabId,
-		FOnSpawnTab::CreateRaw(this, &FTraceInsightsModule::SpawnStartPageTab))
-		.SetDisplayName(NSLOCTEXT("FTraceInsightsModule", "StartPageTabTitle", "Unreal Insights"))
-		.SetTooltipText(NSLOCTEXT("FTraceInsightsModule", "StartPageTooltipText", "Open the start page for Unreal Insights."))
-		.SetIcon(FSlateIcon(FInsightsStyle::GetStyleSetName(), "StartPage.Icon.Small"));
-
-#if WITH_EDITOR
-	StartPageTabSpawnerEntry.SetGroup(WorkspaceMenu::GetMenuStructure().GetDeveloperToolsProfilingCategory());
-#else
-	StartPageTabSpawnerEntry.SetGroup(WorkspaceMenu::GetMenuStructure().GetToolsCategory());
-#endif
-
-	//////////////////////////////////////////////////
-
-	// Register tab spawner for the Session Info.
-	auto& SessionInfoTabSpawnerEntry = FGlobalTabmanager::Get()->RegisterNomadTabSpawner(FInsightsManagerTabs::SessionInfoTabId,
-		FOnSpawnTab::CreateRaw(this, &FTraceInsightsModule::SpawnSessionInfoTab))
-		.SetDisplayName(NSLOCTEXT("FTraceInsightsModule", "SessionInfoTabTitle", "Session Info"))
-		.SetTooltipText(NSLOCTEXT("FTraceInsightsModule", "SessionInfoTooltipText", "Open the Session Info tab."))
-		.SetIcon(FSlateIcon(FInsightsStyle::GetStyleSetName(), "SessionInfo.Icon.Small"));
-
-#if WITH_EDITOR
-	SessionInfoTabSpawnerEntry.SetGroup(WorkspaceMenu::GetMenuStructure().GetDeveloperToolsProfilingCategory());
-#else
-	SessionInfoTabSpawnerEntry.SetGroup(WorkspaceMenu::GetMenuStructure().GetToolsCategory());
-#endif
-
-	//////////////////////////////////////////////////
-
-	// Register tab spawner for the Timing Insights.
-	auto& TimingProfilerTabSpawnerEntry = FGlobalTabmanager::Get()->RegisterNomadTabSpawner(FInsightsManagerTabs::TimingProfilerTabId,
-		FOnSpawnTab::CreateRaw(this, &FTraceInsightsModule::SpawnTimingProfilerTab))
-		.SetDisplayName(NSLOCTEXT("FTraceInsightsModule", "TimingProfilerTabTitle", "Timing Insights"))
-		.SetTooltipText(NSLOCTEXT("FTraceInsightsModule", "TimingProfilerTooltipText", "Open the Timing Insights tab."))
-		.SetIcon(FSlateIcon(FInsightsStyle::GetStyleSetName(), "TimingProfiler.Icon.Small"));
-
-#if WITH_EDITOR
-	TimingProfilerTabSpawnerEntry.SetGroup(WorkspaceMenu::GetMenuStructure().GetDeveloperToolsProfilingCategory());
-#else
-	TimingProfilerTabSpawnerEntry.SetGroup(WorkspaceMenu::GetMenuStructure().GetToolsCategory());
-#endif
-
-	//////////////////////////////////////////////////
-
-	// Register tab spawner for the Asset Loading Insights.
-	auto& LoadingProfilerTabSpawnerEntry = FGlobalTabmanager::Get()->RegisterNomadTabSpawner(FInsightsManagerTabs::LoadingProfilerTabId,
-		FOnSpawnTab::CreateRaw(this, &FTraceInsightsModule::SpawnLoadingProfilerTab))
-		.SetDisplayName(NSLOCTEXT("FTraceInsightsModule", "LoadingProfilerTabTitle", "Asset Loading Insights"))
-		.SetTooltipText(NSLOCTEXT("FTraceInsightsModule", "LoadingProfilerTooltipText", "Open the Asset Loading Insights tab."))
-		.SetIcon(FSlateIcon(FInsightsStyle::GetStyleSetName(), "LoadingProfiler.Icon.Small"));
-
-#if WITH_EDITOR
-	LoadingProfilerTabSpawnerEntry.SetGroup(WorkspaceMenu::GetMenuStructure().GetDeveloperToolsProfilingCategory());
-#else
-	LoadingProfilerTabSpawnerEntry.SetGroup(WorkspaceMenu::GetMenuStructure().GetToolsCategory());
-#endif
-
-	//////////////////////////////////////////////////
-
-	// Register tab spawner for the Networking Insights.
-	auto& NetworkingProfilerTabSpawnerEntry = FGlobalTabmanager::Get()->RegisterNomadTabSpawner(FInsightsManagerTabs::NetworkingProfilerTabId,
-		FOnSpawnTab::CreateRaw(this, &FTraceInsightsModule::SpawnNetworkingProfilerTab))
-		.SetReuseTabMethod(FOnFindTabToReuse::CreateStatic(&NeverReuse))
-		.SetDisplayName(NSLOCTEXT("FTraceInsightsModule", "NetworkingProfilerTabTitle", "Networking Insights"))
-		.SetTooltipText(NSLOCTEXT("FTraceInsightsModule", "NetworkingProfilerTooltipText", "Open the Networking Insights tab."))
-		.SetIcon(FSlateIcon(FInsightsStyle::GetStyleSetName(), "NetworkingProfiler.Icon.Small"));
-
-#if WITH_EDITOR
-	NetworkingProfilerTabSpawnerEntry.SetGroup(WorkspaceMenu::GetMenuStructure().GetDeveloperToolsProfilingCategory());
-#else
-	NetworkingProfilerTabSpawnerEntry.SetGroup(WorkspaceMenu::GetMenuStructure().GetToolsCategory());
-#endif
-
-	//////////////////////////////////////////////////
-
 #if WITH_EDITOR
 	if (TraceSessionService.IsValid())
 	{
 		TraceSessionService->StartRecorderServer();
 	}
 
-	// In editor, load the layout and display it here. In the standalone tool this is done through CreateSessionBrowser or CreateSessionViewer.
-	TSharedRef<FTabManager::FLayout> DefaultLayout = FTabManager::NewLayout("UnrealInsightsLayout_v1.0");
-	AddAreaForSessionViewer(DefaultLayout);
-	PersistentLayout = FLayoutSaveRestore::LoadFromConfig(GEditorLayoutIni, DefaultLayout);
-	FGlobalTabmanager::Get()->RestoreFrom(PersistentLayout.ToSharedRef(), TSharedPtr<SWindow>());
-
 	FCoreDelegates::OnExit.AddRaw(this, &FTraceInsightsModule::HandleExit);
 #endif
 }
@@ -281,6 +132,79 @@ void FTraceInsightsModule::HandleExit()
 
 ////////////////////////////////////////////////////////////////////////////////////////////////////
 
+void FTraceInsightsModule::RegisterTabSpawners()
+{
+	TSharedRef<FWorkspaceItem> ToolsCategory = WorkspaceMenu::GetMenuStructure().GetToolsCategory(); 
+
+	const FInsightsMajorTabConfig& StartPageConfig = FindMajorTabConfig(FInsightsManagerTabs::StartPageTabId);
+	if(StartPageConfig.bIsAvailable)
+	{
+		// Register tab spawner for the Start Page.
+		FTabSpawnerEntry& StartPageTabSpawnerEntry = FGlobalTabmanager::Get()->RegisterNomadTabSpawner(FInsightsManagerTabs::StartPageTabId,
+			FOnSpawnTab::CreateRaw(this, &FTraceInsightsModule::SpawnStartPageTab))
+			.SetDisplayName(StartPageConfig.TabLabel.IsSet() ? StartPageConfig.TabLabel.GetValue() : NSLOCTEXT("FTraceInsightsModule", "StartPageTabTitle", "Unreal Insights"))
+			.SetTooltipText(StartPageConfig.TabTooltip.IsSet() ? StartPageConfig.TabTooltip.GetValue() : NSLOCTEXT("FTraceInsightsModule", "StartPageTooltipText", "Open the start page for Unreal Insights."))
+			.SetIcon(StartPageConfig.TabIcon.IsSet() ? StartPageConfig.TabIcon.GetValue() : FSlateIcon(FInsightsStyle::GetStyleSetName(), "StartPage.Icon.Small"));
+
+		StartPageTabSpawnerEntry.SetGroup(StartPageConfig.WorkspaceGroup.IsValid() ? StartPageConfig.WorkspaceGroup.ToSharedRef() : ToolsCategory);
+	}
+
+	const FInsightsMajorTabConfig& SessionInfoConfig = FindMajorTabConfig(FInsightsManagerTabs::SessionInfoTabId);
+	if(SessionInfoConfig.bIsAvailable)
+	{
+		// Register tab spawner for the Session Info.
+		FTabSpawnerEntry& SessionInfoTabSpawnerEntry = FGlobalTabmanager::Get()->RegisterNomadTabSpawner(FInsightsManagerTabs::SessionInfoTabId,
+			FOnSpawnTab::CreateRaw(this, &FTraceInsightsModule::SpawnSessionInfoTab))
+			.SetDisplayName(NSLOCTEXT("FTraceInsightsModule", "SessionInfoTabTitle", "Session Info"))
+			.SetTooltipText(NSLOCTEXT("FTraceInsightsModule", "SessionInfoTooltipText", "Open the Session Info tab."))
+			.SetIcon(FSlateIcon(FInsightsStyle::GetStyleSetName(), "SessionInfo.Icon.Small"));
+
+		SessionInfoTabSpawnerEntry.SetGroup(SessionInfoConfig.WorkspaceGroup.IsValid() ? SessionInfoConfig.WorkspaceGroup.ToSharedRef() : ToolsCategory);
+	}
+
+	const FInsightsMajorTabConfig& TimingProfilerConfig = FindMajorTabConfig(FInsightsManagerTabs::TimingProfilerTabId);
+	if(TimingProfilerConfig.bIsAvailable)
+	{
+		// Register tab spawner for the Timing Insights.
+		FTabSpawnerEntry& TimingProfilerTabSpawnerEntry = FGlobalTabmanager::Get()->RegisterNomadTabSpawner(FInsightsManagerTabs::TimingProfilerTabId,
+			FOnSpawnTab::CreateRaw(this, &FTraceInsightsModule::SpawnTimingProfilerTab))
+			.SetDisplayName(TimingProfilerConfig.TabLabel.IsSet() ? TimingProfilerConfig.TabLabel.GetValue() : NSLOCTEXT("FTraceInsightsModule", "TimingProfilerTabTitle", "Timing Insights"))
+			.SetTooltipText(TimingProfilerConfig.TabTooltip.IsSet() ? TimingProfilerConfig.TabTooltip.GetValue() : NSLOCTEXT("FTraceInsightsModule", "TimingProfilerTooltipText", "Open the Timing Insights tab."))
+			.SetIcon(TimingProfilerConfig.TabIcon.IsSet() ? TimingProfilerConfig.TabIcon.GetValue() : FSlateIcon(FInsightsStyle::GetStyleSetName(), "TimingProfiler.Icon.Small"));
+
+		TimingProfilerTabSpawnerEntry.SetGroup(TimingProfilerConfig.WorkspaceGroup.IsValid() ? TimingProfilerConfig.WorkspaceGroup.ToSharedRef() : ToolsCategory);
+	}
+
+	const FInsightsMajorTabConfig& LoadingProfilerConfig = FindMajorTabConfig(FInsightsManagerTabs::LoadingProfilerTabId);
+	if(LoadingProfilerConfig.bIsAvailable)
+	{
+		// Register tab spawner for the Asset Loading Insights.
+		FTabSpawnerEntry& LoadingProfilerTabSpawnerEntry = FGlobalTabmanager::Get()->RegisterNomadTabSpawner(FInsightsManagerTabs::LoadingProfilerTabId,
+			FOnSpawnTab::CreateRaw(this, &FTraceInsightsModule::SpawnLoadingProfilerTab))
+			.SetDisplayName(LoadingProfilerConfig.TabLabel.IsSet() ? LoadingProfilerConfig.TabLabel.GetValue() : NSLOCTEXT("FTraceInsightsModule", "LoadingProfilerTabTitle", "Asset Loading Insights"))
+			.SetTooltipText(LoadingProfilerConfig.TabTooltip.IsSet() ? LoadingProfilerConfig.TabTooltip.GetValue() : NSLOCTEXT("FTraceInsightsModule", "LoadingProfilerTooltipText", "Open the Asset Loading Insights tab."))
+			.SetIcon(LoadingProfilerConfig.TabIcon.IsSet() ? LoadingProfilerConfig.TabIcon.GetValue() : FSlateIcon(FInsightsStyle::GetStyleSetName(), "LoadingProfiler.Icon.Small"));
+
+		LoadingProfilerTabSpawnerEntry.SetGroup(LoadingProfilerConfig.WorkspaceGroup.IsValid() ? LoadingProfilerConfig.WorkspaceGroup.ToSharedRef() : ToolsCategory);
+	}
+
+	const FInsightsMajorTabConfig& NetworkingProfilerConfig = FindMajorTabConfig(FInsightsManagerTabs::NetworkingProfilerTabId);
+	if(NetworkingProfilerConfig.bIsAvailable)
+	{
+		// Register tab spawner for the Networking Insights.
+		FTabSpawnerEntry& NetworkingProfilerTabSpawnerEntry = FGlobalTabmanager::Get()->RegisterNomadTabSpawner(FInsightsManagerTabs::NetworkingProfilerTabId,
+			FOnSpawnTab::CreateRaw(this, &FTraceInsightsModule::SpawnNetworkingProfilerTab))
+			.SetReuseTabMethod(FOnFindTabToReuse::CreateStatic(&NeverReuse))
+			.SetDisplayName(NetworkingProfilerConfig.TabLabel.IsSet() ? NetworkingProfilerConfig.TabLabel.GetValue() : NSLOCTEXT("FTraceInsightsModule", "NetworkingProfilerTabTitle", "Networking Insights"))
+			.SetTooltipText(NetworkingProfilerConfig.TabTooltip.IsSet() ? NetworkingProfilerConfig.TabTooltip.GetValue() : NSLOCTEXT("FTraceInsightsModule", "NetworkingProfilerTooltipText", "Open the Networking Insights tab."))
+			.SetIcon(NetworkingProfilerConfig.TabIcon.IsSet() ? NetworkingProfilerConfig.TabIcon.GetValue() : FSlateIcon(FInsightsStyle::GetStyleSetName(), "NetworkingProfiler.Icon.Small"));
+
+		NetworkingProfilerTabSpawnerEntry.SetGroup(NetworkingProfilerConfig.WorkspaceGroup.IsValid() ? NetworkingProfilerConfig.WorkspaceGroup.ToSharedRef() : ToolsCategory);
+	}
+}
+
+////////////////////////////////////////////////////////////////////////////////////////////////////
+
 void FTraceInsightsModule::CreateSessionBrowser(bool bAllowDebugTools, bool bSingleProcess)
 {
 	bBrowserMode = true;
@@ -291,6 +215,10 @@ void FTraceInsightsModule::CreateSessionBrowser(bool bAllowDebugTools, bool bSin
 		TraceSessionService->StartRecorderServer();
 	}
 
+	RegisterTabSpawners();
+
+	TSharedRef<FWorkspaceItem> ToolsCategory = WorkspaceMenu::GetMenuStructure().GetToolsCategory(); 
+
 	const float DPIScaleFactor = FPlatformApplicationMisc::GetDPIScaleFactorAtPoint(10.0f, 10.0f);
 
 	TSharedRef<FTabManager::FLayout> DefaultLayout = FTabManager::NewLayout("TraceSessionBrowserLayout_v1.0");
@@ -328,10 +256,14 @@ void FTraceInsightsModule::CreateSessionViewer(bool bAllowDebugTools)
 {
 	bBrowserMode = false;
 
+	RegisterTabSpawners();
+
 	TSharedRef<FTabManager::FLayout> DefaultLayout = FTabManager::NewLayout("UnrealInsightsLayout_v1.0");
 
 	AddAreaForSessionViewer(DefaultLayout);
+#if !WITH_EDITOR
 	AddAreaForWidgetReflector(DefaultLayout, bAllowDebugTools);
+#endif
 
 	// Restore application layout.
 	UnrealInsightsLayoutIni = FPaths::GetPath(GEngineIni) + "/UnrealInsightsLayout.ini";
@@ -399,6 +331,34 @@ void FTraceInsightsModule::ShutdownUserInterface()
 
 ////////////////////////////////////////////////////////////////////////////////////////////////////
 
+void FTraceInsightsModule::RegisterMajorTabConfig(const FName& InMajorTabId, const FInsightsMajorTabConfig& InConfig)
+{
+	TabConfigs.Add(InMajorTabId, InConfig);
+}
+
+////////////////////////////////////////////////////////////////////////////////////////////////////
+
+void FTraceInsightsModule::UnregisterMajorTabConfig(const FName& InMajorTabId)
+{
+	TabConfigs.Remove(InMajorTabId);
+}
+
+////////////////////////////////////////////////////////////////////////////////////////////////////
+
+const FInsightsMajorTabConfig& FTraceInsightsModule::FindMajorTabConfig(const FName& InMajorTabId) const
+{
+	const FInsightsMajorTabConfig* FoundConfig = TabConfigs.Find(InMajorTabId);
+	if(FoundConfig != nullptr)
+	{
+		return *FoundConfig;
+	}
+
+	static FInsightsMajorTabConfig DefaultConfig;
+	return DefaultConfig;
+}
+
+////////////////////////////////////////////////////////////////////////////////////////////////////
+
 void FTraceInsightsModule::StartAnalysisForTraceFile(const TCHAR* InTraceFile)
 {
 	if (InTraceFile != nullptr)
@@ -419,6 +379,13 @@ void FTraceInsightsModule::StartAnalysisForSession(const TCHAR* InSessionId)
 
 ////////////////////////////////////////////////////////////////////////////////////////////////////
 
+void FTraceInsightsModule::StartAnalysisForLastLiveSession()
+{
+	FInsightsManager::Get()->LoadLastLiveSession();
+}
+
+////////////////////////////////////////////////////////////////////////////////////////////////////
+
 TSharedRef<SDockTab> FTraceInsightsModule::SpawnStartPageTab(const FSpawnTabArgs& Args)
 {
 	const TSharedRef<SDockTab> DockTab = SNew(SDockTab)
diff --git a/Engine/Source/Developer/TraceInsights/Private/Insights/TraceInsightsModule.h b/Engine/Source/Developer/TraceInsights/Private/Insights/TraceInsightsModule.h
new file mode 100644
--- /dev/null
+++ b/Engine/Source/Developer/TraceInsights/Private/Insights/TraceInsightsModule.h
@@ -0,0 +1,96 @@
+// Copyright Epic Games, Inc. All Rights Reserved
+
+#pragma once
+
+#include "CoreMinimal.h"
+#include "Insights/IUnrealInsightsModule.h"
+#include "Framework/Docking/TabManager.h"
+
+namespace Trace 
+{ 
+	class IAnalysisService;
+	class ISessionService;
+	class IModuleService;
+}
+
+class SDockTab;
+class FSpawnTabArgs;
+
+/**
+ * Implements the Trace Insights module.
+ */
+class FTraceInsightsModule : public IUnrealInsightsModule
+{
+public:
+	virtual void StartupModule() override;
+	virtual void ShutdownModule() override;
+
+	virtual bool SupportsDynamicReloading() override
+	{
+		return false;
+	}
+
+	virtual void CreateSessionBrowser(bool bAllowDebugTools, bool bSingleProcess) override;
+	virtual void CreateSessionViewer(bool bAllowDebugTools) override;
+	virtual void StartAnalysisForTraceFile(const TCHAR* InTraceFile) override;
+	virtual void StartAnalysisForSession(const TCHAR* InSessionId) override;
+	virtual void StartAnalysisForLastLiveSession() override;
+	virtual void ShutdownUserInterface() override;
+	virtual void RegisterMajorTabConfig(const FName& InMajorTabId, const FInsightsMajorTabConfig& InConfig) override;
+	virtual void UnregisterMajorTabConfig(const FName& InMajorTabId) override;
+	virtual FOnInsightsMajorTabCreated& OnMajorTabCreated() override { return OnInsightsMajorTabCreatedDelegate; }
+
+	/** Find a major tab config for the specified ID */
+	const FInsightsMajorTabConfig& FindMajorTabConfig(const FName& InMajorTabId) const;
+
+protected:
+	void RegisterTabSpawners();
+
+	void AddAreaForSessionViewer(TSharedRef<FTabManager::FLayout> Layout);
+	void AddAreaForWidgetReflector(TSharedRef<FTabManager::FLayout> Layout, bool bAllowDebugTools);
+
+	/** Callback called when a major tab is closed. */
+	void OnTabBeingClosed(TSharedRef<SDockTab> TabBeingClosed);
+
+	/** Callback called when the Timing Profiler major tab is closed. */
+	void OnTimingProfilerTabBeingClosed(TSharedRef<SDockTab> TabBeingClosed);
+
+	/** Callback called when the Loading Profiler major tab is closed. */
+	void OnLoadingProfilerTabBeingClosed(TSharedRef<SDockTab> TabBeingClosed);
+
+	/** Callback called when the Networking Profiler major tab is closed. */
+	void OnNetworkingProfilerTabBeingClosed(TSharedRef<SDockTab> TabBeingClosed);
+
+	/** Start Page */
+	TSharedRef<SDockTab> SpawnStartPageTab(const FSpawnTabArgs& Args);
+
+	/** Session Info */
+	TSharedRef<SDockTab> SpawnSessionInfoTab(const FSpawnTabArgs& Args);
+
+	/** Timing Profiler */
+	TSharedRef<SDockTab> SpawnTimingProfilerTab(const FSpawnTabArgs& Args);
+
+	/** Loading Profiler */
+	TSharedRef<SDockTab> SpawnLoadingProfilerTab(const FSpawnTabArgs& Args);
+
+	/** Networking Profiler */
+	TSharedRef<SDockTab> SpawnNetworkingProfilerTab(const FSpawnTabArgs& Args);
+
+#if WITH_EDITOR
+	/** Handle exit */
+	void HandleExit();
+#endif
+
+protected:
+	TSharedPtr<FTabManager::FLayout> PersistentLayout;
+	bool bBrowserMode;
+	FString UnrealInsightsLayoutIni;
+
+	TSharedPtr<Trace::IAnalysisService> TraceAnalysisService;
+	TSharedPtr<Trace::ISessionService> TraceSessionService;
+	TSharedPtr<Trace::IModuleService> TraceModuleService;
+
+	TMap<FName, FInsightsMajorTabConfig> TabConfigs;
+
+	FOnInsightsMajorTabCreated OnInsightsMajorTabCreatedDelegate;
+};
\ No newline at end of file
diff --git a/Engine/Source/Developer/TraceInsights/Private/Insights/Widgets/SStartPageWindow.cpp b/Engine/Source/Developer/TraceInsights/Private/Insights/Widgets/SStartPageWindow.cpp
--- a/Engine/Source/Developer/TraceInsights/Private/Insights/Widgets/SStartPageWindow.cpp
+++ b/Engine/Source/Developer/TraceInsights/Private/Insights/Widgets/SStartPageWindow.cpp
@@ -557,12 +557,6 @@ void SStartPageWindow::Construct(const FArguments& InArgs)
 			.Expose(OverlaySettingsSlot)
 		];
 
-#if WITH_EDITOR
-	// In the editor, attempt to connect to ourselves so we can choose our own session.
-	TSharedRef<Trace::ISessionService> SessionService = FInsightsManager::Get()->GetSessionService();
-	SessionService->ConnectSession(TEXT("127.0.0.1"));
-#endif
-
 	RefreshTraceSessionList();
 
 	FSlateApplication::Get().SetKeyboardFocus(TraceSessionsListView);
diff --git a/Engine/Source/Developer/TraceInsights/Private/Insights/Widgets/STimingProfilerWindow.cpp b/Engine/Source/Developer/TraceInsights/Private/Insights/Widgets/STimingProfilerWindow.cpp
--- a/Engine/Source/Developer/TraceInsights/Private/Insights/Widgets/STimingProfilerWindow.cpp
+++ b/Engine/Source/Developer/TraceInsights/Private/Insights/Widgets/STimingProfilerWindow.cpp
@@ -34,6 +34,7 @@
 #include "Insights/Widgets/STimerTreeView.h"
 #include "Insights/Widgets/STimingProfilerToolbar.h"
 #include "Insights/Widgets/STimingView.h"
+#include "Insights/TraceInsightsModule.h"
 
 ////////////////////////////////////////////////////////////////////////////////////////////////////
 
@@ -366,89 +367,119 @@ void STimingProfilerWindow::Construct(const FArguments& InArgs, const TSharedRef
 	TSharedPtr<FTimingProfilerManager> TimingProfilerManager = FTimingProfilerManager::Get();
 	ensure(TimingProfilerManager.IsValid());
 
-	// Create tab layout.
-	const TSharedRef<FTabManager::FLayout> Layout = FTabManager::NewLayout("InsightsTimingProfilerLayout_v1.0")
-		->AddArea
-		(
-			FTabManager::NewPrimaryArea()
-			->SetOrientation(Orient_Vertical)
-			->Split
-			(
-				FTabManager::NewStack()
-				->AddTab(FTimingProfilerTabs::ToolbarID, ETabState::OpenedTab)
-				->SetHideTabWell(true)
-			)
-			->Split
+	TSharedPtr<FTabManager::FLayout> Layout;
+
+	// Check for layout overrides
+	FTraceInsightsModule& TraceInsightsModule = FModuleManager::GetModuleChecked<FTraceInsightsModule>("TraceInsights");
+	FInsightsMajorTabConfig TabConfig = TraceInsightsModule.FindMajorTabConfig(FInsightsManagerTabs::TimingProfilerTabId);
+
+	// Register any new minor tabs
+	for(const FInsightsMinorTabConfig& MinorTabConfig : TabConfig.MinorTabs)
+	{
+		FTabSpawnerEntry& TabSpawnerEntry = TabManager->RegisterTabSpawner(MinorTabConfig.TabId, MinorTabConfig.OnSpawnTab);
+
+		TabSpawnerEntry
+		.SetDisplayName(MinorTabConfig.TabLabel)
+		.SetTooltipText(MinorTabConfig.TabTooltip)
+		.SetIcon(MinorTabConfig.TabIcon)
+		.SetReuseTabMethod(MinorTabConfig.OnFindTabToReuse);
+
+		if(MinorTabConfig.WorkspaceGroup.IsValid())
+		{
+			TabSpawnerEntry.SetGroup(MinorTabConfig.WorkspaceGroup.ToSharedRef());
+		}
+	}
+
+	if(TabConfig.Layout.IsValid())
+	{
+		Layout = TabConfig.Layout;
+	}
+	else
+	{
+		// Create tab layout.
+		Layout = FTabManager::NewLayout("InsightsTimingProfilerLayout_v1.0")
+			->AddArea
 			(
-				FTabManager::NewSplitter()
-				->SetOrientation(Orient_Horizontal)
-				->SetSizeCoefficient(1.0f)
+				FTabManager::NewPrimaryArea()
+				->SetOrientation(Orient_Vertical)
 				->Split
 				(
-					FTabManager::NewSplitter()
-					->SetOrientation(Orient_Vertical)
-					->SetSizeCoefficient(0.65f)
-					->Split
-					(
-						FTabManager::NewStack()
-						->SetSizeCoefficient(0.1f)
-						->SetHideTabWell(true)
-						->AddTab(FTimingProfilerTabs::FramesTrackID, TimingProfilerManager->IsFramesTrackVisible() ? ETabState::OpenedTab : ETabState::ClosedTab)
-					)
-					->Split
-					(
-						FTabManager::NewStack()
-						->SetSizeCoefficient(0.5f)
-						->SetHideTabWell(true)
-						->AddTab(FTimingProfilerTabs::TimingViewID, TimingProfilerManager->IsTimingViewVisible() ? ETabState::OpenedTab : ETabState::ClosedTab)
-					)
-					->Split
-					(
-						FTabManager::NewStack()
-						->SetSizeCoefficient(0.2f)
-						->SetHideTabWell(true)
-						->AddTab(FTimingProfilerTabs::LogViewID, TimingProfilerManager->IsLogViewVisible() ? ETabState::OpenedTab : ETabState::ClosedTab)
-					)
+					FTabManager::NewStack()
+					->AddTab(FTimingProfilerTabs::ToolbarID, ETabState::OpenedTab)
+					->SetHideTabWell(true)
 				)
 				->Split
 				(
 					FTabManager::NewSplitter()
-					->SetOrientation(Orient_Vertical)
-					->SetSizeCoefficient(0.35f)
+					->SetOrientation(Orient_Horizontal)
+					->SetSizeCoefficient(1.0f)
 					->Split
 					(
-						FTabManager::NewStack()
-						->SetSizeCoefficient(0.67f)
-						->AddTab(FTimingProfilerTabs::TimersID, TimingProfilerManager->IsTimersViewVisible() ? ETabState::OpenedTab : ETabState::ClosedTab)
-						->AddTab(FTimingProfilerTabs::StatsCountersID, TimingProfilerManager->IsStatsCountersViewVisible() ? ETabState::OpenedTab : ETabState::ClosedTab)
-						->SetForegroundTab(FTimingProfilerTabs::TimersID)
+						FTabManager::NewSplitter()
+						->SetOrientation(Orient_Vertical)
+						->SetSizeCoefficient(0.65f)
+						->Split
+						(
+							FTabManager::NewStack()
+							->SetSizeCoefficient(0.1f)
+							->SetHideTabWell(true)
+							->AddTab(FTimingProfilerTabs::FramesTrackID, TimingProfilerManager->IsFramesTrackVisible() ? ETabState::OpenedTab : ETabState::ClosedTab)
+						)
+						->Split
+						(
+							FTabManager::NewStack()
+							->SetSizeCoefficient(0.5f)
+							->SetHideTabWell(true)
+							->AddTab(FTimingProfilerTabs::TimingViewID, TimingProfilerManager->IsTimingViewVisible() ? ETabState::OpenedTab : ETabState::ClosedTab)
+						)
+						->Split
+						(
+							FTabManager::NewStack()
+							->SetSizeCoefficient(0.2f)
+							->SetHideTabWell(true)
+							->AddTab(FTimingProfilerTabs::LogViewID, TimingProfilerManager->IsLogViewVisible() ? ETabState::OpenedTab : ETabState::ClosedTab)
+						)
 					)
 					->Split
 					(
-						FTabManager::NewStack()
-						->SetSizeCoefficient(0.165f)
-						->SetHideTabWell(true)
-						->AddTab(FTimingProfilerTabs::CallersID, TimingProfilerManager->IsCallersTreeViewVisible() ? ETabState::OpenedTab : ETabState::ClosedTab)
-					)
-					->Split
-					(
-						FTabManager::NewStack()
-						->SetSizeCoefficient(0.165f)
-						->SetHideTabWell(true)
-						->AddTab(FTimingProfilerTabs::CalleesID, TimingProfilerManager->IsCalleesTreeViewVisible() ? ETabState::OpenedTab : ETabState::ClosedTab)
+						FTabManager::NewSplitter()
+						->SetOrientation(Orient_Vertical)
+						->SetSizeCoefficient(0.35f)
+						->Split
+						(
+							FTabManager::NewStack()
+							->SetSizeCoefficient(0.67f)
+							->AddTab(FTimingProfilerTabs::TimersID, TimingProfilerManager->IsTimersViewVisible() ? ETabState::OpenedTab : ETabState::ClosedTab)
+							->AddTab(FTimingProfilerTabs::StatsCountersID, TimingProfilerManager->IsStatsCountersViewVisible() ? ETabState::OpenedTab : ETabState::ClosedTab)
+							->SetForegroundTab(FTimingProfilerTabs::TimersID)
+						)
+						->Split
+						(
+							FTabManager::NewStack()
+							->SetSizeCoefficient(0.165f)
+							->SetHideTabWell(true)
+							->AddTab(FTimingProfilerTabs::CallersID, TimingProfilerManager->IsCallersTreeViewVisible() ? ETabState::OpenedTab : ETabState::ClosedTab)
+						)
+						->Split
+						(
+							FTabManager::NewStack()
+							->SetSizeCoefficient(0.165f)
+							->SetHideTabWell(true)
+							->AddTab(FTimingProfilerTabs::CalleesID, TimingProfilerManager->IsCalleesTreeViewVisible() ? ETabState::OpenedTab : ETabState::ClosedTab)
+						)
 					)
 				)
-			)
-		);
+			);
+	}
 
 	// Create & initialize main menu.
-	FMenuBarBuilder MenuBarBuilder = FMenuBarBuilder(TSharedPtr<FUICommandList>());
+	FMenuBarBuilder MenuBarBuilder = FMenuBarBuilder(TSharedPtr<FUICommandList>(), TabConfig.MenuExtender);
 
 	MenuBarBuilder.AddPullDownMenu(
-		LOCTEXT("MenuLabel", "MENU"),
+		LOCTEXT("MenuLabel", "Menu"),
 		FText::GetEmpty(),
 		FNewMenuDelegate::CreateStatic(&STimingProfilerWindow::FillMenu, TabManager),
-		FName(TEXT("MENU"))
+		FName(TEXT("Menu"))
 	);
 
 	ChildSlot
@@ -483,7 +514,7 @@ void STimingProfilerWindow::Construct(const FArguments& InArgs, const TSharedRef
 					+ SVerticalBox::Slot()
 						.FillHeight(1.0f)
 						[
-							TabManager->RestoreFrom(Layout, ConstructUnderWindow).ToSharedRef()
+							TabManager->RestoreFrom(Layout.ToSharedRef(), ConstructUnderWindow).ToSharedRef()
 						]
 				]
 
@@ -505,6 +536,9 @@ void STimingProfilerWindow::Construct(const FArguments& InArgs, const TSharedRef
 
 	// Tell tab-manager about the global menu bar.
 	TabManager->SetMenuMultiBox(MenuBarBuilder.GetMultiBox());
+
+	// Tell clients about creation
+	TraceInsightsModule.OnMajorTabCreated().Broadcast(FInsightsManagerTabs::TimingProfilerTabId, TabManager.ToSharedRef());
 }
 
 END_SLATE_FUNCTION_BUILD_OPTIMIZATION
diff --git a/Engine/Source/Developer/TraceInsights/Private/Insights/Widgets/STimingProfilerWindow.h b/Engine/Source/Developer/TraceInsights/Private/Insights/Widgets/STimingProfilerWindow.h
--- a/Engine/Source/Developer/TraceInsights/Private/Insights/Widgets/STimingProfilerWindow.h
+++ b/Engine/Source/Developer/TraceInsights/Private/Insights/Widgets/STimingProfilerWindow.h
@@ -30,21 +30,6 @@ class SVerticalBox;
 
 ////////////////////////////////////////////////////////////////////////////////////////////////////
 
-struct FTimingProfilerTabs
-{
-	// Tab identifiers
-	static const FName ToolbarID;
-	static const FName FramesTrackID;
-	static const FName TimingViewID;
-	static const FName TimersID;
-	static const FName CallersID;
-	static const FName CalleesID;
-	static const FName StatsCountersID;
-	static const FName LogViewID;
-};
-
-////////////////////////////////////////////////////////////////////////////////////////////////////
-
 /** Implements the timing profiler window. */
 class STimingProfilerWindow : public SCompoundWidget
 {
diff --git a/Engine/Source/Developer/TraceInsights/Private/Insights/Widgets/STimingView.h b/Engine/Source/Developer/TraceInsights/Private/Insights/Widgets/STimingView.h
--- a/Engine/Source/Developer/TraceInsights/Private/Insights/Widgets/STimingView.h
+++ b/Engine/Source/Developer/TraceInsights/Private/Insights/Widgets/STimingView.h
@@ -228,6 +228,8 @@ public:
 	virtual Insights::FSelectedTrackChangedDelegate& OnSelectedTrackChanged() override { return OnSelectedTrackChangedDelegate; }
 	virtual Insights::FSelectedEventChangedDelegate& OnSelectedEventChanged() override { return OnSelectedEventChangedDelegate; }
 
+	virtual double GetTimeMarker() const override { return TimeMarker; }
+
 	////////////////////////////////////////////////////////////////////////////////////////////////////
 
 	void UpdateScrollableTracksOrder();
diff --git a/Engine/Source/Developer/TraceInsights/Public/Insights/ITimingViewSession.h b/Engine/Source/Developer/TraceInsights/Public/Insights/ITimingViewSession.h
--- a/Engine/Source/Developer/TraceInsights/Public/Insights/ITimingViewSession.h
+++ b/Engine/Source/Developer/TraceInsights/Public/Insights/ITimingViewSession.h
@@ -79,6 +79,9 @@ public:
 
 	/** Gets the delegate to be invoked when the selected timing event has changed. */
 	virtual FSelectedEventChangedDelegate& OnSelectedEventChanged() = 0;
+
+	/** Get the current marker time */
+	virtual double GetTimeMarker() const = 0;
 };
 
 } // namespace Insights
diff --git a/Engine/Source/Developer/TraceInsights/Public/Insights/IUnrealInsightsModule.h b/Engine/Source/Developer/TraceInsights/Public/Insights/IUnrealInsightsModule.h
--- a/Engine/Source/Developer/TraceInsights/Public/Insights/IUnrealInsightsModule.h
+++ b/Engine/Source/Developer/TraceInsights/Public/Insights/IUnrealInsightsModule.h
@@ -6,6 +6,101 @@
 #include "Modules/ModuleInterface.h"
 #include "Framework/Docking/TabManager.h"
 
+class FExtender;
+
+/** Major tab IDs for Insights tools */
+struct TRACEINSIGHTS_API FInsightsManagerTabs
+{
+	static const FName StartPageTabId;
+	static const FName SessionInfoTabId;
+	static const FName TimingProfilerTabId;
+	static const FName LoadingProfilerTabId;
+	static const FName NetworkingProfilerTabId;
+};
+
+/** Tab IDs for the timing profiler */
+struct TRACEINSIGHTS_API FTimingProfilerTabs
+{
+	// Tab identifiers
+	static const FName ToolbarID;
+	static const FName FramesTrackID;
+	static const FName TimingViewID;
+	static const FName TimersID;
+	static const FName CallersID;
+	static const FName CalleesID;
+	static const FName StatsCountersID;
+	static const FName LogViewID;
+};
+
+/** Delegate invoked when a major tab is created */
+DECLARE_MULTICAST_DELEGATE_TwoParams(FOnInsightsMajorTabCreated, FName /*MajorTabId*/, TSharedRef<FTabManager> /*TabManager*/)
+
+/** Configuration for an Insights minor tab. This is used to augment the standard supplied tabs from plugins. */
+struct TRACEINSIGHTS_API FInsightsMinorTabConfig
+{
+	FName TabId;
+
+	FText TabLabel;
+
+	FText TabTooltip;
+
+	FSlateIcon TabIcon;
+
+	FOnSpawnTab OnSpawnTab;
+
+	FOnFindTabToReuse OnFindTabToReuse;
+
+	TSharedPtr<FWorkspaceItem> WorkspaceGroup;
+};
+
+/** Configuration for an Insights major tab */
+struct TRACEINSIGHTS_API FInsightsMajorTabConfig
+{
+	FInsightsMajorTabConfig()
+		: Layout(nullptr)
+		, WorkspaceGroup(nullptr)
+		, bIsAvailable(true)
+	{}
+
+	/** Helper function for creating unavailable tab configs */
+	static FInsightsMajorTabConfig Unavailable()
+	{
+		FInsightsMajorTabConfig Config;
+		Config.bIsAvailable = false;
+		return Config;
+	}
+
+	/** Identifier for this config */
+	FName ConfigId;
+
+	/** Display name for this config */
+	FText ConfigDisplayName;
+
+	/** Label for the tab. If this is not set the default will be used */
+	TOptional<FText> TabLabel;
+
+	/** Tooltip for the tab. If this is not set the default will be used */
+	TOptional<FText> TabTooltip;
+
+	/** Icon for the tab. If this is not set the default will be used */
+	TOptional<FSlateIcon> TabIcon;
+
+	/** The tab layout to use. If not specified, the default will be used. */
+	TSharedPtr<FTabManager::FLayout> Layout;
+
+	/** The menu workspace group to use. If not specified, the default will be used. */
+	TSharedPtr<FWorkspaceItem> WorkspaceGroup;
+
+	/** Extender used to add to the menu for this tab */
+	TSharedPtr<FExtender> MenuExtender;
+
+	/** Any additional minor tabs to add */
+	TArray<FInsightsMinorTabConfig> MinorTabs;
+
+	/** Whether the tab is available for selection (i.e. registered with the tab manager) */
+	bool bIsAvailable;
+};
+
 /** Interface for an Unreal Insights module. */
 class IUnrealInsightsModule : public IModuleInterface
 {
@@ -30,12 +125,34 @@ public:
 	/**
 	 * Starts analysis of the specified session. Called when the application starts in "Viewer" mode.
 	 *
-	 * @param InSessionId The id of the session to analyze. If nullptr, the app will wait for a live session.
+	 * @param InSessionId The id of the session to analyze.
 	 */
 	virtual void StartAnalysisForSession(const TCHAR* InSessionId) = 0;
 
+	/**
+	 * Starts analysis of the last live session. Called when the application starts in "Viewer" mode.
+	 */
+	virtual void StartAnalysisForLastLiveSession() = 0;
+
 	/**
 	 * Called when the application shutsdown.
 	 */
 	virtual void ShutdownUserInterface() = 0;
+
+	/** 
+	 * Registers a major tab layout. This defines how the major tab will appear when spawned.
+	 * If this is not called prior to tabs being spawned then the built-in default layout will be used.
+	 * @param	InMajorTabId	The major tab ID we are supplying a layout to
+	 * @param	InConfig		The config to use when spawning the major tab
+	 */
+	virtual void RegisterMajorTabConfig(const FName& InMajorTabId, const FInsightsMajorTabConfig& InConfig) = 0;
+
+	/** 
+	 * Unregisters a major tab layout. This will revert the major tab to spawning with its default layout
+	 * @param	InMajorTabId	The major tab ID we are supplying a layout to
+	 */
+	virtual void UnregisterMajorTabConfig(const FName& InMajorTabId) = 0;
+
+	/** Callback invoked when a major tab is created */
+	virtual FOnInsightsMajorTabCreated& OnMajorTabCreated() = 0;
 };
diff --git a/Engine/Source/Developer/TraceServices/Private/TraceServicesModule.cpp b/Engine/Source/Developer/TraceServices/Private/TraceServicesModule.cpp
--- a/Engine/Source/Developer/TraceServices/Private/TraceServicesModule.cpp
+++ b/Engine/Source/Developer/TraceServices/Private/TraceServicesModule.cpp
@@ -45,6 +45,7 @@ TSharedPtr<Trace::ISessionService> FTraceServicesModule::GetSessionService()
 	if (!SessionService.IsValid())
 	{
 		GetModuleService();
+		GetAnalysisService();
 		SessionService = MakeShared<Trace::FSessionService>(*ModuleService.Get(), *AnalysisService.Get());
 	}
 	return SessionService;
