Workaround/fix for crash when shutting down the editor

diff --git a/Engine/Source/Developer/TraceInsights/Private/Insights/TraceInsightsModule.cpp b/Engine/Source/Developer/TraceInsights/Private/Insights/TraceInsightsModule.cpp
--- a/Engine/Source/Developer/TraceInsights/Private/Insights/TraceInsightsModule.cpp
+++ b/Engine/Source/Developer/TraceInsights/Private/Insights/TraceInsightsModule.cpp
@@ -73,9 +73,12 @@ void FTraceInsightsModule::StartupModule()
 void FTraceInsightsModule::ShutdownModule()
 {
 #if WITH_EDITOR
-	// Save application layout.
-	FLayoutSaveRestore::SaveToConfig(GEditorLayoutIni, PersistentLayout.ToSharedRef());
-	GConfig->Flush(false, GEditorLayoutIni);
+	if (PersistentLayout.IsValid())
+	{
+		// Save application layout.
+		FLayoutSaveRestore::SaveToConfig(GEditorLayoutIni, PersistentLayout.ToSharedRef());
+		GConfig->Flush(false, GEditorLayoutIni);
+	}
 #endif
 
 #if !WITH_EDITOR
