[Insights]
  - Changed trace session management to use the new store API.
  - Removed Start/Stop recorder functionality.
  - Enabled TRACE_WITH_ASIO (will be removed in a following commit).

diff --git a/Engine/Source/Developer/TraceAnalysis/Private/Asio/Asio.h b/Engine/Source/Developer/TraceAnalysis/Private/Asio/Asio.h
--- a/Engine/Source/Developer/TraceAnalysis/Private/Asio/Asio.h
+++ b/Engine/Source/Developer/TraceAnalysis/Private/Asio/Asio.h
@@ -3,7 +3,7 @@
 #pragma once
 
 #if !defined(TRACE_WITH_ASIO)
-#	define TRACE_WITH_ASIO 0
+#	define TRACE_WITH_ASIO 1
 #endif // 0
 
 #if TRACE_WITH_ASIO
diff --git a/Engine/Source/Developer/TraceInsights/Private/Insights/InsightsManager.cpp b/Engine/Source/Developer/TraceInsights/Private/Insights/InsightsManager.cpp
--- a/Engine/Source/Developer/TraceInsights/Private/Insights/InsightsManager.cpp
+++ b/Engine/Source/Developer/TraceInsights/Private/Insights/InsightsManager.cpp
@@ -3,17 +3,19 @@
 #include "InsightsManager.h"
 
 #include "Framework/Application/SlateApplication.h"
+#include "Misc/CString.h"
 #include "Modules/ModuleManager.h"
 #include "Templates/UniquePtr.h"
+#include "Trace/StoreClient.h"
 #include "TraceServices/Model/NetProfiler.h"
 
 // Insights
+#include "Insights/IUnrealInsightsModule.h"
 #include "Insights/LoadingProfiler/LoadingProfilerManager.h"
 #include "Insights/NetworkingProfiler/NetworkingProfilerManager.h"
 #include "Insights/TimingProfilerManager.h"
 #include "Insights/Widgets/SStartPageWindow.h"
 #include "Insights/Widgets/STimingProfilerWindow.h"
-#include "Insights/IUnrealInsightsModule.h"
 
 ////////////////////////////////////////////////////////////////////////////////////////////////////
 
@@ -35,6 +37,10 @@ FInsightsManager::FInsightsManager(TSharedRef<Trace::IAnalysisService> InTraceAn
 	: AnalysisService(InTraceAnalysisService)
 	, SessionService(InTraceSessionService)
 	, ModuleService(InTraceModuleService)
+	, StoreDir()
+	, StoreHost()
+	, StorePort(0)
+	, StoreClient()
 	, CommandList(new FUICommandList())
 	, ActionManager(this)
 	, Settings()
@@ -95,13 +101,6 @@ TSharedPtr<const Trace::IAnalysisSession> FInsightsManager::GetSession() const
 
 ////////////////////////////////////////////////////////////////////////////////////////////////////
 
-Trace::FSessionHandle FInsightsManager::GetSessionHandle() const
-{
-	return CurrentSessionHandle;
-}
-
-////////////////////////////////////////////////////////////////////////////////////////////////////
-
 const TSharedRef<FUICommandList> FInsightsManager::GetCommandList() const
 {
 	return CommandList;
@@ -130,6 +129,18 @@ FInsightsSettings& FInsightsManager::GetSettings()
 
 ////////////////////////////////////////////////////////////////////////////////////////////////////
 
+bool FInsightsManager::ConnectToStore(const TCHAR* Host, uint32 Port)
+{
+	using namespace Trace;
+	StoreHost = Host;
+	StorePort = Port;
+	FStoreClient* Client = FStoreClient::Connect(Host, Port);
+	StoreClient = TUniquePtr<FStoreClient>(Client);
+	return StoreClient.IsValid();
+}
+
+////////////////////////////////////////////////////////////////////////////////////////////////////
+
 bool FInsightsManager::Tick(float DeltaTime)
 {
 	if (!bIsNetworkingProfilerAvailable)
@@ -167,7 +178,7 @@ void FInsightsManager::ResetSession()
 	if (Session.IsValid())
 	{
 		Session.Reset();
-		CurrentSessionHandle = 0;
+		CurrentTraceId = 0;
 		bIsNetworkingProfilerAvailable = false;
 		OnSessionChanged();
 	}
@@ -259,86 +270,61 @@ void FInsightsManager::ActivateTimingInsightsTab()
 
 ////////////////////////////////////////////////////////////////////////////////////////////////////
 
-bool FInsightsManager::IsAnyLiveSessionAvailable(Trace::FSessionHandle& OutLastLiveSessionHandle) const
+void FInsightsManager::LoadLastLiveSession()
 {
-	TArray<Trace::FSessionHandle> LiveSessions;
-	SessionService->GetLiveSessions(LiveSessions);
+	ResetSession();
 
-	if (LiveSessions.Num() > 0)
+	if (StoreClient.IsValid())
 	{
-		OutLastLiveSessionHandle = LiveSessions.Last();
-		return true;
+		return;
 	}
 
-	return false;
-}
-
-////////////////////////////////////////////////////////////////////////////////////////////////////
-
-bool FInsightsManager::IsAnySessionAvailable(Trace::FSessionHandle& OutLastSessionHandle) const
-{
-	TArray<Trace::FSessionHandle> AvailableSessions;
-	SessionService->GetAvailableSessions(AvailableSessions);
+	const int32 SessionCount = StoreClient->GetSessionCount();
+	if (SessionCount == 0)
+	{
+		return;
+	}
 
-	if (AvailableSessions.Num() > 0)
+	const Trace::FStoreClient::FSessionInfo* SessionInfo = StoreClient->GetSessionInfo(SessionCount - 1);
+	if (SessionInfo == nullptr)
 	{
-		OutLastSessionHandle = AvailableSessions.Last();
-		return true;
+		return;
 	}
 
-	return false;
+	LoadTrace(SessionInfo->GetTraceId());
 }
 
 ////////////////////////////////////////////////////////////////////////////////////////////////////
 
-void FInsightsManager::LoadLastLiveSession()
+void FInsightsManager::LoadTrace(uint32 InTraceId)
 {
 	ResetSession();
 
-	TArray<Trace::FSessionHandle> AvailableSessions;
-	SessionService->GetAvailableSessions(AvailableSessions);
-
-	// Iterate in reverse order as we want the most recent live session first.
-	for (int32 SessionIndex = AvailableSessions.Num() - 1; SessionIndex >= 0; --SessionIndex)
+	if (StoreClient == nullptr)
 	{
-		Trace::FSessionHandle SessionHandle = AvailableSessions[SessionIndex];
-
-		Trace::FSessionInfo SessionInfo;
-		SessionService->GetSessionInfo(SessionHandle, SessionInfo);
-
-		if (SessionInfo.bIsLive)
-		{
-			LoadSession(SessionHandle);
-			break;
-		}
+		return;
 	}
-}
-
-////////////////////////////////////////////////////////////////////////////////////////////////////
-
-void FInsightsManager::LoadLastSession()
-{
-	ResetSession();
 
-	TArray<Trace::FSessionHandle> AvailableSessions;
-	SessionService->GetAvailableSessions(AvailableSessions);
+	Trace::FStoreClient::FTraceData TraceData = StoreClient->ReadTrace(InTraceId);
+	if (!TraceData)
+	{
+		return;
+	}
 
-	if (AvailableSessions.Num() > 0)
+	FString TraceName;
+	const Trace::FStoreClient::FTraceInfo* TraceInfo = StoreClient->GetTraceInfoById(InTraceId);
+	if (TraceInfo != nullptr)
 	{
-		LoadSession(AvailableSessions.Last());
+		FAnsiStringView Name = TraceInfo->GetName();
+		TraceName = FString(Name.Len(), Name.GetData());
 	}
-}
 
-////////////////////////////////////////////////////////////////////////////////////////////////////
+	Session = AnalysisService->StartAnalysis(*TraceName, MoveTemp(TraceData));
 
-void FInsightsManager::LoadSession(Trace::FSessionHandle SessionHandle)
-{
-	ResetSession();
-
-	Session = SessionService->StartAnalysis(SessionHandle);
 	if (Session)
 	{
-		CurrentSessionHandle = SessionHandle;
+		CurrentTraceId = InTraceId;
+		CurrentTraceFilename = TraceName;
 		bIsNetworkingProfilerAvailable = false;
 		SpawnAndActivateTabs();
 		OnSessionChanged();
@@ -347,14 +333,23 @@ void FInsightsManager::LoadSession(Trace::FSessionHandle SessionHandle)
 
 ////////////////////////////////////////////////////////////////////////////////////////////////////
 
-void FInsightsManager::LoadTraceFile(const FString& TraceFilepath)
+void FInsightsManager::LoadTraceFile(const FString& InTraceFilename)
 {
+	IPlatformFile& PlatformFile = IPlatformFile::GetPlatformPhysical();
+	if (!PlatformFile.FileExists(*InTraceFilename))
+	{
+		const uint32 TraceId = uint32(FCString::Strtoui64(*InTraceFilename, nullptr, 10));
+		return LoadTrace(TraceId);
+	}
+
 	ResetSession();
 
-	Session = AnalysisService->StartAnalysis(*TraceFilepath);
+	Session = AnalysisService->StartAnalysis(*InTraceFilename);
+
 	if (Session)
 	{
-		CurrentSessionHandle = 0;
+		CurrentTraceId = 0;
+		CurrentTraceFilename = InTraceFilename;
 		bIsNetworkingProfilerAvailable = false;
 		SpawnAndActivateTabs();
 		OnSessionChanged();
diff --git a/Engine/Source/Developer/TraceInsights/Private/Insights/InsightsManager.h b/Engine/Source/Developer/TraceInsights/Private/Insights/InsightsManager.h
--- a/Engine/Source/Developer/TraceInsights/Private/Insights/InsightsManager.h
+++ b/Engine/Source/Developer/TraceInsights/Private/Insights/InsightsManager.h
@@ -15,6 +15,11 @@
 
 ////////////////////////////////////////////////////////////////////////////////////////////////////
 
+namespace Trace
+{
+	class FStoreClient;
+}
+
 class SStartPageWindow;
 
 ////////////////////////////////////////////////////////////////////////////////////////////////////
@@ -72,11 +77,22 @@ public:
 	TSharedRef<Trace::ISessionService> GetSessionService() const { return SessionService; }
 	TSharedRef<Trace::IModuleService> GetModuleService() const { return ModuleService; }
 
+	void SetStoreDir(const FString& InStoreDir) { StoreDir = InStoreDir; }
+	const FString& GetStoreDir() const { return StoreDir; }
+
+	bool ConnectToStore(const TCHAR* Host, uint32 Port);
+	const FString& GetStoreHost() const { return StoreHost; }
+	uint32 GetStorePort() const { return StorePort; }
+	Trace::FStoreClient* GetStoreClient() const { return StoreClient.Get(); }
+
 	/** @return an instance of the trace analysis session. */
 	TSharedPtr<const Trace::IAnalysisSession> GetSession() const;
 
-	/** @return the session handle of the trace analysis session. */
-	Trace::FSessionHandle GetSessionHandle() const;
+	/** @return the id of the trace being analyzed. */
+	uint32 GetTraceId() const { return CurrentTraceId; }
+
+	/** @return the filename of the trace being analyzed. */
+	const FString& GetTraceFilename() const { return CurrentTraceFilename; }
 
 	/** @returns UI command list for the main manager. */
 	const TSharedRef<FUICommandList> GetCommandList() const;
@@ -116,26 +132,20 @@ public:
 	bool ShouldOpenAnalysisInSeparateProcess() const { return bShouldOpenAnalysisInSeparateProcess; }
 	void SetOpenAnalysisInSeparateProcess(bool bOnOff) { bShouldOpenAnalysisInSeparateProcess = bOnOff; }
 
-	bool IsAnyLiveSessionAvailable(Trace::FSessionHandle& OutLastLiveSessionHandle) const;
-	bool IsAnySessionAvailable(Trace::FSessionHandle& OutLastSessionHandle) const;
-
-	/** Creates a new analysis session instance and loads the latest available trace session that is live. */
+	/** Creates a new analysis session instance and loads the latest available trace that is live. */
 	void LoadLastLiveSession();
 
-	/** Creates a new analysis session instance and loads the latest available trace session. */
-	void LoadLastSession();
-
 	/**
-	 * Creates a new analysis session instance using specified session handle.
-	 * @param SessionHandle - The handle for session to analyze
+	 * Creates a new analysis session instance using specified trace id.
+	 * @param TraceId - The id of the trace to analyze
 	 */
-	void LoadSession(Trace::FSessionHandle SessionHandle);
+	void LoadTrace(uint32 TraceId);
 
 	/**
 	 * Creates a new analysis session instance and loads a trace file from the specified location.
-	 * @param TraceFilepath - The path to the trace file
+	 * @param TraceFilename - The trace file to analyze
 	 */
-	void LoadTraceFile(const FString& TraceFilepath);
+	void LoadTraceFile(const FString& TraceFilename);
 
 	/** Opens the Settings dialog. */
 	void OpenSettings();
@@ -183,11 +193,26 @@ private:
 	TSharedRef<Trace::ISessionService> SessionService;
 	TSharedRef<Trace::IModuleService> ModuleService;
 
+	/** The location of the trace files managed by the trace store. */
+	FString StoreDir;
+
+	/** The host name of the trace store. */
+	FString StoreHost;
+
+	/** The port of the trace store. */
+	uint32 StorePort;
+
+	/** The client used to connect to the trace store. */
+	TUniquePtr<Trace::FStoreClient> StoreClient;
+
 	/** The trace analysis session. */
 	TSharedPtr<const Trace::IAnalysisSession> Session;
 
-	/** The session handle. */
-	Trace::FSessionHandle CurrentSessionHandle;
+	/** The id of the trace being analyzed. */
+	uint32 CurrentTraceId;
+
+	/** The filename of the trace being analyzed. */
+	FString CurrentTraceFilename;
 
 	/** List of UI commands for this manager. This will be filled by this and corresponding classes. */
 	TSharedRef<FUICommandList> CommandList;
diff --git a/Engine/Source/Developer/TraceInsights/Private/Insights/TraceInsightsModule.cpp b/Engine/Source/Developer/TraceInsights/Private/Insights/TraceInsightsModule.cpp
--- a/Engine/Source/Developer/TraceInsights/Private/Insights/TraceInsightsModule.cpp
+++ b/Engine/Source/Developer/TraceInsights/Private/Insights/TraceInsightsModule.cpp
@@ -3,6 +3,8 @@
 #include "TraceInsightsModule.h"
 #include "HAL/PlatformApplicationMisc.h"
 #include "Modules/ModuleManager.h"
+#include "Trace/StoreClient.h"
+#include "Trace/StoreService.h"
 #include "TraceServices/ITraceServicesModule.h"
 #include "TraceServices/SessionService.h"
 #include "Widgets/DeclarativeSyntaxSupport.h"
@@ -56,14 +58,8 @@ void FTraceInsightsModule::StartupModule()
 	FLoadingProfilerManager::Initialize();
 	FNetworkingProfilerManager::Initialize();
 
-	//////////////////////////////////////////////////
-
 #if WITH_EDITOR
-	if (TraceSessionService.IsValid())
-	{
-		TraceSessionService->StartRecorderServer();
-	}
-
+	//...
 	FCoreDelegates::OnExit.AddRaw(this, &FTraceInsightsModule::HandleExit);
 #endif
 }
@@ -81,18 +77,7 @@ void FTraceInsightsModule::ShutdownModule()
 	}
 #endif
 
-#if !WITH_EDITOR
-	if (TraceSessionService.IsValid())
-	{
-		TraceSessionService->StopRecorderServer();
-	}
-#endif
-
-	FGlobalTabmanager::Get()->UnregisterNomadTabSpawner(FInsightsManagerTabs::NetworkingProfilerTabId);
-	FGlobalTabmanager::Get()->UnregisterNomadTabSpawner(FInsightsManagerTabs::LoadingProfilerTabId);
-	FGlobalTabmanager::Get()->UnregisterNomadTabSpawner(FInsightsManagerTabs::TimingProfilerTabId);
-	FGlobalTabmanager::Get()->UnregisterNomadTabSpawner(FInsightsManagerTabs::SessionInfoTabId);
-	FGlobalTabmanager::Get()->UnregisterNomadTabSpawner(FInsightsManagerTabs::StartPageTabId);
+	UnregisterTabSpawners();
 
 	if (FNetworkingProfilerManager::Get().IsValid())
 	{
@@ -125,11 +110,8 @@ void FTraceInsightsModule::ShutdownModule()
 void FTraceInsightsModule::HandleExit()
 {
 	// In editor, as module lifetimes are different, we need to shut down the recorder service before
-	// threads get killed in the shutdown process otherwise we will hang forever waiting on the recorder server
-	if (TraceSessionService.IsValid())
-	{
-		TraceSessionService->StopRecorderServer();
-	}
+	// threads get killed in the shutdown process otherwise we will hang forever waiting on the recorder server.
+	//...
 }
 #endif
 
@@ -140,7 +122,7 @@ void FTraceInsightsModule::RegisterTabSpawners()
 	TSharedRef<FWorkspaceItem> ToolsCategory = WorkspaceMenu::GetMenuStructure().GetToolsCategory(); 
 
 	const FInsightsMajorTabConfig& StartPageConfig = FindMajorTabConfig(FInsightsManagerTabs::StartPageTabId);
-	if(StartPageConfig.bIsAvailable)
+	if (StartPageConfig.bIsAvailable)
 	{
 		// Register tab spawner for the Start Page.
 		FTabSpawnerEntry& StartPageTabSpawnerEntry = FGlobalTabmanager::Get()->RegisterNomadTabSpawner(FInsightsManagerTabs::StartPageTabId,
@@ -153,20 +135,20 @@ void FTraceInsightsModule::RegisterTabSpawners()
 	}
 
 	const FInsightsMajorTabConfig& SessionInfoConfig = FindMajorTabConfig(FInsightsManagerTabs::SessionInfoTabId);
-	if(SessionInfoConfig.bIsAvailable)
+	if (SessionInfoConfig.bIsAvailable)
 	{
 		// Register tab spawner for the Session Info.
 		FTabSpawnerEntry& SessionInfoTabSpawnerEntry = FGlobalTabmanager::Get()->RegisterNomadTabSpawner(FInsightsManagerTabs::SessionInfoTabId,
 			FOnSpawnTab::CreateRaw(this, &FTraceInsightsModule::SpawnSessionInfoTab))
-			.SetDisplayName(NSLOCTEXT("FTraceInsightsModule", "SessionInfoTabTitle", "Session Info"))
-			.SetTooltipText(NSLOCTEXT("FTraceInsightsModule", "SessionInfoTooltipText", "Open the Session Info tab."))
-			.SetIcon(FSlateIcon(FInsightsStyle::GetStyleSetName(), "SessionInfo.Icon.Small"));
+			.SetDisplayName(SessionInfoConfig.TabLabel.IsSet() ? SessionInfoConfig.TabLabel.GetValue() : NSLOCTEXT("FTraceInsightsModule", "SessionInfoTabTitle", "Session Info"))
+			.SetTooltipText(SessionInfoConfig.TabTooltip.IsSet() ? SessionInfoConfig.TabTooltip.GetValue() : NSLOCTEXT("FTraceInsightsModule", "SessionInfoTooltipText", "Open the Session Info tab."))
+			.SetIcon(SessionInfoConfig.TabIcon.IsSet() ? SessionInfoConfig.TabIcon.GetValue() : FSlateIcon(FInsightsStyle::GetStyleSetName(), "SessionInfo.Icon.Small"));
 
 		SessionInfoTabSpawnerEntry.SetGroup(SessionInfoConfig.WorkspaceGroup.IsValid() ? SessionInfoConfig.WorkspaceGroup.ToSharedRef() : ToolsCategory);
 	}
 
 	const FInsightsMajorTabConfig& TimingProfilerConfig = FindMajorTabConfig(FInsightsManagerTabs::TimingProfilerTabId);
-	if(TimingProfilerConfig.bIsAvailable)
+	if (TimingProfilerConfig.bIsAvailable)
 	{
 		// Register tab spawner for the Timing Insights.
 		FTabSpawnerEntry& TimingProfilerTabSpawnerEntry = FGlobalTabmanager::Get()->RegisterNomadTabSpawner(FInsightsManagerTabs::TimingProfilerTabId,
@@ -179,7 +161,7 @@ void FTraceInsightsModule::RegisterTabSpawners()
 	}
 
 	const FInsightsMajorTabConfig& LoadingProfilerConfig = FindMajorTabConfig(FInsightsManagerTabs::LoadingProfilerTabId);
-	if(LoadingProfilerConfig.bIsAvailable)
+	if (LoadingProfilerConfig.bIsAvailable)
 	{
 		// Register tab spawner for the Asset Loading Insights.
 		FTabSpawnerEntry& LoadingProfilerTabSpawnerEntry = FGlobalTabmanager::Get()->RegisterNomadTabSpawner(FInsightsManagerTabs::LoadingProfilerTabId,
@@ -192,7 +174,7 @@ void FTraceInsightsModule::RegisterTabSpawners()
 	}
 
 	const FInsightsMajorTabConfig& NetworkingProfilerConfig = FindMajorTabConfig(FInsightsManagerTabs::NetworkingProfilerTabId);
-	if(NetworkingProfilerConfig.bIsAvailable)
+	if (NetworkingProfilerConfig.bIsAvailable)
 	{
 		// Register tab spawner for the Networking Insights.
 		FTabSpawnerEntry& NetworkingProfilerTabSpawnerEntry = FGlobalTabmanager::Get()->RegisterNomadTabSpawner(FInsightsManagerTabs::NetworkingProfilerTabId,
@@ -208,14 +190,33 @@ void FTraceInsightsModule::RegisterTabSpawners()
 
 ////////////////////////////////////////////////////////////////////////////////////////////////////
 
+void FTraceInsightsModule::UnregisterTabSpawners()
+{
+	FGlobalTabmanager::Get()->UnregisterNomadTabSpawner(FInsightsManagerTabs::NetworkingProfilerTabId);
+	FGlobalTabmanager::Get()->UnregisterNomadTabSpawner(FInsightsManagerTabs::LoadingProfilerTabId);
+	FGlobalTabmanager::Get()->UnregisterNomadTabSpawner(FInsightsManagerTabs::TimingProfilerTabId);
+	FGlobalTabmanager::Get()->UnregisterNomadTabSpawner(FInsightsManagerTabs::SessionInfoTabId);
+	FGlobalTabmanager::Get()->UnregisterNomadTabSpawner(FInsightsManagerTabs::StartPageTabId);
+}
+
+////////////////////////////////////////////////////////////////////////////////////////////////////
+
 void FTraceInsightsModule::CreateSessionBrowser(bool bAllowDebugTools, bool bSingleProcess)
 {
-	bBrowserMode = true;
 	FInsightsManager::Get()->SetOpenAnalysisInSeparateProcess(!bSingleProcess);
+	FInsightsManager::Get()->SetStoreDir(FPaths::ProjectSavedDir() / TEXT("TraceSessions"));
 
-	if (TraceSessionService.IsValid())
+	// Create the Store Service.
+	Trace::FStoreService::FDesc StoreServiceDesc;
+	StoreServiceDesc.StoreDir = *FInsightsManager::Get()->GetStoreDir();
+	StoreServiceDesc.RecorderPort = 1980;
+	StoreServiceDesc.ThreadCount = 2;
+	StoreService = TUniquePtr<Trace::FStoreService>(Trace::FStoreService::Create(StoreServiceDesc));
+
+	if (StoreService.IsValid())
 	{
-		TraceSessionService->StartRecorderServer();
+		// Create the Store Client.
+		FInsightsManager::Get()->ConnectToStore(TEXT("127.0.0.1"), StoreService->GetPort());
 	}
 
 	RegisterTabSpawners();
@@ -257,8 +258,6 @@ void FTraceInsightsModule::CreateSessionBrowser(bool bAllowDebugTools, bool bSin
 
 void FTraceInsightsModule::CreateSessionViewer(bool bAllowDebugTools)
 {
-	bBrowserMode = false;
-
 	RegisterTabSpawners();
 
 	TSharedRef<FTabManager::FLayout> DefaultLayout = FTabManager::NewLayout("UnrealInsightsLayout_v1.0");
@@ -351,7 +350,7 @@ void FTraceInsightsModule::UnregisterMajorTabConfig(const FName& InMajorTabId)
 const FInsightsMajorTabConfig& FTraceInsightsModule::FindMajorTabConfig(const FName& InMajorTabId) const
 {
 	const FInsightsMajorTabConfig* FoundConfig = TabConfigs.Find(InMajorTabId);
-	if(FoundConfig != nullptr)
+	if (FoundConfig != nullptr)
 	{
 		return *FoundConfig;
 	}
@@ -372,11 +371,13 @@ void FTraceInsightsModule::StartAnalysisForTraceFile(const TCHAR* InTraceFile)
 
 ////////////////////////////////////////////////////////////////////////////////////////////////////
 
-void FTraceInsightsModule::StartAnalysisForSession(const TCHAR* InSessionId)
+void FTraceInsightsModule::StartAnalysisForTrace(const TCHAR* InStoreHost, uint32 InStorePort, uint32 InTraceId)
 {
-	if (InSessionId != nullptr)
+	FInsightsManager::Get()->ConnectToStore(InStoreHost, InStorePort);
+
+	if (InTraceId != 0)
 	{
-		//TODO: FInsightsManager::Get()->LoadSession(FString(InSessionId));
+		FInsightsManager::Get()->LoadTrace(InTraceId);
 	}
 }
 
diff --git a/Engine/Source/Developer/TraceInsights/Private/Insights/TraceInsightsModule.h b/Engine/Source/Developer/TraceInsights/Private/Insights/TraceInsightsModule.h
--- a/Engine/Source/Developer/TraceInsights/Private/Insights/TraceInsightsModule.h
+++ b/Engine/Source/Developer/TraceInsights/Private/Insights/TraceInsightsModule.h
@@ -6,8 +6,9 @@
 #include "Insights/IUnrealInsightsModule.h"
 #include "Framework/Docking/TabManager.h"
 
-namespace Trace 
-{ 
+namespace Trace
+{
+	class FStoreService;
 	class IAnalysisService;
 	class ISessionService;
 	class IModuleService;
@@ -33,8 +34,9 @@ public:
 	virtual void CreateSessionBrowser(bool bAllowDebugTools, bool bSingleProcess) override;
 	virtual void CreateSessionViewer(bool bAllowDebugTools) override;
 	virtual void StartAnalysisForTraceFile(const TCHAR* InTraceFile) override;
-	virtual void StartAnalysisForSession(const TCHAR* InSessionId) override;
-	virtual void StartAnalysisForLastLiveSession() override;
+	virtual void StartAnalysisForTrace(const TCHAR* InStoreHost, uint32 InStorePort, uint32 InTraceId) override;
+	virtual void StartAnalysisForLastLiveSession(/*StoreClient*/) override;
+	//virtual void StartAnalysisForTrace(StoreClient or FTraceData..) override;
 	virtual void ShutdownUserInterface() override;
 	virtual void RegisterMajorTabConfig(const FName& InMajorTabId, const FInsightsMajorTabConfig& InConfig) override;
 	virtual void UnregisterMajorTabConfig(const FName& InMajorTabId) override;
@@ -44,7 +46,10 @@ public:
 	const FInsightsMajorTabConfig& FindMajorTabConfig(const FName& InMajorTabId) const;
 
 protected:
+	void InitTraceStore();
+
 	void RegisterTabSpawners();
+	void UnregisterTabSpawners();
 
 	void AddAreaForSessionViewer(TSharedRef<FTabManager::FLayout> Layout);
 	void AddAreaForWidgetReflector(TSharedRef<FTabManager::FLayout> Layout, bool bAllowDebugTools);
@@ -82,9 +87,7 @@ protected:
 #endif
 
 protected:
-	TSharedPtr<FTabManager::FLayout> PersistentLayout;
-	bool bBrowserMode;
-	FString UnrealInsightsLayoutIni;
+	TUniquePtr<Trace::FStoreService> StoreService;
 
 	TSharedPtr<Trace::IAnalysisService> TraceAnalysisService;
 	TSharedPtr<Trace::ISessionService> TraceSessionService;
@@ -93,4 +96,7 @@ protected:
 	TMap<FName, FInsightsMajorTabConfig> TabConfigs;
 
 	FOnInsightsMajorTabCreated OnInsightsMajorTabCreatedDelegate;
-};
\ No newline at end of file
+
+	TSharedPtr<FTabManager::FLayout> PersistentLayout;
+	FString UnrealInsightsLayoutIni;
+};
diff --git a/Engine/Source/Developer/TraceInsights/Private/Insights/Widgets/SStartPageWindow.cpp b/Engine/Source/Developer/TraceInsights/Private/Insights/Widgets/SStartPageWindow.cpp
--- a/Engine/Source/Developer/TraceInsights/Private/Insights/Widgets/SStartPageWindow.cpp
+++ b/Engine/Source/Developer/TraceInsights/Private/Insights/Widgets/SStartPageWindow.cpp
@@ -2,6 +2,7 @@
 
 #include "SStartPageWindow.h"
 
+#include "Containers/StringView.h"
 #include "DesktopPlatformModule.h"
 #include "EditorStyleSet.h"
 #include "Framework/Application/SlateApplication.h"
@@ -12,6 +13,8 @@
 #include "Internationalization/Text.h"
 #include "SlateOptMacros.h"
 #include "Styling/CoreStyle.h"
+#include "Trace/ControlClient.h"
+#include "Trace/StoreClient.h"
 #include "Widgets/Docking/SDockTab.h"
 #include "Widgets/Images/SImage.h"
 #include "Widgets/Input/SCheckBox.h"
@@ -37,6 +40,7 @@
 
 // Insights
 #include "Insights/InsightsManager.h"
+#include "Insights/TimingProfilerCommon.h" // for UE_LOG(TimingProfiler, ...
 #include "Insights/TimingProfilerManager.h"
 #include "Insights/InsightsStyle.h"
 #include "Insights/Version.h"
@@ -240,12 +244,12 @@ public:
 		return FText::GetEmpty();
 	}
 
-	FText GetTraceSessionTimeStamp() const
+	FText GetTraceSessionTimestamp() const
 	{
 		TSharedPtr<FTraceSession> TraceSessionPin = WeakTraceSession.Pin();
 		if (TraceSessionPin.IsValid())
 		{
-			return FText::AsDate(TraceSessionPin->TimeStamp);
+			return FText::AsDate(TraceSessionPin->Timestamp);
 		}
 		else
 		{
@@ -325,8 +329,10 @@ public:
 		{
 			FTextBuilder TextBuilder;
 
+			const FString TraceIdStr = FString::Printf(TEXT("0x%08X"), TraceSessionPin->TraceId);
+			TextBuilder.AppendLineFormat(LOCTEXT("TraceSessionTooltip_Id", "Trace {0} ({1})"), FText::AsNumber(TraceSessionPin->TraceIndex), FText::FromString(TraceIdStr));
 			TextBuilder.AppendLineFormat(LOCTEXT("TraceSessionTooltip_Name", "{0}"), TraceSessionPin->Name);
-			TextBuilder.AppendLineFormat(LOCTEXT("TraceSessionTooltip_Uri", "{0}"), TraceSessionPin->Uri);
+			TextBuilder.AppendLineFormat(LOCTEXT("TraceSessionTooltip_Uri", "Uri: {0}"), TraceSessionPin->Uri);
 			TextBuilder.AppendLineFormat(LOCTEXT("TraceSessionTooltip_Platform", "Platform: {0}"), TraceSessionPin->Platform);
 			TextBuilder.AppendLineFormat(LOCTEXT("TraceSessionTooltip_AppName", "App Name: {0}"), TraceSessionPin->AppName);
 			TextBuilder.AppendLineFormat(LOCTEXT("TraceSessionTooltip_CommandLine", "Command Line: {0}"), TraceSessionPin->CommandLine);
@@ -334,7 +340,7 @@ public:
 				TraceSessionPin->ConfigurationType == EBuildConfiguration::Unknown ? FText::GetEmpty() : EBuildConfigurations::ToText(TraceSessionPin->ConfigurationType));
 			TextBuilder.AppendLineFormat(LOCTEXT("TraceSessionTooltip_BuildTarget", "Build Target: {0}"),
 				TraceSessionPin->TargetType == EBuildTargetType::Unknown ? FText::GetEmpty() : FText::FromString(LexToString(TraceSessionPin->TargetType)));
-			TextBuilder.AppendLineFormat(LOCTEXT("TraceSessionTooltip_Timestamp", "Timestamp: {0}"), FText::AsDateTime(TraceSessionPin->TimeStamp));
+			TextBuilder.AppendLineFormat(LOCTEXT("TraceSessionTooltip_Timestamp", "Timestamp: {0}"), FText::AsDateTime(TraceSessionPin->Timestamp));
 			if (TraceSessionPin->Size > 1024)
 			{
 				TextBuilder.AppendLineFormat(LOCTEXT("TraceSessionTooltip_FileSize2", "File Size: {0} bytes ({1})"), FText::AsNumber(TraceSessionPin->Size), FText::AsMemory(TraceSessionPin->Size));
@@ -370,7 +376,6 @@ SStartPageWindow::SStartPageWindow()
 	, DurationActive(0.0f)
 	, ActiveTimerHandle()
 	, MainContentPanel()
-	, AvailableSessionCount(0)
 	, LiveSessionCount(0)
 #if WITH_EDITOR
 	, bAutoStartAnalysisForLiveSessions(false)
@@ -473,23 +478,23 @@ void SStartPageWindow::Construct(const FArguments& InArgs)
 							]
 						]
 
-					+ SVerticalBox::Slot()
-					.AutoHeight()
-					.HAlign(HAlign_Center)
-					.Padding(3.0f, 3.0f)
-					[
-						SNew(SBox)
-						.WidthOverride(256.0f)
-						[
-							SNew(SBorder)
-							.BorderImage(FEditorStyle::GetBrush("NotificationList.ItemBackground"))
-							.Padding(8.0f)
-							.HAlign(HAlign_Fill)
-							[
-								ConstructRecorderPanel()
-							]
-						]
-					]
+					//+ SVerticalBox::Slot()
+					//.AutoHeight()
+					//.HAlign(HAlign_Center)
+					//.Padding(3.0f, 3.0f)
+					//[
+					//	SNew(SBox)
+					//	.WidthOverride(256.0f)
+					//	[
+					//		SNew(SBorder)
+					//		.BorderImage(FEditorStyle::GetBrush("NotificationList.ItemBackground"))
+					//		.Padding(8.0f)
+					//		.HAlign(HAlign_Fill)
+					//		[
+					//			ConstructRecorderPanel()
+					//		]
+					//	]
+					//]
 
 					+ SVerticalBox::Slot()
 					.AutoHeight()
@@ -666,7 +671,7 @@ TSharedRef<SWidget> SStartPageWindow::ConstructLoadPanel()
 			SNew(SButton)
 			.IsEnabled(this, &SStartPageWindow::Open_IsEnabled)
 			.OnClicked(this, &SStartPageWindow::Open_OnClicked)
-			.ToolTipText(LOCTEXT("OpenButtonTooltip", "Start analysis for selected session."))
+			.ToolTipText(LOCTEXT("OpenButtonTooltip", "Start analysis for selected trace session."))
 			.ContentPadding(FMargin(4.0f, 1.0f))
 			.Content()
 			[
@@ -693,7 +698,7 @@ TSharedRef<SWidget> SStartPageWindow::ConstructLoadPanel()
 		.AutoWidth()
 		[
 			SNew(SComboButton)
-			.ToolTipText(LOCTEXT("MRU_Tooltip", "Open a file or choose a session."))
+			.ToolTipText(LOCTEXT("MRU_Tooltip", "Open a trace file or choose a trace session."))
 			.OnGetMenuContent(this, &SStartPageWindow::MakeSessionListMenu)
 			.HasDownArrow(true)
 			.ContentPadding(FMargin(1.0f, 1.0f, 1.0f, 1.0f))
@@ -1030,14 +1035,17 @@ FReply SStartPageWindow::RefreshTraceSessions_OnClicked()
 
 FReply SStartPageWindow::Connect_OnClicked()
 {
-	TSharedRef<Trace::ISessionService> SessionService = FInsightsManager::Get()->GetSessionService();
-
 	FText HostText = HostTextBox->GetText();
 	if (HostText.IsEmptyOrWhitespace())
 	{
 		// nothing to do
+		return FReply::Handled();
 	}
-	else if (SessionService->ConnectSession(*HostText.ToString()))
+
+	TSharedRef<Trace::ISessionService> SessionService = FInsightsManager::Get()->GetSessionService();
+	const bool bConnectedSuccessfully = SessionService->ConnectSession(*HostText.ToString());
+
+	if (bConnectedSuccessfully)
 	{
 		FNotificationInfo NotificationInfo(FText::Format(LOCTEXT("ConnectSuccess", "Successfully connected to \"{0}\"!"), HostText));
 		NotificationInfo.bFireAndForget = false;
@@ -1070,85 +1078,89 @@ FReply SStartPageWindow::Connect_OnClicked()
 
 void SStartPageWindow::RefreshTraceSessionList()
 {
-	TSharedRef<Trace::ISessionService> SessionService = FInsightsManager::Get()->GetSessionService();
+	Trace::FStoreClient* StoreClient = FInsightsManager::Get()->GetStoreClient();
+	if (StoreClient == nullptr)
+	{
+		return;
+	}
 
-	TArray<Trace::FSessionHandle> AvailableSessions;
-	SessionService->GetAvailableSessions(AvailableSessions);
+	bool bTraceListChanged = false;
 
-	AvailableSessionCount = AvailableSessions.Num();
+	// Update file metadata (size and timestamp).
+	{
+		int32 AvailableTraceCount = 0;
 
-	bool bSessionChanged = false;
+		const int32 TraceCount = StoreClient->GetTraceCount();
+		for (int32 TraceIndex = 0; TraceIndex < TraceCount; ++TraceIndex)
+		{
+			const Trace::FStoreClient::FTraceInfo* TraceInfo = StoreClient->GetTraceInfo(TraceIndex);
+			if (TraceInfo == nullptr)
+			{
+				continue;
+			}
 
-	FString AutoStartPlatformFilterStr = AutoStartPlatformFilter->GetText().ToString();
-	FString AutoStartAppNameFilterStr = AutoStartAppNameFilter->GetText().ToString();
+			AvailableTraceCount++;
 
-	// Count number of live sessions and update file sizes.
-	int32 OldLiveSessionCount = LiveSessionCount;
-	LiveSessionCount = 0;
-	for (Trace::FSessionHandle SessionHandle : AvailableSessions)
-	{
-		Trace::FSessionInfo SessionInfo;
-		SessionService->GetSessionInfo(SessionHandle, SessionInfo);
-		if (SessionInfo.bIsLive)
-		{
-			++LiveSessionCount;
+			const uint32 TraceId = TraceInfo->GetId();
 
-			// Auto start.
-			if (bAutoStartAnalysisForLiveSessions && !AutoStartedSessions.Contains(SessionHandle))
+			TSharedPtr<FTraceSession>* TraceSessionPtrPtr = TraceSessionsMap.Find(TraceId);
+			if (TraceSessionPtrPtr)
 			{
-				if ((AutoStartPlatformFilterStr.IsEmpty() || FCString::Strcmp(*AutoStartPlatformFilterStr, SessionInfo.Platform) == 0) &&
-					(AutoStartAppNameFilterStr.IsEmpty() || FCString::Strcmp(*AutoStartAppNameFilterStr, SessionInfo.AppName) == 0) &&
-					(AutoStartConfigurationTypeFilter == EBuildConfiguration::Unknown || AutoStartConfigurationTypeFilter == SessionInfo.ConfigurationType) &&
-					(AutoStartTargetTypeFilter == EBuildTargetType::Unknown || AutoStartTargetTypeFilter == SessionInfo.TargetType))
-				{
-					AutoStartedSessions.Add(SessionHandle);
-					LoadTraceFile(SessionInfo.Uri);
-				}
+				FTraceSession& TraceSession = **TraceSessionPtrPtr;
+
+				// Reset live status for all traces. Live status will be updated at the end of this function.
+				TraceSession.bIsLive = false;
+				TraceSession.IpAddress = 0;
+
+				TraceSession.Size = TraceInfo->GetSize();
+				TraceSession.Timestamp = FTraceSession::ConvertTimestamp(TraceInfo->GetTimestamp());
+			}
+			else
+			{
+				// New trace detected.
+				bTraceListChanged = true;
+				break;
 			}
 		}
 
-		TSharedPtr<FTraceSession>* TraceSessionPtrPtr = TraceSessionsMap.Find(SessionHandle);
-		if (TraceSessionPtrPtr)
-		{
-			(*TraceSessionPtrPtr)->Size = SessionInfo.Size;
-		}
-		else
-		{
-			bSessionChanged = true;
-		}
+		bTraceListChanged = bTraceListChanged || (AvailableTraceCount != TraceSessions.Num());
 	}
 
-	// If session list has changed on analysis side, recreate the TraceSessions list view widget.
-	//TODO: if (AvailableSessionsChangeNumber != SessionService->GetAvailableSessionsChangeNumber())
-	if (bSessionChanged ||
-		AvailableSessionCount != TraceSessions.Num() ||
-		LiveSessionCount != OldLiveSessionCount)
+	// If trace list has changed on store side, recreate the TraceSessions list view widget.
+	if (bTraceListChanged)
 	{
 		TSharedPtr<FTraceSession> NewSelectedTraceSession;
 
 		TraceSessions.Reset();
 		TraceSessionsMap.Reset();
 
-		for (Trace::FSessionHandle SessionHandle : AvailableSessions)
+		const int32 TraceCount = StoreClient->GetTraceCount();
+		for (int32 TraceIndex = 0; TraceIndex < TraceCount; ++TraceIndex)
 		{
-			Trace::FSessionInfo SessionInfo;
-			SessionService->GetSessionInfo(SessionHandle, SessionInfo);
+			const Trace::FStoreClient::FTraceInfo* TraceInfo = StoreClient->GetTraceInfo(TraceIndex);
+			if (TraceInfo == nullptr)
+			{
+				continue;
+			}
 
-			const TSharedPtr<FTraceSession> TraceSessionPtr = MakeShareable(new FTraceSession(SessionHandle, SessionInfo));
-			TraceSessions.Add(TraceSessionPtr);
-			TraceSessionsMap.Add(SessionHandle, TraceSessionPtr);
+			const TSharedRef<FTraceSession> TraceSession = MakeShared<FTraceSession>(TraceInfo);
+			TraceSession->TraceIndex = TraceIndex;
+			TraceSession->Uri = FText::FromString(FInsightsManager::Get()->GetStoreDir() + TEXT("/") + TraceSession->Name.ToString() + TEXT(".utrace"));
+			TraceSessions.Add(TraceSession);
+			TraceSessionsMap.Add(TraceSession->TraceId, TraceSession);
 
-			// Identify the previously selected session (if stil available) to ensure selection remains unchanged.
-			if (SelectedTraceSession && SelectedTraceSession->SessionHandle == SessionHandle)
+			// Identify the previously selected trace session (if stil available) to ensure selection remains unchanged.
+			if (SelectedTraceSession && SelectedTraceSession->TraceId == TraceSession->TraceId)
 			{
-				NewSelectedTraceSession = TraceSessionPtr;
+				NewSelectedTraceSession = TraceSession;
 			}
 		}
-		Algo::SortBy(TraceSessions, &FTraceSession::TimeStamp);
+
+		Algo::SortBy(TraceSessions, &FTraceSession::Timestamp);
 
 		TraceSessionsListView->RebuildList();
 
-		// If no selection, auto select the last (newest) session.
+		// If no selection, auto select the last (newest) trace session.
 		if (!NewSelectedTraceSession.IsValid() && TraceSessions.Num() > 0)
 		{
 			NewSelectedTraceSession = TraceSessions.Last();
@@ -1156,13 +1168,60 @@ void SStartPageWindow::RefreshTraceSessionList()
 
 		TraceSessionsListView->ScrollToBottom();
 
-		// Restore selection and ensure is visible.
+		// Restore selection and ensure it is visible.
 		if (NewSelectedTraceSession.IsValid())
 		{
 			TraceSessionsListView->SetItemSelection(NewSelectedTraceSession, true);
 			TraceSessionsListView->RequestScrollIntoView(NewSelectedTraceSession);
 		}
 	}
+
+	// Process the connected recorder sessions.
+	{
+		const FString AutoStartPlatformFilterStr = AutoStartPlatformFilter->GetText().ToString();
+		const FString AutoStartAppNameFilterStr = AutoStartAppNameFilter->GetText().ToString();
+
+		const int32 SessionCount = StoreClient->GetSessionCount();
+		for (int32 SessionIndex = 0; SessionIndex < SessionCount; ++SessionIndex)
+		{
+			const Trace::FStoreClient::FSessionInfo* SessionInfo = StoreClient->GetSessionInfo(SessionIndex);
+			if (SessionInfo == nullptr)
+			{
+				continue;
+			}
+
+			const uint32 TraceId = SessionInfo->GetTraceId();
+
+			TSharedPtr<FTraceSession>* TraceSessionPtrPtr = TraceSessionsMap.Find(TraceId);
+			if (TraceSessionPtrPtr)
+			{
+				FTraceSession& TraceSession = **TraceSessionPtrPtr;
+
+				// Update live status.
+				TraceSession.bIsLive = true;
+				TraceSession.IpAddress = SessionInfo->GetIpAddress();
+
+				// Auto start anaysis for a live session.
+				if (bAutoStartAnalysisForLiveSessions && // is auto start enabled?
+					!AutoStartedSessions.Contains(TraceId)) // is not already auto-started?
+				{
+					// matches filter?
+					if ((AutoStartPlatformFilterStr.IsEmpty() || FCString::Strcmp(*AutoStartPlatformFilterStr, *TraceSession.Platform.ToString()) == 0) &&
+						(AutoStartAppNameFilterStr.IsEmpty() || FCString::Strcmp(*AutoStartAppNameFilterStr, *TraceSession.AppName.ToString()) == 0) &&
+						(AutoStartConfigurationTypeFilter == EBuildConfiguration::Unknown || AutoStartConfigurationTypeFilter == TraceSession.ConfigurationType) &&
+						(AutoStartTargetTypeFilter == EBuildTargetType::Unknown || AutoStartTargetTypeFilter == TraceSession.TargetType))
+					{
+						AutoStartedSessions.Add(TraceSession.TraceId);
+						LoadTrace(TraceSession.TraceId);
+					}
+				}
+			}
+			else
+			{
+				// Trace not found. It will be picked up by the next RefreshTraceSessionList() call.
+			}
+		}
+	}
 }
 
 ////////////////////////////////////////////////////////////////////////////////////////////////////
@@ -1332,7 +1391,7 @@ FReply SStartPageWindow::OnDrop(const FGeometry& MyGeometry, const FDragDropEven
 
 bool SStartPageWindow::Open_IsEnabled() const
 {
-	return AvailableSessionCount > 0;
+	return TraceSessions.Num() > 0;
 }
 
 ////////////////////////////////////////////////////////////////////////////////////////////////////
@@ -1347,8 +1406,7 @@ FReply SStartPageWindow::Open_OnClicked()
 
 void SStartPageWindow::OpenFileDialog()
 {
-	TSharedRef<Trace::ISessionService> SessionService = FInsightsManager::Get()->GetSessionService();
-	const FString ProfilingDirectory(FPaths::ConvertRelativePathToFull(SessionService->GetLocalSessionDirectory()));
+	const FString ProfilingDirectory(FPaths::ConvertRelativePathToFull(FInsightsManager::Get()->GetStoreDir()));
 
 	TArray<FString> OutFiles;
 	bool bOpened = false;
@@ -1385,14 +1443,7 @@ void SStartPageWindow::LoadTraceSession(TSharedPtr<FTraceSession> InTraceSession
 {
 	if (InTraceSession.IsValid())
 	{
-		if (FInsightsManager::Get()->ShouldOpenAnalysisInSeparateProcess())
-		{
-			LoadTraceFile(InTraceSession->Uri.ToString());
-		}
-		else
-		{
-			LoadSession(InTraceSession->SessionHandle);
-		}
+		LoadTrace(InTraceSession->TraceId);
 	}
 }
 
@@ -1402,9 +1453,11 @@ void SStartPageWindow::LoadTraceFile(const FString& InTraceFile)
 {
 	if (FInsightsManager::Get()->ShouldOpenAnalysisInSeparateProcess())
 	{
+		UE_LOG(TimingProfiler, Log, TEXT("Start analysis (in separate process) for trace file: \"%s\""), *InTraceFile);
+
 		const TCHAR* ExecutablePath = FPlatformProcess::ExecutablePath();
 
-		FString CmdLine = TEXT("-Trace=\"") + InTraceFile + TEXT("\"");
+		FString CmdLine = TEXT("-TraceFile=\"") + InTraceFile + TEXT("\"");
 
 		constexpr bool bLaunchDetached = false;
 		constexpr bool bLaunchHidden = false;
@@ -1425,24 +1478,50 @@ void SStartPageWindow::LoadTraceFile(const FString& InTraceFile)
 	}
 	else
 	{
+		UE_LOG(TimingProfiler, Log, TEXT("Start analysis for trace file: \"%s\""), *InTraceFile);
 		FInsightsManager::Get()->LoadTraceFile(InTraceFile);
 	}
 }
 
 ////////////////////////////////////////////////////////////////////////////////////////////////////
 
-void SStartPageWindow::LoadSession(Trace::FSessionHandle SessionHandle)
+void SStartPageWindow::LoadTrace(uint32 InTraceId)
 {
 	if (FInsightsManager::Get()->ShouldOpenAnalysisInSeparateProcess())
 	{
-		TSharedRef<Trace::ISessionService> SessionService = FInsightsManager::Get()->GetSessionService();
-		Trace::FSessionInfo SessionInfo;
-		SessionService->GetSessionInfo(SessionHandle, SessionInfo);
-		LoadTraceFile(SessionInfo.Uri);
+		UE_LOG(TimingProfiler, Log, TEXT("Start analysis (in separate process) for trace id: 0x%08X"), InTraceId);
+
+		const TCHAR* ExecutablePath = FPlatformProcess::ExecutablePath();
+
+		const uint32 StorePort = FInsightsManager::Get()->GetStorePort();
+		FString CmdLine = FString::Printf(TEXT("-TraceId=%d -StorePort=%d"), InTraceId, StorePort);
+
+		constexpr bool bLaunchDetached = false;
+		constexpr bool bLaunchHidden = false;
+		constexpr bool bLaunchReallyHidden = false;
+
+		uint32 ProcessID = 0;
+		const int32 PriorityModifier = 0;
+		const TCHAR* OptionalWorkingDirectory = nullptr;
+
+		void* PipeWriteChild = nullptr;
+		void* PipeReadChild = nullptr;
+
+		FProcHandle Handle = FPlatformProcess::CreateProc(ExecutablePath, *CmdLine, bLaunchDetached, bLaunchHidden, bLaunchReallyHidden, &ProcessID, PriorityModifier, OptionalWorkingDirectory, PipeWriteChild, PipeReadChild);
+		FPlatformProcess::CloseProc(Handle);
+
+		TSharedPtr<FTraceSession>* TraceSessionPtrPtr = TraceSessionsMap.Find(InTraceId);
+		if (TraceSessionPtrPtr)
+		{
+			FTraceSession& TraceSession = **TraceSessionPtrPtr;
+			SplashScreenOverlayTraceFile = FPaths::GetBaseFilename(TraceSession.Uri.ToString());
+		}
+		ShowSplashScreenOverlay();
 	}
 	else
 	{
-		FInsightsManager::Get()->LoadSession(SessionHandle);
+		UE_LOG(TimingProfiler, Log, TEXT("Start analysis for trace id: 0x%08X"), InTraceId);
+		FInsightsManager::Get()->LoadTrace(InTraceId);
 	}
 }
 
@@ -1467,93 +1546,37 @@ TSharedRef<SWidget> SStartPageWindow::MakeSessionListMenu()
 	}
 	MenuBuilder.EndSection();
 
-	/* TODO: persistent MRU
-	MenuBuilder.BeginSection("MostRecentlyUsedSessions", LOCTEXT("MostRecentlyUsedSessionsHeading", "Most Recently Used Sessions"));
-	{
-		TSharedRef<Trace::ISessionService> SessionService = FInsightsManager::Get()->GetSessionService();
-
-		TArray<Trace::FSessionHandle> MruSessionList;
-
-		//TODO: real mru list
-		SessionService->GetAvailableSessions(MruSessionList); // mock the list using available sessions
-		int32 NumSessions = FMath::Min(3, MruSessionList.Num());
-
-		for (int32 SessionIndex = 0; SessionIndex < NumSessions; ++SessionIndex)
-		{
-			Trace::FSessionHandle SessionHandle = MruSessionList[SessionIndex];
-
-			Trace::FSessionInfo SessionInfo;
-			SessionService->GetSessionInfo(SessionHandle, SessionInfo);
-
-			FText Label = FText::FromString(SessionInfo.Name);
-			if (SessionInfo.bIsLive)
-			{
-				Label = FText::Format(LOCTEXT("LiveSessionTextFmt", "{0} (LIVE!)"), Label);
-			}
-
-			MenuBuilder.AddMenuEntry(
-				Label,
-				TAttribute<FText>(), // no tooltip
-				FSlateIcon(),
-				FUIAction(FExecuteAction::CreateSP(this, &SStartPageWindow::LoadSession, SessionHandle)),
-				NAME_None,
-				EUserInterfaceActionType::Button
-			);
-		}
-	}
-	MenuBuilder.EndSection();
-	*/
-
 	MenuBuilder.BeginSection("AvailableSessions", LOCTEXT("AvailableSessionsHeading", "Top 10 Most Recently Created Sessions"));
 	{
-		TSharedRef<Trace::ISessionService> SessionService = FInsightsManager::Get()->GetSessionService();
-
-		TArray<Trace::FSessionHandle> AvailableSessions;
-		SessionService->GetAvailableSessions(AvailableSessions);
-
-		struct FSessionInfoEx
-		{
-			FDateTime GetTimeStamp() const { return SessionInfo.TimeStamp; }
-
-			Trace::FSessionHandle SessionHandle;
-			Trace::FSessionInfo  SessionInfo;
-		};
-
-		TArray<TSharedRef<FSessionInfoEx>> SortedAvailableSessions;
-		SortedAvailableSessions.Reserve(AvailableSessions.Num());
-
-		for (Trace::FSessionHandle SessionHandle : AvailableSessions)
+		Trace::FStoreClient* StoreClient = FInsightsManager::Get()->GetStoreClient();
+		if (StoreClient != nullptr)
 		{
-			TSharedRef<FSessionInfoEx> SessionInfoEx = MakeShared<FSessionInfoEx>();
-			SessionInfoEx->SessionHandle = SessionHandle;
-			SessionService->GetSessionInfo(SessionHandle, SessionInfoEx->SessionInfo);
-			SortedAvailableSessions.Add(SessionInfoEx);
-		}
+			// Make a copy of the sessions list (to allow TraceSessions to be sorted by other criteria).
+			TArray<TSharedPtr<FTraceSession>> SortedTraceSessions(TraceSessions);
+			Algo::SortBy(SortedTraceSessions, &FTraceSession::Timestamp);
 
-		Algo::SortBy(SortedAvailableSessions, &FSessionInfoEx::GetTimeStamp);
+			int32 TraceCountLimit = 10; // top 10
 
-		int32 SessionCountLimit = 10; // top 10
+			// Iterate in reverse order as we want most recent sessions first.
+			for (int32 TraceSessionIndex = SortedTraceSessions.Num() - 1; TraceSessionIndex >= 0 && TraceCountLimit > 0; --TraceSessionIndex, --TraceCountLimit)
+			{
+				const FTraceSession& TraceSession = *SortedTraceSessions[TraceSessionIndex];
 
-		// Iterate in reverse order as we want most recent sessions first.
-		for (int32 SessionIndex = SortedAvailableSessions.Num() - 1; SessionIndex >= 0 && SessionCountLimit > 0; --SessionIndex, --SessionCountLimit)
-		{
-			const Trace::FSessionHandle SessionHandle = SortedAvailableSessions[SessionIndex]->SessionHandle;
-			const Trace::FSessionInfo& SessionInfo = SortedAvailableSessions[SessionIndex]->SessionInfo;
+				FText Label = TraceSession.Name;
+				if (TraceSession.bIsLive)
+				{
+					Label = FText::Format(LOCTEXT("LiveSessionTextFmt", "{0} (LIVE!)"), Label);
+				}
 
-			FText Label = FText::FromString(SessionInfo.Name);
-			if (SessionInfo.bIsLive)
-			{
-				Label = FText::Format(LOCTEXT("LiveSessionTextFmt", "{0} (LIVE!)"), Label);
+				MenuBuilder.AddMenuEntry(
+					Label,
+					TAttribute<FText>(), // no tooltip
+					FSlateIcon(),
+					FUIAction(FExecuteAction::CreateSP(this, &SStartPageWindow::LoadTrace, TraceSession.TraceId)),
+					NAME_None,
+					EUserInterfaceActionType::Button
+				);
 			}
-
-			MenuBuilder.AddMenuEntry(
-				Label,
-				TAttribute<FText>(), // no tooltip
-				FSlateIcon(),
-				FUIAction(FExecuteAction::CreateSP(this, &SStartPageWindow::LoadSession, SessionHandle)),
-				NAME_None,
-				EUserInterfaceActionType::Button
-			);
 		}
 	}
 	MenuBuilder.EndSection();
@@ -1565,16 +1588,14 @@ TSharedRef<SWidget> SStartPageWindow::MakeSessionListMenu()
 
 FText SStartPageWindow::GetLocalSessionDirectory() const
 {
-	TSharedRef<Trace::ISessionService> SessionService = FInsightsManager::Get()->GetSessionService();
-	return FText::FromString(FPaths::ConvertRelativePathToFull(SessionService->GetLocalSessionDirectory()));
+	return FText::FromString(FPaths::ConvertRelativePathToFull(FInsightsManager::Get()->GetStoreDir()));
 }
 
 ////////////////////////////////////////////////////////////////////////////////////////////////////
 
 FReply SStartPageWindow::ExploreLocalSessionDirectory_OnClicked()
 {
-	TSharedRef<Trace::ISessionService> SessionService = FInsightsManager::Get()->GetSessionService();
-	FString FullPath(FPaths::ConvertRelativePathToFull(SessionService->GetLocalSessionDirectory()));
+	FString FullPath(FPaths::ConvertRelativePathToFull(FInsightsManager::Get()->GetStoreDir()));
 	FPlatformProcess::ExploreFolder(*FullPath);
 	return FReply::Handled();
 }
@@ -1583,8 +1604,10 @@ FReply SStartPageWindow::ExploreLocalSessionDirectory_OnClicked()
 
 FText SStartPageWindow::GetRecorderStatusText() const
 {
-	TSharedRef<Trace::ISessionService> SessionService = FInsightsManager::Get()->GetSessionService();
-	if (SessionService->IsRecorderServerRunning())
+	Trace::FStoreClient* StoreClient = FInsightsManager::Get()->GetStoreClient();
+	const bool bIsRecorderServerRunning = (StoreClient != nullptr); // TODO: StoreClient->IsRecorderServerRunning();
+
+	if (bIsRecorderServerRunning)
 	{
 		return FText(LOCTEXT("RecorderServerRunning", "Running"));
 	}
@@ -1598,24 +1621,27 @@ FText SStartPageWindow::GetRecorderStatusText() const
 
 EVisibility SStartPageWindow::StartTraceRecorder_Visibility() const
 {
-	TSharedRef<Trace::ISessionService> SessionService = FInsightsManager::Get()->GetSessionService();
-	return SessionService->IsRecorderServerRunning() ? EVisibility::Collapsed : EVisibility::Visible;
+	Trace::FStoreClient* StoreClient = FInsightsManager::Get()->GetStoreClient();
+	const bool bIsRecorderServerRunning = (StoreClient != nullptr); // TODO: StoreClient->IsRecorderServerRunning();
+
+	return bIsRecorderServerRunning ? EVisibility::Collapsed : EVisibility::Visible;
 }
 
 ////////////////////////////////////////////////////////////////////////////////////////////////////
 
 EVisibility SStartPageWindow::StopTraceRecorder_Visibility() const
 {
-	TSharedRef<Trace::ISessionService> SessionService = FInsightsManager::Get()->GetSessionService();
-	return SessionService->IsRecorderServerRunning() ? EVisibility::Visible : EVisibility::Collapsed;
+	Trace::FStoreClient* StoreClient = FInsightsManager::Get()->GetStoreClient();
+	const bool bIsRecorderServerRunning = (StoreClient != nullptr); // TODO: StoreClient->IsRecorderServerRunning();
+
+	return bIsRecorderServerRunning ? EVisibility::Visible : EVisibility::Collapsed;
 }
 
 ////////////////////////////////////////////////////////////////////////////////////////////////////
 
 FReply SStartPageWindow::StartTraceRecorder_OnClicked()
 {
-	TSharedRef<Trace::ISessionService> SessionService = FInsightsManager::Get()->GetSessionService();
-	SessionService->StartRecorderServer();
+	//TODO: StoreClient->StartRecorderServer();
 	RefreshTraceSessionList();
 	return FReply::Handled();
 }
@@ -1624,8 +1650,7 @@ FReply SStartPageWindow::StartTraceRecorder_OnClicked()
 
 FReply SStartPageWindow::StopTraceRecorder_OnClicked()
 {
-	TSharedRef<Trace::ISessionService> SessionService = FInsightsManager::Get()->GetSessionService();
-	SessionService->StopRecorderServer();
+	//TODO: StoreClient->StopRecorderServer();
 	RefreshTraceSessionList();
 	return FReply::Handled();
 }
diff --git a/Engine/Source/Developer/TraceInsights/Private/Insights/Widgets/SStartPageWindow.h b/Engine/Source/Developer/TraceInsights/Private/Insights/Widgets/SStartPageWindow.h
--- a/Engine/Source/Developer/TraceInsights/Private/Insights/Widgets/SStartPageWindow.h
+++ b/Engine/Source/Developer/TraceInsights/Private/Insights/Widgets/SStartPageWindow.h
@@ -9,6 +9,7 @@
 #include "Layout/Visibility.h"
 #include "Misc/Guid.h"
 #include "SlateFwd.h"
+#include "Trace/StoreClient.h"
 #include "TraceServices/ModuleService.h"
 #include "TraceServices/SessionService.h"
 #include "Widgets/DeclarativeSyntaxSupport.h"
@@ -42,33 +43,52 @@ typedef TWeakPtr<class SNotificationItem> SNotificationItemWeak;
 
 struct FTraceSession
 {
-	Trace::FSessionHandle SessionHandle;
+	uint32 TraceId;
+	int32 TraceIndex; // debug
+
 	FText Name;
 	FText Uri;
+	FDateTime Timestamp = 0;
+	uint64 Size = 0;
+
+	bool bIsLive = false;
+	uint32 IpAddress = 0;
+
 	FText Platform;
 	FText AppName;
 	FText CommandLine;
-	EBuildConfiguration ConfigurationType;
-	EBuildTargetType TargetType;
-	FDateTime TimeStamp;
-	uint64 Size;
-	bool bIsLive;
-
-	FTraceSession(const Trace::FSessionHandle InSessionHandle, const Trace::FSessionInfo& InSessionInfo)
-		: SessionHandle(InSessionHandle)
-		, Name(FText::FromString(InSessionInfo.Name))
-		, Uri(FText::FromString(InSessionInfo.Uri))
-		, Platform(FText::FromString(InSessionInfo.Platform))
-		, AppName(FText::FromString(InSessionInfo.AppName))
-		, CommandLine(FText::FromString(InSessionInfo.CommandLine))
-		//, ConfigurationType(EBuildConfiguration::Unknown)
-		, ConfigurationType(InSessionInfo.ConfigurationType)
-		//, TargetType(EBuildTargetType::Unknown)
-		, TargetType(InSessionInfo.TargetType)
-		, TimeStamp(InSessionInfo.TimeStamp)
-		, Size(InSessionInfo.Size)
-		, bIsLive(InSessionInfo.bIsLive)
-	{}
+	EBuildConfiguration ConfigurationType = EBuildConfiguration::Unknown;
+	EBuildTargetType TargetType = EBuildTargetType::Unknown;
+
+	//FTraceSession() = default;
+	FTraceSession(const Trace::FStoreClient::FTraceInfo* InTraceInfo)
+		: TraceId(InTraceInfo->GetId())
+		, TraceIndex(-1)
+		, Name(AnsiStringViewToText(InTraceInfo->GetName()))
+		, Uri(/*TODO: InTraceInfo->GetUri()*/)
+		, Timestamp(ConvertTimestamp(InTraceInfo->GetTimestamp()))
+		, Size(InTraceInfo->GetSize())
+		, bIsLive(false)
+		, IpAddress(0)
+		, Platform()
+		, AppName()
+		, CommandLine()
+		, ConfigurationType(EBuildConfiguration::Unknown)
+		, TargetType(EBuildTargetType::Unknown)
+	{
+	}
+
+	static FDateTime ConvertTimestamp(uint64 InTimestamp)
+	{
+		return FDateTime(static_cast<int64>(InTimestamp));
+	}
+
+private:
+	static FText AnsiStringViewToText(const FAnsiStringView& AnsiStringView)
+	{
+		FString FatString(AnsiStringView.Len(), AnsiStringView.GetData());
+		return FText::FromString(FatString);
+	}
 };
 
 /** Implements the Start Page window. */
@@ -122,7 +142,7 @@ private:
 
 	void LoadTraceSession(TSharedPtr<FTraceSession> InTraceSession);
 	void LoadTraceFile(const FString& InTraceFile);
-	void LoadSession(Trace::FSessionHandle InSessionHandle);
+	void LoadTrace(uint32 InTraceId);
 
 	TSharedRef<SWidget> MakeSessionListMenu();
 
@@ -227,11 +247,10 @@ private:
 	/** Holds all widgets for the profiler window like menu bar, toolbar and tabs. */
 	TSharedPtr<SVerticalBox> MainContentPanel;
 
-	int32 AvailableSessionCount;
 	int32 LiveSessionCount;
 
 	bool bAutoStartAnalysisForLiveSessions;
-	TArray<Trace::FSessionHandle> AutoStartedSessions; // tracks sessions that were auto started (in order to not start them again)
+	TArray<uint32> AutoStartedSessions; // tracks sessions that were auto started (in order to not start them again)
 
 	TSharedPtr<SSearchBox> AutoStartPlatformFilter;
 	TSharedPtr<SSearchBox> AutoStartAppNameFilter;
@@ -240,7 +259,7 @@ private:
 
 	TSharedPtr<SListView<TSharedPtr<FTraceSession>>> TraceSessionsListView;
 	TArray<TSharedPtr<FTraceSession>> TraceSessions;
-	TMap<Trace::FSessionHandle, TSharedPtr<FTraceSession>> TraceSessionsMap;
+	TMap<uint32, TSharedPtr<FTraceSession>> TraceSessionsMap;
 	TSharedPtr<SEditableTextBox> HostTextBox;
 	TSharedPtr<FTraceSession> SelectedTraceSession;
 
diff --git a/Engine/Source/Developer/TraceInsights/Public/Insights/IUnrealInsightsModule.h b/Engine/Source/Developer/TraceInsights/Public/Insights/IUnrealInsightsModule.h
--- a/Engine/Source/Developer/TraceInsights/Public/Insights/IUnrealInsightsModule.h
+++ b/Engine/Source/Developer/TraceInsights/Public/Insights/IUnrealInsightsModule.h
@@ -8,6 +8,11 @@
 
 class FExtender;
 
+namespace Trace
+{
+	class FStoreClient;
+}
+
 /** Major tab IDs for Insights tools */
 struct TRACEINSIGHTS_API FInsightsManagerTabs
 {
@@ -118,16 +123,18 @@ public:
 	/**
 	 * Starts analysis of the specified *.utrace file. Called when the application starts in "Viewer" mode.
 	 *
-	 * @param InTraceFile The file path to the *.utrace file to analyze.
+	 * @param InTraceFile The filename (*.utrace) of the trace to analyze.
 	 */
 	virtual void StartAnalysisForTraceFile(const TCHAR* InTraceFile) = 0;
 
 	/**
-	 * Starts analysis of the specified session. Called when the application starts in "Viewer" mode.
+	 * Starts analysis of the specified trace. Called when the application starts in "Viewer" mode.
 	 *
-	 * @param InSessionId The id of the session to analyze.
+	 * @param InStoreHost The host of the store to connect to.
+	 * @param InStorePort The port of the store to connect to.
+	 * @param InTraceId The id of the trace to analyze.
 	 */
-	virtual void StartAnalysisForSession(const TCHAR* InSessionId) = 0;
+	virtual void StartAnalysisForTrace(const TCHAR* InStoreHost, uint32 InStorePort, uint32 InTraceId) = 0;
 
 	/**
 	 * Starts analysis of the last live session. Called when the application starts in "Viewer" mode.
@@ -139,7 +146,7 @@ public:
 	 */
 	virtual void ShutdownUserInterface() = 0;
 
-	/** 
+	/**
 	 * Registers a major tab layout. This defines how the major tab will appear when spawned.
 	 * If this is not called prior to tabs being spawned then the built-in default layout will be used.
 	 * @param	InMajorTabId	The major tab ID we are supplying a layout to
@@ -147,7 +154,7 @@ public:
 	 */
 	virtual void RegisterMajorTabConfig(const FName& InMajorTabId, const FInsightsMajorTabConfig& InConfig) = 0;
 
-	/** 
+	/**
 	 * Unregisters a major tab layout. This will revert the major tab to spawning with its default layout
 	 * @param	InMajorTabId	The major tab ID we are supplying a layout to
 	 */
diff --git a/Engine/Source/Developer/TraceInsights/TraceInsights.Build.cs b/Engine/Source/Developer/TraceInsights/TraceInsights.Build.cs
--- a/Engine/Source/Developer/TraceInsights/TraceInsights.Build.cs
+++ b/Engine/Source/Developer/TraceInsights/TraceInsights.Build.cs
@@ -24,6 +24,7 @@ public class TraceInsights : ModuleRules
 				"Slate",
 				"EditorStyle",
 				"TraceLog",
+				"TraceAnalysis",
 				"TraceServices",
 				"DesktopPlatform",
 				"WorkspaceMenuStructure",
diff --git a/Engine/Source/Programs/UnrealInsights/Private/UserInterfaceCommand.cpp b/Engine/Source/Programs/UnrealInsights/Private/UserInterfaceCommand.cpp
--- a/Engine/Source/Programs/UnrealInsights/Private/UserInterfaceCommand.cpp
+++ b/Engine/Source/Programs/UnrealInsights/Private/UserInterfaceCommand.cpp
@@ -96,17 +96,20 @@ void FUserInterfaceCommand::Run()
 {
 	// Only a single instance of Session Browser window/process is allowed.
 	{
-		const uint32 MaxPath = FPlatformMisc::GetMaxPathLength();
-		TCHAR* TraceFile = new TCHAR[MaxPath + 1];
-		TraceFile[0] = 0;
-		bool bUseTraceFile = FParse::Value(FCommandLine::Get(), TEXT("-Trace="), TraceFile, MaxPath, true);
+		bool bBrowserMode = true;
+
+		if (bBrowserMode)
+		{
+			bBrowserMode = FCString::Strifind(FCommandLine::Get(), TEXT("-TraceId=")) == nullptr;
+		}
+		if (bBrowserMode)
+		{
+			bBrowserMode = FCString::Strifind(FCommandLine::Get(), TEXT("-TraceFile=")) == nullptr;
+		}
 
-		if (!bUseTraceFile)
+		if (bBrowserMode && !CheckSessionBrowserSingleInstance())
 		{
-			if (!CheckSessionBrowserSingleInstance())
-			{
-				return;
-			}
+			return;
 		}
 	}
 
@@ -206,22 +209,43 @@ void FUserInterfaceCommand::InitializeSlateApplication()
 	IUnrealInsightsModule& TraceInsightsModule = FModuleManager::LoadModuleChecked<IUnrealInsightsModule>("TraceInsights");
 
 	const uint32 MaxPath = FPlatformMisc::GetMaxPathLength();
-	TCHAR* TraceFile = new TCHAR[MaxPath + 1];
-	TraceFile[0] = 0;
-	bool bUseTraceFile = FParse::Value(FCommandLine::Get(), TEXT("-Trace="), TraceFile, MaxPath, true);
 
-	if (bUseTraceFile)
+	uint32 TraceId = 0;
+	bool bUseTraceId = FParse::Value(FCommandLine::Get(), TEXT("-TraceId="), TraceId);
+
+	if (bUseTraceId)
 	{
+		TCHAR* StoreHost = new TCHAR[MaxPath + 1];
+		FCString::Strcpy(StoreHost, MaxPath, TEXT("127.0.0.1"));
+		FParse::Value(FCommandLine::Get(), TEXT("-StoreHost="), StoreHost, MaxPath, true);
+
+		uint32 StorePort = 0;
+		FParse::Value(FCommandLine::Get(), TEXT("-StorePort="), StorePort);
+
 		TraceInsightsModule.CreateSessionViewer(bAllowDebugTools);
-		TraceInsightsModule.StartAnalysisForTraceFile(TraceFile);
+		TraceInsightsModule.StartAnalysisForTrace(StoreHost, StorePort, TraceId);
+
+		delete[] StoreHost;
 	}
 	else
 	{
-		const bool bSingleProcess = FParse::Param(FCommandLine::Get(), TEXT("SingleProcess"));
-		TraceInsightsModule.CreateSessionBrowser(bAllowDebugTools, bSingleProcess);
-	}
+		TCHAR* TraceFile = new TCHAR[MaxPath + 1];
+		TraceFile[0] = 0;
+		bool bUseTraceFile = FParse::Value(FCommandLine::Get(), TEXT("-TraceFile="), TraceFile, MaxPath, true);
+
+		if (bUseTraceFile)
+		{
+			TraceInsightsModule.CreateSessionViewer(bAllowDebugTools);
+			TraceInsightsModule.StartAnalysisForTraceFile(TraceFile);
+		}
+		else
+		{
+			const bool bSingleProcess = FParse::Param(FCommandLine::Get(), TEXT("SingleProcess"));
+			TraceInsightsModule.CreateSessionBrowser(bAllowDebugTools, bSingleProcess);
+		}
 
-	delete[] TraceFile;
+		delete[] TraceFile;
+	}
 }
 
 ////////////////////////////////////////////////////////////////////////////////////////////////////
