Network Insights
- Changed find packet text box to use packet sequence numbers instead of packet index as this makes more sense when trying to match up data between different connections

diff --git a/Engine/Source/Developer/TraceInsights/Private/Insights/NetworkingProfiler/Widgets/SPacketContentView.cpp b/Engine/Source/Developer/TraceInsights/Private/Insights/NetworkingProfiler/Widgets/SPacketContentView.cpp
--- a/Engine/Source/Developer/TraceInsights/Private/Insights/NetworkingProfiler/Widgets/SPacketContentView.cpp
+++ b/Engine/Source/Developer/TraceInsights/Private/Insights/NetworkingProfiler/Widgets/SPacketContentView.cpp
@@ -59,6 +59,7 @@ void SPacketContentView::Reset()
 	ConnectionIndex = 0;
 	ConnectionMode = Trace::ENetProfilerConnectionMode::Outgoing;
 	PacketIndex = 0;
+	PacketSequence = 0;
 	PacketBitSize = 0;
 
 	bFilterByEventType = false;
@@ -376,7 +377,7 @@ FReply SPacketContentView::FindNextPacket_OnClicked()
 
 FText SPacketContentView::GetPacketText() const
 {
-	return FText::AsNumber(PacketIndex);
+	return FText::AsNumber(PacketSequence);
 }
 
 ////////////////////////////////////////////////////////////////////////////////////////////////////
@@ -385,13 +386,13 @@ void SPacketContentView::Packet_OnTextCommitted(const FText& InNewText, ETextCom
 {
 	if (InNewText.IsNumeric())
 	{
-		int32 NewPacketIndex = 0;
-		TTypeFromString<int32>::FromString(NewPacketIndex, *InNewText.ToString());
+		uint32 NewPacketSequence = 0;
+		TTypeFromString<uint32>::FromString(NewPacketSequence, *InNewText.ToString());
 
 		TSharedPtr<SPacketView> PacketView = ProfilerWindow->GetPacketView();
 		if (PacketView.IsValid())
 		{
-			PacketView->SetSelectedPacket(NewPacketIndex);
+			PacketView->SelectPacketBySequenceNumber(NewPacketSequence);
 		}
 	}
 }
@@ -621,6 +622,7 @@ void SPacketContentView::ResetPacket()
 	ConnectionIndex = 0;
 	ConnectionMode = Trace::ENetProfilerConnectionMode::Outgoing;
 	PacketIndex = 0;
+	PacketSequence = 0;
 	PacketBitSize = 0;
 
 	FAxisViewportDouble& ViewportX = Viewport.GetHorizontalAxisViewport();
@@ -644,6 +646,7 @@ void SPacketContentView::SetPacket(uint32 InGameInstanceIndex, uint32 InConnecti
 	ConnectionIndex = InConnectionIndex;
 	ConnectionMode = InConnectionMode;
 	PacketIndex = InPacketIndex;
+	PacketSequence = GetPacketSequence(InPacketIndex);
 	PacketBitSize = InPacketBitSize;
 
 	FAxisViewportDouble& ViewportX = Viewport.GetHorizontalAxisViewport();
@@ -722,6 +725,27 @@ void SPacketContentView::EnableFilterEventType(const uint32 InEventTypeIndex)
 
 ////////////////////////////////////////////////////////////////////////////////////////////////////
 
+uint32 SPacketContentView::GetPacketSequence(int32 InPacketIndex) const
+{
+	uint32 NewSequenceNumber = 0U;
+
+	TSharedPtr<const Trace::IAnalysisSession> Session = FInsightsManager::Get()->GetSession();
+	if (Session.IsValid())
+	{
+		Trace::FAnalysisSessionReadScope SessionReadScope(*Session.Get());
+		const Trace::INetProfilerProvider& NetProfilerProvider = Trace::ReadNetProfilerProvider(*Session.Get());
+
+		NetProfilerProvider.EnumeratePackets(ConnectionIndex, ConnectionMode, InPacketIndex, InPacketIndex, [&NewSequenceNumber](const Trace::FNetProfilerPacket& Packet)
+		{
+			NewSequenceNumber = Packet.SequenceNumber;
+		});
+	}
+
+	return NewSequenceNumber;
+}
+
+////////////////////////////////////////////////////////////////////////////////////////////////////
+
 void SPacketContentView::DisableFilterEventType()
 {
 	bFilterByEventType = false;
diff --git a/Engine/Source/Developer/TraceInsights/Private/Insights/NetworkingProfiler/Widgets/SPacketContentView.h b/Engine/Source/Developer/TraceInsights/Private/Insights/NetworkingProfiler/Widgets/SPacketContentView.h
--- a/Engine/Source/Developer/TraceInsights/Private/Insights/NetworkingProfiler/Widgets/SPacketContentView.h
+++ b/Engine/Source/Developer/TraceInsights/Private/Insights/NetworkingProfiler/Widgets/SPacketContentView.h
@@ -194,6 +194,8 @@ private:
 
 	////////////////////////////////////////////////////////////////////////////////////////////////////
 
+	uint32 GetPacketSequence(int32 PacketIndex) const;
+
 	void UpdateState();
 
 	void UpdateHoveredEvent();
@@ -215,6 +217,7 @@ private:
 	uint32 ConnectionIndex;
 	Trace::ENetProfilerConnectionMode ConnectionMode;
 	uint32 PacketIndex;
+	uint32 PacketSequence;
 	int64 PacketBitSize; // total number of bits; [bit]
 
 	bool bFilterByEventType;
diff --git a/Engine/Source/Developer/TraceInsights/Private/Insights/NetworkingProfiler/Widgets/SPacketView.cpp b/Engine/Source/Developer/TraceInsights/Private/Insights/NetworkingProfiler/Widgets/SPacketView.cpp
--- a/Engine/Source/Developer/TraceInsights/Private/Insights/NetworkingProfiler/Widgets/SPacketView.cpp
+++ b/Engine/Source/Developer/TraceInsights/Private/Insights/NetworkingProfiler/Widgets/SPacketView.cpp
@@ -1322,6 +1322,25 @@ void SPacketView::SetSelectedPacket(const int32 InPacketIndex)
 
 ////////////////////////////////////////////////////////////////////////////////////////////////////
 
+void SPacketView::SelectPacketBySequenceNumber(const uint32 InSequenceNumber)
+{
+	// Find the PacketIndex from sequence number
+	TSharedPtr<const Trace::IAnalysisSession> Session = FInsightsManager::Get()->GetSession();
+	if (Session.IsValid())
+	{
+		Trace::FAnalysisSessionReadScope SessionReadScope(*Session.Get());
+		const Trace::INetProfilerProvider& NetProfilerProvider = Trace::ReadNetProfilerProvider(*Session.Get());
+
+		const int32 PacketId = NetProfilerProvider.FindPacketIndexFromPacketSequence(ConnectionIndex, ConnectionMode, InSequenceNumber);
+		if (PacketId != -1)
+		{
+			SetSelectedPacket(PacketId);
+		}
+	}	
+}
+
+////////////////////////////////////////////////////////////////////////////////////////////////////
+
 void SPacketView::SelectPreviousPacket()
 {
 	const FAxisViewportInt32& ViewportX = Viewport.GetHorizontalAxisViewport();
diff --git a/Engine/Source/Developer/TraceInsights/Private/Insights/NetworkingProfiler/Widgets/SPacketView.h b/Engine/Source/Developer/TraceInsights/Private/Insights/NetworkingProfiler/Widgets/SPacketView.h
--- a/Engine/Source/Developer/TraceInsights/Private/Insights/NetworkingProfiler/Widgets/SPacketView.h
+++ b/Engine/Source/Developer/TraceInsights/Private/Insights/NetworkingProfiler/Widgets/SPacketView.h
@@ -134,6 +134,7 @@ public:
 
 	void EnsurePacketIsVisible(const int InPacketIndex);
 	void SetSelectedPacket(const int32 InPacketIndex);
+	void SelectPacketBySequenceNumber(const uint32 InSequenceNumber);
 	void SelectPreviousPacket();
 	void SelectNextPacket();
 	void ExtendLeftSideOfSelectedInterval();
diff --git a/Engine/Source/Developer/TraceServices/Private/Model/NetProfilerProvider.cpp b/Engine/Source/Developer/TraceServices/Private/Model/NetProfilerProvider.cpp
--- a/Engine/Source/Developer/TraceServices/Private/Model/NetProfilerProvider.cpp
+++ b/Engine/Source/Developer/TraceServices/Private/Model/NetProfilerProvider.cpp
@@ -359,6 +359,42 @@ const INetProfilerProvider& ReadNetProfilerProvider(const IAnalysisSession& Sess
 	return *Session.ReadProvider<INetProfilerProvider>(FNetProfilerProvider::ProviderName);
 }
 
+int32 FNetProfilerProvider::FindPacketIndexFromPacketSequence(uint32 ConnectionIndex, ENetProfilerConnectionMode Mode, uint32 SequenceNumber) const
+{
+	Session.ReadAccessCheck();
+
+	check(ConnectionIndex < Connections.Num());
+
+	const auto& Packets = Connections[ConnectionIndex].Data[Mode]->Packets;
+	const int32 PacketCount = (int32)Packets.Num();
+
+	if (PacketCount == 0)
+	{
+		return -1;
+	}
+
+	if (SequenceNumber < Packets[0].SequenceNumber)
+	{
+		return -1;
+	}
+
+	if (SequenceNumber > Packets[PacketCount - 1].SequenceNumber)
+	{
+		return -1;
+	}
+
+	// Brute force it, we can cache some data to speed this up if necessary
+	for (int32 PacketIt = 0, PacketEndIt = PacketCount - 1; PacketIt <= PacketEndIt; ++PacketIt)
+	{
+		if (Packets[PacketIt].SequenceNumber == SequenceNumber)
+		{
+			return PacketIt;
+		}
+	}
+
+	return -1;	
+}
+
 uint32 FNetProfilerProvider::GetPacketCount(uint32 ConnectionIndex, ENetProfilerConnectionMode Mode) const
 {
 	Session.ReadAccessCheck();
diff --git a/Engine/Source/Developer/TraceServices/Private/Model/NetProfilerProvider.h b/Engine/Source/Developer/TraceServices/Private/Model/NetProfilerProvider.h
--- a/Engine/Source/Developer/TraceServices/Private/Model/NetProfilerProvider.h
+++ b/Engine/Source/Developer/TraceServices/Private/Model/NetProfilerProvider.h
@@ -101,6 +101,9 @@ public:
 	virtual void ReadObject(uint32 GameInstanceIndex, uint32 ObjectIndex, TFunctionRef<void(const FNetProfilerObjectInstance&)> Callback) const override;
 	virtual uint32 GetObjectsChangeCount(uint32 GameInstanceIndex) const override;
 
+	// Find Packet Index from SequenceNumber
+	virtual int32 FindPacketIndexFromPacketSequence(uint32 ConnectionIndex, ENetProfilerConnectionMode Mode, uint32 SequenceNumber) const override;
+
 	// Enumerate packets in the provided packet interval
 	virtual uint32 GetPacketCount(uint32 ConnectionIndex, ENetProfilerConnectionMode Mode) const override;
 	virtual void EnumeratePackets(uint32 ConnectionIndex, ENetProfilerConnectionMode Mode, uint32 PacketIndexIntervalStart, uint32 PacketIndexIntervalEnd, TFunctionRef<void(const FNetProfilerPacket&)> Callback) const override;
diff --git a/Engine/Source/Developer/TraceServices/Public/TraceServices/Model/NetProfiler.h b/Engine/Source/Developer/TraceServices/Public/TraceServices/Model/NetProfiler.h
--- a/Engine/Source/Developer/TraceServices/Public/TraceServices/Model/NetProfiler.h
+++ b/Engine/Source/Developer/TraceServices/Public/TraceServices/Model/NetProfiler.h
@@ -155,6 +155,9 @@ public:
 	virtual void ReadObject(uint32 GameInstanceIndex, uint32 ObjectIndex, TFunctionRef<void(const FNetProfilerObjectInstance&)> Callback) const = 0;
 	virtual uint32 GetObjectsChangeCount(uint32 GameInstanceIndex) const = 0;
 
+	// Find Packet Index from SequenceNumber
+	virtual int32 FindPacketIndexFromPacketSequence(uint32 ConnectionIndex, ENetProfilerConnectionMode Mode, uint32 SequenceNumber) const = 0;
+
 	// Enumerate packets in the provided packet interval
 	virtual uint32 GetPacketCount(uint32 ConnectionIndex, ENetProfilerConnectionMode Mode) const = 0;
 	virtual void EnumeratePackets(uint32 ConnectionIndex, ENetProfilerConnectionMode Mode, uint32 PacketIndexIntervalStart, uint32 PacketIndexIntervalEnd, TFunctionRef<void(const FNetProfilerPacket&)> Callback) const = 0;
