Submitting JB shelf to expose analysis-session and session store from Unreal Insights Module

diff --git a/Engine/Source/Developer/TraceInsights/Private/Insights/TraceInsightsModule.cpp b/Engine/Source/Developer/TraceInsights/Private/Insights/TraceInsightsModule.cpp
--- a/Engine/Source/Developer/TraceInsights/Private/Insights/TraceInsightsModule.cpp
+++ b/Engine/Source/Developer/TraceInsights/Private/Insights/TraceInsightsModule.cpp
@@ -360,6 +360,19 @@ FOnRegisterMajorTabExtensions& FTraceInsightsModule::OnRegisterMajorTabExtension
 	return MajorTabExtensionDelegates.FindOrAdd(InMajorTabId);
 }
 
+////////////////////////////////////////////////////////////////////////////////////////////////////
+Trace::FStoreClient* FTraceInsightsModule::GetStoreClient()
+{
+	return FInsightsManager::Get()->GetStoreClient();
+}
+
+////////////////////////////////////////////////////////////////////////////////////////////////////
+
+TSharedPtr<const Trace::IAnalysisSession> FTraceInsightsModule::GetAnalysisSession() const
+{
+	return FInsightsManager::Get()->GetSession();
+}
+
 ////////////////////////////////////////////////////////////////////////////////////////////////////
 
 const FInsightsMajorTabConfig& FTraceInsightsModule::FindMajorTabConfig(const FName& InMajorTabId) const
diff --git a/Engine/Source/Developer/TraceInsights/Private/Insights/TraceInsightsModule.h b/Engine/Source/Developer/TraceInsights/Private/Insights/TraceInsightsModule.h
--- a/Engine/Source/Developer/TraceInsights/Private/Insights/TraceInsightsModule.h
+++ b/Engine/Source/Developer/TraceInsights/Private/Insights/TraceInsightsModule.h
@@ -42,6 +42,8 @@ public:
 	virtual void RegisterMajorTabConfig(const FName& InMajorTabId, const FInsightsMajorTabConfig& InConfig) override;
 	virtual void UnregisterMajorTabConfig(const FName& InMajorTabId) override;
 	virtual FOnInsightsMajorTabCreated& OnMajorTabCreated() override { return OnInsightsMajorTabCreatedDelegate; }
+	virtual Trace::FStoreClient* GetStoreClient() override;
+	virtual TSharedPtr<const Trace::IAnalysisSession> GetAnalysisSession() const override;
 	virtual FOnRegisterMajorTabExtensions& OnRegisterMajorTabExtension(const FName& InMajorTabId) override;
 
 	/** Find a major tab config for the specified ID */
diff --git a/Engine/Source/Developer/TraceInsights/Public/Insights/IUnrealInsightsModule.h b/Engine/Source/Developer/TraceInsights/Public/Insights/IUnrealInsightsModule.h
--- a/Engine/Source/Developer/TraceInsights/Public/Insights/IUnrealInsightsModule.h
+++ b/Engine/Source/Developer/TraceInsights/Public/Insights/IUnrealInsightsModule.h
@@ -13,6 +13,7 @@ class FExtender;
 namespace Trace
 {
 	class FStoreClient;
+	class IAnalysisSession;
 }
 
 /** Major tab IDs for Insights tools */
@@ -195,6 +196,16 @@ public:
 	 */
 	virtual FOnRegisterMajorTabExtensions& OnRegisterMajorTabExtension(const FName& InMajorTabId) = 0;
 
+
+	/**
+	 * Expose the session service directly
+	 */
+	virtual Trace::FStoreClient* GetStoreClient() = 0;
+
+	/**
+	 * Get the session service.
+	 */
+	virtual TSharedPtr<const Trace::IAnalysisSession> GetAnalysisSession() const = 0;
 	/** Callback invoked when a major tab is created */
 	virtual FOnInsightsMajorTabCreated& OnMajorTabCreated() = 0;
 };
