Added overlay extensions to timing view

This allows extenders to add their own widgets to the timing view

diff --git a/Engine/Source/Developer/TraceInsights/Private/Insights/Widgets/STimingView.cpp b/Engine/Source/Developer/TraceInsights/Private/Insights/Widgets/STimingView.cpp
--- a/Engine/Source/Developer/TraceInsights/Private/Insights/Widgets/STimingView.cpp
+++ b/Engine/Source/Developer/TraceInsights/Private/Insights/Widgets/STimingView.cpp
@@ -90,6 +90,8 @@ STimingView::STimingView()
 	IModularFeatures::Get().RegisterModularFeature(Insights::TimingViewExtenderFeatureName, LoadingSharedState.Get());
 	IModularFeatures::Get().RegisterModularFeature(Insights::TimingViewExtenderFeatureName, FileActivitySharedState.Get());
 
+	ExtensionOverlay = SNew(SOverlay).Visibility(EVisibility::SelfHitTestInvisible);
+
 	Reset();
 }
 
@@ -186,6 +188,13 @@ void STimingView::Construct(const FArguments& InArgs)
 				]
 			]
 		]
+		+ SOverlay::Slot()
+		.HAlign(HAlign_Fill)
+		.VAlign(VAlign_Fill)
+		.Padding(FMargin(0.0f, 0.0f, 0.0f, 0.0f))
+		[
+			ExtensionOverlay.ToSharedRef()
+		]
 	];
 
 	UpdateHorizontalScrollBar();
@@ -238,6 +247,10 @@ void STimingView::Reset()
 
 	//////////////////////////////////////////////////
 
+	ExtensionOverlay->ClearChildren();
+
+	//////////////////////////////////////////////////
+
 	MousePosition = FVector2D::ZeroVector;
 
 	MousePositionOnButtonDown = FVector2D::ZeroVector;
@@ -2597,6 +2610,28 @@ void STimingView::RaiseTimeMarkerChanged()
 
 ////////////////////////////////////////////////////////////////////////////////////////////////////
 
+void STimingView::SetTimeMarker(double InMarkerTime)
+{ 
+	TimeMarker = InMarkerTime; 
+
+	RaiseTimeMarkerChanged();
+}
+
+////////////////////////////////////////////////////////////////////////////////////////////////////
+
+void STimingView::AddOverlayWidget(const TSharedRef<SWidget>& InWidget)
+{
+	if(ExtensionOverlay.IsValid())
+	{
+		ExtensionOverlay->AddSlot()
+		[
+			InWidget
+		];
+	}
+}
+
+////////////////////////////////////////////////////////////////////////////////////////////////////
+
 void STimingView::UpdateAggregatedStats()
 {
 	if (bAssetLoadingMode)
diff --git a/Engine/Source/Developer/TraceInsights/Private/Insights/Widgets/STimingView.h b/Engine/Source/Developer/TraceInsights/Private/Insights/Widgets/STimingView.h
--- a/Engine/Source/Developer/TraceInsights/Private/Insights/Widgets/STimingView.h
+++ b/Engine/Source/Developer/TraceInsights/Private/Insights/Widgets/STimingView.h
@@ -32,6 +32,7 @@ class FTimeRulerTrack;
 class FTimingGraphTrack;
 class FTimingViewDrawHelper;
 class SScrollBar;
+class SOverlay;
 namespace Insights { class ITimingViewExtender; }
 
 /** A custom widget used to display timing events. */
@@ -228,6 +229,8 @@ public:
 	virtual Insights::FSelectedTrackChangedDelegate& OnSelectedTrackChanged() override { return OnSelectedTrackChangedDelegate; }
 	virtual Insights::FSelectedEventChangedDelegate& OnSelectedEventChanged() override { return OnSelectedEventChangedDelegate; }
 
+	virtual void AddOverlayWidget(const TSharedRef<SWidget>& InWidget) override;
+	virtual void SetTimeMarker(double InMarkerTime) override;
 	virtual double GetTimeMarker() const override { return TimeMarker; }
 
 	////////////////////////////////////////////////////////////////////////////////////////////////////
@@ -394,6 +397,9 @@ protected:
 
 	////////////////////////////////////////////////////////////
 
+	/** The extension overlay containing external sub-widgets */
+	TSharedPtr<SOverlay> ExtensionOverlay;
+
 	/** Horizontal scroll bar, used for scrolling timing events' viewport. */
 	TSharedPtr<SScrollBar> HorizontalScrollBar;
 
diff --git a/Engine/Source/Developer/TraceInsights/Public/Insights/ITimingViewSession.h b/Engine/Source/Developer/TraceInsights/Public/Insights/ITimingViewSession.h
--- a/Engine/Source/Developer/TraceInsights/Public/Insights/ITimingViewSession.h
+++ b/Engine/Source/Developer/TraceInsights/Public/Insights/ITimingViewSession.h
@@ -7,6 +7,7 @@
 class FBaseTimingTrack;
 class FTimingEventsTrack;
 class ITimingEvent;
+class SWidget;
 
 namespace Insights
 {
@@ -82,6 +83,12 @@ public:
 
 	/** Get the current marker time */
 	virtual double GetTimeMarker() const = 0;
+
+	/** Set the current marker time */
+	virtual void SetTimeMarker(double InMarkerTime) = 0;
+
+	/** Add a slot to the overlay */
+	virtual void AddOverlayWidget(const TSharedRef<SWidget>& InWidget) = 0;
 };
 
 } // namespace Insights
