[Insights]
  - Session Browser (start page): Updated layout.
  - Session Browser (start page): Changed "auto start analysis" checkbbox to be disabled by default. Also moved on same line with Open button.
  - Session Browser (start page): Fixed meta-data for trace sessions.
  - Counters: Changed double click on a counter to also turn visibility on for the Graph track in addition to adding a graph series for respective counter.
  - Timing view: Added workaround for cases where a log message has a null category name.
  - Added "TraceInsights" log category and a Log.h (to avoid using "TimingProfiler" category).
  - Split FTraceInsightsModule::StartAnalysisForTrace(const TCHAR* InStoreHost, uint32 InStorePort, uint32 InTraceId) into ConnectToStore(const TCHAR* InStoreHost, uint32 InStorePort) + StartAnalysisForTrace(InTraceId).

diff --git a/Engine/Source/Developer/TraceInsights/Private/Insights/Common/TimeUtils.cpp b/Engine/Source/Developer/TraceInsights/Private/Insights/Common/TimeUtils.cpp
--- a/Engine/Source/Developer/TraceInsights/Private/Insights/Common/TimeUtils.cpp
+++ b/Engine/Source/Developer/TraceInsights/Private/Insights/Common/TimeUtils.cpp
@@ -5,7 +5,7 @@
 #include <limits>
 //#include <math.h>
 
-//#include "TimingProfilerCommon.h" // for UE_LOG(TimingProfiler, ...
+//#include "Insights/Log.h"
 
 namespace TimeUtils
 {
@@ -701,15 +701,15 @@ void TestTimeFormatting()
 {
 	double T1 = 1 * TimeUtils::Day + 2 * TimeUtils::Hour + 3 * TimeUtils::Minute + 4.567890123456789;
 	FString S1 = FormatTime(T1, TimeUtils::Day);
-	//UE_LOG(TimingProfiler, Log, TEXT("D-T: %s"), *S1);
+	//UE_LOG(TraceInsights, Log, TEXT("D-T: %s"), *S1);
 	FString S2 = FormatTime(T1, TimeUtils::Hour);
-	//UE_LOG(TimingProfiler, Log, TEXT("H-T: %s"), *S2);
+	//UE_LOG(TraceInsights, Log, TEXT("H-T: %s"), *S2);
 	FString S3 = FormatTime(T1, TimeUtils::Minute);
-	//UE_LOG(TimingProfiler, Log, TEXT("M-T: %s"), *S3);
+	//UE_LOG(TraceInsights, Log, TEXT("M-T: %s"), *S3);
 	for (double P = 10.0; P >= TimeUtils::Picosecond; P /= 10.0)
 	{
 		FString SP = FormatTime(T1, P);
-		//UE_LOG(TimingProfiler, Log, TEXT("P:%g T: %s"), P, *SP);
+		//UE_LOG(TraceInsights, Log, TEXT("P:%g T: %s"), P, *SP);
 	}
 }
 
@@ -781,7 +781,7 @@ void TestTimeAutoFormatting()
 	for (int32 Index = 0; Index < DataCount; ++Index)
 	{
 		FString Str = FormatTimeAuto(Data[Index].T);
-		//UE_LOG(TimingProfiler, Log, TEXT("%s : %s"), Data[Index].Msg, *Str);
+		//UE_LOG(TraceInsights, Log, TEXT("%s : %s"), Data[Index].Msg, *Str);
 	}
 }
 
@@ -811,7 +811,7 @@ void TestOptimizationIssue()
 	constexpr double Ns = 0.000000001;
 	int32 D1 = GetNumDigits(Ns);
 	int32 D2 = GetNumDigitsOptDisabled(Ns);
-	//UE_LOG(TimingProfiler, Log, TEXT("D1 = %d, D2 = %d"), D1, D2);
+	//UE_LOG(TraceInsights, Log, TEXT("D1 = %d, D2 = %d"), D1, D2);
 	ensure(D1 == 9); // 10 ?
 	ensure(D1 == D2);
 }
diff --git a/Engine/Source/Developer/TraceInsights/Private/Insights/InsightsManager.h b/Engine/Source/Developer/TraceInsights/Private/Insights/InsightsManager.h
--- a/Engine/Source/Developer/TraceInsights/Private/Insights/InsightsManager.h
+++ b/Engine/Source/Developer/TraceInsights/Private/Insights/InsightsManager.h
@@ -191,7 +191,7 @@ private:
 	TSharedRef<Trace::ISessionService> SessionService;
 	TSharedRef<Trace::IModuleService> ModuleService;
 
-	/** The location of the trace files managed by trace store. */
+	/** The location of the trace files managed by the trace store. */
 	FString StoreDir;
 
 	/** The client used to connect to the trace store. */
@@ -203,7 +203,7 @@ private:
 	/** The id of the trace being analyzed. */
 	uint32 CurrentTraceId;
 
-	/** The filename of trace being analyzed. */
+	/** The filename of the trace being analyzed. */
 	FString CurrentTraceFilename;
 
 	/** List of UI commands for this manager. This will be filled by this and corresponding classes. */
diff --git a/Engine/Source/Developer/TraceInsights/Private/Insights/TraceInsightsModule.cpp b/Engine/Source/Developer/TraceInsights/Private/Insights/TraceInsightsModule.cpp
--- a/Engine/Source/Developer/TraceInsights/Private/Insights/TraceInsightsModule.cpp
+++ b/Engine/Source/Developer/TraceInsights/Private/Insights/TraceInsightsModule.cpp
@@ -22,6 +22,7 @@
 #include "Insights/IUnrealInsightsModule.h"
 #include "Insights/LoadingProfiler/LoadingProfilerManager.h"
 #include "Insights/LoadingProfiler/Widgets/SLoadingProfilerWindow.h"
+#include "Insights/Log.h"
 #include "Insights/NetworkingProfiler/NetworkingProfilerManager.h"
 #include "Insights/NetworkingProfiler/Widgets/SNetworkingProfilerWindow.h"
 #include "Insights/TimingProfilerManager.h"
@@ -31,6 +32,8 @@
 
 ////////////////////////////////////////////////////////////////////////////////////////////////////
 
+DEFINE_LOG_CATEGORY(TraceInsights);
+
 IMPLEMENT_MODULE(FTraceInsightsModule, TraceInsights);
 
 ////////////////////////////////////////////////////////////////////////////////////////////////////
@@ -201,13 +204,6 @@ void FTraceInsightsModule::UnregisterTabSpawners()
 
 ////////////////////////////////////////////////////////////////////////////////////////////////////
 
-bool FTraceInsightsModule::ConnectToStore(const TCHAR* Host, uint32 Port)
-{
-	return FInsightsManager::Get()->ConnectToStore(Host, Port);
-}
-
-////////////////////////////////////////////////////////////////////////////////////////////////////
-
 void FTraceInsightsModule::CreateSessionBrowser(bool bAllowDebugTools, bool bSingleProcess)
 {
 	FInsightsManager::Get()->SetOpenAnalysisInSeparateProcess(!bSingleProcess);
@@ -360,19 +356,6 @@ FOnRegisterMajorTabExtensions& FTraceInsightsModule::OnRegisterMajorTabExtension
 	return MajorTabExtensionDelegates.FindOrAdd(InMajorTabId);
 }
 
-////////////////////////////////////////////////////////////////////////////////////////////////////
-Trace::FStoreClient* FTraceInsightsModule::GetStoreClient()
-{
-	return FInsightsManager::Get()->GetStoreClient();
-}
-
-////////////////////////////////////////////////////////////////////////////////////////////////////
-
-TSharedPtr<const Trace::IAnalysisSession> FTraceInsightsModule::GetAnalysisSession() const
-{
-	return FInsightsManager::Get()->GetSession();
-}
-
 ////////////////////////////////////////////////////////////////////////////////////////////////////
 
 const FInsightsMajorTabConfig& FTraceInsightsModule::FindMajorTabConfig(const FName& InMajorTabId) const
@@ -396,20 +379,29 @@ const FOnRegisterMajorTabExtensions* FTraceInsightsModule::FindMajorTabLayoutExt
 
 ////////////////////////////////////////////////////////////////////////////////////////////////////
 
-void FTraceInsightsModule::StartAnalysisForTraceFile(const TCHAR* InTraceFile)
+Trace::FStoreClient* FTraceInsightsModule::GetStoreClient()
 {
-	if (InTraceFile != nullptr)
-	{
-		FInsightsManager::Get()->LoadTraceFile(FString(InTraceFile));
-	}
+	return FInsightsManager::Get()->GetStoreClient();
 }
 
 ////////////////////////////////////////////////////////////////////////////////////////////////////
 
-void FTraceInsightsModule::StartAnalysisForTrace(const TCHAR* InStoreHost, uint32 InStorePort, uint32 InTraceId)
+bool FTraceInsightsModule::ConnectToStore(const TCHAR* Host, uint32 Port)
 {
-	FInsightsManager::Get()->ConnectToStore(InStoreHost, InStorePort);
+	return FInsightsManager::Get()->ConnectToStore(Host, Port);
+}
+
+////////////////////////////////////////////////////////////////////////////////////////////////////
 
+TSharedPtr<const Trace::IAnalysisSession> FTraceInsightsModule::GetAnalysisSession() const
+{
+	return FInsightsManager::Get()->GetSession();
+}
+
+////////////////////////////////////////////////////////////////////////////////////////////////////
+
+void FTraceInsightsModule::StartAnalysisForTrace(uint32 InTraceId)
+{
 	if (InTraceId != 0)
 	{
 		FInsightsManager::Get()->LoadTrace(InTraceId);
@@ -425,6 +417,16 @@ void FTraceInsightsModule::StartAnalysisForLastLiveSession()
 
 ////////////////////////////////////////////////////////////////////////////////////////////////////
 
+void FTraceInsightsModule::StartAnalysisForTraceFile(const TCHAR* InTraceFile)
+{
+	if (InTraceFile != nullptr)
+	{
+		FInsightsManager::Get()->LoadTraceFile(FString(InTraceFile));
+	}
+}
+
+////////////////////////////////////////////////////////////////////////////////////////////////////
+
 TSharedRef<SDockTab> FTraceInsightsModule::SpawnStartPageTab(const FSpawnTabArgs& Args)
 {
 	const TSharedRef<SDockTab> DockTab = SNew(SDockTab)
diff --git a/Engine/Source/Developer/TraceInsights/Private/Insights/TraceInsightsModule.h b/Engine/Source/Developer/TraceInsights/Private/Insights/TraceInsightsModule.h
--- a/Engine/Source/Developer/TraceInsights/Private/Insights/TraceInsightsModule.h
+++ b/Engine/Source/Developer/TraceInsights/Private/Insights/TraceInsightsModule.h
@@ -31,25 +31,29 @@ public:
 		return false;
 	}
 
-	virtual bool ConnectToStore(const TCHAR* Host, uint32 Port) override;
 	virtual void CreateSessionBrowser(bool bAllowDebugTools, bool bSingleProcess) override;
 	virtual void CreateSessionViewer(bool bAllowDebugTools) override;
+
+	virtual Trace::FStoreClient* GetStoreClient() override;
+	virtual bool ConnectToStore(const TCHAR* InStoreHost, uint32 InStorePort) override;
+
+	virtual TSharedPtr<const Trace::IAnalysisSession> GetAnalysisSession() const override;
+	virtual void StartAnalysisForTrace(uint32 InTraceId) override;
+	virtual void StartAnalysisForLastLiveSession() override;
 	virtual void StartAnalysisForTraceFile(const TCHAR* InTraceFile) override;
-	virtual void StartAnalysisForTrace(const TCHAR* InStoreHost, uint32 InStorePort, uint32 InTraceId) override;
-	virtual void StartAnalysisForLastLiveSession(/*StoreClient*/) override;
-	//virtual void StartAnalysisForTrace(StoreClient or FTraceData..) override;
+
 	virtual void ShutdownUserInterface() override;
+
 	virtual void RegisterMajorTabConfig(const FName& InMajorTabId, const FInsightsMajorTabConfig& InConfig) override;
 	virtual void UnregisterMajorTabConfig(const FName& InMajorTabId) override;
 	virtual FOnInsightsMajorTabCreated& OnMajorTabCreated() override { return OnInsightsMajorTabCreatedDelegate; }
-	virtual Trace::FStoreClient* GetStoreClient() override;
-	virtual TSharedPtr<const Trace::IAnalysisSession> GetAnalysisSession() const override;
 	virtual FOnRegisterMajorTabExtensions& OnRegisterMajorTabExtension(const FName& InMajorTabId) override;
 
 	/** Find a major tab config for the specified ID */
 	const FInsightsMajorTabConfig& FindMajorTabConfig(const FName& InMajorTabId) const;
 
 	const FOnRegisterMajorTabExtensions* FindMajorTabLayoutExtension(const FName& InMajorTabId) const;
+
 protected:
 	void InitTraceStore();
 
diff --git a/Engine/Source/Developer/TraceInsights/Private/Insights/ViewModels/LogFilter.cpp b/Engine/Source/Developer/TraceInsights/Private/Insights/ViewModels/LogFilter.cpp
--- a/Engine/Source/Developer/TraceInsights/Private/Insights/ViewModels/LogFilter.cpp
+++ b/Engine/Source/Developer/TraceInsights/Private/Insights/ViewModels/LogFilter.cpp
@@ -5,7 +5,7 @@
 #include "Async/AsyncWork.h"
 
 // Insights
-#include "Insights/TimingProfilerCommon.h" // for UE_LOG
+#include "Insights/Log.h"
 #include "Insights/Widgets/SLogView.h"
 
 ////////////////////////////////////////////////////////////////////////////////////////////////////
@@ -273,11 +273,11 @@ void FLogFilteringAsyncTask::DoWork()
 
 	if (Filter.IsFilterSetByText())
 	{
-		UE_LOG(TimingProfiler, Log, TEXT("[LogView] FLogFilteringAsyncTask::DoWork [%d to %d] by Text (\"%s\")"), StartIndex, EndIndex, *Filter.GetFilterText().ToString());
+		UE_LOG(TraceInsights, Log, TEXT("[LogView] FLogFilteringAsyncTask::DoWork [%d to %d] by Text (\"%s\")"), StartIndex, EndIndex, *Filter.GetFilterText().ToString());
 	}
 	else
 	{
-		UE_LOG(TimingProfiler, Log, TEXT("[LogView] FLogFilteringAsyncTask::DoWork [%d to %d]"), StartIndex, EndIndex);
+		UE_LOG(TraceInsights, Log, TEXT("[LogView] FLogFilteringAsyncTask::DoWork [%d to %d]"), StartIndex, EndIndex);
 	}
 
 	for (int32 Index = StartIndex; Index < EndIndex && !bCanceled; ++Index)
@@ -295,11 +295,11 @@ void FLogFilteringAsyncTask::DoWork()
 	if (bCanceled)
 	{
 		FilteredMessages.Reset();
-		UE_LOG(TimingProfiler, Log, TEXT("[LogView] FLogFilteringAsyncTask::DoWork CANCELED"));
+		UE_LOG(TraceInsights, Log, TEXT("[LogView] FLogFilteringAsyncTask::DoWork CANCELED"));
 	}
 	else
 	{
-		UE_LOG(TimingProfiler, Log, TEXT("[LogView] FLogFilteringAsyncTask::DoWork DONE (%d filtered messages)"), FilteredMessages.Num());
+		UE_LOG(TraceInsights, Log, TEXT("[LogView] FLogFilteringAsyncTask::DoWork DONE (%d filtered messages)"), FilteredMessages.Num());
 	}
 }
 
diff --git a/Engine/Source/Developer/TraceInsights/Private/Insights/ViewModels/LogFilter.h b/Engine/Source/Developer/TraceInsights/Private/Insights/ViewModels/LogFilter.h
--- a/Engine/Source/Developer/TraceInsights/Private/Insights/ViewModels/LogFilter.h
+++ b/Engine/Source/Developer/TraceInsights/Private/Insights/ViewModels/LogFilter.h
@@ -155,8 +155,6 @@ private:
 
 ////////////////////////////////////////////////////////////////////////////////////////////////////
 
-class FAsyncLogFilterProxy;
-
 class FLogFilteringAsyncTask : public FNonAbandonableTask
 {
 public:
diff --git a/Engine/Source/Developer/TraceInsights/Private/Insights/ViewModels/MarkersTimingTrack.cpp b/Engine/Source/Developer/TraceInsights/Private/Insights/ViewModels/MarkersTimingTrack.cpp
--- a/Engine/Source/Developer/TraceInsights/Private/Insights/ViewModels/MarkersTimingTrack.cpp
+++ b/Engine/Source/Developer/TraceInsights/Private/Insights/ViewModels/MarkersTimingTrack.cpp
@@ -280,14 +280,19 @@ void FTimeMarkerTrackBuilder::AddLogMessage(const Trace::FLogMessage& Message)
 			[this](const Trace::FLogMessage& Message) { AddLogMessage(Message); });
 	}
 
-	if (!Track.bUseOnlyBookmarks || FCString::Strcmp(Message.Category->Name, TEXT("LogBookmark")) == 0)
+	check(Message.Category != nullptr);
+	//check(Message.Category->Name != nullptr);
+
+	const TCHAR* CategoryName = Message.Category->Name != nullptr ? Message.Category->Name : TEXT("");
+
+	if (!Track.bUseOnlyBookmarks || FCString::Strcmp(CategoryName, TEXT("LogBookmark")) == 0)
 	{
 		float X = Viewport.TimeToSlateUnitsRounded(Message.Time);
-		if (X < 0)
+		if (X < 0.0f)
 		{
 			X = -1.0f;
 		}
-		AddTimeMarker(X, Message.Index, Message.Verbosity, Message.Category->Name, Message.Message);
+		AddTimeMarker(X, Message.Index, Message.Verbosity, CategoryName, Message.Message);
 	}
 }
 
@@ -318,14 +323,14 @@ FLinearColor FTimeMarkerTrackBuilder::GetColorByVerbosity(const ELogVerbosity::T
 {
 	static FLinearColor Colors[] =
 	{
-		FLinearColor(0.0, 0.0, 0.0, 1.0), // NoLogging
-		FLinearColor(1.0, 0.0, 0.0, 1.0), // Fatal
-		FLinearColor(1.0, 0.3, 0.0, 1.0), // Error
-		FLinearColor(0.7, 0.5, 0.0, 1.0), // Warning
-		FLinearColor(0.0, 0.7, 0.0, 1.0), // Display
-		FLinearColor(0.0, 0.7, 1.0, 1.0), // Log
-		FLinearColor(0.7, 0.7, 0.7, 1.0), // Verbose
-		FLinearColor(1.0, 1.0, 1.0, 1.0), // VeryVerbose
+		FLinearColor(0.0f, 0.0f, 0.0f, 1.0f), // NoLogging
+		FLinearColor(1.0f, 0.0f, 0.0f, 1.0f), // Fatal
+		FLinearColor(1.0f, 0.3f, 0.0f, 1.0f), // Error
+		FLinearColor(0.7f, 0.5f, 0.0f, 1.0f), // Warning
+		FLinearColor(0.0f, 0.7f, 0.0f, 1.0f), // Display
+		FLinearColor(0.0f, 0.7f, 1.0f, 1.0f), // Log
+		FLinearColor(0.7f, 0.7f, 0.7f, 1.0f), // Verbose
+		FLinearColor(1.0f, 1.0f, 1.0f, 1.0f), // VeryVerbose
 	};
 	static_assert(sizeof(Colors) / sizeof(FLinearColor) == (int)ELogVerbosity::Type::All + 1, "ELogVerbosity::Type has changed!?");
 	//return Colors[Verbosity & ELogVerbosity::VerbosityMask];
diff --git a/Engine/Source/Developer/TraceInsights/Private/Insights/Widgets/SLogView.cpp b/Engine/Source/Developer/TraceInsights/Private/Insights/Widgets/SLogView.cpp
--- a/Engine/Source/Developer/TraceInsights/Private/Insights/Widgets/SLogView.cpp
+++ b/Engine/Source/Developer/TraceInsights/Private/Insights/Widgets/SLogView.cpp
@@ -18,7 +18,7 @@
 #include "Insights/Common/TimeUtils.h"
 #include "Insights/InsightsManager.h"
 #include "Insights/InsightsStyle.h"
-#include "Insights/TimingProfilerCommon.h"
+#include "Insights/Log.h"
 #include "Insights/TimingProfilerManager.h"
 #include "Insights/ViewModels/MarkersTimingTrack.h" // for FTimeMarkerTrackBuilder::GetColorBy*
 #include "Insights/ViewModels/TimingViewDrawHelper.h"
@@ -722,19 +722,12 @@ void SLogView::Tick(const FGeometry& AllottedGeometry, const double InCurrentTim
 
 		//TODO: show only categories that are used in current trace
 		//TODO: cause of duplicates: a) runtime, b) case insensitive, c) stripped "Log" prefix
-		//TODO: FString vs. FText vs. FName ?
-
-		//TODO: int32 NumCategories = static_cast<int32>(LogProvider.GetCategoriesCount());
-		int32 NumCategories = 0;
-		LogProvider.EnumerateCategories([&NumCategories](const Trace::FLogCategory& Category)
-		{
-			NumCategories++;
-		});
 
+		const int32 NumCategories = static_cast<int32>(LogProvider.GetCategoryCount());
 		if (NumCategories != TotalNumCategories)
 		{
 			TotalNumCategories = NumCategories;
-			UE_LOG(TimingProfiler, Log, TEXT("[LogView] Total Log Categories: %d"), TotalNumCategories);
+			UE_LOG(TraceInsights, Log, TEXT("[LogView] Total Log Categories: %d"), TotalNumCategories);
 
 			TSet<FName> Categories;
 			LogProvider.EnumerateCategories([&Categories](const Trace::FLogCategory& Category)
@@ -746,12 +739,12 @@ void SLogView::Tick(const FGeometry& AllottedGeometry, const double InCurrentTim
 				}
 				if (Categories.Contains(FName(*CategoryStr)))
 				{
-					UE_LOG(TimingProfiler, Log, TEXT("[LogView] Duplicated Log Category: \"%s\""), Category.Name);
+					UE_LOG(TraceInsights, Log, TEXT("[LogView] Duplicated Log Category: \"%s\""), Category.Name);
 				}
 				Categories.Add(FName(*CategoryStr));
 			});
 			Filter.SyncAvailableCategories(Categories);
-			UE_LOG(TimingProfiler, Log, TEXT("[LogView] Unique Log Categories: %d"), Filter.GetAvailableLogCategories().Num());
+			UE_LOG(TraceInsights, Log, TEXT("[LogView] Unique Log Categories: %d"), Filter.GetAvailableLogCategories().Num());
 
 			//Cache.Reset();
 			Messages.Reset();
@@ -781,7 +774,7 @@ void SLogView::Tick(const FGeometry& AllottedGeometry, const double InCurrentTim
 					FilteringEndIndex = NewMessageCount;
 					FilteringChangeNumber = Filter.GetChangeNumber();
 					FilteringAsyncTask = MakeUnique<FAsyncTask<FLogFilteringAsyncTask>>(FilteringStartIndex, FilteringEndIndex, Filter, SharedThis(this));
-					UE_LOG(TimingProfiler, Log, TEXT("[LogView] Start async task for filtering by%s%s%s (\"%s\") (%d to %d)"),
+					UE_LOG(TraceInsights, Log, TEXT("[LogView] Start async task for filtering by%s%s%s (\"%s\") (%d to %d)"),
 						Filter.IsFilterSetByVerbosity() ? TEXT(" Verbosity,") : TEXT(""),
 						Filter.IsFilterSetByCategory() ? TEXT(" Category,") : TEXT(""),
 						Filter.IsFilterSetByText() ? TEXT(" Text") : TEXT(""),
@@ -831,7 +824,7 @@ void SLogView::Tick(const FGeometry& AllottedGeometry, const double InCurrentTim
 				uint64 DurationMs = FilteringStopwatch.GetAccumulatedTimeMs();
 				if (DurationMs > 10) // avoids spams
 				{
-					UE_LOG(TimingProfiler, Log, TEXT("[LogView] Updated (no filter; %d added / %d total messages) in %llu ms."),
+					UE_LOG(TraceInsights, Log, TEXT("[LogView] Updated (no filter; %d added / %d total messages) in %llu ms."),
 						NumAddedMessages, NewMessageCount, DurationMs);
 				}
 			}
@@ -839,7 +832,7 @@ void SLogView::Tick(const FGeometry& AllottedGeometry, const double InCurrentTim
 		else // if (NewMessageCount < TotalNumMessages)
 		{
 			// Just reset. On next Tick() the list will grow if needed.
-			UE_LOG(TimingProfiler, Log, TEXT("[LogView] RESET"));
+			UE_LOG(TraceInsights, Log, TEXT("[LogView] RESET"));
 			Cache.Reset();
 			Messages.Reset();
 			TotalNumMessages = 0;
@@ -900,7 +893,7 @@ void SLogView::Tick(const FGeometry& AllottedGeometry, const double InCurrentTim
 			{
 				int32 NumAsyncFilteredMessages = Task.GetEndIndex() - Task.GetStartIndex();
 				double Speed = static_cast<double>(NumAsyncFilteredMessages) / FilteringStopwatch.GetAccumulatedTime();
-				UE_LOG(TimingProfiler, Log, TEXT("[LogView] Updated (%d added / %d async filtered / %d total messages) in %llu ms (%.2f messages/second)."),
+				UE_LOG(TraceInsights, Log, TEXT("[LogView] Updated (%d added / %d async filtered / %d total messages) in %llu ms (%.2f messages/second)."),
 					NumFilteredMessages, NumAsyncFilteredMessages, TotalNumMessages, DurationMs, Speed);
 			}
 		}
@@ -1001,7 +994,7 @@ void SLogView::FilterTextBox_OnTextChanged(const FText& InFilterText)
 void SLogView::OnFilterChanged()
 {
 	const FString FilterText = FilterTextBox->GetText().ToString();
-	UE_LOG(TimingProfiler, Log, TEXT("[LogView] OnFilterChanged: \"%s\""), *FilterText);
+	UE_LOG(TraceInsights, Log, TEXT("[LogView] OnFilterChanged: \"%s\""), *FilterText);
 	Cache.Reset();
 	Messages.Reset();
 	TotalNumMessages = 0;
diff --git a/Engine/Source/Developer/TraceInsights/Private/Insights/Widgets/SSessionInfoWindow.h b/Engine/Source/Developer/TraceInsights/Private/Insights/Widgets/SSessionInfoWindow.h
--- a/Engine/Source/Developer/TraceInsights/Private/Insights/Widgets/SSessionInfoWindow.h
+++ b/Engine/Source/Developer/TraceInsights/Private/Insights/Widgets/SSessionInfoWindow.h
@@ -125,15 +125,17 @@ private:
 	/** The handle to the active update duration tick */
 	TWeakPtr<FActiveTimerHandle> ActiveTimerHandle;
 
-	//Trace::FSessionHandle SessionHandle;
+	//uint32 TraceId;
 	//FText Name;
-	//FText Uri;
+
 	//FText Platform;
 	//FText AppName;
 	//FText CommandLine;
 	//EBuildConfiguration ConfigurationType;
 	//EBuildTargetType TargetType;
-	//FDateTime TimeStamp;
+
+	//FText Uri;
+	//FDateTime Timestamp;
 	//uint64 Size;
 	//bool bIsLive;
 };
diff --git a/Engine/Source/Developer/TraceInsights/Private/Insights/Widgets/SStartPageWindow.cpp b/Engine/Source/Developer/TraceInsights/Private/Insights/Widgets/SStartPageWindow.cpp
--- a/Engine/Source/Developer/TraceInsights/Private/Insights/Widgets/SStartPageWindow.cpp
+++ b/Engine/Source/Developer/TraceInsights/Private/Insights/Widgets/SStartPageWindow.cpp
@@ -13,6 +13,8 @@
 #include "Internationalization/Text.h"
 #include "SlateOptMacros.h"
 #include "Styling/CoreStyle.h"
+#include "Trace/Analysis.h"
+#include "Trace/Analyzer.h"
 #include "Trace/ControlClient.h"
 #include "Trace/StoreClient.h"
 #include "Widgets/Docking/SDockTab.h"
@@ -22,6 +24,7 @@
 #include "Widgets/Layout/SBorder.h"
 #include "Widgets/Layout/SBox.h"
 #include "Widgets/Layout/SScrollBox.h"
+#include "Widgets/Layout/SSeparator.h"
 #include "Widgets/Layout/SSpacer.h"
 #include "Widgets/Notifications/SNotificationList.h"
 #include "Widgets/SBoxPanel.h"
@@ -39,13 +42,13 @@
 #endif // WITH_EDITOR
 
 // Insights
+#include "Insights/Common/Stopwatch.h"
 #include "Insights/InsightsManager.h"
-#include "Insights/TimingProfilerCommon.h" // for UE_LOG(TimingProfiler, ...
+#include "Insights/Log.h"
 #include "Insights/TimingProfilerManager.h"
 #include "Insights/InsightsStyle.h"
 #include "Insights/Version.h"
 #include "Insights/Widgets/SInsightsSettings.h"
-#include "Insights/Widgets/STimingProfilerWindow.h"
 
 ////////////////////////////////////////////////////////////////////////////////////////////////////
 
@@ -377,11 +380,7 @@ SStartPageWindow::SStartPageWindow()
 	, ActiveTimerHandle()
 	, MainContentPanel()
 	, LiveSessionCount(0)
-#if WITH_EDITOR
 	, bAutoStartAnalysisForLiveSessions(false)
-#else
-	, bAutoStartAnalysisForLiveSessions(true)
-#endif
 	, AutoStartedSessions()
 	, AutoStartPlatformFilter()
 	, AutoStartAppNameFilter()
@@ -433,85 +432,56 @@ void SStartPageWindow::Construct(const FArguments& InArgs)
 			// Overlay slot for the main window area
 			+ SOverlay::Slot()
 			.HAlign(HAlign_Fill)
-			.VAlign(VAlign_Center)
+			.VAlign(VAlign_Fill)
 			[
-				SNew(SScrollBox)
-				.Orientation(Orient_Vertical)
+				SAssignNew(MainContentPanel, SVerticalBox)
 
-				+ SScrollBox::Slot()
+				+ SVerticalBox::Slot()
+				.HAlign(HAlign_Fill)
+				.VAlign(VAlign_Fill)
+				.FillHeight(1.0f)
+				.Padding(3.0f, 3.0f)
 				[
-					SAssignNew(MainContentPanel, SVerticalBox)
-
-					+ SVerticalBox::Slot()
-					.AutoHeight()
-					.HAlign(HAlign_Center)
-					.Padding(3.0f, 3.0f)
+					SNew(SBox)
 					[
-						SNew(SBox)
-						.WidthOverride(1024.0f)
+						SNew(SBorder)
+						.BorderImage(FEditorStyle::GetBrush("NotificationList.ItemBackground"))
+						.Padding(8.0f)
 						[
-							SNew(SBorder)
-							.BorderImage(FEditorStyle::GetBrush("NotificationList.ItemBackground"))
-							.Padding(8.0f)
-							.HAlign(HAlign_Fill)
-							[
-								ConstructSessionsPanel()
-							]
+							ConstructSessionsPanel()
 						]
 					]
+				]
 
-						+ SVerticalBox::Slot()
-						.AutoHeight()
-						.HAlign(HAlign_Center)
-						.Padding(3.0f, 3.0f)
+				+ SVerticalBox::Slot()
+				.HAlign(HAlign_Fill)
+				.AutoHeight()
+				.Padding(3.0f, 3.0f)
+				[
+					SNew(SBox)
+					[
+						SNew(SBorder)
+						.BorderImage(FEditorStyle::GetBrush("NotificationList.ItemBackground"))
+						.Padding(8.0f)
 						[
-							SNew(SBox)
-							.WidthOverride(1024.0f)
-							[
-								SNew(SBorder)
-								.BorderImage(FEditorStyle::GetBrush("NotificationList.ItemBackground"))
-								.Padding(8.0f)
-								.HAlign(HAlign_Fill)
-								[
-									ConstructAutoStartPanel()
-								]
-							]
+							ConstructLocalSessionDirectoryPanel()
 						]
+					]
+				]
 
-					//+ SVerticalBox::Slot()
-					//.AutoHeight()
-					//.HAlign(HAlign_Center)
-					//.Padding(3.0f, 3.0f)
-					//[
-					//	SNew(SBox)
-					//	.WidthOverride(256.0f)
-					//	[
-					//		SNew(SBorder)
-					//		.BorderImage(FEditorStyle::GetBrush("NotificationList.ItemBackground"))
-					//		.Padding(8.0f)
-					//		.HAlign(HAlign_Fill)
-					//		[
-					//			ConstructRecorderPanel()
-					//		]
-					//	]
-					//]
-
-					+ SVerticalBox::Slot()
-					.AutoHeight()
-					.HAlign(HAlign_Center)
-					.Padding(3.0f, 3.0f)
+				+ SVerticalBox::Slot()
+				.HAlign(HAlign_Fill)
+				.AutoHeight()
+				.Padding(3.0f, 3.0f)
+				[
+					SNew(SBox)
+					.Visibility(this, &SStartPageWindow::StopTraceRecorder_Visibility)
 					[
-						SNew(SBox)
-						.WidthOverride(512.0f)
-						.Visibility(this, &SStartPageWindow::StopTraceRecorder_Visibility)
+						SNew(SBorder)
+						.BorderImage(FEditorStyle::GetBrush("NotificationList.ItemBackground"))
+						.Padding(8.0f)
 						[
-							SNew(SBorder)
-							.BorderImage(FEditorStyle::GetBrush("NotificationList.ItemBackground"))
-							.Padding(8.0f)
-							.HAlign(HAlign_Fill)
-							[
-								ConstructConnectPanel()
-							]
+							ConstructConnectPanel()
 						]
 					]
 				]
@@ -575,86 +545,88 @@ TSharedRef<SWidget> SStartPageWindow::ConstructSessionsPanel()
 	TSharedRef<SWidget> Widget = SNew(SVerticalBox)
 
 	+ SVerticalBox::Slot()
-		.AutoHeight()
-		.HAlign(HAlign_Center)
-		.Padding(0.0f, 2.0f)
-		[
-			SNew(STextBlock)
-			.Text(LOCTEXT("SessionsPanelTitle", "Trace Sessions"))
-			.Font(FCoreStyle::GetDefaultFontStyle("Bold", 11))
-			.ColorAndOpacity(FLinearColor::Gray)
-		]
+	.AutoHeight()
+	.HAlign(HAlign_Left)
+	.Padding(0.0f, 2.0f)
+	[
+		SNew(STextBlock)
+		.Text(LOCTEXT("SessionsPanelTitle", "Trace Sessions"))
+		.Font(FCoreStyle::GetDefaultFontStyle("Bold", 11))
+		.ColorAndOpacity(FLinearColor::Gray)
+	]
 
 	+ SVerticalBox::Slot()
-		.AutoHeight()
-		.HAlign(HAlign_Fill)
-		.Padding(0.0f, 1.0f, 0.0f, 2.0f)
-		.MaxHeight(22.0f + 20 * 14.0f) // max 20 rows
-		[
-			SAssignNew(TraceSessionsListView, SListView<TSharedPtr<FTraceSession>>)
-			.IsFocusable(true)
-			.ItemHeight(20.0f)
-			.SelectionMode(ESelectionMode::Single)
-			.OnSelectionChanged(this, &SStartPageWindow::TraceSessions_OnSelectionChanged)
-			.OnMouseButtonDoubleClick(this, &SStartPageWindow::TraceSessions_OnMouseButtonDoubleClick)
-			.ListItemsSource(&TraceSessions)
-			.OnGenerateRow(this, &SStartPageWindow::TraceSessions_OnGenerateRow)
-			.ConsumeMouseWheel(EConsumeMouseWheel::Always)
-			//.OnContextMenuOpening(FOnContextMenuOpening::CreateSP(this, &SStartPageWindow::TraceSessions_GetContextMenu))
-			.HeaderRow
-			(
-				SNew(SHeaderRow)
-
-				+ SHeaderRow::Column(FName(TEXT("Name")))
-				.FillWidth(0.25f)
-				.DefaultLabel(LOCTEXT("NameColumn", "Name"))
-
-				+ SHeaderRow::Column(FName(TEXT("Platform")))
-				.FillWidth(0.1f)
-				.DefaultLabel(LOCTEXT("PlatformColumn", "Platform"))
-
-				+ SHeaderRow::Column(FName(TEXT("AppName")))
-				.FillWidth(0.1f)
-				.DefaultLabel(LOCTEXT("AppNameColumn", "App Name"))
-
-				+ SHeaderRow::Column(FName(TEXT("BuildConfig")))
-				.FillWidth(0.1f)
-				.DefaultLabel(LOCTEXT("BuildConfigColumn", "Build Config"))
-
-				+ SHeaderRow::Column(FName(TEXT("BuildTarget")))
-				.FillWidth(0.1f)
-				.DefaultLabel(LOCTEXT("BuildTargetColumn", "Build Target"))
-
-				+ SHeaderRow::Column(FName(TEXT("Size")))
-				.FixedWidth(100.0f)
-				.HAlignHeader(HAlign_Right)
-				.HAlignCell(HAlign_Right)
-				.DefaultLabel(LOCTEXT("SizeColumn", "File Size"))
-
-				+ SHeaderRow::Column(FName(TEXT("Status")))
-				.FixedWidth(60.0f)
-				.HAlignHeader(HAlign_Right)
-				.HAlignCell(HAlign_Right)
-				.DefaultLabel(LOCTEXT("StatusColumn", "Status"))
-			)
-		]
+	.HAlign(HAlign_Fill)
+	.VAlign(VAlign_Fill)
+	.Padding(0.0f, 1.0f, 0.0f, 2.0f)
+	[
+		SAssignNew(TraceSessionsListView, SListView<TSharedPtr<FTraceSession>>)
+		.IsFocusable(true)
+		.ItemHeight(20.0f)
+		.SelectionMode(ESelectionMode::Single)
+		.OnSelectionChanged(this, &SStartPageWindow::TraceSessions_OnSelectionChanged)
+		.OnMouseButtonDoubleClick(this, &SStartPageWindow::TraceSessions_OnMouseButtonDoubleClick)
+		.ListItemsSource(&TraceSessions)
+		.OnGenerateRow(this, &SStartPageWindow::TraceSessions_OnGenerateRow)
+		.ConsumeMouseWheel(EConsumeMouseWheel::Always)
+		//.OnContextMenuOpening(FOnContextMenuOpening::CreateSP(this, &SStartPageWindow::TraceSessions_GetContextMenu))
+		.HeaderRow
+		(
+			SNew(SHeaderRow)
+
+			+ SHeaderRow::Column(FName(TEXT("Name")))
+			.FillWidth(0.25f)
+			.DefaultLabel(LOCTEXT("NameColumn", "Name"))
+
+			+ SHeaderRow::Column(FName(TEXT("Platform")))
+			.FillWidth(0.1f)
+			.DefaultLabel(LOCTEXT("PlatformColumn", "Platform"))
+
+			+ SHeaderRow::Column(FName(TEXT("AppName")))
+			.FillWidth(0.1f)
+			.DefaultLabel(LOCTEXT("AppNameColumn", "App Name"))
+
+			+ SHeaderRow::Column(FName(TEXT("BuildConfig")))
+			.FillWidth(0.1f)
+			.DefaultLabel(LOCTEXT("BuildConfigColumn", "Build Config"))
+
+			+ SHeaderRow::Column(FName(TEXT("BuildTarget")))
+			.FillWidth(0.1f)
+			.DefaultLabel(LOCTEXT("BuildTargetColumn", "Build Target"))
+
+			+ SHeaderRow::Column(FName(TEXT("Size")))
+			.FixedWidth(100.0f)
+			.HAlignHeader(HAlign_Right)
+			.HAlignCell(HAlign_Right)
+			.DefaultLabel(LOCTEXT("SizeColumn", "File Size"))
+
+			+ SHeaderRow::Column(FName(TEXT("Status")))
+			.FixedWidth(60.0f)
+			.HAlignHeader(HAlign_Right)
+			.HAlignCell(HAlign_Right)
+			.DefaultLabel(LOCTEXT("StatusColumn", "Status"))
+		)
+	]
 
 	+ SVerticalBox::Slot()
-		.AutoHeight()
-		.HAlign(HAlign_Right)
-		.Padding(0.0f, 2.0f)
-		[
-			ConstructLoadPanel()
-		]
+	.AutoHeight()
+	.HAlign(HAlign_Fill)
+	.Padding(0.0f, 4.0f, 0.0f, 6.0f)
+	[
+		SNew(SSeparator)
+		.Orientation(Orient_Horizontal)
+		.SeparatorImage(FInsightsStyle::Get().GetBrush("WhiteBrush"))
+		.ColorAndOpacity(FLinearColor::Black)
+		.Thickness(1.0f)
+	]
 
 	+ SVerticalBox::Slot()
-		.AutoHeight()
-		.HAlign(HAlign_Left)
-		.Padding(0.0f, 2.0f)
-		[
-			ConstructLocalSessionDirectoryPanel()
-		]
-	;
+	.AutoHeight()
+	.HAlign(HAlign_Fill)
+	.Padding(0.0f, 2.0f)
+	[
+		ConstructLoadPanel()
+	];
 
 	return Widget;
 }
@@ -666,44 +638,53 @@ TSharedRef<SWidget> SStartPageWindow::ConstructLoadPanel()
 	TSharedRef<SWidget> Widget = SNew(SHorizontalBox)
 
 	+ SHorizontalBox::Slot()
-		.AutoWidth()
+	.FillWidth(1.0f)
+	[
+		SNew(SBox)
+		.HAlign(HAlign_Left)
 		[
-			SNew(SButton)
-			.IsEnabled(this, &SStartPageWindow::Open_IsEnabled)
-			.OnClicked(this, &SStartPageWindow::Open_OnClicked)
-			.ToolTipText(LOCTEXT("OpenButtonTooltip", "Start analysis for selected trace session."))
-			.ContentPadding(FMargin(4.0f, 1.0f))
-			.Content()
-			[
-				SNew(SHorizontalBox)
-
-				+ SHorizontalBox::Slot()
-					.AutoWidth()
-					[
-						SNew(SImage)
-						.Image(FInsightsStyle::GetBrush("Open.Icon.Small"))
-					]
-
-				+ SHorizontalBox::Slot()
-					.AutoWidth()
-					.VAlign(VAlign_Center)
-					[
-						SNew(STextBlock)
-						.Text(LOCTEXT("OpenButtonText", "Open"))
-					]
-			]
+			ConstructAutoStartPanel()
 		]
+	]
 
 	+ SHorizontalBox::Slot()
-		.AutoWidth()
+	.AutoWidth()
+	[
+		SNew(SButton)
+		.IsEnabled(this, &SStartPageWindow::Open_IsEnabled)
+		.OnClicked(this, &SStartPageWindow::Open_OnClicked)
+		.ToolTipText(LOCTEXT("OpenButtonTooltip", "Start analysis for selected trace session."))
+		.ContentPadding(FMargin(4.0f, 1.0f))
+		.Content()
 		[
-			SNew(SComboButton)
-			.ToolTipText(LOCTEXT("MRU_Tooltip", "Open a trace file or choose a trace session."))
-			.OnGetMenuContent(this, &SStartPageWindow::MakeSessionListMenu)
-			.HasDownArrow(true)
-			.ContentPadding(FMargin(1.0f, 1.0f, 1.0f, 1.0f))
+			SNew(SHorizontalBox)
+
+			+ SHorizontalBox::Slot()
+				.AutoWidth()
+				[
+					SNew(SImage)
+					.Image(FInsightsStyle::GetBrush("Open.Icon.Small"))
+				]
+
+			+ SHorizontalBox::Slot()
+				.AutoWidth()
+				.VAlign(VAlign_Center)
+				[
+					SNew(STextBlock)
+					.Text(LOCTEXT("OpenButtonText", "Open"))
+				]
 		]
-	;
+	]
+
+	+ SHorizontalBox::Slot()
+	.AutoWidth()
+	[
+		SNew(SComboButton)
+		.ToolTipText(LOCTEXT("MRU_Tooltip", "Open a trace file or choose a trace session."))
+		.OnGetMenuContent(this, &SStartPageWindow::MakeSessionListMenu)
+		.HasDownArrow(true)
+		.ContentPadding(FMargin(1.0f, 1.0f, 1.0f, 1.0f))
+	];
 
 	return Widget;
 }
@@ -714,43 +695,42 @@ TSharedRef<SWidget> SStartPageWindow::ConstructLocalSessionDirectoryPanel()
 	TSharedRef<SWidget> Widget = SNew(SVerticalBox)
 
 	+ SVerticalBox::Slot()
-		.AutoHeight()
-		.HAlign(HAlign_Left)
+	.AutoHeight()
+	.HAlign(HAlign_Left)
+	.Padding(0.0f, 0.0f)
+	[
+		SNew(STextBlock)
+		.Text(LOCTEXT("LocalSessionDirectoryText", "Local Session Directory:"))
+		.ColorAndOpacity(FLinearColor::Gray)
+	]
+
+	+ SVerticalBox::Slot()
+	.AutoHeight()
+	.HAlign(HAlign_Left)
+	.Padding(0.0f, 0.0f)
+	[
+		SNew(SHorizontalBox)
+
+		+ SHorizontalBox::Slot()
 		.Padding(0.0f, 0.0f)
+		.VAlign(VAlign_Center)
 		[
 			SNew(STextBlock)
-			.Text(LOCTEXT("LocalSessionDirectoryText", "Local Session Directory:"))
-			.ColorAndOpacity(FLinearColor::Gray)
+			.Text(this, &SStartPageWindow::GetLocalSessionDirectory)
+			.Justification(ETextJustify::Right)
 		]
 
-	+ SVerticalBox::Slot()
-		.AutoHeight()
-		.HAlign(HAlign_Left)
-		.Padding(0.0f, 0.0f)
+		+ SHorizontalBox::Slot()
+		.AutoWidth()
+		.Padding(4.0f, 0.0f, 0.0f, 0.0f)
+		.VAlign(VAlign_Center)
 		[
-			SNew(SHorizontalBox)
-
-			+ SHorizontalBox::Slot()
-				.Padding(0.0f, 0.0f)
-				.VAlign(VAlign_Center)
-				[
-					SNew(STextBlock)
-					.Text(this, &SStartPageWindow::GetLocalSessionDirectory)
-					.Justification(ETextJustify::Right)
-				]
-
-			+ SHorizontalBox::Slot()
-				.AutoWidth()
-				.Padding(4.0f, 0.0f, 0.0f, 0.0f)
-				.VAlign(VAlign_Center)
-				[
-					SNew(SButton)
-					.Text(LOCTEXT("ExploreLocalSessionDirButton", "..."))
-					.ToolTipText(LOCTEXT("ExploreLocalSessionDirButtonToolTip", "Explore the Local Session Directory"))
-					.OnClicked(this, &SStartPageWindow::ExploreLocalSessionDirectory_OnClicked)
-				]
+			SNew(SButton)
+			.Text(LOCTEXT("ExploreLocalSessionDirButton", "..."))
+			.ToolTipText(LOCTEXT("ExploreLocalSessionDirButtonToolTip", "Explore the Local Session Directory"))
+			.OnClicked(this, &SStartPageWindow::ExploreLocalSessionDirectory_OnClicked)
 		]
-	;
+	];
 
 	return Widget;
 }
@@ -762,40 +742,42 @@ TSharedRef<SWidget> SStartPageWindow::ConstructAutoStartPanel()
 	TSharedRef<SWidget> Widget = SNew(SHorizontalBox)
 
 	+ SHorizontalBox::Slot()
-		.AutoWidth()
-		.Padding(0.0f, 0.0f, 0.0f, 0.0f)
-		.VAlign(VAlign_Center)
+	.AutoWidth()
+	.Padding(0.0f, 0.0f, 0.0f, 0.0f)
+	.HAlign(HAlign_Left)
+	.VAlign(VAlign_Center)
+	[
+		SNew(SCheckBox)
+		.ToolTipText(LOCTEXT("AutoStart_Tooltip", "Enable auto-start analysis for LIVE sessions."))
+		.IsChecked(this, &SStartPageWindow::AutoStart_IsChecked)
+		.OnCheckStateChanged(this, &SStartPageWindow::AutoStart_OnCheckStateChanged)
 		[
-			SNew(SCheckBox)
-			.ToolTipText(LOCTEXT("AutoStart_Tooltip", "Enable auto-start analysis for LIVE sessions."))
-			.IsChecked(this, &SStartPageWindow::AutoStart_IsChecked)
-			.OnCheckStateChanged(this, &SStartPageWindow::AutoStart_OnCheckStateChanged)
-			[
-				SNew(STextBlock)
-				.Text(LOCTEXT("AutoStart_Text", "Auto Start Analysis for LIVE Sessions"))
-			]
+			SNew(STextBlock)
+			.Text(LOCTEXT("AutoStart_Text", "Auto Start Analysis for LIVE Sessions"))
 		]
+	]
 
 	+ SHorizontalBox::Slot()
-		.AutoWidth()
-		.Padding(4.0f, 0.0f, 0.0f, 0.0f)
-		.VAlign(VAlign_Center)
-		[
-			SAssignNew(AutoStartPlatformFilter, SSearchBox)
-			.HintText(LOCTEXT("AutoStartPlatformFilter_Hint", "Platform"))
-			.ToolTipText(LOCTEXT("AutoStartPlatformFilter_Tooltip", "Type here to specify the Platform filter.\nAuto-start analysis will be enabled only for live sessions with this specified Platform."))
-		]
+	.AutoWidth()
+	.Padding(4.0f, 0.0f, 0.0f, 0.0f)
+	.HAlign(HAlign_Left)
+	.VAlign(VAlign_Center)
+	[
+		SAssignNew(AutoStartPlatformFilter, SSearchBox)
+		.HintText(LOCTEXT("AutoStartPlatformFilter_Hint", "Platform"))
+		.ToolTipText(LOCTEXT("AutoStartPlatformFilter_Tooltip", "Type here to specify the Platform filter.\nAuto-start analysis will be enabled only for live sessions with this specified Platform."))
+	]
 
 	+ SHorizontalBox::Slot()
-		.AutoWidth()
-		.Padding(4.0f, 0.0f, 0.0f, 0.0f)
-		.VAlign(VAlign_Center)
-		[
-			SAssignNew(AutoStartAppNameFilter, SSearchBox)
-			.HintText(LOCTEXT("AutoStartAppNameFilter_Hint", "AppName"))
-			.ToolTipText(LOCTEXT("AutoStartAppNameFilter_Tooltip", "Type here to specify the AppName filter.\nAuto-start analysis will be enabled only for live sessions with this specified AppName."))
-		]
-	;
+	.AutoWidth()
+	.Padding(4.0f, 0.0f, 0.0f, 0.0f)
+	.HAlign(HAlign_Left)
+	.VAlign(VAlign_Center)
+	[
+		SAssignNew(AutoStartAppNameFilter, SSearchBox)
+		.HintText(LOCTEXT("AutoStartAppNameFilter_Hint", "AppName"))
+		.ToolTipText(LOCTEXT("AutoStartAppNameFilter_Tooltip", "Type here to specify the AppName filter.\nAuto-start analysis will be enabled only for live sessions with this specified AppName."))
+	];
 
 	return Widget;
 }
@@ -807,89 +789,88 @@ TSharedRef<SWidget> SStartPageWindow::ConstructRecorderPanel()
 	TSharedRef<SWidget> Widget = SNew(SVerticalBox)
 
 	+ SVerticalBox::Slot()
-		.AutoHeight()
-		.HAlign(HAlign_Center)
-		.Padding(0.0f, 2.0f)
+	.AutoHeight()
+	.HAlign(HAlign_Center)
+	.Padding(0.0f, 2.0f)
+	[
+		SNew(STextBlock)
+		.Text(LOCTEXT("RecorderPanelTitle", "Trace Recorder"))
+		.Font(FCoreStyle::GetDefaultFontStyle("Bold", 11))
+		.ColorAndOpacity(FLinearColor::Gray)
+	]
+
+	+ SVerticalBox::Slot()
+	.AutoHeight()
+	.HAlign(HAlign_Left)
+	.Padding(0.0f, 2.0f)
+	[
+		SNew(SHorizontalBox)
+
+		+ SHorizontalBox::Slot()
+		.AutoWidth()
+		.Padding(0.0f, 0.0f, 2.0f, 0.0f)
+		.VAlign(VAlign_Center)
 		[
 			SNew(STextBlock)
-			.Text(LOCTEXT("RecorderPanelTitle", "Trace Recorder"))
-			.Font(FCoreStyle::GetDefaultFontStyle("Bold", 11))
+			.Text(LOCTEXT("RecorderStatusTitle", "Status:"))
 			.ColorAndOpacity(FLinearColor::Gray)
 		]
 
-	+ SVerticalBox::Slot()
-		.AutoHeight()
-		.HAlign(HAlign_Left)
-		.Padding(0.0f, 2.0f)
+		+ SHorizontalBox::Slot()
+		.AutoWidth()
+		.Padding(2.0f, 0.0f)
+		.VAlign(VAlign_Center)
 		[
-			SNew(SHorizontalBox)
-
-			+ SHorizontalBox::Slot()
-				.AutoWidth()
-				.Padding(0.0f, 0.0f, 2.0f, 0.0f)
-				.VAlign(VAlign_Center)
-				[
-					SNew(STextBlock)
-					.Text(LOCTEXT("RecorderStatusTitle", "Status:"))
-					.ColorAndOpacity(FLinearColor::Gray)
-				]
-
-			+ SHorizontalBox::Slot()
-				.AutoWidth()
-				.Padding(2.0f, 0.0f)
-				.VAlign(VAlign_Center)
-				[
-					SNew(STextBlock)
-					.Text(this, &SStartPageWindow::GetRecorderStatusText)
-				]
-
-			+ SHorizontalBox::Slot()
-				.AutoWidth()
-				.Padding(2.0f, 0.0f)
-				.VAlign(VAlign_Center)
-				[
-					SNew(SButton)
-					.Text(LOCTEXT("StartRecorder", "Start"))
-					.ToolTipText(LOCTEXT("StartRecorderToolTip", "Start the Trace Recorder"))
-					.OnClicked(this, &SStartPageWindow::StartTraceRecorder_OnClicked)
-					.Visibility(this, &SStartPageWindow::StartTraceRecorder_Visibility)
-				]
+			SNew(STextBlock)
+			.Text(this, &SStartPageWindow::GetRecorderStatusText)
+		]
 
-			+ SHorizontalBox::Slot()
-				.AutoWidth()
-				.Padding(2.0f, 0.0f)
-				.VAlign(VAlign_Center)
-				[
-					SNew(SButton)
-					.Text(LOCTEXT("StopRecorder", "Stop"))
-					.ToolTipText(LOCTEXT("StopRecorderToolTip", "Stop the Trace Recorder"))
-					.OnClicked(this, &SStartPageWindow::StopTraceRecorder_OnClicked)
-					.Visibility(this, &SStartPageWindow::StopTraceRecorder_Visibility)
-				]
+		+ SHorizontalBox::Slot()
+		.AutoWidth()
+		.Padding(2.0f, 0.0f)
+		.VAlign(VAlign_Center)
+		[
+			SNew(SButton)
+			.Text(LOCTEXT("StartRecorder", "Start"))
+			.ToolTipText(LOCTEXT("StartRecorderToolTip", "Start the Trace Recorder"))
+			.OnClicked(this, &SStartPageWindow::StartTraceRecorder_OnClicked)
+			.Visibility(this, &SStartPageWindow::StartTraceRecorder_Visibility)
 		]
 
-	+ SVerticalBox::Slot()
-		.AutoHeight()
-		.HAlign(HAlign_Left)
-		.Padding(0.0f, 2.0f, 0.0f, 1.0f)
+		+ SHorizontalBox::Slot()
+		.AutoWidth()
+		.Padding(2.0f, 0.0f)
+		.VAlign(VAlign_Center)
 		[
-			SNew(SHorizontalBox)
+			SNew(SButton)
+			.Text(LOCTEXT("StopRecorder", "Stop"))
+			.ToolTipText(LOCTEXT("StopRecorderToolTip", "Stop the Trace Recorder"))
+			.OnClicked(this, &SStartPageWindow::StopTraceRecorder_OnClicked)
 			.Visibility(this, &SStartPageWindow::StopTraceRecorder_Visibility)
+		]
+	]
 
-			+ SHorizontalBox::Slot()
-				.AutoWidth()
-				.Padding(0.0f, 0.0f, 2.0f, 0.0f)
-				.VAlign(VAlign_Center)
-				[
-					SNew(STextBlock)
-					.Text_Lambda([this]() -> FText
-					{
-						return FText::Format(LOCTEXT("ConnectionCountFormat", "Connections / live sessions: {0}"), FText::AsNumber(LiveSessionCount));
-					})
-					.ColorAndOpacity(FLinearColor::Gray)
-				]
+	+ SVerticalBox::Slot()
+	.AutoHeight()
+	.HAlign(HAlign_Left)
+	.Padding(0.0f, 2.0f, 0.0f, 1.0f)
+	[
+		SNew(SHorizontalBox)
+		.Visibility(this, &SStartPageWindow::StopTraceRecorder_Visibility)
+
+		+ SHorizontalBox::Slot()
+		.AutoWidth()
+		.Padding(0.0f, 0.0f, 2.0f, 0.0f)
+		.VAlign(VAlign_Center)
+		[
+			SNew(STextBlock)
+			.Text_Lambda([this]() -> FText
+			{
+				return FText::Format(LOCTEXT("ConnectionCountFormat", "Connections / live sessions: {0}"), FText::AsNumber(LiveSessionCount));
+			})
+			.ColorAndOpacity(FLinearColor::Gray)
 		]
-	;
+	];
 
 	return Widget;
 }
@@ -901,52 +882,55 @@ TSharedRef<SWidget> SStartPageWindow::ConstructConnectPanel()
 	TSharedRef<SWidget> Widget = SNew(SVerticalBox)
 
 	+ SVerticalBox::Slot()
-		.AutoHeight()
-		.HAlign(HAlign_Center)
-		.Padding(0.0f, 2.0f)
+	.AutoHeight()
+	.HAlign(HAlign_Left)
+	.Padding(0.0f, 2.0f)
+	[
+		SNew(STextBlock)
+		.Text(LOCTEXT("ConnectPanelTitle", "New Connection"))
+		.Font(FCoreStyle::GetDefaultFontStyle("Bold", 11))
+		.ColorAndOpacity(FLinearColor::Gray)
+	]
+
+	+ SVerticalBox::Slot()
+	.AutoHeight()
+	.HAlign(HAlign_Fill)
+	.Padding(0.0f, 2.0f)
+	[
+		SNew(SHorizontalBox)
+
+		+ SHorizontalBox::Slot()
+		.AutoWidth()
+		.Padding(0.0f, 0.0f, 2.0f, 0.0f)
+		.VAlign(VAlign_Center)
 		[
 			SNew(STextBlock)
-			.Text(LOCTEXT("ConnectPanelTitle", "New Connection"))
-			.Font(FCoreStyle::GetDefaultFontStyle("Bold", 11))
+			.Text(LOCTEXT("HostTitle", "Running instance IP:"))
 			.ColorAndOpacity(FLinearColor::Gray)
 		]
 
-	+ SVerticalBox::Slot()
-		.AutoHeight()
-		.HAlign(HAlign_Fill)
-		.Padding(0.0f, 2.0f)
+		+ SHorizontalBox::Slot()
+		.AutoWidth()
+		.VAlign(VAlign_Center)
 		[
-			SNew(SHorizontalBox)
-
-			+ SHorizontalBox::Slot()
-				.AutoWidth()
-				.Padding(0.0f, 0.0f, 2.0f, 0.0f)
-				.VAlign(VAlign_Center)
-				[
-					SNew(STextBlock)
-					.Text(LOCTEXT("HostTitle", "Running instance IP:"))
-					.ColorAndOpacity(FLinearColor::Gray)
-				]
-
-			+ SHorizontalBox::Slot()
-				.FillWidth(1.0)
-				.VAlign(VAlign_Center)
-				[
-					SAssignNew(HostTextBox, SEditableTextBox)
-				]
+			SNew(SBox)
+			.MinDesiredWidth(120.0f)
+			[
+				SAssignNew(HostTextBox, SEditableTextBox)
+			]
+		]
 
-			+ SHorizontalBox::Slot()
-				.AutoWidth()
-				.Padding(2.0f, 0.0f)
-				.VAlign(VAlign_Center)
-				[
-					SNew(SButton)
-					.Text(LOCTEXT("Connect", "Connect"))
-					.ToolTipText(LOCTEXT("ConnectToolTip", "Try connecting to host."))
-					.OnClicked(this, &SStartPageWindow::Connect_OnClicked)
-				]
+		+ SHorizontalBox::Slot()
+		.AutoWidth()
+		.Padding(2.0f, 0.0f)
+		.VAlign(VAlign_Center)
+		[
+			SNew(SButton)
+			.Text(LOCTEXT("Connect", "Connect"))
+			.ToolTipText(LOCTEXT("ConnectToolTip", "Connect the running instance at specified ip with the local trace recorder."))
+			.OnClicked(this, &SStartPageWindow::Connect_OnClicked)
 		]
-	;
+	];
 
 	return Widget;
 }
@@ -1129,6 +1113,11 @@ void SStartPageWindow::RefreshTraceSessionList()
 	// If trace list has changed on store side, recreate the TraceSessions list view widget.
 	if (bTraceListChanged)
 	{
+		UE_LOG(TraceInsights, Log, TEXT("List of trace sessions has changed. Updating..."));
+
+		FStopwatch Stopwatch;
+		Stopwatch.Start();
+
 		TSharedPtr<FTraceSession> NewSelectedTraceSession;
 
 		TraceSessions.Reset();
@@ -1146,6 +1135,7 @@ void SStartPageWindow::RefreshTraceSessionList()
 			const TSharedRef<FTraceSession> TraceSession = MakeShared<FTraceSession>(TraceInfo);
 			TraceSession->TraceIndex = TraceIndex;
 			TraceSession->Uri = FText::FromString(FInsightsManager::Get()->GetStoreDir() + TEXT("/") + TraceSession->Name.ToString() + TEXT(".utrace"));
+			UpdateMetadata(*TraceSession);
 			TraceSessions.Add(TraceSession);
 			TraceSessionsMap.Add(TraceSession->TraceId, TraceSession);
 
@@ -1174,6 +1164,9 @@ void SStartPageWindow::RefreshTraceSessionList()
 			TraceSessionsListView->SetItemSelection(NewSelectedTraceSession, true);
 			TraceSessionsListView->RequestScrollIntoView(NewSelectedTraceSession);
 		}
+
+		Stopwatch.Stop();
+		UE_LOG(TraceInsights, Log, TEXT("List of trace sessions updated in %llu ms."), Stopwatch.GetAccumulatedTimeMs());
 	}
 
 	// Process the connected recorder sessions.
@@ -1211,6 +1204,7 @@ void SStartPageWindow::RefreshTraceSessionList()
 						(AutoStartConfigurationTypeFilter == EBuildConfiguration::Unknown || AutoStartConfigurationTypeFilter == TraceSession.ConfigurationType) &&
 						(AutoStartTargetTypeFilter == EBuildTargetType::Unknown || AutoStartTargetTypeFilter == TraceSession.TargetType))
 					{
+						UE_LOG(TraceInsights, Log, TEXT("Auto starting analysis for trace with id 0x%08X..."), TraceSession.TraceId);
 						AutoStartedSessions.Add(TraceSession.TraceId);
 						LoadTrace(TraceSession.TraceId);
 					}
@@ -1453,7 +1447,7 @@ void SStartPageWindow::LoadTraceFile(const FString& InTraceFile)
 {
 	if (FInsightsManager::Get()->ShouldOpenAnalysisInSeparateProcess())
 	{
-		UE_LOG(TimingProfiler, Log, TEXT("Start analysis (in separate process) for trace file: \"%s\""), *InTraceFile);
+		UE_LOG(TraceInsights, Log, TEXT("Start analysis (in separate process) for trace file: \"%s\""), *InTraceFile);
 
 		const TCHAR* ExecutablePath = FPlatformProcess::ExecutablePath();
 
@@ -1478,7 +1472,7 @@ void SStartPageWindow::LoadTraceFile(const FString& InTraceFile)
 	}
 	else
 	{
-		UE_LOG(TimingProfiler, Log, TEXT("Start analysis for trace file: \"%s\""), *InTraceFile);
+		UE_LOG(TraceInsights, Log, TEXT("Start analysis for trace file: \"%s\""), *InTraceFile);
 		FInsightsManager::Get()->LoadTraceFile(InTraceFile);
 	}
 }
@@ -1489,7 +1483,7 @@ void SStartPageWindow::LoadTrace(uint32 InTraceId)
 {
 	if (FInsightsManager::Get()->ShouldOpenAnalysisInSeparateProcess())
 	{
-		UE_LOG(TimingProfiler, Log, TEXT("Start analysis (in separate process) for trace id: 0x%08X"), InTraceId);
+		UE_LOG(TraceInsights, Log, TEXT("Start analysis (in separate process) for trace id: 0x%08X"), InTraceId);
 
 		const TCHAR* ExecutablePath = FPlatformProcess::ExecutablePath();
 
@@ -1520,7 +1514,7 @@ void SStartPageWindow::LoadTrace(uint32 InTraceId)
 	}
 	else
 	{
-		UE_LOG(TimingProfiler, Log, TEXT("Start analysis for trace id: 0x%08X"), InTraceId);
+		UE_LOG(TraceInsights, Log, TEXT("Start analysis for trace id: 0x%08X"), InTraceId);
 		FInsightsManager::Get()->LoadTrace(InTraceId);
 	}
 }
@@ -1546,7 +1540,7 @@ TSharedRef<SWidget> SStartPageWindow::MakeSessionListMenu()
 	}
 	MenuBuilder.EndSection();
 
-	MenuBuilder.BeginSection("AvailableSessions", LOCTEXT("AvailableSessionsHeading", "Top 10 Most Recently Created Sessions"));
+	MenuBuilder.BeginSection("AvailableSessions", LOCTEXT("AvailableSessionsHeading", "Top Most Recently Created Sessions"));
 	{
 		Trace::FStoreClient* StoreClient = FInsightsManager::Get()->GetStoreClient();
 		if (StoreClient != nullptr)
@@ -1687,4 +1681,109 @@ void SStartPageWindow::CloseSettings()
 
 ////////////////////////////////////////////////////////////////////////////////////////////////////
 
+struct FDiagnosticsSessionAnalyzer : public Trace::IAnalyzer
+{
+	virtual void OnAnalysisBegin(const FOnAnalysisContext& Context) override
+	{
+		Context.InterfaceBuilder.RouteEvent(0, "Diagnostics", "Session");
+	}
+
+	virtual bool OnEvent(uint16, const FOnEventContext& Context) override
+	{
+		const FEventData& EventData = Context.EventData;
+
+		const uint8* Attachment = EventData.GetAttachment();
+		if (Attachment == nullptr)
+		{
+			return false;
+		}
+
+		uint8 AppNameOffset = EventData.GetValue<uint8>("AppNameOffset");
+		uint8 CommandLineOffset = EventData.GetValue<uint8>("CommandLineOffset");
+
+		Platform = FString(AppNameOffset, (const ANSICHAR*)Attachment);
+
+		Attachment += AppNameOffset;
+		int32 AppNameLength = CommandLineOffset - AppNameOffset;
+		AppName = FString(AppNameLength, (const ANSICHAR*)Attachment);
+
+		Attachment += AppNameLength;
+		int32 CommandLineLength = EventData.GetAttachmentSize() - CommandLineOffset;
+		CommandLine = FString(CommandLineLength, (const ANSICHAR*)Attachment);
+
+		ConfigurationType = EventData.GetValue<int8>("ConfigurationType");
+		TargetType = EventData.GetValue<int8>("TargetType");
+
+		return false;
+	}
+
+	FString Platform;
+	FString AppName;
+	FString CommandLine;
+	int8 ConfigurationType = 0;
+	int8 TargetType = 0;
+};
+
+////////////////////////////////////////////////////////////////////////////////////////////////////
+
+void SStartPageWindow::UpdateMetadata(FTraceSession& TraceSession)
+{
+	if (TraceSession.bIsMetadataUpdated)
+	{
+		return;
+	}
+
+	Trace::FStoreClient* StoreClient = FInsightsManager::Get()->GetStoreClient();
+	if (StoreClient == nullptr)
+	{
+		return;
+	}
+
+	Trace::FStoreClient::FTraceData TraceData = StoreClient->ReadTrace(TraceSession.TraceId);
+	if (!TraceData)
+	{
+		return;
+	}
+
+	struct FDataStream : public Trace::IInDataStream
+	{
+		virtual int32 Read(void* Data, uint32 Size) override
+		{
+			if (BytesRead >= 48 * 1024)
+			{
+				return 0;
+			}
+
+			int32 InnerBytesRead = Inner->Read(Data, Size);
+			BytesRead += InnerBytesRead;
+			return InnerBytesRead;
+		}
+
+		int32 BytesRead = 0;
+		Trace::IInDataStream* Inner;
+	};
+
+	FDataStream DataStream;
+	DataStream.Inner = TraceData.Get();
+
+	FDiagnosticsSessionAnalyzer Analyzer;
+	Trace::FAnalysisContext Context;
+	Context.AddAnalyzer(Analyzer);
+	Context.Process(DataStream).Wait();
+
+	if (Analyzer.Platform.Len() != 0)
+	{
+		TraceSession.Platform = FText::FromString(Analyzer.Platform);
+		TraceSession.AppName = FText::FromString(Analyzer.AppName);
+		TraceSession.CommandLine = FText::FromString(Analyzer.CommandLine);
+		TraceSession.ConfigurationType = static_cast<EBuildConfiguration>(Analyzer.ConfigurationType);
+
+		TraceSession.TargetType = static_cast<EBuildTargetType>(Analyzer.TargetType);
+	}
+
+	TraceSession.bIsMetadataUpdated = true;
+}
+
+////////////////////////////////////////////////////////////////////////////////////////////////////
+
 #undef LOCTEXT_NAMESPACE
diff --git a/Engine/Source/Developer/TraceInsights/Private/Insights/Widgets/SStartPageWindow.h b/Engine/Source/Developer/TraceInsights/Private/Insights/Widgets/SStartPageWindow.h
--- a/Engine/Source/Developer/TraceInsights/Private/Insights/Widgets/SStartPageWindow.h
+++ b/Engine/Source/Developer/TraceInsights/Private/Insights/Widgets/SStartPageWindow.h
@@ -54,6 +54,7 @@ struct FTraceSession
 	bool bIsLive = false;
 	uint32 IpAddress = 0;
 
+	bool bIsMetadataUpdated = false;
 	FText Platform;
 	FText AppName;
 	FText CommandLine;
@@ -70,6 +71,7 @@ struct FTraceSession
 		, Size(InTraceInfo->GetSize())
 		, bIsLive(false)
 		, IpAddress(0)
+		, bIsMetadataUpdated(false)
 		, Platform()
 		, AppName()
 		, CommandLine()
@@ -112,9 +114,9 @@ public:
 
 private:
 	TSharedRef<SWidget> ConstructSessionsPanel();
-	TSharedRef<SWidget> ConstructAutoStartPanel();
 	TSharedRef<SWidget> ConstructLoadPanel();
 	TSharedRef<SWidget> ConstructLocalSessionDirectoryPanel();
+	TSharedRef<SWidget> ConstructAutoStartPanel();
 	TSharedRef<SWidget> ConstructRecorderPanel();
 	TSharedRef<SWidget> ConstructConnectPanel();
 
@@ -151,6 +153,7 @@ private:
 	void TraceSessions_OnMouseButtonDoubleClick(TSharedPtr<FTraceSession> TraceSession);
 	EVisibility TraceSessions_Visibility() const;
 	FReply RefreshTraceSessions_OnClicked();
+	void UpdateMetadata(FTraceSession& TraceSession);
 
 	ECheckBoxState AutoStart_IsChecked() const;
 	void AutoStart_OnCheckStateChanged(ECheckBoxState NewState);
diff --git a/Engine/Source/Developer/TraceInsights/Private/Insights/Widgets/SStatsView.cpp b/Engine/Source/Developer/TraceInsights/Private/Insights/Widgets/SStatsView.cpp
--- a/Engine/Source/Developer/TraceInsights/Private/Insights/Widgets/SStatsView.cpp
+++ b/Engine/Source/Developer/TraceInsights/Private/Insights/Widgets/SStatsView.cpp
@@ -919,6 +919,7 @@ void SStatsView::TreeView_OnMouseButtonDoubleClick(FStatsNodePtr StatsNode)
 					}
 					else
 					{
+						GraphTrack->Show();
 						GraphTrack->AddStatsCounterSeries(StatsCounterId, StatsNode->GetColor());
 						GraphTrack->SetDirtyFlag();
 						StatsNode->SetAddedToGraphFlag(true);
diff --git a/Engine/Source/Developer/TraceInsights/Private/Insights/Widgets/STimingProfilerWindow.cpp b/Engine/Source/Developer/TraceInsights/Private/Insights/Widgets/STimingProfilerWindow.cpp
--- a/Engine/Source/Developer/TraceInsights/Private/Insights/Widgets/STimingProfilerWindow.cpp
+++ b/Engine/Source/Developer/TraceInsights/Private/Insights/Widgets/STimingProfilerWindow.cpp
@@ -3,9 +3,9 @@
 #include "STimingProfilerWindow.h"
 
 #include "EditorStyleSet.h"
+#include "Framework/Docking/LayoutExtender.h"
 #include "Framework/Docking/WorkspaceItem.h"
 #include "Framework/MultiBox/MultiBoxBuilder.h"
-#include "Framework/Docking/LayoutExtender.h"
 #include "SlateOptMacros.h"
 #include "Widgets/Docking/SDockTab.h"
 #include "Widgets/Images/SImage.h"
@@ -27,6 +27,7 @@
 #include "Insights/InsightsManager.h"
 #include "Insights/InsightsStyle.h"
 #include "Insights/TimingProfilerManager.h"
+#include "Insights/TraceInsightsModule.h"
 #include "Insights/Version.h"
 #include "Insights/Widgets/SFrameTrack.h"
 #include "Insights/Widgets/SInsightsSettings.h"
@@ -36,7 +37,6 @@
 #include "Insights/Widgets/STimerTreeView.h"
 #include "Insights/Widgets/STimingProfilerToolbar.h"
 #include "Insights/Widgets/STimingView.h"
-#include "Insights/TraceInsightsModule.h"
 
 ////////////////////////////////////////////////////////////////////////////////////////////////////
 
@@ -373,7 +373,7 @@ void STimingProfilerWindow::Construct(const FArguments& InArgs, const TSharedRef
 
 	TSharedPtr<FTabManager::FLayout> Layout;
 
-	// Check for layout overrides
+	// Check for layout overrides.
 	FTraceInsightsModule& TraceInsightsModule = FModuleManager::GetModuleChecked<FTraceInsightsModule>("TraceInsights");
 	FInsightsMajorTabConfig TabConfig = TraceInsightsModule.FindMajorTabConfig(FInsightsManagerTabs::TimingProfilerTabId);
 
@@ -383,8 +383,8 @@ void STimingProfilerWindow::Construct(const FArguments& InArgs, const TSharedRef
 		ExtensionDelegate->Broadcast(*Extension);
 	}
 
-	// Register any new minor tabs
-	for(const FInsightsMinorTabConfig& MinorTabConfig : Extension->GetMinorTabs())
+	// Register any new minor tabs.
+	for (const FInsightsMinorTabConfig& MinorTabConfig : Extension->GetMinorTabs())
 	{
 		FTabSpawnerEntry& TabSpawnerEntry = TabManager->RegisterTabSpawner(MinorTabConfig.TabId, MinorTabConfig.OnSpawnTab);
 
@@ -394,13 +394,13 @@ void STimingProfilerWindow::Construct(const FArguments& InArgs, const TSharedRef
 		.SetIcon(MinorTabConfig.TabIcon)
 		.SetReuseTabMethod(MinorTabConfig.OnFindTabToReuse);
 
-		if(MinorTabConfig.WorkspaceGroup.IsValid())
+		if (MinorTabConfig.WorkspaceGroup.IsValid())
 		{
 			TabSpawnerEntry.SetGroup(MinorTabConfig.WorkspaceGroup.ToSharedRef());
 		}
 	}
 
-	if(TabConfig.Layout.IsValid())
+	if (TabConfig.Layout.IsValid())
 	{
 		Layout = TabConfig.Layout;
 	}
diff --git a/Engine/Source/Developer/TraceInsights/Private/Insights/Widgets/STimingView.h b/Engine/Source/Developer/TraceInsights/Private/Insights/Widgets/STimingView.h
--- a/Engine/Source/Developer/TraceInsights/Private/Insights/Widgets/STimingView.h
+++ b/Engine/Source/Developer/TraceInsights/Private/Insights/Widgets/STimingView.h
@@ -31,8 +31,8 @@ class FThreadTimingSharedState;
 class FTimeRulerTrack; 
 class FTimingGraphTrack;
 class FTimingViewDrawHelper;
-class SScrollBar;
 class SOverlay;
+class SScrollBar;
 namespace Insights { class ITimingViewExtender; }
 
 /** A custom widget used to display timing events. */
@@ -222,6 +222,9 @@ public:
 	virtual void InvalidateScrollableTracksOrder() override;
 	//TODO: virtual void InvalidateScrollableTracksVisibility() override;
 
+	virtual double GetTimeMarker() const override { return TimeMarker; }
+	virtual void SetTimeMarker(double InMarkerTime) override;
+
 	virtual Insights::FSelectionChangedDelegate& OnSelectionChanged() override { return OnSelectionChangedDelegate; }
 	virtual Insights::FTimeMarkerChangedDelegate& OnTimeMarkerChanged() override { return OnTimeMarkerChangedDelegate; }
 	virtual Insights::FHoveredTrackChangedDelegate& OnHoveredTrackChanged() override { return OnHoveredTrackChangedDelegate; }
@@ -230,8 +233,6 @@ public:
 	virtual Insights::FSelectedEventChangedDelegate& OnSelectedEventChanged() override { return OnSelectedEventChangedDelegate; }
 
 	virtual void AddOverlayWidget(const TSharedRef<SWidget>& InWidget) override;
-	virtual void SetTimeMarker(double InMarkerTime) override;
-	virtual double GetTimeMarker() const override { return TimeMarker; }
 
 	////////////////////////////////////////////////////////////////////////////////////////////////////
 
diff --git a/Engine/Source/Developer/TraceInsights/Public/Insights/IUnrealInsightsModule.h b/Engine/Source/Developer/TraceInsights/Public/Insights/IUnrealInsightsModule.h
--- a/Engine/Source/Developer/TraceInsights/Public/Insights/IUnrealInsightsModule.h
+++ b/Engine/Source/Developer/TraceInsights/Public/Insights/IUnrealInsightsModule.h
@@ -4,8 +4,8 @@
 
 #include "CoreMinimal.h"
 #include "Modules/ModuleInterface.h"
-#include "Framework/Docking/TabManager.h"
 #include "Framework/Docking/LayoutExtender.h"
+#include "Framework/Docking/TabManager.h"
 #include "Framework/MultiBox/MultiBoxExtender.h"
 
 class FExtender;
@@ -107,12 +107,13 @@ struct TRACEINSIGHTS_API FInsightsMajorTabConfig
 struct TRACEINSIGHTS_API FInsightsMajorTabExtender
 {
 	FInsightsMajorTabExtender(TSharedPtr<FTabManager>& InTabManager) : MenuExtender(MakeShared<FExtender>()), TabManager(InTabManager) {}
-	
+
 	TSharedPtr<FExtender>& GetMenuExtender() { return MenuExtender; }
 	FLayoutExtender& GetLayoutExtender() { return LayoutExtender; }
-	FInsightsMinorTabConfig& AddMinorTabConfig() { return MinorTabs.AddDefaulted_GetRef(); }	
+	FInsightsMinorTabConfig& AddMinorTabConfig() { return MinorTabs.AddDefaulted_GetRef(); }
 	TSharedPtr<FTabManager> GetTabManager() const { return TabManager; }
 	const TArray<FInsightsMinorTabConfig>& GetMinorTabs() const { return MinorTabs; }
+
 protected:
 	/** Extender used to add to the menu for this tab */
 	TSharedPtr<FExtender> MenuExtender;
@@ -134,12 +135,6 @@ DECLARE_MULTICAST_DELEGATE_OneParam(FOnRegisterMajorTabExtensions, FInsightsMajo
 class IUnrealInsightsModule : public IModuleInterface
 {
 public:
-
-	/**
-	 * Connect to specific store.
-	 */
-	virtual bool ConnectToStore(const TCHAR* Host, uint32 Port) = 0;
-
 	/**
 	 * Called when the application starts in "Browser" mode.
 	 */
@@ -151,26 +146,43 @@ public:
 	virtual void CreateSessionViewer(bool bAllowDebugTools) = 0;
 
 	/**
-	 * Starts analysis of the specified *.utrace file. Called when the application starts in "Viewer" mode.
-	 *
-	 * @param InTraceFile The filename (*.utrace) of the trace to analyze.
+	 * Gets the store client.
 	 */
-	virtual void StartAnalysisForTraceFile(const TCHAR* InTraceFile) = 0;
+	virtual Trace::FStoreClient* GetStoreClient() = 0;
 
 	/**
-	 * Starts analysis of the specified trace. Called when the application starts in "Viewer" mode.
+	 * Connects to a specified store.
 	 *
 	 * @param InStoreHost The host of the store to connect to.
 	 * @param InStorePort The port of the store to connect to.
+	 * @return If connected succesfully or not.
+	 */
+	virtual bool ConnectToStore(const TCHAR* InStoreHost, uint32 InStorePort) = 0;
+
+	/**
+	 * Gets the current analysis session.
+	 */
+	virtual TSharedPtr<const Trace::IAnalysisSession> GetAnalysisSession() const = 0;
+
+	/**
+	 * Starts analysis of the specified trace. Called when the application starts in "Viewer" mode.
+	 *
 	 * @param InTraceId The id of the trace to analyze.
 	 */
-	virtual void StartAnalysisForTrace(const TCHAR* InStoreHost, uint32 InStorePort, uint32 InTraceId) = 0;
+	virtual void StartAnalysisForTrace(uint32 InTraceId) = 0;
 
 	/**
 	 * Starts analysis of the last live session. Called when the application starts in "Viewer" mode.
 	 */
 	virtual void StartAnalysisForLastLiveSession() = 0;
 
+	/**
+	 * Starts analysis of the specified *.utrace file. Called when the application starts in "Viewer" mode.
+	 *
+	 * @param InTraceFile The filename (*.utrace) of the trace to analyze.
+	 */
+	virtual void StartAnalysisForTraceFile(const TCHAR* InTraceFile) = 0;
+
 	/**
 	 * Called when the application shutsdown.
 	 */
@@ -179,33 +191,23 @@ public:
 	/**
 	 * Registers a major tab layout. This defines how the major tab will appear when spawned.
 	 * If this is not called prior to tabs being spawned then the built-in default layout will be used.
-	 * @param	InMajorTabId	The major tab ID we are supplying a layout to
-	 * @param	InConfig		The config to use when spawning the major tab
+	 * @param InMajorTabId The major tab ID we are supplying a layout to
+	 * @param InConfig The config to use when spawning the major tab
 	 */
 	virtual void RegisterMajorTabConfig(const FName& InMajorTabId, const FInsightsMajorTabConfig& InConfig) = 0;
 
 	/**
 	 * Unregisters a major tab layout. This will revert the major tab to spawning with its default layout
-	 * @param	InMajorTabId	The major tab ID we are supplying a layout to
+	 * @param InMajorTabId The major tab ID we are supplying a layout to
 	 */
 	virtual void UnregisterMajorTabConfig(const FName& InMajorTabId) = 0;
-	
-	/**
-	 * Allows for registering a delegate callback for populating a FInsightsMajorTabExtender structure	 
-	 * @param	InMajorTabId	The major tab ID to register the delegate for
-	 */
-	virtual FOnRegisterMajorTabExtensions& OnRegisterMajorTabExtension(const FName& InMajorTabId) = 0;
-
 
 	/**
-	 * Expose the session service directly
+	 * Allows for registering a delegate callback for populating a FInsightsMajorTabExtender structure.
+	 * @param InMajorTabId The major tab ID to register the delegate for
 	 */
-	virtual Trace::FStoreClient* GetStoreClient() = 0;
+	virtual FOnRegisterMajorTabExtensions& OnRegisterMajorTabExtension(const FName& InMajorTabId) = 0;
 
-	/**
-	 * Get the session service.
-	 */
-	virtual TSharedPtr<const Trace::IAnalysisSession> GetAnalysisSession() const = 0;
 	/** Callback invoked when a major tab is created */
 	virtual FOnInsightsMajorTabCreated& OnMajorTabCreated() = 0;
 };
diff --git a/Engine/Source/Developer/TraceInsights/Public/Insights/Log.h b/Engine/Source/Developer/TraceInsights/Public/Insights/Log.h
new file mode 100644
--- /dev/null
+++ b/Engine/Source/Developer/TraceInsights/Public/Insights/Log.h
@@ -0,0 +1,8 @@
+// Copyright Epic Games, Inc. All Rights Reserved.
+
+#pragma once
+
+#include "CoreMinimal.h"
+#include "Logging/LogMacros.h"
+
+DECLARE_LOG_CATEGORY_EXTERN(TraceInsights, Verbose, All);
diff --git a/Engine/Source/Developer/TraceInsights/Public/Insights/TimingProfilerCommon.h b/Engine/Source/Developer/TraceInsights/Public/Insights/TimingProfilerCommon.h
--- a/Engine/Source/Developer/TraceInsights/Public/Insights/TimingProfilerCommon.h
+++ b/Engine/Source/Developer/TraceInsights/Public/Insights/TimingProfilerCommon.h
@@ -9,14 +9,6 @@
 
 DECLARE_LOG_CATEGORY_EXTERN(TimingProfiler, Log, All);
 
-#define DEBUG_TIMING_PROFILER_PERFORMANCE 0
-
-#if DEBUG_TIMING_PROFILER_PERFORMANCE==1
-#define TP_SCOPE_LOG_TIME(arg0,arg1) SCOPE_LOG_TIME(arg0,arg1)
-#else
-#define TP_SCOPE_LOG_TIME(arg0,arg1)
-#endif
-
 DECLARE_STATS_GROUP(TEXT("TimingProfiler"), STATGROUP_TimingProfiler, STATCAT_Advanced);
 
 /** Time spent on Frame Track drawing. */
diff --git a/Engine/Source/Programs/UnrealInsights/Private/UserInterfaceCommand.cpp b/Engine/Source/Programs/UnrealInsights/Private/UserInterfaceCommand.cpp
--- a/Engine/Source/Programs/UnrealInsights/Private/UserInterfaceCommand.cpp
+++ b/Engine/Source/Programs/UnrealInsights/Private/UserInterfaceCommand.cpp
@@ -225,7 +225,8 @@ void FUserInterfaceCommand::InitializeSlateApplication()
 		FParse::Value(FCommandLine::Get(), TEXT("-StorePort="), StorePort);
 
 		TraceInsightsModule.CreateSessionViewer(bAllowDebugTools);
-		TraceInsightsModule.StartAnalysisForTrace(StoreHost, StorePort, TraceId);
+		TraceInsightsModule.ConnectToStore(StoreHost, StorePort);
+		TraceInsightsModule.StartAnalysisForTrace(TraceId);
 
 		delete[] StoreHost;
 	}
