[Insights]  Session Browser: Optimized the update of trace metadata when list of traces changes (ex.: new live traces sessions; file renames; file added/removed).

diff --git a/Engine/Source/Developer/TraceInsights/Private/Insights/Widgets/SStartPageWindow.cpp b/Engine/Source/Developer/TraceInsights/Private/Insights/Widgets/SStartPageWindow.cpp
--- a/Engine/Source/Developer/TraceInsights/Private/Insights/Widgets/SStartPageWindow.cpp
+++ b/Engine/Source/Developer/TraceInsights/Private/Insights/Widgets/SStartPageWindow.cpp
@@ -1096,8 +1096,8 @@ void SStartPageWindow::RefreshTraceSessionList()
 				TraceSession.bIsLive = false;
 				TraceSession.IpAddress = 0;
 
-				TraceSession.Size = TraceInfo->GetSize();
 				TraceSession.Timestamp = FTraceSession::ConvertTimestamp(TraceInfo->GetTimestamp());
+				TraceSession.Size = TraceInfo->GetSize();
 			}
 			else
 			{
@@ -1118,10 +1118,12 @@ void SStartPageWindow::RefreshTraceSessionList()
 		FStopwatch Stopwatch;
 		Stopwatch.Start();
 
+		int32 MetadataUpdateCount = 0; // for debugging
+
 		TSharedPtr<FTraceSession> NewSelectedTraceSession;
 
+		// Rebuild list of traces.
 		TraceSessions.Reset();
-		TraceSessionsMap.Reset();
 
 		const int32 TraceCount = StoreClient->GetTraceCount();
 		for (int32 TraceIndex = 0; TraceIndex < TraceCount; ++TraceIndex)
@@ -1132,20 +1134,57 @@ void SStartPageWindow::RefreshTraceSessionList()
 				continue;
 			}
 
-			const TSharedRef<FTraceSession> TraceSession = MakeShared<FTraceSession>(TraceInfo);
-			TraceSession->TraceIndex = TraceIndex;
-			TraceSession->Uri = FText::FromString(FInsightsManager::Get()->GetStoreDir() + TEXT("/") + TraceSession->Name.ToString() + TEXT(".utrace"));
-			UpdateMetadata(*TraceSession);
-			TraceSessions.Add(TraceSession);
-			TraceSessionsMap.Add(TraceSession->TraceId, TraceSession);
+			TSharedPtr<FTraceSession> TraceSessionPtr;
+
+			// Use previous TraceSessionsMap as a cache.
+			TSharedPtr<FTraceSession>* TraceSessionPtrPtr = TraceSessionsMap.Find(TraceInfo->GetId());
+			if (TraceSessionPtrPtr)
+			{
+				// Reuse an existing FTraceSession.
+				TraceSessionPtr = *TraceSessionPtrPtr;
+			}
+			else
+			{
+				// Make new FTraceSession.
+				TraceSessionPtr = MakeShared<FTraceSession>();
+
+				TraceSessionPtr->TraceId = TraceInfo->GetId();
+
+				const FAnsiStringView AnsiTraceName = TraceInfo->GetName();
+				const FString TraceName(AnsiTraceName.Len(), AnsiTraceName.GetData());
+				TraceSessionPtr->Name = FText::FromString(TraceName);
+				TraceSessionPtr->Uri = FText::FromString(FInsightsManager::Get()->GetStoreDir() + TEXT("/") + TraceName + TEXT(".utrace"));
+			}
+
+			FTraceSession& TraceSession = *TraceSessionPtr;
+
+			TraceSession.TraceIndex = TraceIndex;
+
+			TraceSession.Timestamp = FTraceSession::ConvertTimestamp(TraceInfo->GetTimestamp());
+			TraceSession.Size = TraceInfo->GetSize();
+
+			if (!TraceSession.bIsMetadataUpdated)
+			{
+				UpdateMetadata(TraceSession);
+				MetadataUpdateCount++;
+			}
+
+			TraceSessions.Add(TraceSessionPtr);
 
 			// Identify the previously selected trace session (if stil available) to ensure selection remains unchanged.
-			if (SelectedTraceSession && SelectedTraceSession->TraceId == TraceSession->TraceId)
+			if (SelectedTraceSession && SelectedTraceSession->TraceId == TraceSession.TraceId)
 			{
-				NewSelectedTraceSession = TraceSession;
+				NewSelectedTraceSession = TraceSessionPtr;
 			}
 		}
 
+		// Rebuild map.
+		TraceSessionsMap.Reset();
+		for (TSharedPtr<FTraceSession>& TraceSession : TraceSessions)
+		{
+			TraceSessionsMap.Add(TraceSession->TraceId, TraceSession);
+		}
+
 		Algo::SortBy(TraceSessions, &FTraceSession::Timestamp);
 
 		TraceSessionsListView->RebuildList();
@@ -1166,7 +1205,7 @@ void SStartPageWindow::RefreshTraceSessionList()
 		}
 
 		Stopwatch.Stop();
-		UE_LOG(TraceInsights, Log, TEXT("List of trace sessions updated in %llu ms."), Stopwatch.GetAccumulatedTimeMs());
+		UE_LOG(TraceInsights, Log, TEXT("List of trace sessions updated in %llu ms (metadata updated for %d trace%s)."), Stopwatch.GetAccumulatedTimeMs(), MetadataUpdateCount, MetadataUpdateCount == 1 ? TEXT("") : TEXT("s"));
 	}
 
 	// Process the connected recorder sessions.
diff --git a/Engine/Source/Developer/TraceInsights/Private/Insights/Widgets/SStartPageWindow.h b/Engine/Source/Developer/TraceInsights/Private/Insights/Widgets/SStartPageWindow.h
--- a/Engine/Source/Developer/TraceInsights/Private/Insights/Widgets/SStartPageWindow.h
+++ b/Engine/Source/Developer/TraceInsights/Private/Insights/Widgets/SStartPageWindow.h
@@ -61,31 +61,13 @@ struct FTraceSession
 	EBuildConfiguration ConfigurationType = EBuildConfiguration::Unknown;
 	EBuildTargetType TargetType = EBuildTargetType::Unknown;
 
-	//FTraceSession() = default;
-	FTraceSession(const Trace::FStoreClient::FTraceInfo* InTraceInfo)
-		: TraceId(InTraceInfo->GetId())
-		, TraceIndex(-1)
-		, Name(AnsiStringViewToText(InTraceInfo->GetName()))
-		, Uri(/*TODO: InTraceInfo->GetUri()*/)
-		, Timestamp(ConvertTimestamp(InTraceInfo->GetTimestamp()))
-		, Size(InTraceInfo->GetSize())
-		, bIsLive(false)
-		, IpAddress(0)
-		, bIsMetadataUpdated(false)
-		, Platform()
-		, AppName()
-		, CommandLine()
-		, ConfigurationType(EBuildConfiguration::Unknown)
-		, TargetType(EBuildTargetType::Unknown)
-	{
-	}
+	FTraceSession() = default;
 
 	static FDateTime ConvertTimestamp(uint64 InTimestamp)
 	{
 		return FDateTime(static_cast<int64>(InTimestamp));
 	}
 
-private:
 	static FText AnsiStringViewToText(const FAnsiStringView& AnsiStringView)
 	{
 		FString FatString(AnsiStringView.Len(), AnsiStringView.GetData());
