Added the ability for external systems to scroll the view to specified times

diff --git a/Engine/Source/Developer/TraceInsights/Private/Insights/Widgets/STimingView.cpp b/Engine/Source/Developer/TraceInsights/Private/Insights/Widgets/STimingView.cpp
--- a/Engine/Source/Developer/TraceInsights/Private/Insights/Widgets/STimingView.cpp
+++ b/Engine/Source/Developer/TraceInsights/Private/Insights/Widgets/STimingView.cpp
@@ -312,7 +312,7 @@ void STimingView::Reset()
 
 	LastSelectionType = ESelectionType::None;
 
-	TimeMarker = std::numeric_limits<double>::infinity();
+	SetTimeMarker(std::numeric_limits<double>::infinity());
 	bIsScrubbing = false;
 
 	//ThisGeometry
@@ -2648,7 +2648,7 @@ void STimingView::UpdateAggregatedStats()
 
 void STimingView::SetAndCenterOnTimeMarker(double Time)
 {
-	TimeMarker = Time;
+	SetTimeMarker(Time);
 
 	double MinT, MaxT;
 	Viewport.GetHorizontalScrollLimits(MinT, MaxT);
@@ -2674,7 +2674,7 @@ void STimingView::SelectToTimeMarker(double Time)
 		SelectTimeInterval(Time, TimeMarker - Time);
 	}
 
-	TimeMarker = Time;
+	SetTimeMarker(Time);
 }
 
 ////////////////////////////////////////////////////////////////////////////////////////////////////
diff --git a/Engine/Source/Developer/TraceInsights/Private/Insights/Widgets/STimingView.h b/Engine/Source/Developer/TraceInsights/Private/Insights/Widgets/STimingView.h
--- a/Engine/Source/Developer/TraceInsights/Private/Insights/Widgets/STimingView.h
+++ b/Engine/Source/Developer/TraceInsights/Private/Insights/Widgets/STimingView.h
@@ -224,6 +224,7 @@ public:
 
 	virtual double GetTimeMarker() const override { return TimeMarker; }
 	virtual void SetTimeMarker(double InMarkerTime) override;
+	virtual void SetAndCenterOnTimeMarker(double InTime) override;
 
 	virtual Insights::FSelectionChangedDelegate& OnSelectionChanged() override { return OnSelectionChangedDelegate; }
 	virtual Insights::FTimeMarkerChangedDelegate& OnTimeMarkerChanged() override { return OnTimeMarkerChangedDelegate; }
@@ -263,7 +264,6 @@ public:
 	void CenterOnTimeInterval(double IntervalStartTime, double IntervalDuration);
 	void BringIntoView(double StartTime, double EndTime);
 	void SelectTimeInterval(double IntervalStartTime, double IntervalDuration);
-	void SetAndCenterOnTimeMarker(double Time);
 	void SelectToTimeMarker(double Time);
 
 	//bool AreTimeMarkersVisible() { return MarkersTrack->IsVisible(); }
diff --git a/Engine/Source/Developer/TraceInsights/Public/Insights/ITimingViewSession.h b/Engine/Source/Developer/TraceInsights/Public/Insights/ITimingViewSession.h
--- a/Engine/Source/Developer/TraceInsights/Public/Insights/ITimingViewSession.h
+++ b/Engine/Source/Developer/TraceInsights/Public/Insights/ITimingViewSession.h
@@ -69,6 +69,9 @@ public:
 	/** Set the current marker time */
 	virtual void SetTimeMarker(double InMarkerTime) = 0;
 
+	/** Set the current marker time and center the view on it */
+	virtual void SetAndCenterOnTimeMarker(double InTime) = 0;
+
 	/** Gets the delegate to be invoked when the selection have been changed. */
 	virtual FSelectionChangedDelegate& OnSelectionChanged() = 0;
 
