Restarting the UE4Editor with a Gameplay Insights enabled project always starts with the Gameplay Insights window undocked
#fix Ensure that tabs are not always activated and spawned whenever a Trace is opened/loaded

diff --git a/Engine/Source/Developer/TraceInsights/Private/Insights/InsightsManager.cpp b/Engine/Source/Developer/TraceInsights/Private/Insights/InsightsManager.cpp
--- a/Engine/Source/Developer/TraceInsights/Private/Insights/InsightsManager.cpp
+++ b/Engine/Source/Developer/TraceInsights/Private/Insights/InsightsManager.cpp
@@ -329,7 +329,6 @@ void FInsightsManager::LoadTrace(uint32 InTraceId)
 		CurrentTraceId = InTraceId;
 		CurrentTraceFilename = TraceName;
 		bIsNetworkingProfilerAvailable = false;
-		SpawnAndActivateTabs();
 		OnSessionChanged();
 	}
 }
@@ -354,7 +353,6 @@ void FInsightsManager::LoadTraceFile(const FString& InTraceFilename)
 		CurrentTraceId = 0;
 		CurrentTraceFilename = InTraceFilename;
 		bIsNetworkingProfilerAvailable = false;
-		SpawnAndActivateTabs();
 		OnSessionChanged();
 	}
 }
