[Insights]
  - Bug fix: Browser: Renamed "..." button to "Explore" (Explore the Trace Store Directory). Also aligned the button to the right. (UE-89187)
  - Bug fix: Renamed command line params to analyze a session from "-TraceFile" / "-TraceId" to "-OpenTraceFile" / "-OpenTraceId" (in order to avoid conflict with "-TraceFile" cmd line param already in use to direct tracing to a file).

diff --git a/Engine/Source/Developer/TraceInsights/Private/Insights/Widgets/SStartPageWindow.cpp b/Engine/Source/Developer/TraceInsights/Private/Insights/Widgets/SStartPageWindow.cpp
--- a/Engine/Source/Developer/TraceInsights/Private/Insights/Widgets/SStartPageWindow.cpp
+++ b/Engine/Source/Developer/TraceInsights/Private/Insights/Widgets/SStartPageWindow.cpp
@@ -851,27 +851,28 @@ TSharedRef<SWidget> SStartPageWindow::ConstructTraceStoreDirectoryPanel()
 
 	+ SVerticalBox::Slot()
 	.AutoHeight()
-	.HAlign(HAlign_Left)
+	.HAlign(HAlign_Fill)
 	.Padding(0.0f, 0.0f)
 	[
 		SNew(SHorizontalBox)
 
 		+ SHorizontalBox::Slot()
 		.Padding(0.0f, 0.0f)
+		.HAlign(HAlign_Fill)
 		.VAlign(VAlign_Center)
 		[
 			SNew(STextBlock)
 			.Text(this, &SStartPageWindow::GetTraceStoreDirectory)
-			.Justification(ETextJustify::Right)
 		]
 
 		+ SHorizontalBox::Slot()
 		.AutoWidth()
 		.Padding(4.0f, 0.0f, 0.0f, 0.0f)
+		.HAlign(HAlign_Right)
 		.VAlign(VAlign_Center)
 		[
 			SNew(SButton)
-			.Text(LOCTEXT("ExploreTraceStoreDirButton", "..."))
+			.Text(LOCTEXT("ExploreTraceStoreDirButton", "Explore"))
 			.ToolTipText(LOCTEXT("ExploreTraceStoreDirButtonToolTip", "Explore the Trace Store Directory"))
 			.OnClicked(this, &SStartPageWindow::ExploreTraceStoreDirectory_OnClicked)
 		]
@@ -1608,7 +1609,7 @@ void SStartPageWindow::LoadTraceFile(const FString& InTraceFile)
 
 		const TCHAR* ExecutablePath = FPlatformProcess::ExecutablePath();
 
-		FString CmdLine = TEXT("-TraceFile=\"") + InTraceFile + TEXT("\"");
+		FString CmdLine = TEXT("-OpenTraceFile=\"") + InTraceFile + TEXT("\"");
 
 		constexpr bool bLaunchDetached = true;
 		constexpr bool bLaunchHidden = false;
@@ -1648,7 +1649,7 @@ void SStartPageWindow::LoadTrace(uint32 InTraceId)
 		const TCHAR* ExecutablePath = FPlatformProcess::ExecutablePath();
 
 		const uint32 StorePort = FInsightsManager::Get()->GetStoreClient()->GetStorePort();
-		FString CmdLine = FString::Printf(TEXT("-TraceId=%d -StorePort=%d"), InTraceId, StorePort);
+		FString CmdLine = FString::Printf(TEXT("-OpenTraceId=%d -StorePort=%d"), InTraceId, StorePort);
 
 		constexpr bool bLaunchDetached = true;
 		constexpr bool bLaunchHidden = false;
diff --git a/Engine/Source/Programs/UnrealInsights/Private/UserInterfaceCommand.cpp b/Engine/Source/Programs/UnrealInsights/Private/UserInterfaceCommand.cpp
--- a/Engine/Source/Programs/UnrealInsights/Private/UserInterfaceCommand.cpp
+++ b/Engine/Source/Programs/UnrealInsights/Private/UserInterfaceCommand.cpp
@@ -110,11 +110,11 @@ void FUserInterfaceCommand::Run()
 
 		if (bBrowserMode)
 		{
-			bBrowserMode = FCString::Strifind(FCommandLine::Get(), TEXT("-TraceId=")) == nullptr;
+			bBrowserMode = FCString::Strifind(FCommandLine::Get(), TEXT("-OpenTraceId=")) == nullptr;
 		}
 		if (bBrowserMode)
 		{
-			bBrowserMode = FCString::Strifind(FCommandLine::Get(), TEXT("-TraceFile=")) == nullptr;
+			bBrowserMode = FCString::Strifind(FCommandLine::Get(), TEXT("-OpenTraceFile=")) == nullptr;
 		}
 
 		if (bBrowserMode && !CheckSessionBrowserSingleInstance())
@@ -222,7 +222,7 @@ void FUserInterfaceCommand::InitializeSlateApplication()
 	const uint32 MaxPath = FPlatformMisc::GetMaxPathLength();
 
 	uint32 TraceId = 0;
-	bool bUseTraceId = FParse::Value(FCommandLine::Get(), TEXT("-TraceId="), TraceId);
+	bool bUseTraceId = FParse::Value(FCommandLine::Get(), TEXT("-OpenTraceId="), TraceId);
 
 	TCHAR* StoreHost = new TCHAR[MaxPath + 1];
 	FCString::Strcpy(StoreHost, MaxPath, TEXT("127.0.0.1"));
@@ -259,7 +259,7 @@ void FUserInterfaceCommand::InitializeSlateApplication()
 	{
 		TCHAR* TraceFile = new TCHAR[MaxPath + 1];
 		TraceFile[0] = 0;
-		bool bUseTraceFile = FParse::Value(FCommandLine::Get(), TEXT("-TraceFile="), TraceFile, MaxPath, true);
+		bool bUseTraceFile = FParse::Value(FCommandLine::Get(), TEXT("-OpenTraceFile="), TraceFile, MaxPath, true);
 
 		if (bUseTraceFile)
 		{
