[Insights] Fixed Asset Loading Profiler tab not being spawned when Unreal Insights runs in standalone mode.

diff --git a/Engine/Source/Developer/TraceInsights/Private/Insights/TraceInsightsModule.cpp b/Engine/Source/Developer/TraceInsights/Private/Insights/TraceInsightsModule.cpp
--- a/Engine/Source/Developer/TraceInsights/Private/Insights/TraceInsightsModule.cpp
+++ b/Engine/Source/Developer/TraceInsights/Private/Insights/TraceInsightsModule.cpp
@@ -317,11 +317,12 @@ void FTraceInsightsModule::AddAreaForSessionViewer(TSharedRef<FTabManager::FLayo
 			// In editor, we default to all tabs closed.
 			->AddTab(FInsightsManagerTabs::SessionInfoTabId, ETabState::ClosedTab)
 			->AddTab(FInsightsManagerTabs::TimingProfilerTabId, ETabState::ClosedTab)
+			->AddTab(FInsightsManagerTabs::LoadingProfilerTabId, ETabState::ClosedTab)
 #else
 			->AddTab(FInsightsManagerTabs::SessionInfoTabId, ETabState::OpenedTab)
 			->AddTab(FInsightsManagerTabs::TimingProfilerTabId, ETabState::OpenedTab)
+			->AddTab(FInsightsManagerTabs::LoadingProfilerTabId, ETabState::OpenedTab)
 #endif
-			->AddTab(FInsightsManagerTabs::LoadingProfilerTabId, ETabState::ClosedTab)
 			->AddTab(FInsightsManagerTabs::NetworkingProfilerTabId, ETabState::ClosedTab)
 			->SetForegroundTab(FTabId(FInsightsManagerTabs::TimingProfilerTabId))
 		)
