Remove unused private fields - Insights

diff --git a/Engine/Source/Developer/TraceInsights/Private/Insights/NetworkingProfiler/Widgets/SPacketContentView.h b/Engine/Source/Developer/TraceInsights/Private/Insights/NetworkingProfiler/Widgets/SPacketContentView.h
--- a/Engine/Source/Developer/TraceInsights/Private/Insights/NetworkingProfiler/Widgets/SPacketContentView.h
+++ b/Engine/Source/Developer/TraceInsights/Private/Insights/NetworkingProfiler/Widgets/SPacketContentView.h
@@ -211,8 +211,6 @@ private:
 	FPacketContentViewport Viewport;
 	bool bIsViewportDirty;
 
-	float NetEventsOffsetY;
-
 	uint32 GameInstanceIndex;
 	uint32 ConnectionIndex;
 	Trace::ENetProfilerConnectionMode ConnectionMode;
diff --git a/Engine/Source/Developer/TraceInsights/Private/Insights/Table/ViewModels/TableColumn.h b/Engine/Source/Developer/TraceInsights/Private/Insights/Table/ViewModels/TableColumn.h
--- a/Engine/Source/Developer/TraceInsights/Private/Insights/Table/ViewModels/TableColumn.h
+++ b/Engine/Source/Developer/TraceInsights/Private/Insights/Table/ViewModels/TableColumn.h
@@ -56,7 +56,6 @@ public:
 	/** Initialization constructor. */
 	FTableColumn(const FName InId)
 		: Id(InId)
-		, Order(0)
 		, Index(-1)
 		, ShortName()
 		, TitleName()
@@ -191,9 +190,6 @@ private:
 	/** Id of the column. */
 	FName Id;
 
-	/** Order value, to sort columns in the list/tree view. */
-	int32 Order;
-
 	/** Column index in source table. */
 	int32 Index;
 
diff --git a/Engine/Source/Developer/TraceInsights/Private/Insights/ViewModels/FrameTimingTrack.h b/Engine/Source/Developer/TraceInsights/Private/Insights/ViewModels/FrameTimingTrack.h
--- a/Engine/Source/Developer/TraceInsights/Private/Insights/ViewModels/FrameTimingTrack.h
+++ b/Engine/Source/Developer/TraceInsights/Private/Insights/ViewModels/FrameTimingTrack.h
@@ -64,7 +64,6 @@ class FFrameTimingTrack : public FTimingEventsTrack
 public:
 	explicit FFrameTimingTrack(FFrameSharedState& InSharedState, const FString& InName, uint32 InFrameType)
 		: FTimingEventsTrack(InName)
-		, SharedState(InSharedState)
 		, FrameType(InFrameType)
 		, Header(*this)
 	{
@@ -121,8 +120,6 @@ private:
 	bool FindFrame(const FTimingEventSearchParameters& InParameters, TFunctionRef<void(double, double, uint32, const Trace::FFrame&)> InFoundPredicate) const;
 
 private:
-	FFrameSharedState& SharedState;
-
 	uint32 FrameType; // ETraceFrameType
 
 	FTrackHeader Header;
diff --git a/Engine/Source/Developer/TraceInsights/Private/Insights/ViewModels/ThreadTimingTrack.h b/Engine/Source/Developer/TraceInsights/Private/Insights/ViewModels/ThreadTimingTrack.h
--- a/Engine/Source/Developer/TraceInsights/Private/Insights/ViewModels/ThreadTimingTrack.h
+++ b/Engine/Source/Developer/TraceInsights/Private/Insights/ViewModels/ThreadTimingTrack.h
@@ -95,7 +95,6 @@ class FThreadTimingTrack : public FTimingEventsTrack
 public:
 	explicit FThreadTimingTrack(FThreadTimingSharedState& InSharedState, const FString& InName, const TCHAR* InGroupName, uint32 InTimelineIndex, uint32 InThreadId)
 		: FTimingEventsTrack(InName)
-		, SharedState(InSharedState)
 		, GroupName(InGroupName)
 		, TimelineIndex(InTimelineIndex)
 		, ThreadId(InThreadId)
@@ -139,8 +138,6 @@ private:
 						  Trace::FTimingProfilerEvent& OutRootEvent) const;
 
 private:
-	FThreadTimingSharedState& SharedState;
-
 	const TCHAR* GroupName;
 	uint32 TimelineIndex;
 	uint32 ThreadId;
diff --git a/Engine/Source/Developer/TraceInsights/Private/Insights/Widgets/STimerTreeView.cpp b/Engine/Source/Developer/TraceInsights/Private/Insights/Widgets/STimerTreeView.cpp
--- a/Engine/Source/Developer/TraceInsights/Private/Insights/Widgets/STimerTreeView.cpp
+++ b/Engine/Source/Developer/TraceInsights/Private/Insights/Widgets/STimerTreeView.cpp
@@ -35,9 +35,6 @@ STimerTreeView::STimerTreeView()
 	, CurrentSorter(nullptr)
 	, ColumnBeingSorted(GetDefaultColumnBeingSorted())
 	, ColumnSortMode(GetDefaultColumnSortMode())
-	, StatsStartTime(0.0)
-	, StatsEndTime(0.0)
-	, StatsTimerId(0)
 {
 }
 
diff --git a/Engine/Source/Developer/TraceInsights/Private/Insights/Widgets/STimerTreeView.h b/Engine/Source/Developer/TraceInsights/Private/Insights/Widgets/STimerTreeView.h
--- a/Engine/Source/Developer/TraceInsights/Private/Insights/Widgets/STimerTreeView.h
+++ b/Engine/Source/Developer/TraceInsights/Private/Insights/Widgets/STimerTreeView.h
@@ -217,12 +217,6 @@ private:
 
 	/** How we sort the nodes? Ascending or Descending. */
 	EColumnSortMode::Type ColumnSortMode;
-
-	//////////////////////////////////////////////////
-
-	double StatsStartTime;
-	double StatsEndTime;
-	uint32 StatsTimerId;
 };
 
 ////////////////////////////////////////////////////////////////////////////////////////////////////
