[Insights] Forward -InsightsTest and -DebugTools arguments from the Session Browser to the Session Viewer. Add options to enable Debug Tools and InsightsTest from the Session Browser UI.

diff --git a/Engine/Source/Developer/TraceInsights/Private/Insights/TraceInsightsModule.cpp b/Engine/Source/Developer/TraceInsights/Private/Insights/TraceInsightsModule.cpp
--- a/Engine/Source/Developer/TraceInsights/Private/Insights/TraceInsightsModule.cpp
+++ b/Engine/Source/Developer/TraceInsights/Private/Insights/TraceInsightsModule.cpp
@@ -20,6 +20,7 @@
 #include "Insights/NetworkingProfiler/NetworkingProfilerManager.h"
 #include "Insights/Tests/InsightsTestRunner.h"
 #include "Insights/TimingProfilerManager.h"
+#include "Insights/Widgets/SStartPageWindow.h"
 
 ////////////////////////////////////////////////////////////////////////////////////////////////////
 
@@ -155,7 +156,7 @@ void FTraceInsightsModule::UnregisterTabSpawners()
 
 ////////////////////////////////////////////////////////////////////////////////////////////////////
 
-void FTraceInsightsModule::CreateSessionBrowser(bool bAllowDebugTools, bool bSingleProcess)
+void FTraceInsightsModule::CreateSessionBrowser(bool bAllowDebugTools, bool bSingleProcess, bool bInitializeTesting)
 {
 	FInsightsManager::Get()->SetOpenAnalysisInSeparateProcess(!bSingleProcess);
 
@@ -227,6 +228,14 @@ void FTraceInsightsModule::CreateSessionBrowser(bool bAllowDebugTools, bool bSin
 	const EOutputCanBeNullptr OutputCanBeNullptr = EOutputCanBeNullptr::Never;
 	TSharedPtr<SWidget> Content = FGlobalTabmanager::Get()->RestoreFrom(PersistentLayout.ToSharedRef(), RootWindow, bEmbedTitleAreaContent, OutputCanBeNullptr);
 	RootWindow->SetContent(Content.ToSharedRef());
+
+	// Set up command line parameter forwarding.
+	TSharedPtr<class SStartPageWindow> StartPageWnd = FInsightsManager::Get()->GetStartPageWindow();
+	if (StartPageWnd.IsValid())
+	{
+		StartPageWnd->SetEnableAutomaticTesting(bInitializeTesting);
+		StartPageWnd->SetEnableDebugTools(bAllowDebugTools);
+	}
 }
 
 ////////////////////////////////////////////////////////////////////////////////////////////////////
diff --git a/Engine/Source/Developer/TraceInsights/Private/Insights/TraceInsightsModule.h b/Engine/Source/Developer/TraceInsights/Private/Insights/TraceInsightsModule.h
--- a/Engine/Source/Developer/TraceInsights/Private/Insights/TraceInsightsModule.h
+++ b/Engine/Source/Developer/TraceInsights/Private/Insights/TraceInsightsModule.h
@@ -36,8 +36,8 @@ public:
 	virtual Trace::FStoreClient* GetStoreClient() override;
 	virtual bool ConnectToStore(const TCHAR* InStoreHost, uint32 InStorePort) override;
 
-	virtual void CreateSessionBrowser(bool bAllowDebugTools, bool bSingleProcess) override;
-	virtual void CreateSessionViewer(bool bAllowDebugTools) override;
+	virtual void CreateSessionBrowser(bool bAllowDebugTools = false, bool bSingleProcess = false, bool bInitializeTesting = false) override;
+	virtual void CreateSessionViewer(bool bAllowDebugTools = false) override;
 
 	virtual TSharedPtr<const Trace::IAnalysisSession> GetAnalysisSession() const override;
 	virtual void StartAnalysisForTrace(uint32 InTraceId, bool InAutoQuit = false) override;
diff --git a/Engine/Source/Developer/TraceInsights/Private/Insights/Widgets/SStartPageWindow.cpp b/Engine/Source/Developer/TraceInsights/Private/Insights/Widgets/SStartPageWindow.cpp
--- a/Engine/Source/Developer/TraceInsights/Private/Insights/Widgets/SStartPageWindow.cpp
+++ b/Engine/Source/Developer/TraceInsights/Private/Insights/Widgets/SStartPageWindow.cpp
@@ -1612,6 +1612,10 @@ void SStartPageWindow::LoadTraceFile(const FString& InTraceFile)
 
 		FString CmdLine = TEXT("-OpenTraceFile=\"") + InTraceFile + TEXT("\"");
 
+		FString ExtraCmdParams;
+		GetExtraCommandLineParams(ExtraCmdParams);
+		CmdLine += ExtraCmdParams;
+
 		constexpr bool bLaunchDetached = true;
 		constexpr bool bLaunchHidden = false;
 		constexpr bool bLaunchReallyHidden = false;
@@ -1651,6 +1655,10 @@ void SStartPageWindow::LoadTrace(uint32 InTraceId)
 
 		const uint32 StorePort = FInsightsManager::Get()->GetStoreClient()->GetStorePort();
 		FString CmdLine = FString::Printf(TEXT("-OpenTraceId=%d -StorePort=%d"), InTraceId, StorePort);
+		
+		FString ExtraCmdParams;
+		GetExtraCommandLineParams(ExtraCmdParams);
+		CmdLine += ExtraCmdParams;
 
 		constexpr bool bLaunchDetached = true;
 		constexpr bool bLaunchHidden = false;
@@ -1740,6 +1748,58 @@ TSharedRef<SWidget> SStartPageWindow::MakeTraceListMenu()
 	}
 	MenuBuilder.EndSection();
 
+	MenuBuilder.BeginSection("DebugOptions", LOCTEXT("OptionsHeading", "Debug Options"));
+
+	auto CanExecute = []() {return true; };
+
+	// Enable Automation Tests Option.
+	{
+		FUIAction ToogleAutomationTestsAction;
+		ToogleAutomationTestsAction.CanExecuteAction = FCanExecuteAction::CreateLambda(CanExecute);
+		ToogleAutomationTestsAction.ExecuteAction = FExecuteAction::CreateLambda([this]()
+			{
+				this->SetEnableAutomaticTesting(!this->GetEnableAutomaticTesting());
+			});
+		ToogleAutomationTestsAction.GetActionCheckState = FGetActionCheckState::CreateLambda([this]()
+			{
+				return this->GetEnableAutomaticTesting() ? ECheckBoxState::Checked : ECheckBoxState::Unchecked;
+			});
+
+		MenuBuilder.AddMenuEntry(
+			LOCTEXT("EnableAutomatedTesting", "Enable Automation Testing"),
+			LOCTEXT("EnableAutomatedTestingDesc", "Activates the automatic test system for sessions opened from this window."),
+			FSlateIcon(),
+			ToogleAutomationTestsAction,
+			NAME_None,
+			EUserInterfaceActionType::ToggleButton
+		);
+	}
+
+	// Enable Debug Tools Option.
+	{
+		FUIAction ToogleDebugToolsAction;
+		ToogleDebugToolsAction.CanExecuteAction = FCanExecuteAction::CreateLambda(CanExecute);
+		ToogleDebugToolsAction.ExecuteAction = FExecuteAction::CreateLambda([this]()
+			{
+				this->SetEnableDebugTools(!this->GetEnableDebugTools());
+			});
+		ToogleDebugToolsAction.GetActionCheckState = FGetActionCheckState::CreateLambda([this]()
+			{
+				return this->GetEnableDebugTools() ? ECheckBoxState::Checked : ECheckBoxState::Unchecked;
+			});
+
+		MenuBuilder.AddMenuEntry(
+			LOCTEXT("EnableDebugTools", "Enable Debug Tools"),
+			LOCTEXT("EnableDebugToolsDesc", "Enables debug tools for sessions opened from this window."),
+			FSlateIcon(),
+			ToogleDebugToolsAction,
+			NAME_None,
+			EUserInterfaceActionType::ToggleButton
+		);
+	}
+
+	MenuBuilder.EndSection();
+
 	return MenuBuilder.MakeWidget();
 }
 
@@ -1846,4 +1906,18 @@ void SStartPageWindow::CloseSettings()
 
 ////////////////////////////////////////////////////////////////////////////////////////////////////
 
+void SStartPageWindow::GetExtraCommandLineParams(FString& OutParams) const
+{
+	if (bEnableAutomaticTesting)
+	{
+		OutParams.Append(TEXT(" -InsightsTest"));
+	}
+	if (bEnableDebugTools)
+	{
+		OutParams.Append(TEXT(" -DebugTools"));
+	}
+}
+
+////////////////////////////////////////////////////////////////////////////////////////////////////
+
 #undef LOCTEXT_NAMESPACE
diff --git a/Engine/Source/Developer/TraceInsights/Private/Insights/Widgets/SStartPageWindow.h b/Engine/Source/Developer/TraceInsights/Private/Insights/Widgets/SStartPageWindow.h
--- a/Engine/Source/Developer/TraceInsights/Private/Insights/Widgets/SStartPageWindow.h
+++ b/Engine/Source/Developer/TraceInsights/Private/Insights/Widgets/SStartPageWindow.h
@@ -102,6 +102,14 @@ public:
 	void OpenSettings();
 	void CloseSettings();
 
+	void GetExtraCommandLineParams(FString& OutParams) const;
+
+	void SetEnableAutomaticTesting(bool InValue) { bEnableAutomaticTesting = InValue; };
+	bool GetEnableAutomaticTesting() const { return bEnableAutomaticTesting; };
+
+	void SetEnableDebugTools(bool InValue) { bEnableDebugTools = InValue; };
+	bool GetEnableDebugTools() const { return bEnableDebugTools; };
+
 private:
 	TSharedRef<SWidget> ConstructSessionsPanel();
 	TSharedRef<SWidget> ConstructLoadPanel();
@@ -278,4 +286,7 @@ private:
 
 	FString SplashScreenOverlayTraceFile;
 	float SplashScreenOverlayFadeTime;
+
+	bool bEnableAutomaticTesting = false;
+	bool bEnableDebugTools = false;
 };
diff --git a/Engine/Source/Developer/TraceInsights/Public/Insights/IUnrealInsightsModule.h b/Engine/Source/Developer/TraceInsights/Public/Insights/IUnrealInsightsModule.h
--- a/Engine/Source/Developer/TraceInsights/Public/Insights/IUnrealInsightsModule.h
+++ b/Engine/Source/Developer/TraceInsights/Public/Insights/IUnrealInsightsModule.h
@@ -247,12 +247,12 @@ public:
 	/**
 	 * Called when the application starts in "Browser" mode.
 	 */
-	virtual void CreateSessionBrowser(bool bAllowDebugTools, bool bSingleProcess) = 0;
+	virtual void CreateSessionBrowser(bool bAllowDebugTools = false, bool bSingleProcess = false, bool bInitializeTesting = false) = 0;
 
 	/**
 	 * Called when the application starts in "Viewer" mode.
 	 */
-	virtual void CreateSessionViewer(bool bAllowDebugTools) = 0;
+	virtual void CreateSessionViewer(bool bAllowDebugTools = false) = 0;
 
 	/**
 	 * Called when the application shutsdown.
diff --git a/Engine/Source/Programs/UnrealInsights/Private/UserInterfaceCommand.cpp b/Engine/Source/Programs/UnrealInsights/Private/UserInterfaceCommand.cpp
--- a/Engine/Source/Programs/UnrealInsights/Private/UserInterfaceCommand.cpp
+++ b/Engine/Source/Programs/UnrealInsights/Private/UserInterfaceCommand.cpp
@@ -322,7 +322,7 @@ void FUserInterfaceCommand::InitializeSlateApplication(bool bOpenTraceFile, cons
 				TraceInsightsModule.ConnectToStore(StoreHost, StorePort);
 			}
 			const bool bSingleProcess = FParse::Param(FCommandLine::Get(), TEXT("SingleProcess"));
-			TraceInsightsModule.CreateSessionBrowser(bAllowDebugTools, bSingleProcess);
+			TraceInsightsModule.CreateSessionBrowser(bAllowDebugTools, bSingleProcess, bInitializeTesting);
 		}
 	}
 
