[Insights] Register the Message Log Spawner from Insights stand alone. Cleanup InsightsMessageLog code.

diff --git a/Engine/Source/Developer/TraceInsights/Private/Insights/InsightsManager.cpp b/Engine/Source/Developer/TraceInsights/Private/Insights/InsightsManager.cpp
--- a/Engine/Source/Developer/TraceInsights/Private/Insights/InsightsManager.cpp
+++ b/Engine/Source/Developer/TraceInsights/Private/Insights/InsightsManager.cpp
@@ -3,6 +3,7 @@
 #include "InsightsManager.h"
 
 #include "Framework/Application/SlateApplication.h"
+#include "MessageLog/Public/MessageLogModule.h"
 #include "Misc/CString.h"
 #include "Modules/ModuleManager.h"
 #include "Templates/UniquePtr.h"
@@ -19,7 +20,6 @@
 #include "Insights/NetworkingProfiler/NetworkingProfilerManager.h"
 #include "Insights/Tests/InsightsTestRunner.h"
 #include "Insights/TimingProfilerManager.h"
-#include "Insights/ViewModels/InsightsMessageLogViewModel.h"
 #include "Insights/Widgets/SStartPageWindow.h"
 #include "Insights/Widgets/SSessionInfoWindow.h"
 #include "Insights/Widgets/STimingProfilerWindow.h"
@@ -34,7 +34,6 @@ const FName FInsightsManagerTabs::TimingProfilerTabId(TEXT("TimingProfiler"));
 const FName FInsightsManagerTabs::LoadingProfilerTabId(TEXT("LoadingProfiler"));
 const FName FInsightsManagerTabs::NetworkingProfilerTabId(TEXT("NetworkingProfiler"));
 const FName FInsightsManagerTabs::MemoryProfilerTabId(TEXT("MemoryProfiler"));
-const FName FInsightsManagerTabs::InsightsMessageLogTabId(TEXT("MessageLog"));
 const FName FInsightsManagerTabs::AutomationWindowTabId(TEXT("AutomationWindow"));
 
 ////////////////////////////////////////////////////////////////////////////////////////////////////
@@ -145,8 +144,6 @@ void FInsightsManager::Initialize(IUnrealInsightsModule& InsightsModule)
 	}
 	bIsInitialized = true;
 
-	InsightsMessageLogViewModel = MakeShared<FInsightsMessageLogViewModel>("InsightsLog", InsightsMessageLog);
-
 	InsightsMenuBuilder = MakeShared<FInsightsMenuBuilder>();
 
 	// Register tick functions.
@@ -226,18 +223,8 @@ void FInsightsManager::RegisterMajorTabs(IUnrealInsightsModule& InsightsModule)
 	}
 
 #if !WITH_EDITOR
-	const FInsightsMajorTabConfig& MessageLogConfig = InsightsModule.FindMajorTabConfig(FInsightsManagerTabs::InsightsMessageLogTabId);
-	if (MessageLogConfig.bIsAvailable)
-	{
-		FTabSpawnerEntry& TabSpawnerEntry = FGlobalTabmanager::Get()->RegisterNomadTabSpawner(FInsightsManagerTabs::InsightsMessageLogTabId,
-			FOnSpawnTab::CreateRaw(this, &FInsightsManager::SpawnMessageLogTab))
-			.SetDisplayName(MessageLogConfig.TabLabel.IsSet() ? MessageLogConfig.TabLabel.GetValue() : LOCTEXT("InsightsMessageLogTabTitle", "Message Log"))
-			.SetTooltipText(MessageLogConfig.TabTooltip.IsSet() ? MessageLogConfig.TabTooltip.GetValue() : LOCTEXT("InsightsMessageLogTooltipText", "Open the Message Log tab."))
-			.SetIcon(MessageLogConfig.TabIcon.IsSet() ? MessageLogConfig.TabIcon.GetValue() : FSlateIcon(FInsightsStyle::GetStyleSetName(), "LogView.Icon.Small"));
-
-		TSharedRef<FWorkspaceItem> Group = MessageLogConfig.WorkspaceGroup.IsValid() ? MessageLogConfig.WorkspaceGroup.ToSharedRef() : GetInsightsMenuBuilder()->GetWindowsGroup();
-		TabSpawnerEntry.SetGroup(Group);
-	}
+	FMessageLogModule& MessageLogModule = FModuleManager::LoadModuleChecked<FMessageLogModule>("MessageLog");
+	MessageLogModule.RegisterMessageLogSpawner(GetInsightsMenuBuilder()->GetWindowsGroup());
 #endif
 }
 
@@ -247,9 +234,6 @@ void FInsightsManager::UnregisterMajorTabs()
 {
 	FGlobalTabmanager::Get()->UnregisterNomadTabSpawner(FInsightsManagerTabs::SessionInfoTabId);
 	FGlobalTabmanager::Get()->UnregisterNomadTabSpawner(FInsightsManagerTabs::StartPageTabId);
-#if !WITH_EDITOR
-	FGlobalTabmanager::Get()->UnregisterNomadTabSpawner(FInsightsManagerTabs::InsightsMessageLogTabId);
-#endif
 }
 
 ////////////////////////////////////////////////////////////////////////////////////////////////////
@@ -324,23 +308,6 @@ void FInsightsManager::OnSessionInfoTabClosed(TSharedRef<SDockTab> TabBeingClose
 
 ////////////////////////////////////////////////////////////////////////////////////////////////////
 
-TSharedRef<SDockTab> FInsightsManager::SpawnMessageLogTab(const FSpawnTabArgs& Args)
-{
-	TSharedRef<SDockTab> SpawnedTab = SNew(SDockTab)
-		.Label(LOCTEXT("InsightsMessageLogTitle", "Message Log"))
-		[
-			SNew(SBox)
-			.AddMetaData<FTagMetaData>(FTagMetaData(TEXT("InsightsLog")))
-			[
-				InsightsMessageLog.ToSharedRef()
-			]
-		];
-
-	return SpawnedTab;
-}
-
-////////////////////////////////////////////////////////////////////////////////////////////////////
-
 bool FInsightsManager::ConnectToStore(const TCHAR* Host, uint32 Port)
 {
 	using namespace Trace;
diff --git a/Engine/Source/Developer/TraceInsights/Private/Insights/InsightsManager.h b/Engine/Source/Developer/TraceInsights/Private/Insights/InsightsManager.h
--- a/Engine/Source/Developer/TraceInsights/Private/Insights/InsightsManager.h
+++ b/Engine/Source/Developer/TraceInsights/Private/Insights/InsightsManager.h
@@ -25,7 +25,6 @@ namespace Trace
 
 class SStartPageWindow;
 class SSessionInfoWindow;
-class FInsightsMessageLogViewModel;
 class FInsightsTestRunner;
 class FInsightsMenuBuilder;
 
@@ -201,8 +200,6 @@ public:
 	double GetAnalysisDuration() const { return AnalysisDuration; }
 	double GetAnalysisSpeedFactor() const { return AnalysisSpeedFactor; }
 
-	TSharedPtr<FInsightsMessageLogViewModel> GetMessageLog() { return InsightsMessageLogViewModel; }
-
 	TSharedPtr<FInsightsMenuBuilder> GetInsightsMenuBuilder() { return InsightsMenuBuilder; }
 
 	////////////////////////////////////////////////////////////////////////////////////////////////////
@@ -245,9 +242,6 @@ private:
 	/** Callback called when the Session Info major tab is closed. */
 	void OnSessionInfoTabClosed(TSharedRef<SDockTab> TabBeingClosed);
 
-	/** Called to spawn the Message Log major tab. */
-	TSharedRef<SDockTab> SpawnMessageLogTab(const FSpawnTabArgs& Args);
-
 	/** Updates this manager, done through FCoreTicker. */
 	bool Tick(float DeltaTime);
 
@@ -302,9 +296,6 @@ private:
 	/** A weak pointer to the Session Info window. */
 	TWeakPtr<class SSessionInfoWindow> SessionInfoWindow;
 
-	TSharedPtr<class SWidget> InsightsMessageLog;
-	TSharedPtr<FInsightsMessageLogViewModel> InsightsMessageLogViewModel;
-
 	/** If enabled, UI can display additional info for debugging purposes. */
 	bool bIsDebugInfoEnabled;
 
diff --git a/Engine/Source/Developer/TraceInsights/Private/Insights/MemoryProfiler/MemoryProfilerManager.cpp b/Engine/Source/Developer/TraceInsights/Private/Insights/MemoryProfiler/MemoryProfilerManager.cpp
--- a/Engine/Source/Developer/TraceInsights/Private/Insights/MemoryProfiler/MemoryProfilerManager.cpp
+++ b/Engine/Source/Developer/TraceInsights/Private/Insights/MemoryProfiler/MemoryProfilerManager.cpp
@@ -3,6 +3,7 @@
 #include "MemoryProfilerManager.h"
 
 #include "Modules/ModuleManager.h"
+#include "MessageLog/Public/MessageLogModule.h"
 #include "TraceServices/AnalysisService.h"
 #include "TraceServices/Model/Memory.h"
 #include "WorkspaceMenuStructure.h"
@@ -57,6 +58,7 @@ FMemoryProfilerManager::FMemoryProfilerManager(TSharedRef<FUICommandList> InComm
 	, ProfilerWindow(nullptr)
 	, bIsTimingViewVisible(false)
 	, bIsMemTagTreeViewVisible(false)
+	, LogListingName(TEXT("MemoryInsights"))
 {
 }
 
@@ -94,6 +96,16 @@ void FMemoryProfilerManager::Shutdown()
 	}
 	bIsInitialized = false;
 
+	// If the MessageLog module was already unloaded as part of the global Shutdown process, do not load it again.
+	if (FModuleManager::Get().IsModuleLoaded("MessageLog"))
+	{
+		FMessageLogModule& MessageLogModule = FModuleManager::LoadModuleChecked<FMessageLogModule>("MessageLog");
+		if (MessageLogModule.IsRegisteredLogListing(GetLogListingName()))
+		{
+			MessageLogModule.UnregisterLogListing(GetLogListingName());
+		}
+	}
+
 	FInsightsManager::Get()->GetSessionChangedEvent().RemoveAll(this);
 
 	FMemoryProfilerCommands::Unregister();
@@ -246,6 +258,10 @@ bool FMemoryProfilerManager::Tick(float DeltaTime)
 		{
 			bIsAvailable = true;
 
+			FMessageLogModule& MessageLogModule = FModuleManager::LoadModuleChecked<FMessageLogModule>("MessageLog");
+			MessageLogModule.RegisterLogListing(GetLogListingName(), LOCTEXT("MemoryInsights", "Memory Insights"));
+			MessageLogModule.EnableMessageLogDisplay(true);
+
 #if !WITH_EDITOR
 			const FName& TabId = FInsightsManagerTabs::MemoryProfilerTabId;
 			if (FGlobalTabmanager::Get()->HasTabSpawner(TabId))
diff --git a/Engine/Source/Developer/TraceInsights/Private/Insights/MemoryProfiler/MemoryProfilerManager.h b/Engine/Source/Developer/TraceInsights/Private/Insights/MemoryProfiler/MemoryProfilerManager.h
--- a/Engine/Source/Developer/TraceInsights/Private/Insights/MemoryProfiler/MemoryProfilerManager.h
+++ b/Engine/Source/Developer/TraceInsights/Private/Insights/MemoryProfiler/MemoryProfilerManager.h
@@ -99,6 +99,8 @@ public:
 
 	////////////////////////////////////////////////////////////////////////////////////////////////////
 
+	const FName& GetLogListingName() const { return LogListingName; }
+
 	void OnSessionChanged();
 
 private:
@@ -142,6 +144,9 @@ private:
 	/** If the Categories tree view is visible or hidden. */
 	bool bIsMemTagTreeViewVisible;
 
+	/** The name of the Memory Insights log listing. */
+	FName LogListingName;
+
 	/** A shared pointer to the global instance of the Memory Profiler manager. */
 	static TSharedPtr<FMemoryProfilerManager> Instance;
 };
diff --git a/Engine/Source/Developer/TraceInsights/Private/Insights/MemoryProfiler/ViewModels/MemorySharedState.cpp b/Engine/Source/Developer/TraceInsights/Private/Insights/MemoryProfiler/ViewModels/MemorySharedState.cpp
--- a/Engine/Source/Developer/TraceInsights/Private/Insights/MemoryProfiler/ViewModels/MemorySharedState.cpp
+++ b/Engine/Source/Developer/TraceInsights/Private/Insights/MemoryProfiler/ViewModels/MemorySharedState.cpp
@@ -2,6 +2,7 @@
 
 #include "MemorySharedState.h"
 #include "Framework/MultiBox/MultiBoxBuilder.h"
+#include "Logging/MessageLog.h"
 #include "TraceServices/AnalysisService.h"
 #include "TraceServices/Model/Memory.h"
 
@@ -20,7 +21,6 @@
 #include "Insights/ViewModels/TimingGraphTrack.h"
 #include "Insights/ViewModels/TooltipDrawState.h"
 #include "Insights/ViewModels/TimingEventSearch.h"
-#include "Insights/ViewModels/InsightsMessageLogViewModel.h"
 #include "Insights/Widgets/STimingView.h"
 
 #include <limits>
@@ -719,13 +719,12 @@ void FMemorySharedState::CreateTracksFromReport(const FString& Filename)
 
 	Insights::FReportXmlParser ReportXmlParser;
 
-	auto MessageLog = FInsightsManager::Get()->GetMessageLog();
-	MessageLog->ClearMessageLog();
-
 	ReportXmlParser.LoadReportTypesXML(ReportConfig, Filename);
 	if (ReportXmlParser.GetStatus() != Insights::FReportXmlParser::EStatus::Completed)
 	{
-		MessageLog->UpdateMessageLog(ReportXmlParser.GetErrorMessages());
+		FMessageLog ReportMessageLog(FMemoryProfilerManager::Get()->GetLogListingName());
+		ReportMessageLog.AddMessages(ReportXmlParser.GetErrorMessages());
+		ReportMessageLog.Notify();
 	}
 
 	CreateTracksFromReport(ReportConfig);
diff --git a/Engine/Source/Developer/TraceInsights/Private/Insights/MemoryProfiler/ViewModels/ReportXmlParser.cpp b/Engine/Source/Developer/TraceInsights/Private/Insights/MemoryProfiler/ViewModels/ReportXmlParser.cpp
--- a/Engine/Source/Developer/TraceInsights/Private/Insights/MemoryProfiler/ViewModels/ReportXmlParser.cpp
+++ b/Engine/Source/Developer/TraceInsights/Private/Insights/MemoryProfiler/ViewModels/ReportXmlParser.cpp
@@ -1,14 +1,15 @@
 // Copyright Epic Games, Inc. All Rights Reserved.
 
 #include "ReportXmlParser.h"
+
+#include "Logging/TokenizedMessage.h"
+#include "Misc/Paths.h"
 #include "XmlNode.h"
 #include "XmlFile.h"
-#include "Misc/Paths.h"
-#include "Logging/MessageLog.h"
 
 // Insights
-#include "Insights/MemoryProfiler/ViewModels/Report.h"
 #include "Insights/Log.h"
+#include "Insights/MemoryProfiler/ViewModels/Report.h"
 
 namespace Insights
 {
diff --git a/Engine/Source/Developer/TraceInsights/Private/Insights/ViewModels/InsightsMessageLogViewModel.cpp b/Engine/Source/Developer/TraceInsights/Private/Insights/ViewModels/InsightsMessageLogViewModel.cpp
deleted file mode 100644
diff --git a/Engine/Source/Developer/TraceInsights/Private/Insights/ViewModels/InsightsMessageLogViewModel.h b/Engine/Source/Developer/TraceInsights/Private/Insights/ViewModels/InsightsMessageLogViewModel.h
deleted file mode 100644
diff --git a/Engine/Source/Developer/TraceInsights/Public/Insights/IUnrealInsightsModule.h b/Engine/Source/Developer/TraceInsights/Public/Insights/IUnrealInsightsModule.h
--- a/Engine/Source/Developer/TraceInsights/Public/Insights/IUnrealInsightsModule.h
+++ b/Engine/Source/Developer/TraceInsights/Public/Insights/IUnrealInsightsModule.h
@@ -27,7 +27,6 @@ struct TRACEINSIGHTS_API FInsightsManagerTabs
 	static const FName LoadingProfilerTabId;
 	static const FName NetworkingProfilerTabId;
 	static const FName MemoryProfilerTabId;
-	static const FName InsightsMessageLogTabId;
 	static const FName AutomationWindowTabId;
 };
 
