[Insights]
  - Improved performance of flat grouping for very large generic tables.
  - Added debug code in STableTreeView to measure peak times spent in sorting, filtering and grouping.
  - Added FInt64ValueFormatterAsHex32 and FInt64ValueFormatterAsHex64 (table cell value formatters for Int64 to hex32/hex64).

diff --git a/Engine/Source/Developer/TraceInsights/Private/Insights/Table/ViewModels/BaseTreeNode.h b/Engine/Source/Developer/TraceInsights/Private/Insights/Table/ViewModels/BaseTreeNode.h
--- a/Engine/Source/Developer/TraceInsights/Private/Insights/Table/ViewModels/BaseTreeNode.h
+++ b/Engine/Source/Developer/TraceInsights/Private/Insights/Table/ViewModels/BaseTreeNode.h
@@ -117,13 +117,13 @@ public:
 	}
 
 	/** Clears children. */
-	void ClearChildren()
+	void ClearChildren(int32 NewSize = 0)
 	{
 		for (FBaseTreeNodePtr& NodePtr : Children)
 		{
 			NodePtr->GroupPtr = nullptr;
 		}
-		Children.Reset();
+		Children.Reset(NewSize);
 	}
 
 	/** Adds specified child to the filtered children. */
diff --git a/Engine/Source/Developer/TraceInsights/Private/Insights/Table/ViewModels/TableCellValueFormatter.cpp b/Engine/Source/Developer/TraceInsights/Private/Insights/Table/ViewModels/TableCellValueFormatter.cpp
--- a/Engine/Source/Developer/TraceInsights/Private/Insights/Table/ViewModels/TableCellValueFormatter.cpp
+++ b/Engine/Source/Developer/TraceInsights/Private/Insights/Table/ViewModels/TableCellValueFormatter.cpp
@@ -20,6 +20,13 @@ FText FTableCellValueFormatter::FormatValue(const FTableColumn& Column, const FB
 
 ////////////////////////////////////////////////////////////////////////////////////////////////////
 
+FText FTableCellValueFormatter::FormatValueForTooltip(const FTableColumn& Column, const FBaseTreeNode& Node) const
+{
+	return FormatValueForTooltip(Column.GetValue(Node));
+}
+
+////////////////////////////////////////////////////////////////////////////////////////////////////
+
 FText FBoolValueFormatterAsTrueFalse::FormatValue(const TOptional<FTableCellValue>& InValue) const
 {
 	if (InValue.IsSet())
@@ -102,6 +109,8 @@ FText FFloatValueFormatterAsTimeAuto::FormatValue(const TOptional<FTableCellValu
 	return FText::GetEmpty();
 }
 
+////////////////////////////////////////////////////////////////////////////////////////////////////
+
 FText FFloatValueFormatterAsTimeAuto::FormatValueForTooltip(const TOptional<FTableCellValue>& InValue) const
 {
 	if (InValue.IsSet())
diff --git a/Engine/Source/Developer/TraceInsights/Private/Insights/Table/ViewModels/TableCellValueFormatter.h b/Engine/Source/Developer/TraceInsights/Private/Insights/Table/ViewModels/TableCellValueFormatter.h
--- a/Engine/Source/Developer/TraceInsights/Private/Insights/Table/ViewModels/TableCellValueFormatter.h
+++ b/Engine/Source/Developer/TraceInsights/Private/Insights/Table/ViewModels/TableCellValueFormatter.h
@@ -36,7 +36,7 @@ public:
 	virtual FText FormatValueForTooltip(const TOptional<FTableCellValue>& InValue) const override { return FormatValue(InValue); }
 
 	virtual FText FormatValue(const FTableColumn& Column, const FBaseTreeNode& Node) const override; // { return FormatValue(Column.GetValue(Node)); }
-	virtual FText FormatValueForTooltip(const FTableColumn& Column, const FBaseTreeNode& Node) const override { return FormatValue(Column, Node); }
+	virtual FText FormatValueForTooltip(const FTableColumn& Column, const FBaseTreeNode& Node) const override; // { return FormatValueForTooltip(Column.GetValue(Node)); }
 };
 
 ////////////////////////////////////////////////////////////////////////////////////////////////////
@@ -98,6 +98,28 @@ public:
 
 ////////////////////////////////////////////////////////////////////////////////////////////////////
 
+class FInt64ValueFormatterAsHex32 : public FTableCellValueFormatter
+{
+public:
+	virtual FText FormatValue(const TOptional<FTableCellValue>& InValue) const override
+	{
+		return InValue.IsSet() ? FText::FromString(FString::Printf(TEXT("0x%08X"), static_cast<uint32>(InValue.GetValue().Int64))) : FText::GetEmpty();
+	}
+};
+
+////////////////////////////////////////////////////////////////////////////////////////////////////
+
+class FInt64ValueFormatterAsHex64 : public FTableCellValueFormatter
+{
+public:
+	virtual FText FormatValue(const TOptional<FTableCellValue>& InValue) const override
+	{
+		return InValue.IsSet() ? FText::FromString(FString::Printf(TEXT("0x%016llX"), static_cast<uint64>(InValue.GetValue().Int64))) : FText::GetEmpty();
+	}
+};
+
+////////////////////////////////////////////////////////////////////////////////////////////////////
+
 class FInt64ValueFormatterAsMemory : public FTableCellValueFormatter
 {
 public:
diff --git a/Engine/Source/Developer/TraceInsights/Private/Insights/Table/ViewModels/TableTreeNode.h b/Engine/Source/Developer/TraceInsights/Private/Insights/Table/ViewModels/TableTreeNode.h
--- a/Engine/Source/Developer/TraceInsights/Private/Insights/Table/ViewModels/TableTreeNode.h
+++ b/Engine/Source/Developer/TraceInsights/Private/Insights/Table/ViewModels/TableTreeNode.h
@@ -18,12 +18,11 @@ struct FTableRowId
 {
 	static constexpr int32 InvalidRowIndex = -1;
 
-	FTableRowId(int32 InRowIndex) : RowIndex(InRowIndex), Flags(0) {}
+	FTableRowId(int32 InRowIndex) : RowIndex(InRowIndex) {}
 
 	bool HasValidIndex() const { return RowIndex >= 0; }
 
 	int32 RowIndex;
-	uint32 Flags;
 };
 
 ////////////////////////////////////////////////////////////////////////////////////////////////////
@@ -71,7 +70,7 @@ public:
 
 	virtual const FName& GetTypeName() const override { return TypeName; }
 
-	TWeakPtr<FTable> GetParentTable() { return ParentTable; }
+	const TWeakPtr<FTable>& GetParentTable() const { return ParentTable; }
 	FTableRowId GetRowId() const { return RowId; }
 	int32 GetRowIndex() const { return RowId.RowIndex; }
 
diff --git a/Engine/Source/Developer/TraceInsights/Private/Insights/Table/ViewModels/TreeNodeGrouping.cpp b/Engine/Source/Developer/TraceInsights/Private/Insights/Table/ViewModels/TreeNodeGrouping.cpp
--- a/Engine/Source/Developer/TraceInsights/Private/Insights/Table/ViewModels/TreeNodeGrouping.cpp
+++ b/Engine/Source/Developer/TraceInsights/Private/Insights/Table/ViewModels/TreeNodeGrouping.cpp
@@ -22,6 +22,42 @@ FTreeNodeGrouping::FTreeNodeGrouping(const FText& InShortName, const FText& InTi
 
 ////////////////////////////////////////////////////////////////////////////////////////////////////
 
+void FTreeNodeGrouping::GroupNodes(const TArray<FTableTreeNodePtr>& Nodes, FTableTreeNode& ParentGroup, TWeakPtr<FTable> InParentTable) const
+{
+	TMap<FName, FTableTreeNodePtr> GroupMap;
+
+	ParentGroup.ClearChildren();
+
+	for (FTableTreeNodePtr NodePtr : Nodes)
+	{
+		ensure(!NodePtr->IsGroup());
+
+		FTableTreeNodePtr GroupPtr = nullptr;
+
+		FTreeNodeGroupInfo GroupInfo = GetGroupForNode(NodePtr);
+		FTableTreeNodePtr* GroupPtrPtr = GroupMap.Find(GroupInfo.Name);
+		if (!GroupPtrPtr)
+		{
+			GroupPtr = MakeShared<FTableTreeNode>(GroupInfo.Name, InParentTable);
+
+			GroupMap.Add(GroupInfo.Name, GroupPtr);
+			ParentGroup.AddChildAndSetGroupPtr(GroupPtr);
+
+			GroupPtr->SetExpansion(GroupInfo.IsExpanded);
+		}
+		else
+		{
+			GroupPtr = *GroupPtrPtr;
+		}
+
+		GroupPtr->AddChildAndSetGroupPtr(NodePtr);
+	}
+}
+
+////////////////////////////////////////////////////////////////////////////////////////////////////
+// FTreeNodeGroupingFlat
+////////////////////////////////////////////////////////////////////////////////////////////////////
+
 FTreeNodeGroupingFlat::FTreeNodeGroupingFlat()
 	: FTreeNodeGrouping(
 		LOCTEXT("Grouping_Flat_ShortName", "All"),
@@ -41,6 +77,25 @@ FTreeNodeGroupInfo FTreeNodeGroupingFlat::GetGroupForNode(const FBaseTreeNodePtr
 
 ////////////////////////////////////////////////////////////////////////////////////////////////////
 
+void FTreeNodeGroupingFlat::GroupNodes(const TArray<FTableTreeNodePtr>& Nodes, FTableTreeNode& ParentGroup, TWeakPtr<FTable> InParentTable) const
+{
+	ParentGroup.ClearChildren(1);
+
+	FTableTreeNodePtr GroupPtr = MakeShared<FTableTreeNode>(FName(TEXT("All")), InParentTable);
+	GroupPtr->SetExpansion(true);
+	ParentGroup.AddChildAndSetGroupPtr(GroupPtr);
+
+	GroupPtr->ClearChildren(Nodes.Num());
+	for (FTableTreeNodePtr NodePtr : Nodes)
+	{
+		GroupPtr->AddChildAndSetGroupPtr(NodePtr);
+	}
+}
+
+////////////////////////////////////////////////////////////////////////////////////////////////////
+// FTreeNodeGroupingByUniqueValue
+////////////////////////////////////////////////////////////////////////////////////////////////////
+
 FTreeNodeGroupingByUniqueValue::FTreeNodeGroupingByUniqueValue(TSharedRef<FTableColumn> InColumnRef)
 	: FTreeNodeGrouping(
 		InColumnRef->GetTitleName(),
diff --git a/Engine/Source/Developer/TraceInsights/Private/Insights/Table/ViewModels/TreeNodeGrouping.h b/Engine/Source/Developer/TraceInsights/Private/Insights/Table/ViewModels/TreeNodeGrouping.h
--- a/Engine/Source/Developer/TraceInsights/Private/Insights/Table/ViewModels/TreeNodeGrouping.h
+++ b/Engine/Source/Developer/TraceInsights/Private/Insights/Table/ViewModels/TreeNodeGrouping.h
@@ -9,10 +9,13 @@
 // Insights
 #include "Insights/Table/ViewModels/BaseTreeNode.h"
 #include "Insights/Table/ViewModels/TableColumn.h"
+#include "Insights/Table/ViewModels/TableTreeNode.h"
 
 namespace Insights
 {
 
+class FTable;
+
 ////////////////////////////////////////////////////////////////////////////////////////////////////
 
 class ITreeNodeGrouping
@@ -26,8 +29,6 @@ public:
 	virtual const FSlateBrush* GetIcon() const = 0;
 
 	virtual FName GetColumnId() const = 0;
-
-	//virtual void CreateGrouping() = 0;
 };
 
 ////////////////////////////////////////////////////////////////////////////////////////////////////
@@ -56,6 +57,7 @@ public:
 	virtual FName GetColumnId() const override { return NAME_None; }
 
 	virtual FTreeNodeGroupInfo GetGroupForNode(const FBaseTreeNodePtr InNode) const { return { FName(), false }; }
+	virtual void GroupNodes(const TArray<FTableTreeNodePtr>& Nodes, FTableTreeNode& ParentGroup, TWeakPtr<FTable> InParentTable) const;
 
 protected:
 	FText ShortName;
@@ -75,6 +77,7 @@ public:
 	virtual ~FTreeNodeGroupingFlat() {}
 
 	virtual FTreeNodeGroupInfo GetGroupForNode(const FBaseTreeNodePtr InNode) const override;
+	virtual void GroupNodes(const TArray<FTableTreeNodePtr>& Nodes, FTableTreeNode& ParentGroup, TWeakPtr<FTable> InParentTable) const override;
 };
 
 ////////////////////////////////////////////////////////////////////////////////////////////////////
diff --git a/Engine/Source/Developer/TraceInsights/Private/Insights/Table/Widgets/STableTreeView.cpp b/Engine/Source/Developer/TraceInsights/Private/Insights/Table/Widgets/STableTreeView.cpp
--- a/Engine/Source/Developer/TraceInsights/Private/Insights/Table/Widgets/STableTreeView.cpp
+++ b/Engine/Source/Developer/TraceInsights/Private/Insights/Table/Widgets/STableTreeView.cpp
@@ -20,6 +20,7 @@
 // Insights
 #include "Insights/Common/Stopwatch.h"
 #include "Insights/InsightsManager.h"
+#include "Insights/Log.h"
 #include "Insights/Table/ViewModels/Table.h"
 #include "Insights/Table/ViewModels/TableColumn.h"
 #include "Insights/Table/ViewModels/TableTreeNode.h"
@@ -580,15 +581,28 @@ void STableTreeView::InsightsManager_OnSessionChanged()
 
 void STableTreeView::UpdateTree()
 {
+	FStopwatch Stopwatch;
+	Stopwatch.Start();
+
 	CreateGroups();
 	SortTreeNodes();
 	ApplyFiltering();
+
+	Stopwatch.Stop();
+	const double TotalTime = Stopwatch.GetAccumulatedTime();
+	if (TotalTime > 0.1)
+	{
+		UE_LOG(TraceInsights, Log, TEXT("[Tree - %s] Tree updated (grouping + sorting + filtering) in %.3fs."), *Table->GetDisplayName().ToString(), TotalTime);
+	}
 }
 
 ////////////////////////////////////////////////////////////////////////////////////////////////////
 
 void STableTreeView::ApplyFiltering()
 {
+	FStopwatch Stopwatch;
+	Stopwatch.Start();
+
 	// Apply filter to all groups and its children.
 	ApplyFilteringForNode(Root);
 
@@ -638,6 +652,13 @@ void STableTreeView::ApplyFiltering()
 
 	// Request tree refresh
 	TreeView->RequestTreeRefresh();
+
+	Stopwatch.Stop();
+	const double TotalTime = Stopwatch.GetAccumulatedTime();
+	if (TotalTime > 0.1)
+	{
+		UE_LOG(TraceInsights, Log, TEXT("[Tree - %s] Filtering completed in %.3fs."), *Table->GetDisplayName().ToString(), TotalTime);
+	}
 }
 
 ////////////////////////////////////////////////////////////////////////////////////////////////////
@@ -832,8 +853,14 @@ bool STableTreeView::SearchBox_IsEnabled() const
 
 void STableTreeView::CreateGroups()
 {
+	FStopwatch Stopwatch;
+	Stopwatch.Start();
+
 	GroupNodesRec(TableTreeNodes, *Root, 0);
 
+	Stopwatch.Update();
+	const double Time1 = Stopwatch.GetAccumulatedTime();
+
 	ResetAggregatedValuesRec(*Root);
 	for (const TSharedRef<FTableColumn>& ColumnRef : Table->GetColumns())
 	{
@@ -856,6 +883,13 @@ void STableTreeView::CreateGroups()
 			}
 		}
 	}
+
+	Stopwatch.Stop();
+	const double TotalTime = Stopwatch.GetAccumulatedTime();
+	if (TotalTime > 0.1)
+	{
+		UE_LOG(TraceInsights, Log, TEXT("[Tree - %s] Grouping completed in %.3fs (%.3fs + %.3fs)."), *Table->GetDisplayName().ToString(), TotalTime, Time1, TotalTime - Time1);
+	}
 }
 
 ////////////////////////////////////////////////////////////////////////////////////////////////////
@@ -866,34 +900,13 @@ void STableTreeView::GroupNodesRec(const TArray<FTableTreeNodePtr>& Nodes, FTabl
 
 	FTreeNodeGrouping& Grouping = *CurrentGroupings[GroupingDepth];
 
-	TMap<FName, FTableTreeNodePtr> GroupMap;
-
-	ParentGroup.ClearChildren();
+	Grouping.GroupNodes(Nodes, ParentGroup, Table);
 
-	for (FTableTreeNodePtr NodePtr : Nodes)
+	for (FBaseTreeNodePtr GroupPtr : ParentGroup.GetChildren())
 	{
-		ensure(!NodePtr->IsGroup());
-
-		FTableTreeNodePtr GroupPtr = nullptr;
-
-		FTreeNodeGroupInfo GroupInfo = Grouping.GetGroupForNode(NodePtr);
-		FTableTreeNodePtr* GroupPtrPtr = GroupMap.Find(GroupInfo.Name);
-		if (!GroupPtrPtr)
-		{
-			GroupPtr = MakeShared<FTableTreeNode>(GroupInfo.Name, Table);
-
-			GroupMap.Add(GroupInfo.Name, GroupPtr);
-			ParentGroup.AddChildAndSetGroupPtr(GroupPtr);
-
-			GroupPtr->SetExpansion(GroupInfo.IsExpanded);
-			TreeView->SetItemExpansion(GroupPtr, GroupInfo.IsExpanded);
-		}
-		else
-		{
-			GroupPtr = *GroupPtrPtr;
-		}
-
-		GroupPtr->AddChildAndSetGroupPtr(NodePtr);
+		ensure(GroupPtr->IsGroup());
+		FTableTreeNodePtr TableTreeGroupPtr = StaticCastSharedPtr<FTableTreeNode>(GroupPtr);
+		TreeView->SetItemExpansion(TableTreeGroupPtr, TableTreeGroupPtr->IsExpanded());
 	}
 
 	if (GroupingDepth < CurrentGroupings.Num() - 1)
@@ -1538,7 +1551,17 @@ void STableTreeView::SortTreeNodes()
 {
 	if (CurrentSorter.IsValid())
 	{
+		FStopwatch Stopwatch;
+		Stopwatch.Start();
+
 		SortTreeNodesRec(*Root, *CurrentSorter);
+
+		Stopwatch.Stop();
+		const double TotalTime = Stopwatch.GetAccumulatedTime();
+		if (TotalTime > 0.1)
+		{
+			UE_LOG(TraceInsights, Log, TEXT("[Tree - %s] Sorting completed in %.3fs."), *Table->GetDisplayName().ToString(), TotalTime);
+		}
 	}
 }
 
diff --git a/Engine/Source/Developer/TraceInsights/Private/Insights/Table/Widgets/STableTreeView.h b/Engine/Source/Developer/TraceInsights/Private/Insights/Table/Widgets/STableTreeView.h
--- a/Engine/Source/Developer/TraceInsights/Private/Insights/Table/Widgets/STableTreeView.h
+++ b/Engine/Source/Developer/TraceInsights/Private/Insights/Table/Widgets/STableTreeView.h
@@ -63,7 +63,8 @@ public:
 	 */
 	void Construct(const FArguments& InArgs, TSharedPtr<Insights::FTable> InTablePtr);
 
-	TSharedPtr<Insights::FTable> GetTable() const { return Table; }
+	TSharedPtr<Insights::FTable>& GetTable() { return Table; }
+	const TSharedPtr<Insights::FTable>& GetTable() const { return Table; }
 
 	virtual void Reset();
 
