[Insights]
  - Added initial Memory Insights (allocation queries) workflow:
    * Added "Investigation" section in the "LLM Tags" panel;
    * Added option to choose query rule from dropdown list of available ones;
    * Added A,B,C,D time markers in timing view of Memory Insights;
    * Addded button to invoke query;
    * The query will create up to 4 tree table views (the result of each query is a set of allocations).
  - Added FMemAllocNode, FMemAllocTable, FMemoryAlloc (view model classes for memory allocation tree view).
  - Added SMemAllocTableTreeView widget.
  - Added FMemoryRule and FMemoryRuleSpec. MemorySharedState now exposes a list for available memory rules (for allocation queries).
  - Improved display of time markers when hovered by mouse.

diff --git a/Engine/Source/Developer/TraceInsights/Private/Insights/MemoryProfiler/ViewModels/MemAllocNode.cpp b/Engine/Source/Developer/TraceInsights/Private/Insights/MemoryProfiler/ViewModels/MemAllocNode.cpp
new file mode 100644
--- /dev/null
+++ b/Engine/Source/Developer/TraceInsights/Private/Insights/MemoryProfiler/ViewModels/MemAllocNode.cpp
@@ -0,0 +1,52 @@
+// Copyright Epic Games, Inc. All Rights Reserved.
+
+#include "MemAllocNode.h"
+
+// Insights
+#include "Insights/MemoryProfiler/MemoryProfilerManager.h"
+#include "Insights/MemoryProfiler/ViewModels/MemorySharedState.h"
+#include "Insights/MemoryProfiler/Widgets/SMemoryProfilerWindow.h"
+
+#define LOCTEXT_NAMESPACE "MemAllocNode"
+
+namespace Insights
+{
+
+const FName FMemAllocNode::TypeName(TEXT("FMemAllocNode"));
+
+////////////////////////////////////////////////////////////////////////////////////////////////////
+
+FText FMemAllocNode::GetAllocMemTagAsText() const
+{
+	const FMemoryAlloc* Alloc = GetMemAlloc();
+	if (Alloc)
+	{
+		return FText::FromString(GetAllocMemTagAsString(Alloc->GetMemTag()));
+	}
+	return FText::GetEmpty();
+}
+
+////////////////////////////////////////////////////////////////////////////////////////////////////
+
+FString FMemAllocNode::GetAllocMemTagAsString(FMemoryTagId InMemTag)
+{
+	TSharedPtr<SMemoryProfilerWindow> ProfilerWindow = FMemoryProfilerManager::Get()->GetProfilerWindow();
+	if (ProfilerWindow)
+	{
+		auto& SharedState = ProfilerWindow->GetSharedState();
+		const Insights::FMemoryTagList& TagList = SharedState.GetTagList();
+
+		const FMemoryTag* Tag = TagList.GetTagById(InMemTag);
+		if (Tag)
+		{
+			return Tag->GetStatName();
+		}
+	}
+	return FString();
+}
+
+////////////////////////////////////////////////////////////////////////////////////////////////////
+
+} // namespace Insights
+
+#undef LOCTEXT_NAMESPACE
diff --git a/Engine/Source/Developer/TraceInsights/Private/Insights/MemoryProfiler/ViewModels/MemAllocNode.h b/Engine/Source/Developer/TraceInsights/Private/Insights/MemoryProfiler/ViewModels/MemAllocNode.h
new file mode 100644
--- /dev/null
+++ b/Engine/Source/Developer/TraceInsights/Private/Insights/MemoryProfiler/ViewModels/MemAllocNode.h
@@ -0,0 +1,98 @@
+// Copyright Epic Games, Inc. All Rights Reserved.
+
+#pragma once
+
+#include "CoreMinimal.h"
+#include "TraceServices/Model/Memory.h"
+
+// Insights
+#include "Insights/MemoryProfiler/ViewModels/MemAllocTable.h"
+#include "Insights/MemoryProfiler/ViewModels/MemoryAlloc.h"
+#include "Insights/MemoryProfiler/ViewModels/MemorySharedState.h"
+#include "Insights/Table/ViewModels/TableTreeNode.h"
+
+////////////////////////////////////////////////////////////////////////////////////////////////////
+
+namespace Insights
+{
+
+enum class EMemAllocNodeType
+{
+	/** The MemAllocNode is an allocation node. */
+	MemAlloc,
+
+	/** The MemAllocNode is a group node. */
+	Group,
+
+	/** Invalid enum type, may be used as a number of enumerations. */
+	InvalidOrMax,
+};
+
+////////////////////////////////////////////////////////////////////////////////////////////////////
+
+class FMemAllocNode;
+
+/** Type definition for shared pointers to instances of FMemAllocNode. */
+typedef TSharedPtr<class FMemAllocNode> FMemAllocNodePtr;
+
+/** Type definition for shared references to instances of FMemAllocNode. */
+typedef TSharedRef<class FMemAllocNode> FMemAllocNodeRef;
+
+/** Type definition for shared references to const instances of FMemAllocNode. */
+typedef TSharedRef<const class FMemAllocNode> FMemAllocNodeRefConst;
+
+/** Type definition for weak references to instances of FMemAllocNode. */
+typedef TWeakPtr<class FMemAllocNode> FMemAllocNodeWeak;
+
+////////////////////////////////////////////////////////////////////////////////////////////////////
+/**
+ * Class used to store information about an allocation node (used in the SMemAllocTreeView).
+ */
+class FMemAllocNode : public FTableTreeNode
+{
+public:
+	static const FName TypeName;
+
+public:
+	/** Initialization constructor for the MemAlloc node. */
+	explicit FMemAllocNode(const FName InName, TWeakPtr<FMemAllocTable> InParentTable, int32 InRowIndex)
+		: FTableTreeNode(InName, InParentTable, InRowIndex)
+		, Type(EMemAllocNodeType::MemAlloc)
+	{
+	}
+
+	/** Initialization constructor for the group node. */
+	explicit FMemAllocNode(const FName InGroupName, TWeakPtr<FMemAllocTable> InParentTable)
+		: FTableTreeNode(InGroupName, InParentTable)
+		, Type(EMemAllocNodeType::Group)
+	{
+	}
+
+	virtual const FName& GetTypeName() const override { return TypeName; }
+
+	/**
+	 * @return a type of this MemAlloc node or EMemAllocNodeType::Group for group nodes.
+	 */
+	const EMemAllocNodeType& GetType() const { return Type; }
+
+	FMemAllocTable& GetMemTableChecked() const
+	{
+		const TSharedPtr<FTable>& TablePin = GetParentTable().Pin();
+		check(TablePin.IsValid());
+		return *StaticCastSharedPtr<FMemAllocTable>(TablePin);
+	}
+
+	bool IsValidMemAlloc() const { return GetMemTableChecked().IsValidRowIndex(GetRowIndex()); }
+	const FMemoryAlloc* GetMemAlloc() const { return GetMemTableChecked().GetMemAlloc(GetRowIndex()); }
+	const FMemoryAlloc& GetMemAllocChecked() const { return GetMemTableChecked().GetMemAllocChecked(GetRowIndex()); }
+
+	FText GetAllocMemTagAsText() const;
+	static FString GetAllocMemTagAsString(FMemoryTagId InMemTag);
+
+private:
+	const EMemAllocNodeType Type;
+};
+
+////////////////////////////////////////////////////////////////////////////////////////////////////
+
+} // namespace Insights
diff --git a/Engine/Source/Developer/TraceInsights/Private/Insights/MemoryProfiler/ViewModels/MemAllocTable.cpp b/Engine/Source/Developer/TraceInsights/Private/Insights/MemoryProfiler/ViewModels/MemAllocTable.cpp
new file mode 100644
--- /dev/null
+++ b/Engine/Source/Developer/TraceInsights/Private/Insights/MemoryProfiler/ViewModels/MemAllocTable.cpp
@@ -0,0 +1,391 @@
+// Copyright Epic Games, Inc. All Rights Reserved.
+
+#include "MemAllocTable.h"
+
+// Insights
+#include "Insights/MemoryProfiler/ViewModels/MemAllocNode.h"
+#include "Insights/MemoryProfiler/ViewModels/MemAllocTable.h"
+#include "Insights/Table/ViewModels/TableCellValueFormatter.h"
+#include "Insights/Table/ViewModels/TableCellValueGetter.h"
+#include "Insights/Table/ViewModels/TableCellValueSorter.h"
+#include "Insights/Table/ViewModels/TableColumn.h"
+
+#define LOCTEXT_NAMESPACE "Insights::FMemAllocTable"
+
+namespace Insights
+{
+
+////////////////////////////////////////////////////////////////////////////////////////////////////
+// FMemAllocTable
+////////////////////////////////////////////////////////////////////////////////////////////////////
+
+FMemAllocTable::FMemAllocTable()
+{
+}
+
+////////////////////////////////////////////////////////////////////////////////////////////////////
+
+FMemAllocTable::~FMemAllocTable()
+{
+}
+
+////////////////////////////////////////////////////////////////////////////////////////////////////
+
+void FMemAllocTable::Reset()
+{
+	//...
+
+	FTable::Reset();
+
+	AddDefaultColumns();
+}
+
+////////////////////////////////////////////////////////////////////////////////////////////////////
+
+void FMemAllocTable::AddDefaultColumns()
+{
+	//////////////////////////////////////////////////
+	// Hierarchy Column
+	{
+		int32 HierarchyColumnIndex = -1;
+		const TCHAR* HierarchyColumnName = nullptr;
+		AddHierarchyColumn(HierarchyColumnIndex, HierarchyColumnName);
+
+		const TSharedRef<FTableColumn>& ColumnRef = GetColumns()[0];
+		ColumnRef->SetInitialWidth(200.0f);
+		ColumnRef->SetShortName(LOCTEXT("AllocationColumnName", "Hierarchy"));
+		ColumnRef->SetTitleName(LOCTEXT("AllocationColumnTitle", "Allocation Hierarchy"));
+		ColumnRef->SetDescription(LOCTEXT("AllocationColumnDesc", "Hierarchy of the allocation's tree"));
+	}
+	//////////////////////////////////////////////////
+	// Start Time
+	{
+		TSharedRef<FTableColumn> ColumnRef = MakeShared<FTableColumn>(TEXT("AllocStartTime"));
+		FTableColumn& Column = *ColumnRef;
+
+		Column.SetIndex(0);
+
+		Column.SetShortName(LOCTEXT("StartTimeColumnName", "Start Time"));
+		Column.SetTitleName(LOCTEXT("StartTimeColumnTitle", "Start Time"));
+		Column.SetDescription(LOCTEXT("StartTimeColumnDesc", "The time when the allocation was allocated."));
+
+		Column.SetFlags(ETableColumnFlags::ShouldBeVisible | ETableColumnFlags::CanBeHidden | ETableColumnFlags::CanBeFiltered);
+
+		Column.SetHorizontalAlignment(HAlign_Right);
+		Column.SetInitialWidth(100.0f);
+
+		Column.SetDataType(ETableCellDataType::Double);
+
+		class FMemAllocStartTimeValueGetter : public FTableCellValueGetter
+		{
+		public:
+			virtual const TOptional<FTableCellValue> GetValue(const FTableColumn& Column, const FBaseTreeNode& Node) const override
+			{
+				const FMemAllocNode& MemAllocNode = static_cast<const FMemAllocNode&>(Node);
+				const FMemoryAlloc* Alloc = MemAllocNode.GetMemAlloc();
+				const double AllocStartTime = Alloc ? Alloc->GetStartTime() : 0.0;
+				return TOptional<FTableCellValue>(FTableCellValue(AllocStartTime));
+			}
+		};
+		TSharedRef<ITableCellValueGetter> Getter = MakeShared<FMemAllocStartTimeValueGetter>();
+		Column.SetValueGetter(Getter);
+
+		TSharedRef<ITableCellValueFormatter> Formatter = MakeShared<FDoubleValueFormatterAsTimeAuto>();
+		Column.SetValueFormatter(Formatter);
+
+		TSharedRef<ITableCellValueSorter> Sorter = MakeShared<FSorterByDoubleValue>(ColumnRef);
+		Column.SetValueSorter(Sorter);
+
+		AddColumn(ColumnRef);
+	}
+	//////////////////////////////////////////////////
+	// End Time
+	{
+		TSharedRef<FTableColumn> ColumnRef = MakeShared<FTableColumn>(TEXT("AllocEndTime"));
+		FTableColumn& Column = *ColumnRef;
+
+		Column.SetIndex(0);
+
+		Column.SetShortName(LOCTEXT("EndTimeColumnName", "End Time"));
+		Column.SetTitleName(LOCTEXT("EndTimeColumnTitle", "End Time"));
+		Column.SetDescription(LOCTEXT("EndTimeColumnDesc", "The time when the allocation was freed."));
+
+		Column.SetFlags(ETableColumnFlags::ShouldBeVisible | ETableColumnFlags::CanBeHidden | ETableColumnFlags::CanBeFiltered);
+
+		Column.SetHorizontalAlignment(HAlign_Right);
+		Column.SetInitialWidth(100.0f);
+
+		Column.SetDataType(ETableCellDataType::Double);
+
+		class FMemAllocStartTimeValueGetter : public FTableCellValueGetter
+		{
+		public:
+			virtual const TOptional<FTableCellValue> GetValue(const FTableColumn& Column, const FBaseTreeNode& Node) const override
+			{
+				const FMemAllocNode& MemAllocNode = static_cast<const FMemAllocNode&>(Node);
+				const FMemoryAlloc* Alloc = MemAllocNode.GetMemAlloc();
+				const double AllocEndTime = Alloc ? Alloc->GetEndTime() : 0.0;
+				return TOptional<FTableCellValue>(FTableCellValue(AllocEndTime));
+			}
+		};
+		TSharedRef<ITableCellValueGetter> Getter = MakeShared<FMemAllocStartTimeValueGetter>();
+		Column.SetValueGetter(Getter);
+
+		TSharedRef<ITableCellValueFormatter> Formatter = MakeShared<FDoubleValueFormatterAsTimeAuto>();
+		Column.SetValueFormatter(Formatter);
+
+		TSharedRef<ITableCellValueSorter> Sorter = MakeShared<FSorterByDoubleValue>(ColumnRef);
+		Column.SetValueSorter(Sorter);
+
+		AddColumn(ColumnRef);
+	}
+	//////////////////////////////////////////////////
+	// Duration
+	{
+		TSharedRef<FTableColumn> ColumnRef = MakeShared<FTableColumn>(TEXT("AllocDuration"));
+		FTableColumn& Column = *ColumnRef;
+
+		Column.SetIndex(0);
+
+		Column.SetShortName(LOCTEXT("DurationColumnName", "Duration"));
+		Column.SetTitleName(LOCTEXT("DurationColumnTitle", "Duration"));
+		Column.SetDescription(LOCTEXT("DurationColumnDesc", "The duration of the allocation's life."));
+
+		Column.SetFlags(ETableColumnFlags::ShouldBeVisible | ETableColumnFlags::CanBeHidden | ETableColumnFlags::CanBeFiltered);
+
+		Column.SetHorizontalAlignment(HAlign_Right);
+		Column.SetInitialWidth(100.0f);
+
+		Column.SetDataType(ETableCellDataType::Double);
+
+		class FMemAllocStartTimeValueGetter : public FTableCellValueGetter
+		{
+		public:
+			virtual const TOptional<FTableCellValue> GetValue(const FTableColumn& Column, const FBaseTreeNode& Node) const override
+			{
+				const FMemAllocNode& MemAllocNode = static_cast<const FMemAllocNode&>(Node);
+				const FMemoryAlloc* Alloc = MemAllocNode.GetMemAlloc();
+				const double AllocDuration = Alloc ? Alloc->GetDuration() : 0.0;
+				return TOptional<FTableCellValue>(FTableCellValue(AllocDuration));
+			}
+		};
+		TSharedRef<ITableCellValueGetter> Getter = MakeShared<FMemAllocStartTimeValueGetter>();
+		Column.SetValueGetter(Getter);
+
+		TSharedRef<ITableCellValueFormatter> Formatter = MakeShared<FDoubleValueFormatterAsTimeAuto>();
+		Column.SetValueFormatter(Formatter);
+
+		TSharedRef<ITableCellValueSorter> Sorter = MakeShared<FSorterByDoubleValue>(ColumnRef);
+		Column.SetValueSorter(Sorter);
+
+		AddColumn(ColumnRef);
+	}
+	//////////////////////////////////////////////////
+	// Address Column
+	{
+		TSharedRef<FTableColumn> ColumnRef = MakeShared<FTableColumn>(TEXT("AllocAddress"));
+		FTableColumn& Column = *ColumnRef;
+
+		Column.SetIndex(0);
+
+		Column.SetShortName(LOCTEXT("AddressColumnName", "Address"));
+		Column.SetTitleName(LOCTEXT("AddressColumnTitle", "Address"));
+		Column.SetDescription(LOCTEXT("AddressColumnDesc", "Address of allocation"));
+
+		Column.SetFlags(ETableColumnFlags::ShouldBeVisible | ETableColumnFlags::CanBeHidden | ETableColumnFlags::CanBeFiltered);
+
+		Column.SetHorizontalAlignment(HAlign_Left);
+		Column.SetInitialWidth(120.0f);
+
+		Column.SetDataType(ETableCellDataType::Int64);
+
+		class FMemAllocAddressValueGetter : public FTableCellValueGetter
+		{
+		public:
+			virtual const TOptional<FTableCellValue> GetValue(const FTableColumn& Column, const FBaseTreeNode& Node) const override
+			{
+				const FMemAllocNode& MemAllocNode = static_cast<const FMemAllocNode&>(Node);
+				const FMemoryAlloc* Alloc = MemAllocNode.GetMemAlloc();
+				const uint64 AllocAddress = Alloc ? Alloc->GetAddress() : 0;
+				return TOptional<FTableCellValue>(FTableCellValue(static_cast<int64>(AllocAddress)));
+			}
+		};
+		TSharedRef<ITableCellValueGetter> Getter = MakeShared<FMemAllocAddressValueGetter>();
+		Column.SetValueGetter(Getter);
+
+		TSharedRef<ITableCellValueFormatter> Formatter = MakeShared<FInt64ValueFormatterAsHex64>();
+		Column.SetValueFormatter(Formatter);
+
+		TSharedRef<ITableCellValueSorter> Sorter = MakeShared<FSorterByInt64Value>(ColumnRef);
+		Column.SetValueSorter(Sorter);
+
+		AddColumn(ColumnRef);
+	}
+	//////////////////////////////////////////////////
+	// Size Column
+	{
+		TSharedRef<FTableColumn> ColumnRef = MakeShared<FTableColumn>(TEXT("AllocSize"));
+		FTableColumn& Column = *ColumnRef;
+
+		Column.SetIndex(0);
+
+		Column.SetShortName(LOCTEXT("SizeColumnName", "Size"));
+		Column.SetTitleName(LOCTEXT("SizeColumnTitle", "Size"));
+		Column.SetDescription(LOCTEXT("SizeColumnDesc", "Size of allocation"));
+
+		Column.SetFlags(ETableColumnFlags::ShouldBeVisible | ETableColumnFlags::CanBeHidden | ETableColumnFlags::CanBeFiltered);
+
+		Column.SetHorizontalAlignment(HAlign_Right);
+		Column.SetInitialWidth(100.0f);
+
+		Column.SetDataType(ETableCellDataType::Int64);
+
+		class FMemAllocSizeValueGetter : public FTableCellValueGetter
+		{
+		public:
+			virtual const TOptional<FTableCellValue> GetValue(const FTableColumn& Column, const FBaseTreeNode& Node) const override
+			{
+				const FMemAllocNode& MemAllocNode = static_cast<const FMemAllocNode&>(Node);
+				const FMemoryAlloc* Alloc = MemAllocNode.GetMemAlloc();
+				const uint64 AllocSize = Alloc ? Alloc->GetSize() : 0;
+				return TOptional<FTableCellValue>(FTableCellValue(static_cast<int64>(AllocSize)));
+			}
+		};
+		TSharedRef<ITableCellValueGetter> Getter = MakeShared<FMemAllocSizeValueGetter>();
+		Column.SetValueGetter(Getter);
+
+		TSharedRef<ITableCellValueFormatter> Formatter = MakeShared<FInt64ValueFormatterAsMemory>();
+		Column.SetValueFormatter(Formatter);
+
+		TSharedRef<ITableCellValueSorter> Sorter = MakeShared<FSorterByInt64Value>(ColumnRef);
+		Column.SetValueSorter(Sorter);
+
+		AddColumn(ColumnRef);
+	}
+	//////////////////////////////////////////////////
+	// LLM Tag Column
+	{
+		TSharedRef<FTableColumn> ColumnRef = MakeShared<FTableColumn>(TEXT("AllocLlmTag"));
+		FTableColumn& Column = *ColumnRef;
+
+		Column.SetIndex(0);
+
+		Column.SetShortName(LOCTEXT("LlmTagColumnName", "LLM Tag"));
+		Column.SetTitleName(LOCTEXT("LlmTagColumnTitle", "LLM Tag"));
+		Column.SetDescription(LOCTEXT("LlmTagColumnDesc", "LLM tag of allocation"));
+
+		Column.SetFlags(ETableColumnFlags::ShouldBeVisible | ETableColumnFlags::CanBeHidden | ETableColumnFlags::CanBeFiltered);
+
+		Column.SetHorizontalAlignment(HAlign_Left);
+		Column.SetInitialWidth(100.0f);
+
+		Column.SetDataType(ETableCellDataType::Int64);
+
+		class FMemAllocLlmTagValueGetter : public FTableCellValueGetter
+		{
+		public:
+			virtual const TOptional<FTableCellValue> GetValue(const FTableColumn& Column, const FBaseTreeNode& Node) const
+			{
+				//if (Node->Is<FMemAllocNode>())
+				{
+					const FMemAllocNode& MemAllocNode = static_cast<const FMemAllocNode&>(Node);
+					const FMemoryAlloc* Alloc = MemAllocNode.GetMemAlloc();
+					const FMemoryTagId AllocMemTagId = Alloc ? Alloc->GetMemTag() : 0;
+					return TOptional<FTableCellValue>(FTableCellValue(static_cast<int64>(AllocMemTagId)));
+				}
+				//else
+				//{
+				//	return TOptional<FTableCellValue>(FText::GetEmpty());
+				//}
+			}
+		};
+		TSharedRef<ITableCellValueGetter> Getter = MakeShared<FMemAllocLlmTagValueGetter>();
+		Column.SetValueGetter(Getter);
+
+		class FMemAllocLlmTagValueFormatter : public FTableCellValueFormatter
+		{
+		public:
+			virtual FText FormatValue(const TOptional<FTableCellValue>& InValue) const override
+			{
+				if (InValue.IsSet())
+				{
+					const FMemoryTagId MemTagId = static_cast<FMemoryTagId>(InValue.GetValue().Int64);
+					const FString& TagName = FMemAllocNode::GetAllocMemTagAsString(MemTagId);
+					if (TagName.IsEmpty())
+					{
+						//FText::AsNumber();
+					}
+					return FText::FromString(TagName);
+				}
+				return FText::GetEmpty();
+			}
+			virtual FText FormatValueForTooltip(const TOptional<FTableCellValue>& InValue) const override
+			{
+				if (InValue.IsSet())
+				{
+					const FMemoryTagId MemTagId = static_cast<FMemoryTagId>(InValue.GetValue().Int64);
+					const FString& TagName = FMemAllocNode::GetAllocMemTagAsString(MemTagId);
+					return FText::FromString(FString::Printf(TEXT("%lli (%s)"), MemTagId, *TagName));
+				}
+				return FText::GetEmpty();
+			}
+			virtual FText FormatValue(const FTableColumn& Column, const FBaseTreeNode& Node) const override { return FormatValue(Column.GetValue(Node)); }
+			virtual FText FormatValueForTooltip(const FTableColumn& Column, const FBaseTreeNode& Node) const override { return FormatValueForTooltip(Column.GetValue(Node)); }
+		};
+		TSharedRef<ITableCellValueFormatter> Formatter = MakeShared<FMemAllocLlmTagValueFormatter>();
+		Column.SetValueFormatter(Formatter);
+
+		TSharedRef<ITableCellValueSorter> Sorter = MakeShared<FSorterByInt64Value>(ColumnRef);
+		Column.SetValueSorter(Sorter);
+
+		AddColumn(ColumnRef);
+	}
+	//////////////////////////////////////////////////
+	// Backtrace Column
+	{
+		TSharedRef<FTableColumn> ColumnRef = MakeShared<FTableColumn>(TEXT("AllocBacktrace"));
+		FTableColumn& Column = *ColumnRef;
+
+		Column.SetIndex(0);
+
+		Column.SetShortName(LOCTEXT("BacktraceColumnName", "Backtrace"));
+		Column.SetTitleName(LOCTEXT("BacktraceColumnTitle", "Backtrace"));
+		Column.SetDescription(LOCTEXT("BacktraceColumnDesc", "Backtrace id of allocation"));
+
+		Column.SetFlags(ETableColumnFlags::ShouldBeVisible | ETableColumnFlags::CanBeHidden | ETableColumnFlags::CanBeFiltered);
+
+		Column.SetHorizontalAlignment(HAlign_Left);
+		Column.SetInitialWidth(120.0f);
+
+		Column.SetDataType(ETableCellDataType::Int64);
+
+		class FMemAllocSizeValueGetter : public FTableCellValueGetter
+		{
+		public:
+			virtual const TOptional<FTableCellValue> GetValue(const FTableColumn& Column, const FBaseTreeNode& Node) const override
+			{
+				const FMemAllocNode& MemAllocNode = static_cast<const FMemAllocNode&>(Node);
+				const FMemoryAlloc* Alloc = MemAllocNode.GetMemAlloc();
+				const uint64 AllocBacktraceId = Alloc ? Alloc->GetBacktraceId() : 0;
+				return TOptional<FTableCellValue>(FTableCellValue(static_cast<int64>(AllocBacktraceId)));
+			}
+		};
+		TSharedRef<ITableCellValueGetter> Getter = MakeShared<FMemAllocSizeValueGetter>();
+		Column.SetValueGetter(Getter);
+
+		TSharedRef<ITableCellValueFormatter> Formatter = MakeShared<FInt64ValueFormatterAsHex64>();
+		Column.SetValueFormatter(Formatter);
+
+		TSharedRef<ITableCellValueSorter> Sorter = MakeShared<FSorterByInt64Value>(ColumnRef);
+		Column.SetValueSorter(Sorter);
+
+		AddColumn(ColumnRef);
+	}
+	//////////////////////////////////////////////////
+}
+
+////////////////////////////////////////////////////////////////////////////////////////////////////
+
+} // namespace Insights
+
+#undef LOCTEXT_NAMESPACE
diff --git a/Engine/Source/Developer/TraceInsights/Private/Insights/MemoryProfiler/ViewModels/MemAllocTable.h b/Engine/Source/Developer/TraceInsights/Private/Insights/MemoryProfiler/ViewModels/MemAllocTable.h
new file mode 100644
--- /dev/null
+++ b/Engine/Source/Developer/TraceInsights/Private/Insights/MemoryProfiler/ViewModels/MemAllocTable.h
@@ -0,0 +1,42 @@
+// Copyright Epic Games, Inc. All Rights Reserved.
+
+#pragma once
+
+#include "CoreMinimal.h"
+
+// Insights
+#include "Insights/MemoryProfiler/ViewModels/MemoryAlloc.h"
+#include "Insights/Table/ViewModels/Table.h"
+
+namespace Insights
+{
+
+class FTableColumn;
+
+////////////////////////////////////////////////////////////////////////////////////////////////////
+
+class FMemAllocTable : public FTable
+{
+public:
+	FMemAllocTable();
+	virtual ~FMemAllocTable();
+
+	virtual void Reset();
+
+	TArray<FMemoryAlloc>& GetAllocs() { return Allocs; }
+	const TArray<FMemoryAlloc>& GetAllocs() const { return Allocs; }
+
+	bool IsValidRowIndex(int32 InIndex) const { return InIndex >= 0 && InIndex < Allocs.Num(); }
+	const FMemoryAlloc* GetMemAlloc(int32 InIndex) const { return IsValidRowIndex(InIndex) ? &Allocs[InIndex] : nullptr; }
+	const FMemoryAlloc& GetMemAllocChecked(int32 InIndex) const { check(IsValidRowIndex(InIndex)); return Allocs[InIndex]; }
+
+private:
+	void AddDefaultColumns();
+
+private:
+	TArray<FMemoryAlloc> Allocs;
+};
+
+////////////////////////////////////////////////////////////////////////////////////////////////////
+
+} // namespace Insights
diff --git a/Engine/Source/Developer/TraceInsights/Private/Insights/MemoryProfiler/ViewModels/MemTagNode.h b/Engine/Source/Developer/TraceInsights/Private/Insights/MemoryProfiler/ViewModels/MemTagNode.h
--- a/Engine/Source/Developer/TraceInsights/Private/Insights/MemoryProfiler/ViewModels/MemTagNode.h
+++ b/Engine/Source/Developer/TraceInsights/Private/Insights/MemoryProfiler/ViewModels/MemTagNode.h
@@ -108,10 +108,7 @@ public:
 	const Trace::FMemoryProfilerAggregatedStats& GetAggregatedStats() const { return AggregatedStats; }
 
 	void ResetAggregatedStats();
-	void SetAggregatedStats(const Trace::FMemoryProfilerAggregatedStats& AggregatedStats);
-
-private:
-	void UpdateFullStatName();
+	//TODO: void SetAggregatedStats(const Trace::FMemoryProfilerAggregatedStats& AggregatedStats);
 
 private:
 	const EMemTagNodeType Type;
diff --git a/Engine/Source/Developer/TraceInsights/Private/Insights/MemoryProfiler/ViewModels/MemoryAlloc.cpp b/Engine/Source/Developer/TraceInsights/Private/Insights/MemoryProfiler/ViewModels/MemoryAlloc.cpp
new file mode 100644
--- /dev/null
+++ b/Engine/Source/Developer/TraceInsights/Private/Insights/MemoryProfiler/ViewModels/MemoryAlloc.cpp
@@ -0,0 +1,28 @@
+// Copyright Epic Games, Inc. All Rights Reserved.
+
+#include "MemoryAlloc.h"
+
+namespace Insights
+{
+
+////////////////////////////////////////////////////////////////////////////////////////////////////
+
+FMemoryAlloc::FMemoryAlloc(double InStartTime, double InEndTime, uint64 InAddress, uint64 InSize, FMemoryTagId InMemTag, uint64 InBacktraceId)
+: StartTime(InStartTime)
+, EndTime(InEndTime)
+, Address(InAddress)
+, Size(InSize)
+, MemTag(InMemTag)
+, BacktraceId(InBacktraceId)
+{
+}
+
+////////////////////////////////////////////////////////////////////////////////////////////////////
+
+FMemoryAlloc::~FMemoryAlloc()
+{
+}
+
+////////////////////////////////////////////////////////////////////////////////////////////////////
+
+} // namespace Insights
diff --git a/Engine/Source/Developer/TraceInsights/Private/Insights/MemoryProfiler/ViewModels/MemoryAlloc.h b/Engine/Source/Developer/TraceInsights/Private/Insights/MemoryProfiler/ViewModels/MemoryAlloc.h
new file mode 100644
--- /dev/null
+++ b/Engine/Source/Developer/TraceInsights/Private/Insights/MemoryProfiler/ViewModels/MemoryAlloc.h
@@ -0,0 +1,41 @@
+// Copyright Epic Games, Inc. All Rights Reserved.
+
+#pragma once
+
+#include "CoreMinimal.h"
+
+#include "Insights/MemoryProfiler/ViewModels/MemoryTag.h"
+
+namespace Insights
+{
+
+////////////////////////////////////////////////////////////////////////////////////////////////////
+
+class FMemoryAlloc
+{
+	friend class SMemAllocTableTreeView;
+
+public:
+	FMemoryAlloc(double InStartTime, double InEndTime, uint64 InAddress, uint64 InSize, FMemoryTagId InMemTag, uint64 InBacktraceId);
+	~FMemoryAlloc();
+
+	double GetStartTime() const { return StartTime; }
+	double GetEndTime() const { return EndTime; }
+	double GetDuration() const { return EndTime - StartTime; }
+	uint64 GetAddress() const { return Address; }
+	uint64 GetSize() const { return Size; }
+	FMemoryTagId GetMemTag() const { return MemTag; }
+	uint64 GetBacktraceId() const { return BacktraceId; }
+
+private:
+	double StartTime;
+	double EndTime;
+	uint64 Address;
+	uint64 Size;
+	FMemoryTagId MemTag;
+	uint64 BacktraceId;
+};
+
+////////////////////////////////////////////////////////////////////////////////////////////////////
+
+} // namespace Insights
diff --git a/Engine/Source/Developer/TraceInsights/Private/Insights/MemoryProfiler/ViewModels/MemorySharedState.cpp b/Engine/Source/Developer/TraceInsights/Private/Insights/MemoryProfiler/ViewModels/MemorySharedState.cpp
--- a/Engine/Source/Developer/TraceInsights/Private/Insights/MemoryProfiler/ViewModels/MemorySharedState.cpp
+++ b/Engine/Source/Developer/TraceInsights/Private/Insights/MemoryProfiler/ViewModels/MemorySharedState.cpp
@@ -39,7 +39,9 @@ FMemorySharedState::FMemorySharedState()
 	, TrackHeightMode(EMemoryTrackHeightMode::Medium)
 	, bShowHideAllMemoryTracks(false)
 	, CreatedDefaultTracks()
+	, CurrentMemoryRule(nullptr)
 {
+	InitMemoryRules();
 }
 
 ////////////////////////////////////////////////////////////////////////////////////////////////////
@@ -63,12 +65,14 @@ void FMemorySharedState::OnBeginSession(Insights::ITimingViewSession& InSession)
 	DefaultTracker = nullptr;
 	CurrentTracker = nullptr;
 
-	MainGraphTrack.Reset();
+	MainGraphTrack = nullptr;
 	AllTracks.Reset();
 
 	bShowHideAllMemoryTracks = true;
 
 	CreatedDefaultTracks.Reset();
+
+	CurrentMemoryRule = MemoryRules[0];
 }
 
 ////////////////////////////////////////////////////////////////////////////////////////////////////
@@ -86,12 +90,14 @@ void FMemorySharedState::OnEndSession(Insights::ITimingViewSession& InSession)
 	DefaultTracker = nullptr;
 	CurrentTracker = nullptr;
 
-	MainGraphTrack.Reset();
+	MainGraphTrack = nullptr;
 	AllTracks.Reset();
 
 	bShowHideAllMemoryTracks = false;
 
 	CreatedDefaultTracks.Reset();
+
+	CurrentMemoryRule = nullptr;
 }
 
 ////////////////////////////////////////////////////////////////////////////////////////////////////
@@ -870,4 +876,103 @@ TSharedPtr<FMemoryGraphTrack> FMemorySharedState::CreateGraphTrack(const Insight
 
 ////////////////////////////////////////////////////////////////////////////////////////////////////
 
+void FMemorySharedState::InitMemoryRules()
+{
+	MemoryRules.Reset();
+
+	MemoryRules.Add(MakeShared<Insights::FMemoryRuleSpec>(
+		Insights::EMemoryRule::aAf, 1,
+		LOCTEXT("MemRule_aAf_Short", "*A*"),
+		LOCTEXT("MemRule_aAf_Verbose", "Active Allocs"),
+		LOCTEXT("MemRule_aAf_Desc", "Identifies active allocations at time A.\n(a ≤ A ≤ f)")));
+
+	MemoryRules.Add(MakeShared<Insights::FMemoryRuleSpec>(
+		Insights::EMemoryRule::afA, 1,
+		LOCTEXT("MemRule_afA_Short", "**A"),
+		LOCTEXT("MemRule_afA_Verbose", "Before"),
+		LOCTEXT("MemRule_afA_Desc", "Identifies allocations allocated and freed before time A.\n(a ≤ f ≤ A)")));
+
+	MemoryRules.Add(MakeShared<Insights::FMemoryRuleSpec>(
+		Insights::EMemoryRule::Aaf, 1,
+		LOCTEXT("MemRule_Aaf_Short", "A**"),
+		LOCTEXT("MemRule_Aaf_Verbose", "After"),
+		LOCTEXT("MemRule_Aaf_Desc", "Identifies allocations allocated and freed after time A.\n(A ≤ a ≤ f)")));
+
+	MemoryRules.Add(MakeShared<Insights::FMemoryRuleSpec>(
+		Insights::EMemoryRule::aAfB, 2,
+		LOCTEXT("MemRule_aAfB_Short", "*A*B"),
+		LOCTEXT("MemRule_aAfB_Verbose", "Decline"),
+		LOCTEXT("MemRule_aAfB_Desc", "Identifies allocations allocated before time A and freed between time A and time B.\n(a ≤ A ≤ f ≤ B)")));
+
+	MemoryRules.Add(MakeShared<Insights::FMemoryRuleSpec>(
+		Insights::EMemoryRule::AaBf, 2,
+		LOCTEXT("MemRule_AaBf_Short", "A*B*"),
+		LOCTEXT("MemRule_AaBf_Verbose", "Growth"),
+		LOCTEXT("MemRule_AaBf_Desc", "Identifies allocations allocated between time A and time B and freed after time B.\n(A ≤ a ≤ B ≤ f)")));
+
+	MemoryRules.Add(MakeShared<Insights::FMemoryRuleSpec>(
+		Insights::EMemoryRule::AafB, 2,
+		LOCTEXT("MemRule_AafB_Short", "A**B"),
+		LOCTEXT("MemRule_AafB_Verbose", "Short Living Allocs"),
+		LOCTEXT("MemRule_AafB_Desc", "Identifies allocations allocated and freed between time A and time B.\n(A ≤ a ≤ f ≤ B)")));
+
+	MemoryRules.Add(MakeShared<Insights::FMemoryRuleSpec>(
+		Insights::EMemoryRule::aABf, 2,
+		LOCTEXT("MemRule_aABf_Short", "*A B*"),
+		LOCTEXT("MemRule_aABf_Verbose", "Long Living Allocs"),
+		LOCTEXT("MemRule_aABf_Desc", "Identifies allocations allocated before time A and freed after time B.\n(a ≤ A ≤ B ≤ f)")));
+
+	MemoryRules.Add(MakeShared<Insights::FMemoryRuleSpec>(
+		Insights::EMemoryRule::AaBCf, 3,
+		LOCTEXT("MemRule_AaBCf_Short", "A*B C*"),
+		LOCTEXT("MemRule_AaBCf_Verbose", "Memory Leaks"),
+		LOCTEXT("MemRule_AaBCf_Desc", "Identifies allocations allocated between time A and time B and freed after time C.\n(A ≤ a ≤ B ≤ C ≤ f)")));
+
+	MemoryRules.Add(MakeShared<Insights::FMemoryRuleSpec>(
+		Insights::EMemoryRule::AaBfC, 3,
+		LOCTEXT("MemRule_AaBfC_Short", "A*B*C"),
+		LOCTEXT("MemRule_AaBfC_Verbose", "Limited Lifetime"),
+		LOCTEXT("MemRule_AaBfC_Desc", "Identifies allocations allocated between time A and time B and freed between time B and time C.\n(A ≤ a ≤ B ≤ f ≤ C)")));
+
+	MemoryRules.Add(MakeShared<Insights::FMemoryRuleSpec>(
+		Insights::EMemoryRule::aABfC, 3,
+		LOCTEXT("MemRule_aABfC_Short", "*A B*C"),
+		LOCTEXT("MemRule_aABfC_Verbose", "Decline of Long Living Allocs"),
+		LOCTEXT("MemRule_aABfC_Desc", "Identifies allocations allocated before time A and freed between time B and time C.\n(a ≤ A ≤ B ≤ f ≤ C)")));
+
+	MemoryRules.Add(MakeShared<Insights::FMemoryRuleSpec>(
+		Insights::EMemoryRule::AaBCfD, 4,
+		LOCTEXT("MemRule_AaBCfD_Short", "A*B C*D"),
+		LOCTEXT("MemRule_AaBCfD_Verbose", "Specific Lifetime"),
+		LOCTEXT("MemRule_AaBCfD_Desc", "Identifies allocations allocated between time A and time B and freed between time C and time D.\n(A ≤ a ≤ B ≤ C ≤ f ≤ D)")));
+
+	MemoryRules.Add(MakeShared<Insights::FMemoryRuleSpec>(
+		Insights::EMemoryRule::A_vs_B, 2,
+		LOCTEXT("MemRule_A_vs_B_Short", "*A* + *B*"),
+		LOCTEXT("MemRule_A_vs_B_Verbose", "Compare A vs. B"),
+		LOCTEXT("MemRule_A_vs_B_Desc", "Compares live allocations at time A with live allocations at time B.\n(*A* vs. *B*)")));
+
+	MemoryRules.Add(MakeShared<Insights::FMemoryRuleSpec>(
+		Insights::EMemoryRule::A_or_B, 2,
+		LOCTEXT("MemRule_A_or_B_Short", "*A* | *B*"),
+		LOCTEXT("MemRule_A_or_B_Verbose", "A or B"),
+		LOCTEXT("MemRule_A_or_B_Desc", "Identifies allocations live at time A or at time B.\n(a ≤ A ≤ f OR a ≤ B ≤ f)\n{*A*} U {*B*}")));
+
+	MemoryRules.Add(MakeShared<Insights::FMemoryRuleSpec>(
+		Insights::EMemoryRule::A_xor_B, 2,
+		LOCTEXT("MemRule_A_xor_B_Short", "*A* ^ *B*"),
+		LOCTEXT("MemRule_A_xor_B_Verbose", "A xor B"),
+		LOCTEXT("MemRule_A_xor_B_Desc", "Identifies allocations live either at time A or at time B (but not both).\n(a ≤ A ≤ f XOR a ≤ B ≤ f)\n({*A*} U {*B*}) \\ {*AB*}")));
+
+	CurrentMemoryRule = MemoryRules[0];
+}
+
+////////////////////////////////////////////////////////////////////////////////////////////////////
+
+void FMemorySharedState::OnMemoryRuleChanged()
+{
+}
+
+////////////////////////////////////////////////////////////////////////////////////////////////////
+
 #undef LOCTEXT_NAMESPACE
diff --git a/Engine/Source/Developer/TraceInsights/Private/Insights/MemoryProfiler/ViewModels/MemorySharedState.h b/Engine/Source/Developer/TraceInsights/Private/Insights/MemoryProfiler/ViewModels/MemorySharedState.h
--- a/Engine/Source/Developer/TraceInsights/Private/Insights/MemoryProfiler/ViewModels/MemorySharedState.h
+++ b/Engine/Source/Developer/TraceInsights/Private/Insights/MemoryProfiler/ViewModels/MemorySharedState.h
@@ -20,7 +20,57 @@ namespace Insights
 	struct FReportTypeConfig;
 	struct FReportTypeGraphConfig;
 	class FMemoryTracker;
-}
+
+// Known memory rules.
+// The enum uses the following naming convention:
+//     A, B, C, D = time markers
+//     a = time when "alloc" event occurs
+//     f = time when "free" event occurs (can be infinite)
+// Ex.: "AaBf" means "all memory allocations allocated between time A and time B and freed after time B".
+enum class EMemoryRule
+{
+	aAf,     // active allocs at A
+	afA,     // before
+	Aaf,     // after
+	aAfB,    // decline
+	AaBf,    // growth
+	AafB,    // short living allocs
+	aABf,    // long living allocs
+	AaBCf,   // memory leaks
+	AaBfC,   // limited lifetime
+	aABfC,   // decline of long living allocs
+	AaBCfD,  // specific lifetime
+	A_vs_B,  // compare A vs. B; {aAf} vs. {aBf}
+	A_or_B,  // live at A or at B; {aAf} U {aBf}
+	A_xor_B, // live either at A or at B; ({aAf} U {aBf}) \ {aABf}
+};
+
+class FMemoryRuleSpec
+{
+public:
+	FMemoryRuleSpec(EMemoryRule InValue, uint32 InNumTimeMarkers, const FText& InShortName, const FText& InVerboseName, const FText& InDescription)
+		: Value(InValue)
+		, NumTimeMarkers(InNumTimeMarkers)
+		, ShortName(InShortName)
+		, VerboseName(InVerboseName)
+		, Description(InDescription)
+	{}
+
+	EMemoryRule GetValue() const { return Value; }
+	uint32 GetNumTimeMarkers() const { return NumTimeMarkers; }
+	FText GetShortName() const { return ShortName; }
+	FText GetVerboseName() const { return VerboseName; }
+	FText GetDescription() const { return Description; }
+
+private:
+	EMemoryRule Value;     // ex.: EMemoryRule::AafB
+	uint32 NumTimeMarkers; // ex.: 2
+	FText ShortName;     // ex.: "A**B"
+	FText VerboseName;   // ex.: "Short Living Allocations"
+	FText Description;   // ex.: "Allocations allocated and freed between time A and time B (A <= a <= f <= B)."
+};
+
+} // namespace Insights
 
 ////////////////////////////////////////////////////////////////////////////////////////////////////
 
@@ -75,12 +125,19 @@ public:
 	void CreateTracksFromReport(const Insights::FReportConfig& ReportConfig);
 	void CreateTracksFromReport(const Insights::FReportTypeConfig& ReportTypeConfig);
 
+	const TArray<TSharedPtr<Insights::FMemoryRuleSpec>>& GetMemoryRules() const { return MemoryRules; }
+	
+	TSharedPtr<Insights::FMemoryRuleSpec> GetCurrentMemoryRule() const { return CurrentMemoryRule; }
+	void SetCurrentMemoryRule(TSharedPtr<Insights::FMemoryRuleSpec> InRule) { CurrentMemoryRule = InRule; OnMemoryRuleChanged(); }
+
 private:
 	void SyncTrackers();
 	void OnTrackerChanged();
 	void SetTrackerIdToAllSeries(TSharedPtr<FMemoryGraphTrack>& GraphTrack, Insights::FMemoryTrackerId TrackerId);
 	int32 GetNextMemoryGraphTrackOrder();
 	TSharedPtr<FMemoryGraphTrack> CreateGraphTrack(const Insights::FReportTypeGraphConfig& ReportTypeGraphConfig);
+	void InitMemoryRules();
+	void OnMemoryRuleChanged();
 
 private:
 	TSharedPtr<STimingView> TimingView;
@@ -99,6 +156,9 @@ private:
 	bool bShowHideAllMemoryTracks;
 
 	TBitArray<> CreatedDefaultTracks;
+
+	TArray<TSharedPtr<Insights::FMemoryRuleSpec>> MemoryRules;
+	TSharedPtr<Insights::FMemoryRuleSpec> CurrentMemoryRule;
 };
 
 ////////////////////////////////////////////////////////////////////////////////////////////////////
diff --git a/Engine/Source/Developer/TraceInsights/Private/Insights/MemoryProfiler/Widgets/SMemAllocTableTreeView.cpp b/Engine/Source/Developer/TraceInsights/Private/Insights/MemoryProfiler/Widgets/SMemAllocTableTreeView.cpp
new file mode 100644
--- /dev/null
+++ b/Engine/Source/Developer/TraceInsights/Private/Insights/MemoryProfiler/Widgets/SMemAllocTableTreeView.cpp
@@ -0,0 +1,412 @@
+// Copyright Epic Games, Inc. All Rights Reserved.
+
+#include "SMemAllocTableTreeView.h"
+
+#include "TraceServices/AnalysisService.h"
+#include "TraceServices/Model/AllocationsProvider.h"
+
+// Insights
+#include "Insights/Common/Stopwatch.h"
+#include "Insights/MemoryProfiler/MemoryProfilerManager.h"
+#include "Insights/MemoryProfiler/ViewModels/MemAllocNode.h"
+#include "Insights/MemoryProfiler/ViewModels/MemAllocTable.h"
+#include "Insights/MemoryProfiler/ViewModels/MemorySharedState.h"
+#include "Insights/TimingProfilerCommon.h"
+
+#include <limits>
+
+#define LOCTEXT_NAMESPACE "SMemAllocTableTreeView"
+
+namespace Insights
+{
+
+////////////////////////////////////////////////////////////////////////////////////////////////////
+
+SMemAllocTableTreeView::SMemAllocTableTreeView()
+{
+}
+
+////////////////////////////////////////////////////////////////////////////////////////////////////
+
+SMemAllocTableTreeView::~SMemAllocTableTreeView()
+{
+}
+
+////////////////////////////////////////////////////////////////////////////////////////////////////
+
+void SMemAllocTableTreeView::Construct(const FArguments& InArgs, TSharedPtr<Insights::FMemAllocTable> InTablePtr)
+{
+	ConstructWidget(InTablePtr);
+}
+
+////////////////////////////////////////////////////////////////////////////////////////////////////
+
+void SMemAllocTableTreeView::Reset()
+{
+	//...
+	STableTreeView::Reset();
+}
+
+////////////////////////////////////////////////////////////////////////////////////////////////////
+/*
+void SMemAllocTableTreeView::UpdateSourceTable(TSharedPtr<Trace::IMemAllocTable> SourceTable)
+{
+	//check(Table->Is<Insights::FMemAllocTable>());
+	TSharedPtr<Insights::FMemAllocTable> MemAllocTable = StaticCastSharedPtr<Insights::FMemAllocTable>(Table);
+
+	if (MemAllocTable->UpdateSourceTable(SourceTable))
+	{
+		RebuildColumns();
+	}
+
+	RebuildTree(true);
+}
+*/
+////////////////////////////////////////////////////////////////////////////////////////////////////
+
+void SMemAllocTableTreeView::Tick(const FGeometry& AllottedGeometry, const double InCurrentTime, const float InDeltaTime)
+{
+	// We need to check if the list of LLM tags has changed.
+	// But, ensure we do not check too often.
+	static uint64 NextTimestamp = 0;
+	uint64 Time = FPlatformTime::Cycles64();
+	if (Time > NextTimestamp)
+	{
+		RebuildTree(false);
+
+		const double WaitTimeSec = 0.5;
+		const uint64 WaitTime = static_cast<uint64>(WaitTimeSec / FPlatformTime::GetSecondsPerCycle64());
+		NextTimestamp = Time + WaitTime;
+	}
+}
+
+////////////////////////////////////////////////////////////////////////////////////////////////////
+
+void SMemAllocTableTreeView::RebuildTree(bool bResync)
+{
+	FStopwatch SyncStopwatch;
+	FStopwatch Stopwatch;
+	Stopwatch.Start();
+
+	if (bResync)
+	{
+		TableTreeNodes.Empty();
+	}
+
+	const int32 PreviousNodeCount = TableTreeNodes.Num();
+
+	//check(Table->Is<Insights::FMemAllocTable>());
+	//TSharedPtr<Insights::FMemAllocTable> MemAllocTable = StaticCastSharedPtr<Insights::FMemAllocTable>(Table);
+
+	TSharedPtr<Insights::FMemAllocTable> MemAllocTable = GetMemAllocTable();
+
+	SyncStopwatch.Start();
+	if (Session.IsValid() && MemAllocTable.IsValid())
+	{
+		UpdateQuery();
+
+		TArray<FMemoryAlloc>&  Allocs = MemAllocTable->GetAllocs();
+
+		const int32 TotalAllocCount = Allocs.Num();
+		if (TotalAllocCount != TableTreeNodes.Num())
+		{
+			if (TableTreeNodes.Num() > TotalAllocCount)
+			{
+				TableTreeNodes.Empty();
+			}
+			TableTreeNodes.Reserve(TotalAllocCount);
+
+			FName BaseNodeName(TEXT("alloc"));
+			for (int32 AllocIndex = TableTreeNodes.Num(); AllocIndex < TotalAllocCount; ++AllocIndex)
+			{
+				FName NodeName(BaseNodeName, AllocIndex + 1);
+				FMemAllocNodePtr NodePtr = MakeShared<FMemAllocNode>(NodeName, MemAllocTable, AllocIndex);
+				TableTreeNodes.Add(NodePtr);
+			}
+			ensure(TableTreeNodes.Num() == TotalAllocCount);
+		}
+	}
+	SyncStopwatch.Stop();
+
+	if (bResync || TableTreeNodes.Num() != PreviousNodeCount)
+	{
+		// Save selection.
+		TArray<FTableTreeNodePtr> SelectedItems;
+		TreeView->GetSelectedItems(SelectedItems);
+
+		UpdateTree();
+
+		TreeView->RebuildList();
+
+		// Restore selection.
+		if (SelectedItems.Num() > 0)
+		{
+			TreeView->ClearSelection();
+			for (FTableTreeNodePtr& NodePtr : SelectedItems)
+			{
+				NodePtr = GetNodeByTableRowIndex(NodePtr->GetRowIndex());
+			}
+			SelectedItems.RemoveAll([](const FTableTreeNodePtr& NodePtr) { return !NodePtr.IsValid(); });
+			if (SelectedItems.Num() > 0)
+			{
+				TreeView->SetItemSelection(SelectedItems, true);
+				TreeView->RequestScrollIntoView(SelectedItems.Last());
+			}
+		}
+	}
+
+	Stopwatch.Stop();
+	const double TotalTime = Stopwatch.GetAccumulatedTime();
+	if (TotalTime > 0.01)
+	{
+		const double SyncTime = SyncStopwatch.GetAccumulatedTime();
+		UE_LOG(MemoryProfiler, Log, TEXT("[MemAlloc] Tree view rebuilt in %.3fs (%.3fs + %.3fs) --> %d allocs (%d added)"),
+			TotalTime, SyncTime, TotalTime - SyncTime, TableTreeNodes.Num(), TableTreeNodes.Num() - PreviousNodeCount);
+	}
+}
+
+////////////////////////////////////////////////////////////////////////////////////////////////////
+
+void SMemAllocTableTreeView::OnQueryInvalidated()
+{
+	CancelQuery();
+
+	TableTreeNodes.Reset();
+
+	TSharedPtr<Insights::FMemAllocTable> MemAllocTable = GetMemAllocTable();
+	if (MemAllocTable)
+	{
+		TArray<FMemoryAlloc>& Allocs = MemAllocTable->GetAllocs();
+		Allocs.Reset(10 * 1024 * 1024);
+	}
+
+	StartQuery();
+
+	RebuildTree(true);
+}
+
+////////////////////////////////////////////////////////////////////////////////////////////////////
+
+void SMemAllocTableTreeView::StartQuery()
+{
+#if defined(UE_USE_ALLOCATIONS_PROVIDER)
+	check(Query == 0);
+
+	if (!Rule)
+	{
+		UE_LOG(MemoryProfiler, Warning, TEXT("[MemAlloc] Invalid query rule!"));
+		return;
+	}
+
+	const IAllocationsProvider* AllocationsProvider = Trace::ReadAllocationsProvider(*Session.Get());
+	if (!AllocationsProvider)
+	{
+		UE_LOG(MemoryProfiler, Warning, TEXT("[MemAlloc] Invalid allocations provider!"));
+		return;
+	}
+	const IAllocationsProvider& Provider = *AllocationsProvider;
+
+	//TODO: Simpler API...
+	//QueryHandle StartQuery(IAllocationsProvider::EQueryRule Rule, double TimeA, double TimeB, double TimeC, double TimeD);
+
+	switch (Rule->GetValue())
+	{
+		case Insights::EMemoryRule::aAf:
+		{
+			Query = Provider.StartQuery(TimeMarkers[0], TimeMarkers[0], IAllocationsProvider::ECrosses::Or);
+			break;
+		}
+		case Insights::EMemoryRule::afA:
+		{
+			Query = Provider.StartQuery(0.0, TimeMarkers[0], IAllocationsProvider::ECrosses::None);
+			break;
+		}
+		case Insights::EMemoryRule::Aaf:
+		{
+			Query = Provider.StartQuery(TimeMarkers[0], std::numeric_limits<double>::infinity(), IAllocationsProvider::ECrosses::None);
+			break;
+		}
+		//case Insights::EMemoryRule::aAfB:
+		//{
+		//	Query = Provider.StartQuery(TimeMarkers[0], TimeMarkers[1], IAllocationsProvider::ECrosses::A);
+		//	break;
+		//}
+		//case Insights::EMemoryRule::AaBf:
+		//{
+		//	Query = Provider.StartQuery(TimeMarkers[0], TimeMarkers[1], IAllocationsProvider::ECrosses::B);
+		//	break;
+		//}
+		case Insights::EMemoryRule::AafB:
+		{
+			Query = Provider.StartQuery(TimeMarkers[0], TimeMarkers[1], IAllocationsProvider::ECrosses::None);
+			break;
+		}
+		case Insights::EMemoryRule::aABf:
+		{
+			Query = Provider.StartQuery(TimeMarkers[0], TimeMarkers[1], IAllocationsProvider::ECrosses::And);
+			break;
+		}
+		case Insights::EMemoryRule::AaBCf:
+		{
+			Query = Provider.StartQuery(TimeMarkers[1], TimeMarkers[2], IAllocationsProvider::ECrosses::And);
+			// needs additional filter Alloc.Start > A
+			break;
+		}
+		//case Insights::EMemoryRule::AaBfC:
+		//	TODO
+		//	break;
+		case Insights::EMemoryRule::aABfC:
+		{
+			Query = Provider.StartQuery(TimeMarkers[0], TimeMarkers[1], IAllocationsProvider::ECrosses::And);
+			// needs additional filter Alloc.End < C
+			break;
+		}
+		//case Insights::EMemoryRule::AaBCfD:
+		//	TODO
+		//	break;
+		case Insights::EMemoryRule::A_vs_B:
+		{
+			Query = Provider.StartQuery(TimeMarkers[0], TimeMarkers[1], IAllocationsProvider::ECrosses::Or);
+			break;
+		}
+		case Insights::EMemoryRule::A_or_B:
+		{
+			Query = Provider.StartQuery(TimeMarkers[0], TimeMarkers[1], IAllocationsProvider::ECrosses::Or);
+			break;
+		}
+		case Insights::EMemoryRule::A_xor_B:
+		{
+			Query = Provider.StartQuery(TimeMarkers[0], TimeMarkers[1], IAllocationsProvider::ECrosses::Xor);
+			break;
+		}
+	}
+
+	if (Query == 0)
+	{
+		UE_LOG(MemoryProfiler, Error, TEXT("[MemAlloc] Unsupported query rule (%s)!"), *Rule->GetShortName().ToString());
+	}
+#endif // defined(UE_USE_ALLOCATIONS_PROVIDER)
+}
+
+////////////////////////////////////////////////////////////////////////////////////////////////////
+
+void SMemAllocTableTreeView::UpdateQuery()
+{
+#if defined(UE_USE_ALLOCATIONS_PROVIDER)
+	if (Query != 0)
+	{
+		const IAllocationsProvider* AllocationsProvider = Trace::ReadAllocationsProvider(*Session.Get());
+		if (!AllocationsProvider)
+		{
+			UE_LOG(MemoryProfiler, Warning, TEXT("[MemAlloc] Invalid allocations provider!"));
+			return;
+		}
+		const IAllocationsProvider& Provider = *AllocationsProvider;
+
+		IAllocationsProvider::FQueryStatus Status = Provider.PollQuery(Query);
+		if (Status.Status <= IAllocationsProvider::EQueryStatus::Done)
+		{
+			Query = 0;
+			return;
+		}
+
+		TSharedPtr<Insights::FMemAllocTable> MemAllocTable = GetMemAllocTable();
+		if (MemAllocTable)
+		{
+			TArray<FMemoryAlloc>& Allocs = MemAllocTable->GetAllocs();
+
+			FStopwatch TotalStopwatch;
+			FStopwatch PageStopwatch;
+			TotalStopwatch.Start();
+			uint32 PageCount = 0;
+			uint32 TotatAllocCount = 0;
+
+			// Multiple 'pages' of results will be returned. No guarantees are made
+			// about the order of pages or the allocations they report.
+			IAllocationsProvider::QueryResult Result = Status.NextResult();
+			while (Result.IsValid())
+			{
+				++PageCount;
+				PageStopwatch.Restart();
+
+				const uint32 AllocCount = Result->Num();
+				TotatAllocCount += AllocCount;
+#if 0
+				Allocs.Reserve(Allocs.Num() + AllocCount);
+
+				for (uint32 AllocIndex = 0; AllocIndex < AllocCount; ++AllocIndex)
+				{
+					const IAllocationsProvider::FAllocation* Allocation = Result->Get(AllocIndex);
+
+					const double AllocStartTime = Allocation->GetStartTime();
+					const double AllocEndTime = Allocation->GetEndTime();
+					const uint64 AllocAddress = Allocation->GetAddress();
+					const uint64 AllocationSize = Allocation->GetSize();
+					const FMemoryTagId AllocTag = static_cast<FMemoryTagId>(Allocation->GetTag());
+					const uint64 BacktraceId = Allocation->GetBacktraceId();
+
+					//Allocs.Add(FMemoryAlloc(AllocStartTime, AllocEndTime, AllocAddress, AllocationSize, AllocTag, BacktraceId));
+					Allocs.Emplace(AllocStartTime, AllocEndTime, AllocAddress, AllocationSize, AllocTag, BacktraceId);
+				}
+#else // 14% faster, but uses more mem
+				uint64 AllocsDestIndex = Allocs.Num();
+				Allocs.AddUninitialized(AllocCount);
+				for (uint32 AllocIndex = 0; AllocIndex < AllocCount; ++AllocIndex, ++AllocsDestIndex)
+				{
+					const IAllocationsProvider::FAllocation* Allocation = Result->Get(AllocIndex);
+					FMemoryAlloc& Alloc = Allocs[AllocsDestIndex];
+					Alloc.StartTime = Allocation->GetStartTime();
+					Alloc.EndTime = Allocation->GetEndTime();
+					Alloc.Address = Allocation->GetAddress();
+					Alloc.Size = Allocation->GetSize();
+					Alloc.MemTag = static_cast<FMemoryTagId>(Allocation->GetTag());
+					Alloc.BacktraceId = Allocation->GetBacktraceId();
+				}
+#endif
+
+				PageStopwatch.Stop();
+				const double PageTime = PageStopwatch.GetAccumulatedTime();
+				if (PageTime > 0.01)
+				{
+					const double Speed = (PageTime * 1000000.0) / AllocCount;
+					UE_LOG(MemoryProfiler, Log, TEXT("[MemAlloc] Query result for page %u (%u allocs, slack=%u) retrieved in %.3fs (speed: %.3f seconds per 1M allocs)."), PageCount, AllocCount, Allocs.GetSlack(), PageTime, Speed);
+				}
+
+				Result = Status.NextResult();
+			}
+
+			TotalStopwatch.Stop();
+			const double TotalTime = TotalStopwatch.GetAccumulatedTime();
+			if (TotalTime > 0.01)
+			{
+				const double Speed = (TotalTime * 1000000.0) / TotatAllocCount;
+				UE_LOG(MemoryProfiler, Log, TEXT("[MemAlloc] Query results (%u pages, %u allocs, slack=%u) retrieved in %.3fs (speed: %.3f seconds per 1M allocs)."), PageCount, TotatAllocCount, Allocs.GetSlack(), TotalTime, Speed);
+			}
+		}
+	}
+#endif // defined(UE_USE_ALLOCATIONS_PROVIDER)
+}
+
+////////////////////////////////////////////////////////////////////////////////////////////////////
+
+void SMemAllocTableTreeView::CancelQuery()
+{
+#if defined(UE_USE_ALLOCATIONS_PROVIDER)
+	if (Query != 0)
+	{
+		const IAllocationsProvider* AllocationsProvider = Trace::ReadAllocationsProvider(*Session.Get());
+		if (AllocationsProvider)
+		{
+			AllocationsProvider->CancelQuery(Query);
+			UE_LOG(MemoryProfiler, Log, TEXT("[MemAlloc] Query canceled."));
+		}
+		Query = 0;
+	}
+#endif // defined(UE_USE_ALLOCATIONS_PROVIDER)
+}
+
+////////////////////////////////////////////////////////////////////////////////////////////////////
+
+} // namespace Insights
+
+#undef LOCTEXT_NAMESPACE
diff --git a/Engine/Source/Developer/TraceInsights/Private/Insights/MemoryProfiler/Widgets/SMemAllocTableTreeView.h b/Engine/Source/Developer/TraceInsights/Private/Insights/MemoryProfiler/Widgets/SMemAllocTableTreeView.h
new file mode 100644
--- /dev/null
+++ b/Engine/Source/Developer/TraceInsights/Private/Insights/MemoryProfiler/Widgets/SMemAllocTableTreeView.h
@@ -0,0 +1,90 @@
+// Copyright Epic Games, Inc. All Rights Reserved.
+
+#pragma once
+
+#include "CoreMinimal.h"
+
+#include "TraceServices/Model/AllocationsProvider.h"
+
+// Insights
+#include "Insights/MemoryProfiler/ViewModels/MemAllocTable.h"
+#include "Insights/Table/Widgets/STableTreeView.h"
+
+namespace Insights
+{
+
+class FMemoryRuleSpec;
+
+////////////////////////////////////////////////////////////////////////////////////////////////////
+
+class SMemAllocTableTreeView : public STableTreeView
+{
+public:
+	/** Default constructor. */
+	SMemAllocTableTreeView();
+
+	/** Virtual destructor. */
+	virtual ~SMemAllocTableTreeView();
+
+	SLATE_BEGIN_ARGS(SMemAllocTableTreeView) {}
+	SLATE_END_ARGS()
+
+	/**
+	 * Construct this widget
+	 * @param InArgs - The declaration data for this widget
+	 */
+	void Construct(const FArguments& InArgs, TSharedPtr<FMemAllocTable> InTablePtr);
+
+	TSharedPtr<FMemAllocTable> GetMemAllocTable() { return StaticCastSharedPtr<FMemAllocTable>(GetTable()); }
+	const TSharedPtr<FMemAllocTable> GetMemAllocTable() const { return StaticCastSharedPtr<FMemAllocTable>(GetTable()); }
+
+	//void UpdateSourceTable(TSharedPtr<Trace::IMemAllocTable> SourceTable);
+
+	virtual void Reset();
+
+	/**
+	 * Ticks this widget.  Override in derived classes, but always call the parent implementation.
+	 *
+	 * @param  AllottedGeometry The space allotted for this widget
+	 * @param  InCurrentTime  Current absolute real time
+	 * @param  InDeltaTime  Real time passed since last tick
+	 */
+	virtual void Tick(const FGeometry& AllottedGeometry, const double InCurrentTime, const float InDeltaTime) override;
+
+	/**
+	 * Rebuilds the tree (if necessary).
+	 * @param bResync - If true, it forces a resync even if the list did not changed since last sync.
+	 */
+	virtual void RebuildTree(bool bResync);
+
+	void SetQueryParams(TSharedPtr<FMemoryRuleSpec> InRule, double TimeA = 0.0, double TimeB = 0.0, double TimeC = 0.0, double TimeD = 0.0)
+	{
+		Rule = InRule;
+		TimeMarkers[0] = TimeA;
+		TimeMarkers[1] = TimeB;
+		TimeMarkers[2] = TimeC;
+		TimeMarkers[3] = TimeD;
+		OnQueryInvalidated();
+	}
+
+	int32 GetTabIndex() const { return TabIndex; }
+	void SetTabIndex(int32 InTabIndex) { TabIndex = InTabIndex; }
+
+private:
+	void OnQueryInvalidated();
+	void StartQuery();
+	void UpdateQuery();
+	void CancelQuery();
+
+private:
+	int32 TabIndex = -1;
+	TSharedPtr<FMemoryRuleSpec> Rule = nullptr;
+	double TimeMarkers[4];
+#if defined(UE_USE_ALLOCATIONS_PROVIDER)
+	IAllocationsProvider::QueryHandle Query = 0;
+#endif
+};
+
+////////////////////////////////////////////////////////////////////////////////////////////////////
+
+} // namespace Insights
diff --git a/Engine/Source/Developer/TraceInsights/Private/Insights/MemoryProfiler/Widgets/SMemTagTreeView.cpp b/Engine/Source/Developer/TraceInsights/Private/Insights/MemoryProfiler/Widgets/SMemTagTreeView.cpp
--- a/Engine/Source/Developer/TraceInsights/Private/Insights/MemoryProfiler/Widgets/SMemTagTreeView.cpp
+++ b/Engine/Source/Developer/TraceInsights/Private/Insights/MemoryProfiler/Widgets/SMemTagTreeView.cpp
@@ -11,25 +11,30 @@
 #include "Widgets/Input/SCheckBox.h"
 #include "Widgets/Layout/SScrollBox.h"
 #include "Widgets/Input/SButton.h"
+#include "Widgets/Input/SEditableTextBox.h"
 #include "Widgets/Input/SSearchBox.h"
 #include "Widgets/Images/SImage.h"
+#include "Widgets/Layout/SExpandableArea.h"
 #include "Widgets/Layout/SGridPanel.h"
 #include "Widgets/Layout/SSeparator.h"
 #include "Widgets/SToolTip.h"
 #include "Widgets/Views/STableViewBase.h"
 
 // Insights
+#include "Insights/Common/TimeUtils.h"
 #include "Insights/InsightsStyle.h"
 #include "Insights/MemoryProfiler/MemoryProfilerManager.h"
 #include "Insights/MemoryProfiler/ViewModels/MemoryTag.h"
 #include "Insights/MemoryProfiler/ViewModels/MemoryTracker.h"
 #include "Insights/MemoryProfiler/ViewModels/MemorySharedState.h"
 #include "Insights/MemoryProfiler/ViewModels/MemTagTreeViewColumnFactory.h"
+#include "Insights/MemoryProfiler/Widgets/SMemAllocTableTreeView.h"
 #include "Insights/MemoryProfiler/Widgets/SMemTagTreeViewTooltip.h"
 #include "Insights/MemoryProfiler/Widgets/SMemTagTreeViewTableRow.h"
 #include "Insights/MemoryProfiler/Widgets/SMemoryProfilerWindow.h"
 #include "Insights/Table/ViewModels/Table.h"
 #include "Insights/Table/ViewModels/TableColumn.h"
+#include "Insights/ViewModels/TimeRulerTrack.h"
 #include "Insights/ViewModels/TimingGraphTrack.h"
 #include "Insights/Widgets/STimingView.h"
 
@@ -89,124 +94,31 @@ void SMemTagTreeView::Construct(const FArguments& InArgs, TSharedPtr<SMemoryProf
 			[
 				SNew(SVerticalBox)
 
+				// Investigation
 				+ SVerticalBox::Slot()
 				.VAlign(VAlign_Center)
 				.Padding(2.0f)
 				.AutoHeight()
 				[
-					SNew(SHorizontalBox)
-
-					// Search box
-					+ SHorizontalBox::Slot()
-					.VAlign(VAlign_Center)
-					.Padding(2.0f)
-					.FillWidth(1.0f)
-					[
-						SAssignNew(SearchBox, SSearchBox)
-						.HintText(LOCTEXT("SearchBoxHint", "Search LLM tags or groups"))
-						.OnTextChanged(this, &SMemTagTreeView::SearchBox_OnTextChanged)
-						.IsEnabled(this, &SMemTagTreeView::SearchBox_IsEnabled)
-						.ToolTipText(LOCTEXT("FilterSearchHint", "Type here to search LLM tags or groups"))
-					]
-
-					// Filter out LLM tags with zero instance count.
-					//+ SHorizontalBox::Slot()
-					//.VAlign(VAlign_Center)
-					//.Padding(2.0f)
-					//.AutoWidth()
-					//[
-					//	SNew(SCheckBox)
-					//	.Style(FEditorStyle::Get(), "ToggleButtonCheckbox")
-					//	.HAlign(HAlign_Center)
-					//	.Padding(2.0f)
-					//	.OnCheckStateChanged(this, &SMemTagTreeView::FilterOutZeroCountMemTags_OnCheckStateChanged)
-					//	.IsChecked(this, &SMemTagTreeView::FilterOutZeroCountMemTags_IsChecked)
-					//	.ToolTipText(LOCTEXT("FilterOutZeroCountMemTags_Tooltip", "Filter out the LLM tags having zero total instance count (aggregated stats)."))
-					//	[
-					//		SNew(STextBlock)
-					//		.Text(LOCTEXT("FilterOutZeroCountMemTags_Button", " !0 "))
-					//		.TextStyle(FEditorStyle::Get(), TEXT("Profiler.Caption"))
-					//	]
-					//]
-
-					// Filter out LLM tags by current tracker.
-					+ SHorizontalBox::Slot()
-					.VAlign(VAlign_Center)
-					.Padding(2.0f)
-					.AutoWidth()
-					[
-						SNew(SCheckBox)
-						.Style(FEditorStyle::Get(), "ToggleButtonCheckbox")
-						.HAlign(HAlign_Center)
-						.Padding(2.0f)
-						.OnCheckStateChanged(this, &SMemTagTreeView::FilterByTracker_OnCheckStateChanged)
-						.IsChecked(this, &SMemTagTreeView::FilterByTracker_IsChecked)
-						.ToolTipText(LOCTEXT("FilterByTracker_Tooltip", "Filter the LLM tags to show only the ones used by the current tracker."))
-						[
-							SNew(STextBlock)
-							.Text(LOCTEXT("FilterByTracker_Button", " T "))
-							.TextStyle(FEditorStyle::Get(), TEXT("Profiler.Caption"))
-						]
-					]
+					ConstructInvestigationWidgetArea()
 				]
 
-				// Group by
+				// Search / Filtering
 				+ SVerticalBox::Slot()
 				.VAlign(VAlign_Center)
 				.Padding(2.0f)
 				.AutoHeight()
 				[
-					SNew(SHorizontalBox)
-
-					+ SHorizontalBox::Slot()
-					.AutoWidth()
-					.VAlign(VAlign_Center)
-					.Padding(0.0f, 0.0f, 4.0f, 0.0f)
-					[
-						SNew(STextBlock)
-						.Text(LOCTEXT("TrackerText", "Tracker:"))
-					]
-
-					+ SHorizontalBox::Slot()
-					.FillWidth(1.0f)
-					.VAlign(VAlign_Center)
-					[
-						SAssignNew(TrackerComboBox, SComboBox<TSharedPtr<Insights::FMemoryTracker>>)
-						.ToolTipText(this, &SMemTagTreeView::Tracker_GetTooltipText)
-						.OptionsSource(GetAvailableTrackers())
-						.OnSelectionChanged(this, &SMemTagTreeView::Tracker_OnSelectionChanged)
-						.OnGenerateWidget(this, &SMemTagTreeView::Tracker_OnGenerateWidget)
-						[
-							SNew(STextBlock)
-							.Text(this, &SMemTagTreeView::Tracker_GetSelectedText)
-						]
-					]
-
-					+ SHorizontalBox::Slot()
-					.AutoWidth()
-					.VAlign(VAlign_Center)
-					.Padding(0.0f, 0.0f, 4.0f, 0.0f)
-					[
-						SNew(STextBlock)
-						.MinDesiredWidth(60.0f)
-						.Justification(ETextJustify::Right)
-						.Text(LOCTEXT("GroupByText", "Group by:"))
-					]
+					ConstructTagsFilteringWidgetArea()
+				]
 
-					+ SHorizontalBox::Slot()
-					.FillWidth(1.0f)
-					.VAlign(VAlign_Center)
-					[
-						SAssignNew(GroupByComboBox, SComboBox<TSharedPtr<EMemTagNodeGroupingMode>>)
-						.ToolTipText(this, &SMemTagTreeView::GroupBy_GetSelectedTooltipText)
-						.OptionsSource(&GroupByOptionsSource)
-						.OnSelectionChanged(this, &SMemTagTreeView::GroupBy_OnSelectionChanged)
-						.OnGenerateWidget(this, &SMemTagTreeView::GroupBy_OnGenerateWidget)
-						[
-							SNew(STextBlock)
-							.Text(this, &SMemTagTreeView::GroupBy_GetSelectedText)
-						]
-					]
+				// Grouping
+				+ SVerticalBox::Slot()
+				.VAlign(VAlign_Center)
+				.Padding(2.0f)
+				.AutoHeight()
+				[
+					ConstructTagsGroupingWidgetArea()
 				]
 
 				// Tracks
@@ -215,102 +127,7 @@ void SMemTagTreeView::Construct(const FArguments& InArgs, TSharedPtr<SMemoryProf
 				.Padding(2.0f)
 				.AutoHeight()
 				[
-					SNew(SHorizontalBox)
-
-					+ SHorizontalBox::Slot()
-					.VAlign(VAlign_Center)
-					.Padding(FMargin(2.0f, 0.0f, 2.0f, 0.0f))
-					.AutoWidth()
-					[
-						SNew(SButton)
-						.ToolTipText(LOCTEXT("HideAll_ToolTip", "Remove all memory graph tracks."))
-						.ContentPadding(FMargin(2.0f, 0.0f, 2.0f, 0.0f))
-						.OnClicked(this, &SMemTagTreeView::HideAllTracks_OnClicked)
-						.Content()
-						[
-							SNew(SImage)
-							.ColorAndOpacity(FSlateColor(FLinearColor(0.5f, 0.0f, 0.0f, 1.0f)))
-							.Image(FInsightsStyle::Get().GetBrush("Mem.Remove.Small"))
-						]
-					]
-
-					+ SHorizontalBox::Slot()
-					.VAlign(VAlign_Center)
-					.Padding(FMargin(2.0f, 0.0f, 2.0f, 0.0f))
-					.AutoWidth()
-					[
-						SNew(SButton)
-						.ToolTipText(LOCTEXT("ShowAll_ToolTip", "Create memory graph tracks for all visible (filtered) LLM tags."))
-						.ContentPadding(FMargin(2.0f, 0.0f, 2.0f, 0.0f))
-						.OnClicked(this, &SMemTagTreeView::ShowAllTracks_OnClicked)
-						.Content()
-						[
-							SNew(SImage)
-							.ColorAndOpacity(FSlateColor(FLinearColor(0.0f, 0.5f, 0.0f, 1.0f)))
-							.Image(FInsightsStyle::Get().GetBrush("Mem.Add.Small"))
-						]
-					]
-
-					+ SHorizontalBox::Slot()
-					.VAlign(VAlign_Center)
-					.Padding(2.0f)
-					.AutoWidth()
-					[
-						SNew(SButton)
-						.ToolTipText(LOCTEXT("LoadReportXML_ToolTip", "Load LLMReportTypes.xml"))
-						.ContentPadding(FMargin(2.0f, 2.0f, 2.0f, 2.0f))
-						.OnClicked(this, &SMemTagTreeView::LoadReportXML_OnClicked)
-						.Content()
-						[
-							SNew(STextBlock)
-							.Text(LOCTEXT("LoadReportXML_Text", "Load XML..."))
-						]
-					]
-
-					+ SHorizontalBox::Slot()
-					.VAlign(VAlign_Center)
-					.Padding(2.0f)
-					[
-						SNew(SButton)
-						.ToolTipText(LOCTEXT("SmallHeight_ToolTip", "Change height of LLM Tag Graph tracks to Small."))
-						.ContentPadding(FMargin(2.0f, 2.0f, 2.0f, 2.0f))
-						.OnClicked(this, &SMemTagTreeView::AllTracksSmallHeight_OnClicked)
-						.Content()
-						[
-							SNew(STextBlock)
-							.Text(LOCTEXT("SmallHeight_Text", "\u2195S"))
-						]
-					]
-
-					+ SHorizontalBox::Slot()
-					.VAlign(VAlign_Center)
-					.Padding(2.0f)
-					[
-						SNew(SButton)
-						.ToolTipText(LOCTEXT("MediumHeight_ToolTip", "Change height of LLM Tag Graph tracks to Medium."))
-						.ContentPadding(FMargin(2.0f, 2.0f, 2.0f, 2.0f))
-						.OnClicked(this, &SMemTagTreeView::AllTracksMediumHeight_OnClicked)
-						.Content()
-						[
-							SNew(STextBlock)
-							.Text(LOCTEXT("MediumHeight_Text", "\u2195M"))
-						]
-					]
-
-					+ SHorizontalBox::Slot()
-					.VAlign(VAlign_Center)
-					.Padding(2.0f)
-					[
-						SNew(SButton)
-						.ToolTipText(LOCTEXT("LargeHeight_ToolTip", "Change height of LLM Tag Graph tracks to Large."))
-						.ContentPadding(FMargin(2.0f, 2.0f, 2.0f, 2.0f))
-						.OnClicked(this, &SMemTagTreeView::AllTracksLargeHeight_OnClicked)
-						.Content()
-						[
-							SNew(STextBlock)
-							.Text(LOCTEXT("LargeHeight_Text", "\u2195L"))
-						]
-					]
+					ConstructTracksMiniToolbar()
 				]
 			]
 		]
@@ -384,10 +201,448 @@ void SMemTagTreeView::Construct(const FArguments& InArgs, TSharedPtr<SMemoryProf
 	// Update the Session (i.e. when analysis session was already started).
 	InsightsManager_OnSessionChanged();
 }
-END_SLATE_FUNCTION_BUILD_OPTIMIZATION
 
 ////////////////////////////////////////////////////////////////////////////////////////////////////
 
+TSharedRef<SWidget> SMemTagTreeView::ConstructInvestigationWidgetArea()
+{
+	TSharedRef<SWidget> Widget = SNew(SExpandableArea)
+
+	.InitiallyCollapsed(false)
+	.BodyBorderImage(FEditorStyle::GetBrush("ToolPanel.GroupBorder"))
+	.BodyBorderBackgroundColor(FLinearColor::White)
+	.Padding(0.0f)
+	.HeaderContent()
+	[
+		SNew(STextBlock)
+		.Text(LOCTEXT("InvestigationTitle", "Investigation"))
+		.Font(FEditorStyle::GetFontStyle("DetailsView.CategoryFontStyle"))
+		.ShadowOffset(FVector2D(1.0f, 1.0f))
+	]
+	.BodyContent()
+	[
+		SNew(SVerticalBox)
+
+		+ SVerticalBox::Slot()
+		.AutoHeight()
+		[
+			SNew(SHorizontalBox)
+
+			+ SHorizontalBox::Slot()
+			.AutoWidth()
+			.VAlign(VAlign_Center)
+			.Padding(0.0f, 0.0f, 4.0f, 0.0f)
+			[
+				SNew(STextBlock)
+				.Text(LOCTEXT("QueryRuleText", "Rule:"))
+			]
+
+			+ SHorizontalBox::Slot()
+			.FillWidth(1.0f)
+			.VAlign(VAlign_Center)
+			[
+				SAssignNew(QueryRuleComboBox, SComboBox<TSharedPtr<Insights::FMemoryRuleSpec>>)
+				.ToolTipText(this, &SMemTagTreeView::QueryRule_GetTooltipText)
+				.OptionsSource(GetAvailableQueryRules())
+				.OnSelectionChanged(this, &SMemTagTreeView::QueryRule_OnSelectionChanged)
+				.OnGenerateWidget(this, &SMemTagTreeView::QueryRule_OnGenerateWidget)
+				[
+					SNew(STextBlock)
+					.Text(this, &SMemTagTreeView::QueryRule_GetSelectedText)
+				]
+			]
+		]
+
+		+ SVerticalBox::Slot()
+		.AutoHeight()
+		[
+			ConstructTimeMarkerWidget(0)
+		]
+
+		+ SVerticalBox::Slot()
+		.AutoHeight()
+		[
+			ConstructTimeMarkerWidget(1)
+		]
+
+		+ SVerticalBox::Slot()
+		.AutoHeight()
+		[
+			ConstructTimeMarkerWidget(2)
+		]
+
+		+ SVerticalBox::Slot()
+		.AutoHeight()
+		[
+			ConstructTimeMarkerWidget(3)
+		]
+
+		+ SVerticalBox::Slot()
+		.AutoHeight()
+		.Padding(0.0f, 4.0f, 0.0f, 2.0f)
+		[
+			SNew(STextBlock)
+			.Text(this, &SMemTagTreeView::QueryRule_GetTooltipText)
+			.ColorAndOpacity(FLinearColor(0.3f, 0.3f, 0.3f, 1.0f))
+			.AutoWrapText(true)
+		]
+
+		+ SVerticalBox::Slot()
+		.AutoHeight()
+		.HAlign(HAlign_Right)
+		[
+			SNew(SButton)
+			.Text(LOCTEXT("RunQueryBtnText", "Run Query"))
+			.ToolTipText(LOCTEXT("RunQueryBtnToolTipText", "Run Memory Query.\nThe resulting list of allocations will be available in a tree view."))
+			.OnClicked(this, &SMemTagTreeView::RunQuery)
+		]
+
+		+ SVerticalBox::Slot()
+		.AutoHeight()
+		.Padding(0.0f, 6.0f, 0.0f, 0.0f)
+		[
+			SNew(SSeparator)
+		]
+	];
+
+	if (ProfilerWindow.IsValid())
+	{
+		FMemorySharedState& SharedState = ProfilerWindow->GetSharedState();
+		QueryRuleComboBox->SetSelectedItem(SharedState.GetCurrentMemoryRule());
+	}
+
+	return Widget;
+}
+
+////////////////////////////////////////////////////////////////////////////////////////////////////
+
+TSharedRef<SWidget> SMemTagTreeView::QueryRule_OnGenerateWidget(TSharedPtr<Insights::FMemoryRuleSpec> InRule)
+{
+	const FText QueryRuleText = FText::Format(LOCTEXT("QueryRuleComboBox_TextFmt", "{0} ({1})"), InRule->GetVerboseName(), InRule->GetShortName());
+
+	TSharedRef<SHorizontalBox> Widget = SNew(SHorizontalBox);
+
+	Widget->AddSlot()
+		.AutoWidth()
+		.HAlign(HAlign_Center)
+		.VAlign(VAlign_Center)
+		.Padding(2.0f, 2.0f, 0.0f, 2.0f)
+		[
+			SNew(SImage)
+			.Visibility_Lambda([Widget]() { return Widget->GetParentWidget()->IsHovered() ? EVisibility::Visible : EVisibility::Hidden; })
+			.Image(FEditorStyle::GetBrush("Profiler.Tooltip.HintIcon10"))
+			.ToolTipText(InRule->GetDescription())
+		];
+
+	Widget->AddSlot()
+		.AutoWidth()
+		[
+			SNew(STextBlock)
+			.Text(QueryRuleText)
+			.Margin(2.0f)
+		];
+
+	return Widget;
+}
+
+////////////////////////////////////////////////////////////////////////////////////////////////////
+
+TSharedRef<SWidget> SMemTagTreeView::ConstructTimeMarkerWidget(uint32 TimeMarkerIndex)
+{
+	if (!ProfilerWindow.IsValid())
+	{
+		return SNew(SBox);
+	}
+
+	FMemorySharedState& SharedState = ProfilerWindow->GetSharedState();
+	const uint32 NumTimeMarkers = ProfilerWindow->GetNumCustomTimeMarkers();
+	if (TimeMarkerIndex >= NumTimeMarkers)
+	{
+		return SNew(SBox);
+	}
+
+	const TSharedRef<Insights::FTimeMarker>& TimeMarker = ProfilerWindow->GetCustomTimeMarker(TimeMarkerIndex);
+
+	TSharedRef<SWidget> Widget = SNew(SHorizontalBox)
+		.Visibility_Lambda([TimeMarker]()
+			{
+				return TimeMarker->IsVisible() ? EVisibility::Visible : EVisibility::Collapsed;
+			})
+
+		+ SHorizontalBox::Slot()
+		.AutoWidth()
+		.VAlign(VAlign_Center)
+		.Padding(20.0f, 1.0f, 4.0f, 1.0f)
+		[
+			SNew(STextBlock)
+			.Text(FText::FromString(TimeMarker->GetName()))
+			.ColorAndOpacity(FSlateColor(TimeMarker->GetColor()))
+		]
+
+		+ SHorizontalBox::Slot()
+		.AutoWidth()
+		.VAlign(VAlign_Center)
+		.Padding(0.0f, 1.0f, 0.0f, 1.0f)
+		[
+			SNew(SEditableTextBox)
+			.MinDesiredWidth(110.0f)
+			.Text_Lambda([TimeMarker]()
+				{
+					return FText::FromString(FString::Printf(TEXT("%.9f"), TimeMarker->GetTime()));
+				})
+			.OnTextCommitted_Lambda([TimeMarker](const FText& InText, ETextCommit::Type InCommitType)
+				{
+					if (InCommitType == ETextCommit::OnEnter)
+					{
+						const double Time = FCString::Atod(*InText.ToString());
+						TimeMarker->SetTime(Time);
+					}
+				})
+		]
+
+		+ SHorizontalBox::Slot()
+		.AutoWidth()
+		.VAlign(VAlign_Center)
+		.Padding(4.0f, 1.0f, 0.0f, 1.0f)
+		[
+			SNew(STextBlock)
+			.Text_Lambda([TimeMarker]()
+				{
+					return FText::FromString(TimeUtils::FormatTime(TimeMarker->GetTime(), 0.1));
+				})
+		];
+
+	return Widget;
+}
+
+////////////////////////////////////////////////////////////////////////////////////////////////////
+
+TSharedRef<SWidget> SMemTagTreeView::ConstructTagsFilteringWidgetArea()
+{
+	TSharedRef<SWidget> Widget = SNew(SHorizontalBox)
+
+	// Search box
+	+ SHorizontalBox::Slot()
+	.VAlign(VAlign_Center)
+	.Padding(2.0f)
+	.FillWidth(1.0f)
+	[
+		SAssignNew(SearchBox, SSearchBox)
+		.HintText(LOCTEXT("SearchBoxHint", "Search LLM tags or groups"))
+		.OnTextChanged(this, &SMemTagTreeView::SearchBox_OnTextChanged)
+		.IsEnabled(this, &SMemTagTreeView::SearchBox_IsEnabled)
+		.ToolTipText(LOCTEXT("FilterSearchHint", "Type here to search LLM tags or groups"))
+	]
+
+	// Filter out LLM tags with zero instance count.
+	//+ SHorizontalBox::Slot()
+	//.VAlign(VAlign_Center)
+	//.Padding(2.0f)
+	//.AutoWidth()
+	//[
+	//	SNew(SCheckBox)
+	//	.Style(FEditorStyle::Get(), "ToggleButtonCheckbox")
+	//	.HAlign(HAlign_Center)
+	//	.Padding(2.0f)
+	//	.OnCheckStateChanged(this, &SMemTagTreeView::FilterOutZeroCountMemTags_OnCheckStateChanged)
+	//	.IsChecked(this, &SMemTagTreeView::FilterOutZeroCountMemTags_IsChecked)
+	//	.ToolTipText(LOCTEXT("FilterOutZeroCountMemTags_Tooltip", "Filter out the LLM tags having zero total instance count (aggregated stats)."))
+	//	[
+	//		SNew(STextBlock)
+	//		.Text(LOCTEXT("FilterOutZeroCountMemTags_Button", " !0 "))
+	//		.TextStyle(FEditorStyle::Get(), TEXT("Profiler.Caption"))
+	//	]
+	//]
+
+	// Filter out LLM tags by current tracker.
+	+ SHorizontalBox::Slot()
+	.VAlign(VAlign_Center)
+	.Padding(2.0f)
+	.AutoWidth()
+	[
+		SNew(SCheckBox)
+		.Style(FEditorStyle::Get(), "ToggleButtonCheckbox")
+		.HAlign(HAlign_Center)
+		.Padding(2.0f)
+		.OnCheckStateChanged(this, &SMemTagTreeView::FilterByTracker_OnCheckStateChanged)
+		.IsChecked(this, &SMemTagTreeView::FilterByTracker_IsChecked)
+		.ToolTipText(LOCTEXT("FilterByTracker_Tooltip", "Filter the LLM tags to show only the ones used by the current tracker."))
+		[
+			SNew(STextBlock)
+			.Text(LOCTEXT("FilterByTracker_Button", " T "))
+			.TextStyle(FEditorStyle::Get(), TEXT("Profiler.Caption"))
+		]
+	];
+
+	return Widget;
+}
+
+////////////////////////////////////////////////////////////////////////////////////////////////////
+
+TSharedRef<SWidget> SMemTagTreeView::ConstructTagsGroupingWidgetArea()
+{
+	TSharedRef<SWidget> Widget = SNew(SHorizontalBox)
+
+	+ SHorizontalBox::Slot()
+	.AutoWidth()
+	.VAlign(VAlign_Center)
+	.Padding(0.0f, 0.0f, 4.0f, 0.0f)
+	[
+		SNew(STextBlock)
+		.Text(LOCTEXT("TrackerText", "Tracker:"))
+	]
+
+	+ SHorizontalBox::Slot()
+	.FillWidth(1.0f)
+	.VAlign(VAlign_Center)
+	[
+		SAssignNew(TrackerComboBox, SComboBox<TSharedPtr<Insights::FMemoryTracker>>)
+		.ToolTipText(this, &SMemTagTreeView::Tracker_GetTooltipText)
+		.OptionsSource(GetAvailableTrackers())
+		.OnSelectionChanged(this, &SMemTagTreeView::Tracker_OnSelectionChanged)
+		.OnGenerateWidget(this, &SMemTagTreeView::Tracker_OnGenerateWidget)
+		[
+			SNew(STextBlock)
+			.Text(this, &SMemTagTreeView::Tracker_GetSelectedText)
+		]
+	]
+
+	+ SHorizontalBox::Slot()
+	.AutoWidth()
+	.VAlign(VAlign_Center)
+	.Padding(0.0f, 0.0f, 4.0f, 0.0f)
+	[
+		SNew(STextBlock)
+		.MinDesiredWidth(60.0f)
+		.Justification(ETextJustify::Right)
+		.Text(LOCTEXT("GroupByText", "Group by:"))
+	]
+
+	+ SHorizontalBox::Slot()
+	.FillWidth(1.0f)
+	.VAlign(VAlign_Center)
+	[
+		SAssignNew(GroupByComboBox, SComboBox<TSharedPtr<EMemTagNodeGroupingMode>>)
+		.ToolTipText(this, &SMemTagTreeView::GroupBy_GetSelectedTooltipText)
+		.OptionsSource(&GroupByOptionsSource)
+		.OnSelectionChanged(this, &SMemTagTreeView::GroupBy_OnSelectionChanged)
+		.OnGenerateWidget(this, &SMemTagTreeView::GroupBy_OnGenerateWidget)
+		[
+			SNew(STextBlock)
+			.Text(this, &SMemTagTreeView::GroupBy_GetSelectedText)
+		]
+	];
+
+	return Widget;
+}
+
+////////////////////////////////////////////////////////////////////////////////////////////////////
+
+TSharedRef<SWidget> SMemTagTreeView::ConstructTracksMiniToolbar()
+{
+	TSharedRef<SWidget> Widget = SNew(SHorizontalBox)
+
+	+ SHorizontalBox::Slot()
+	.VAlign(VAlign_Center)
+	.Padding(FMargin(2.0f, 0.0f, 2.0f, 0.0f))
+	.AutoWidth()
+	[
+		SNew(SButton)
+		.ToolTipText(LOCTEXT("HideAll_ToolTip", "Remove all memory graph tracks."))
+		.ContentPadding(FMargin(2.0f, 0.0f, 2.0f, 0.0f))
+		.OnClicked(this, &SMemTagTreeView::HideAllTracks_OnClicked)
+		.Content()
+		[
+			SNew(SImage)
+			.ColorAndOpacity(FSlateColor(FLinearColor(0.5f, 0.0f, 0.0f, 1.0f)))
+			.Image(FInsightsStyle::Get().GetBrush("Mem.Remove.Small"))
+		]
+	]
+
+	+ SHorizontalBox::Slot()
+	.VAlign(VAlign_Center)
+	.Padding(FMargin(2.0f, 0.0f, 2.0f, 0.0f))
+	.AutoWidth()
+	[
+		SNew(SButton)
+		.ToolTipText(LOCTEXT("ShowAll_ToolTip", "Create memory graph tracks for all visible (filtered) LLM tags."))
+		.ContentPadding(FMargin(2.0f, 0.0f, 2.0f, 0.0f))
+		.OnClicked(this, &SMemTagTreeView::ShowAllTracks_OnClicked)
+		.Content()
+		[
+			SNew(SImage)
+			.ColorAndOpacity(FSlateColor(FLinearColor(0.0f, 0.5f, 0.0f, 1.0f)))
+			.Image(FInsightsStyle::Get().GetBrush("Mem.Add.Small"))
+		]
+	]
+
+	+ SHorizontalBox::Slot()
+	.VAlign(VAlign_Center)
+	.Padding(2.0f)
+	.AutoWidth()
+	[
+		SNew(SButton)
+		.ToolTipText(LOCTEXT("LoadReportXML_ToolTip", "Load LLMReportTypes.xml"))
+		.ContentPadding(FMargin(2.0f, 2.0f, 2.0f, 2.0f))
+		.OnClicked(this, &SMemTagTreeView::LoadReportXML_OnClicked)
+		.Content()
+		[
+			SNew(STextBlock)
+			.Text(LOCTEXT("LoadReportXML_Text", "Load XML..."))
+		]
+	]
+
+	+ SHorizontalBox::Slot()
+	.VAlign(VAlign_Center)
+	.Padding(2.0f)
+	[
+		SNew(SButton)
+		.ToolTipText(LOCTEXT("SmallHeight_ToolTip", "Change height of LLM Tag Graph tracks to Small."))
+		.ContentPadding(FMargin(2.0f, 2.0f, 2.0f, 2.0f))
+		.OnClicked(this, &SMemTagTreeView::AllTracksSmallHeight_OnClicked)
+		.Content()
+		[
+			SNew(STextBlock)
+			.Text(LOCTEXT("SmallHeight_Text", "\u2195S"))
+		]
+	]
+
+	+ SHorizontalBox::Slot()
+	.VAlign(VAlign_Center)
+	.Padding(2.0f)
+	[
+		SNew(SButton)
+		.ToolTipText(LOCTEXT("MediumHeight_ToolTip", "Change height of LLM Tag Graph tracks to Medium."))
+		.ContentPadding(FMargin(2.0f, 2.0f, 2.0f, 2.0f))
+		.OnClicked(this, &SMemTagTreeView::AllTracksMediumHeight_OnClicked)
+		.Content()
+		[
+			SNew(STextBlock)
+			.Text(LOCTEXT("MediumHeight_Text", "\u2195M"))
+		]
+	]
+
+	+ SHorizontalBox::Slot()
+	.VAlign(VAlign_Center)
+	.Padding(2.0f)
+	[
+		SNew(SButton)
+		.ToolTipText(LOCTEXT("LargeHeight_ToolTip", "Change height of LLM Tag Graph tracks to Large."))
+		.ContentPadding(FMargin(2.0f, 2.0f, 2.0f, 2.0f))
+		.OnClicked(this, &SMemTagTreeView::AllTracksLargeHeight_OnClicked)
+		.Content()
+		[
+			SNew(STextBlock)
+			.Text(LOCTEXT("LargeHeight_Text", "\u2195L"))
+		]
+	];
+
+	return Widget;
+}
+
+END_SLATE_FUNCTION_BUILD_OPTIMIZATION
+////////////////////////////////////////////////////////////////////////////////////////////////////
+
 TSharedPtr<SWidget> SMemTagTreeView::TreeView_GetMenuContent()
 {
 	const TArray<FMemTagNodePtr> SelectedNodes = TreeView->GetSelectedItems();
@@ -2016,6 +2271,129 @@ void SMemTagTreeView::SelectMemTagNode(Insights::FMemoryTagId MemTagId)
 	}
 }
 
+////////////////////////////////////////////////////////////////////////////////////////////////////
+// Query Rules
+////////////////////////////////////////////////////////////////////////////////////////////////////
+
+const TArray<TSharedPtr<Insights::FMemoryRuleSpec>>* SMemTagTreeView::GetAvailableQueryRules()
+{
+	if (ProfilerWindow.IsValid())
+	{
+		FMemorySharedState& SharedState = ProfilerWindow->GetSharedState();
+		return &SharedState.GetMemoryRules();
+	}
+	return nullptr;
+}
+
+////////////////////////////////////////////////////////////////////////////////////////////////////
+
+void SMemTagTreeView::QueryRule_OnSelectionChanged(TSharedPtr<Insights::FMemoryRuleSpec> InRule, ESelectInfo::Type SelectInfo)
+{
+	if (SelectInfo != ESelectInfo::Direct)
+	{
+		if (ProfilerWindow.IsValid() && InRule)
+		{
+			FMemorySharedState& SharedState = ProfilerWindow->GetSharedState();
+			SharedState.SetCurrentMemoryRule(InRule);
+			ProfilerWindow->UpdateTimingViewMarkers();
+		}
+	}
+}
+
+////////////////////////////////////////////////////////////////////////////////////////////////////
+
+FText SMemTagTreeView::QueryRule_GetSelectedText() const
+{
+	if (ProfilerWindow.IsValid())
+	{
+		FMemorySharedState& SharedState = ProfilerWindow->GetSharedState();
+		TSharedPtr<Insights::FMemoryRuleSpec> Rule = SharedState.GetCurrentMemoryRule();
+		if (Rule)
+		{
+			return FText::Format(LOCTEXT("QueryRuleComboBox_TextFmt", "{0} ({1})"), Rule->GetVerboseName(), Rule->GetShortName());
+		}
+	}
+	return FText::GetEmpty();
+}
+
+////////////////////////////////////////////////////////////////////////////////////////////////////
+
+FText SMemTagTreeView::QueryRule_GetTooltipText() const
+{
+	if (ProfilerWindow.IsValid())
+	{
+		FMemorySharedState& SharedState = ProfilerWindow->GetSharedState();
+		if (SharedState.GetCurrentMemoryRule())
+		{
+			return SharedState.GetCurrentMemoryRule()->GetDescription();
+		}
+	}
+	return FText::GetEmpty();
+}
+
+////////////////////////////////////////////////////////////////////////////////////////////////////
+
+FReply SMemTagTreeView::RunQuery()
+{
+	if (!ProfilerWindow.IsValid())
+	{
+		UE_LOG(MemoryProfiler, Error, TEXT("[MemQuery] Invalid Profiler Window!"));
+		return FReply::Handled();
+	}
+
+	FMemorySharedState& SharedState = ProfilerWindow->GetSharedState();
+	TSharedPtr<Insights::FMemoryRuleSpec> Rule = SharedState.GetCurrentMemoryRule();
+	if (!Rule)
+	{
+		UE_LOG(MemoryProfiler, Error, TEXT("[MemQuery] Invalid Rule!"));
+		return FReply::Handled();
+	}
+
+	const uint32 NumTimeMarkers = ProfilerWindow->GetNumCustomTimeMarkers();
+	const uint32 RuleNumTimeMarkers = Rule->GetNumTimeMarkers();
+	if (RuleNumTimeMarkers > NumTimeMarkers)
+	{
+		UE_LOG(MemoryProfiler, Error, TEXT("[MemQuery] Only %d time markers available. Current rule (%s) requires %u time markers!"),
+			NumTimeMarkers, *Rule->GetShortName().ToString(), RuleNumTimeMarkers);
+		return FReply::Handled();
+	}
+
+#if !NO_LOGGING
+	TStringBuilder<1024> Builder;
+	if (RuleNumTimeMarkers > 0)
+	{
+		Builder.Append(TEXT(" ("));
+		for (uint32 TimeMarkerIndex = 0; TimeMarkerIndex < RuleNumTimeMarkers; ++TimeMarkerIndex)
+		{
+			if (TimeMarkerIndex != 0)
+			{
+				Builder.Append(TEXT(", "));
+			}
+			Builder.Append(TEXT('A') + TimeMarkerIndex);
+			const TSharedRef<Insights::FTimeMarker>& TimeMarker = ProfilerWindow->GetCustomTimeMarker(TimeMarkerIndex);
+			Builder.Appendf(TEXT("=%.9f"), TimeMarker->GetTime());
+		}
+		Builder.Append(TEXT(")"));
+	}
+	UE_LOG(MemoryProfiler, Log, TEXT("[MemQuery] Run Query %s%s..."), *Rule->GetShortName().ToString(), Builder.ToString());
+#endif
+
+	TSharedPtr<Insights::SMemAllocTableTreeView> MemAllocTableTreeView = ProfilerWindow->ShowMemAllocTableTreeViewTab();
+	if (MemAllocTableTreeView)
+	{
+		double TimeMarkers[4];
+		TimeMarkers[0] = (RuleNumTimeMarkers > 0) ? ProfilerWindow->GetCustomTimeMarker(0)->GetTime() : 0.0;
+		TimeMarkers[1] = (RuleNumTimeMarkers > 1) ? ProfilerWindow->GetCustomTimeMarker(1)->GetTime() : 0.0;
+		TimeMarkers[2] = (RuleNumTimeMarkers > 2) ? ProfilerWindow->GetCustomTimeMarker(2)->GetTime() : 0.0;
+		TimeMarkers[3] = (RuleNumTimeMarkers > 3) ? ProfilerWindow->GetCustomTimeMarker(3)->GetTime() : 0.0;
+		MemAllocTableTreeView->SetQueryParams(Rule, TimeMarkers[0], TimeMarkers[1], TimeMarkers[2], TimeMarkers[3]);
+	}
+
+	return FReply::Handled();
+}
+
+////////////////////////////////////////////////////////////////////////////////////////////////////
+// Trackers
 ////////////////////////////////////////////////////////////////////////////////////////////////////
 
 const TArray<TSharedPtr<Insights::FMemoryTracker>>* SMemTagTreeView::GetAvailableTrackers()
@@ -2117,6 +2495,8 @@ FText SMemTagTreeView::Tracker_GetTooltipText() const
 	return LOCTEXT("TrackerComboBox_Tooltip", "Choose the current memory tracker.");
 }
 
+////////////////////////////////////////////////////////////////////////////////////////////////////
+// Load Report XML button action
 ////////////////////////////////////////////////////////////////////////////////////////////////////
 
 FReply SMemTagTreeView::LoadReportXML_OnClicked()
@@ -2157,6 +2537,8 @@ FReply SMemTagTreeView::LoadReportXML_OnClicked()
 	return FReply::Handled();
 }
 
+////////////////////////////////////////////////////////////////////////////////////////////////////
+// Button actions re graph tracks
 ////////////////////////////////////////////////////////////////////////////////////////////////////
 
 FReply SMemTagTreeView::ShowAllTracks_OnClicked()
diff --git a/Engine/Source/Developer/TraceInsights/Private/Insights/MemoryProfiler/Widgets/SMemTagTreeView.h b/Engine/Source/Developer/TraceInsights/Private/Insights/MemoryProfiler/Widgets/SMemTagTreeView.h
--- a/Engine/Source/Developer/TraceInsights/Private/Insights/MemoryProfiler/Widgets/SMemTagTreeView.h
+++ b/Engine/Source/Developer/TraceInsights/Private/Insights/MemoryProfiler/Widgets/SMemTagTreeView.h
@@ -32,6 +32,7 @@ namespace Trace
 
 namespace Insights
 {
+	class FMemoryRuleSpec;
 	class FTable;
 	class FTableColumn;
 	class ITableCellValueSorter;
@@ -84,6 +85,12 @@ public:
 	void SelectMemTagNode(Insights::FMemoryTagId MemTagId);
 
 private:
+	TSharedRef<SWidget> ConstructInvestigationWidgetArea();
+	TSharedRef<SWidget> ConstructTimeMarkerWidget(uint32 TimeMarkerIndex);
+	TSharedRef<SWidget> ConstructTagsFilteringWidgetArea();
+	TSharedRef<SWidget> ConstructTagsGroupingWidgetArea();
+	TSharedRef<SWidget> ConstructTracksMiniToolbar();
+
 	void UpdateTree();
 
 	void UpdateStatsInternal();
@@ -243,6 +250,15 @@ private:
 
 	////////////////////////////////////////////////////////////////////////////////////////////////////
 
+	const TArray<TSharedPtr<Insights::FMemoryRuleSpec>>* GetAvailableQueryRules();
+	void QueryRule_OnSelectionChanged(TSharedPtr<Insights::FMemoryRuleSpec> InRule, ESelectInfo::Type SelectInfo);
+	TSharedRef<SWidget> QueryRule_OnGenerateWidget(TSharedPtr<Insights::FMemoryRuleSpec> InRule);
+	FText QueryRule_GetSelectedText() const;
+	FText QueryRule_GetTooltipText() const;
+	FReply RunQuery();
+
+	////////////////////////////////////////////////////////////////////////////////////////////////////
+
 	const TArray<TSharedPtr<Insights::FMemoryTracker>>* GetAvailableTrackers();
 	void Tracker_OnSelectionChanged(TSharedPtr<Insights::FMemoryTracker> InTracker, ESelectInfo::Type SelectInfo);
 	TSharedRef<SWidget> Tracker_OnGenerateWidget(TSharedPtr<Insights::FMemoryTracker> InTracker);
@@ -366,6 +382,11 @@ private:
 	/** If true, the expanded nodes have been saved before applying a text filter. */
 	bool bExpansionSaved;
 
+	//////////////////////////////////////////////////
+	// Queries
+
+	TSharedPtr<SComboBox<TSharedPtr<Insights::FMemoryRuleSpec>>> QueryRuleComboBox;
+
 	//////////////////////////////////////////////////
 	// Search box and filters
 
diff --git a/Engine/Source/Developer/TraceInsights/Private/Insights/MemoryProfiler/Widgets/SMemoryProfilerWindow.cpp b/Engine/Source/Developer/TraceInsights/Private/Insights/MemoryProfiler/Widgets/SMemoryProfilerWindow.cpp
--- a/Engine/Source/Developer/TraceInsights/Private/Insights/MemoryProfiler/Widgets/SMemoryProfilerWindow.cpp
+++ b/Engine/Source/Developer/TraceInsights/Private/Insights/MemoryProfiler/Widgets/SMemoryProfilerWindow.cpp
@@ -30,10 +30,12 @@
 #include "Insights/ITimingViewExtender.h"
 #include "Insights/MemoryProfiler/MemoryProfilerManager.h"
 #include "Insights/MemoryProfiler/ViewModels/MemorySharedState.h"
+#include "Insights/MemoryProfiler/Widgets/SMemAllocTableTreeView.h"
 #include "Insights/MemoryProfiler/Widgets/SMemoryProfilerToolbar.h"
 #include "Insights/MemoryProfiler/Widgets/SMemTagTreeView.h"
 #include "Insights/TraceInsightsModule.h"
 #include "Insights/Version.h"
+#include "Insights/ViewModels/TimeRulerTrack.h"
 #include "Insights/Widgets/SInsightsSettings.h"
 #include "Insights/Widgets/STimingView.h"
 
@@ -46,6 +48,7 @@
 const FName FMemoryProfilerTabs::ToolbarID(TEXT("Toolbar"));
 const FName FMemoryProfilerTabs::TimingViewID(TEXT("TimingView"));
 const FName FMemoryProfilerTabs::MemTagTreeViewID(TEXT("LowLevelMemTags"));
+const FName FMemoryProfilerTabs::MemAllocTableTreeViewID(TEXT("MemAllocTableTreeView"));
 
 ////////////////////////////////////////////////////////////////////////////////////////////////////
 
@@ -92,6 +95,7 @@ void SMemoryProfilerWindow::Reset()
 	if (TimingView)
 	{
 		TimingView->Reset();
+		ResetTimingViewMarkers();
 	}
 
 	if (MemTagTreeView)
@@ -104,9 +108,76 @@ void SMemoryProfilerWindow::Reset()
 
 ////////////////////////////////////////////////////////////////////////////////////////////////////
 
+void SMemoryProfilerWindow::ResetTimingViewMarkers()
+{
+	TSharedRef<FTimeRulerTrack> TimeRulerTrack = TimingView->GetTimeRulerTrack();
+
+	TimeRulerTrack->RemoveAllTimeMarkers();
+	CustomTimeMarkers.Reset();
+
+	constexpr uint32 MaxNumTimeMarkers = 5;
+
+	TCHAR TimeMarkerName[2];
+	TimeMarkerName[1] = 0;
+
+	for (uint32 Index = 0; Index < MaxNumTimeMarkers; ++Index)
+	{
+		// Keep (re-add) the "Default Time Marker" as the first time marker and create new ones for the rest.
+		TSharedRef<Insights::FTimeMarker> TimeMarker = (Index == 0) ? TimingView->GetDefaultTimeMarker() : MakeShared<Insights::FTimeMarker>();
+
+		TimeMarkerName[0] = TEXT('A') + Index; // "A", "B", "C", etc.
+		TimeMarker->SetName(TimeMarkerName);
+
+		const uint32 HueStep = 256 / MaxNumTimeMarkers;
+		const uint8 H = uint8(HueStep * Index);
+		const uint8 S = 192;
+		const uint8 V = 255;
+		const FLinearColor Color = FLinearColor::MakeFromHSV8(H, S, V);
+		TimeMarker->SetColor(Color);
+
+		TimeMarker->SetTime(static_cast<float>(Index)); // 0.0f, 1.0f, 2.0f, etc.
+
+		TimeRulerTrack->AddTimeMarker(TimeMarker);
+		CustomTimeMarkers.Add(TimeMarker);
+	}
+
+	UpdateTimingViewMarkers();
+}
+
+////////////////////////////////////////////////////////////////////////////////////////////////////
+
+void SMemoryProfilerWindow::UpdateTimingViewMarkers()
+{
+	TSharedPtr<Insights::FMemoryRuleSpec> Rule = SharedState->GetCurrentMemoryRule();
+	const uint32 NumVisibleTimeMarkers = Rule ? Rule->GetNumTimeMarkers() : 0;
+
+	const uint32 NumTimeMarkers = CustomTimeMarkers.Num();
+	ensure(NumVisibleTimeMarkers <= NumTimeMarkers);
+
+	for (uint32 Index = 0; Index < NumTimeMarkers; ++Index)
+	{
+		TSharedRef<Insights::FTimeMarker>& TimeMarker = CustomTimeMarkers[Index];
+		if (Index < NumVisibleTimeMarkers)
+		{
+			TimeMarker->SetVisibility(true);
+		}
+		else
+		{
+			TimeMarker->SetVisibility(false);
+		}
+	}
+}
+
+////////////////////////////////////////////////////////////////////////////////////////////////////
+
 void SMemoryProfilerWindow::UpdateTableTreeViews()
 {
 	UpdateMemTagTreeView();
+
+	for (TSharedPtr<Insights::SMemAllocTableTreeView>& MemAllocTableTreeView : MemAllocTableTreeViews)
+	{
+		UpdateMemAllocTableTreeView(MemAllocTableTreeView);
+	}
 }
 
 ////////////////////////////////////////////////////////////////////////////////////////////////////
@@ -138,6 +209,16 @@ void SMemoryProfilerWindow::UpdateMemTagTreeView()
 
 ////////////////////////////////////////////////////////////////////////////////////////////////////
 
+void SMemoryProfilerWindow::UpdateMemAllocTableTreeView(TSharedPtr<Insights::SMemAllocTableTreeView> MemAllocTableTreeView)
+{
+	if (MemAllocTableTreeView)
+	{
+		//TODO
+	}
+}
+
+////////////////////////////////////////////////////////////////////////////////////////////////////
+
 BEGIN_SLATE_FUNCTION_BUILD_OPTIMIZATION
 
 TSharedRef<SDockTab> SMemoryProfilerWindow::SpawnTab_Toolbar(const FSpawnTabArgs& Args)
@@ -177,6 +258,7 @@ TSharedRef<SDockTab> SMemoryProfilerWindow::SpawnTab_TimingView(const FSpawnTabA
 	IModularFeatures::Get().RegisterModularFeature(Insights::TimingViewExtenderFeatureName, SharedState.Get());
 
 	TimingView->Reset(true);
+	ResetTimingViewMarkers();
 	TimingView->HideAllDefaultTracks();
 
 	DockTab->SetOnTabClosed(SDockTab::FOnTabClosedCallback::CreateRaw(this, &SMemoryProfilerWindow::OnTimingViewTabClosed));
@@ -226,6 +308,66 @@ void SMemoryProfilerWindow::OnMemTagTreeViewTabClosed(TSharedRef<SDockTab> TabBe
 
 ////////////////////////////////////////////////////////////////////////////////////////////////////
 
+TSharedRef<SDockTab> SMemoryProfilerWindow::SpawnTab_MemAllocTableTreeView(const FSpawnTabArgs& Args, int32 TabIndex)
+{
+	//FMemoryProfilerManager::Get()->SetMemAllocTableTreeViewVisible(TabIndex, true);
+
+	TSharedRef<Insights::FMemAllocTable> MemAllocTable = MakeShared<Insights::FMemAllocTable>();
+	MemAllocTable->Reset();
+
+	TSharedPtr<Insights::SMemAllocTableTreeView> MemAllocTableTreeView;
+
+	const TSharedRef<SDockTab> DockTab = SNew(SDockTab)
+		.ShouldAutosize(false)
+		.TabRole(ETabRole::PanelTab)
+		[
+			SAssignNew(MemAllocTableTreeView, Insights::SMemAllocTableTreeView, MemAllocTable)
+		];
+
+	MemAllocTableTreeView->SetTabIndex(TabIndex);
+	MemAllocTableTreeViews.Add(MemAllocTableTreeView);
+	UpdateMemAllocTableTreeView(MemAllocTableTreeView);
+
+	DockTab->SetOnTabClosed(SDockTab::FOnTabClosedCallback::CreateRaw(this, &SMemoryProfilerWindow::OnMemAllocTableTreeViewTabClosed));
+
+	return DockTab;
+}
+
+////////////////////////////////////////////////////////////////////////////////////////////////////
+
+void SMemoryProfilerWindow::OnMemAllocTableTreeViewTabClosed(TSharedRef<SDockTab> TabBeingClosed)
+{
+	TSharedRef<Insights::SMemAllocTableTreeView> MemAllocTableTreeView = StaticCastSharedRef<Insights::SMemAllocTableTreeView>(TabBeingClosed->GetContent());
+
+	//FMemoryProfilerManager::Get()->SetMemAllocTableTreeViewVisible(MemAllocTableTreeView->GetTabIndex(), false);
+
+	MemAllocTableTreeViews.Remove(MemAllocTableTreeView);
+}
+
+////////////////////////////////////////////////////////////////////////////////////////////////////
+
+TSharedPtr<Insights::SMemAllocTableTreeView> SMemoryProfilerWindow::ShowMemAllocTableTreeViewTab()
+{
+	LastMemAllocTableTreeViewIndex = (LastMemAllocTableTreeViewIndex + 1) % MaxMemAllocTableTreeViews;
+	FName TabId = FMemoryProfilerTabs::MemAllocTableTreeViewID;
+	TabId.SetNumber(LastMemAllocTableTreeViewIndex);
+
+	if (TabManager->HasTabSpawner(TabId))
+	{
+		TSharedPtr<SDockTab> Tab = TabManager->TryInvokeTab(TabId);
+		if (Tab)
+		{
+			TSharedRef<Insights::SMemAllocTableTreeView> MemAllocTableTreeView = StaticCastSharedRef<Insights::SMemAllocTableTreeView>(Tab->GetContent());
+			ensure(MemAllocTableTreeView->GetTabIndex() == LastMemAllocTableTreeViewIndex);
+			return MemAllocTableTreeView;
+		}
+	}
+
+	return nullptr;
+}
+
+////////////////////////////////////////////////////////////////////////////////////////////////////
+
 void SMemoryProfilerWindow::Construct(const FArguments& InArgs, const TSharedRef<SDockTab>& ConstructUnderMajorTab, const TSharedPtr<SWindow>& ConstructUnderWindow)
 {
 	// Create & initialize tab manager.
@@ -253,6 +395,17 @@ void SMemoryProfilerWindow::Construct(const FArguments& InArgs, const TSharedRef
 		.SetIcon(FSlateIcon(FInsightsStyle::GetStyleSetName(), "MemTagTreeView.Icon.Small"))
 		.SetGroup(AppMenuGroup);
 
+	FName MemAllocTableTreeViewTabId = FMemoryProfilerTabs::MemAllocTableTreeViewID;
+	for (int32 TabIndex = 0; TabIndex < MaxMemAllocTableTreeViews; ++TabIndex)
+	{
+		MemAllocTableTreeViewTabId.SetNumber(TabIndex);
+		FText MemAllocTableTreeViewTabDisplayName = FText::Format(LOCTEXT("MemoryProfiler.MemAllocTableTreeViewTabTitle", "Allocs Table {0}"), FText::AsNumber(TabIndex));
+		TabManager->RegisterTabSpawner(MemAllocTableTreeViewTabId, FOnSpawnTab::CreateRaw(this, &SMemoryProfilerWindow::SpawnTab_MemAllocTableTreeView, TabIndex))
+			.SetDisplayName(MemAllocTableTreeViewTabDisplayName)
+			.SetIcon(FSlateIcon(FInsightsStyle::GetStyleSetName(), "MemTagTreeView.Icon.Small")) // TODO
+			.SetGroup(AppMenuGroup);
+	}
+
 	TSharedPtr<FMemoryProfilerManager> MemoryProfilerManager = FMemoryProfilerManager::Get();
 	ensure(MemoryProfilerManager.IsValid());
 
@@ -302,7 +455,7 @@ void SMemoryProfilerWindow::Construct(const FArguments& InArgs, const TSharedRef
 		FName(TEXT("Menu"))
 	);
 
-	TSharedRef<SWidget> MultiboxWidget = MenuBarBuilder.MakeWidget();
+	TSharedRef<SWidget> MenuWidget = MenuBarBuilder.MakeWidget();
 
 	ChildSlot
 		[
@@ -330,7 +483,7 @@ void SMemoryProfilerWindow::Construct(const FArguments& InArgs, const TSharedRef
 					+ SVerticalBox::Slot()
 						.AutoHeight()
 						[
-							MultiboxWidget
+							MenuWidget
 						]
 
 					+ SVerticalBox::Slot()
@@ -357,7 +510,7 @@ void SMemoryProfilerWindow::Construct(const FArguments& InArgs, const TSharedRef
 		];
 
 	// Tell tab-manager about the global menu bar.
-	TabManager->SetMenuMultiBox(MenuBarBuilder.GetMultiBox(), MultiboxWidget);
+	TabManager->SetMenuMultiBox(MenuBarBuilder.GetMultiBox(), MenuWidget);
 }
 
 END_SLATE_FUNCTION_BUILD_OPTIMIZATION
diff --git a/Engine/Source/Developer/TraceInsights/Private/Insights/MemoryProfiler/Widgets/SMemoryProfilerWindow.h b/Engine/Source/Developer/TraceInsights/Private/Insights/MemoryProfiler/Widgets/SMemoryProfilerWindow.h
--- a/Engine/Source/Developer/TraceInsights/Private/Insights/MemoryProfiler/Widgets/SMemoryProfilerWindow.h
+++ b/Engine/Source/Developer/TraceInsights/Private/Insights/MemoryProfiler/Widgets/SMemoryProfilerWindow.h
@@ -26,7 +26,9 @@ class SMemTagTreeView;
 
 namespace Insights
 {
+	class FTimeMarker;
 	class STableTreeView;
+	class SMemAllocTableTreeView;
 }
 
 ////////////////////////////////////////////////////////////////////////////////////////////////////
@@ -37,6 +39,7 @@ struct FMemoryProfilerTabs
 	static const FName ToolbarID;
 	static const FName TimingViewID;
 	static const FName MemTagTreeViewID;
+	static const FName MemAllocTableTreeViewID; // base id
 };
 
 ////////////////////////////////////////////////////////////////////////////////////////////////////
@@ -55,8 +58,11 @@ public:
 	SLATE_END_ARGS()
 
 	void Reset();
+	void ResetTimingViewMarkers();
+	void UpdateTimingViewMarkers();
 	void UpdateTableTreeViews();
 	void UpdateMemTagTreeView();
+	void UpdateMemAllocTableTreeView(TSharedPtr<Insights::SMemAllocTableTreeView> MemAllocTableTreeView);
 
 	/** Constructs this widget. */
 	void Construct(const FArguments& InArgs, const TSharedRef<SDockTab>& ConstructUnderMajorTab, const TSharedPtr<SWindow>& ConstructUnderWindow);
@@ -65,12 +71,19 @@ public:
 	void HideTab(const FName& TabID);
 	void ShowHideTab(const FName& TabID, bool bShow) { bShow ? ShowTab(TabID) : HideTab(TabID); }
 
+	TSharedPtr<Insights::SMemAllocTableTreeView> ShowMemAllocTableTreeViewTab();
+
 	TSharedPtr<FTabManager> GetTabManager() const { return TabManager; }
 
 	TSharedPtr<STimingView> GetTimingView() const { return TimingView; }
 	TSharedPtr<SMemTagTreeView> GetMemTagTreeView() const { return MemTagTreeView; }
 
-	FMemorySharedState& GetSharedState() const { return *SharedState; }
+	uint32 GetNumCustomTimeMarkers() const { return (uint32)CustomTimeMarkers.Num(); }
+	const TSharedRef<Insights::FTimeMarker>& GetCustomTimeMarker(uint32 Index) const { return CustomTimeMarkers[Index]; }
+	const TArray<TSharedRef<Insights::FTimeMarker>>& GetCustomTimeMarkers() const { return CustomTimeMarkers; }
+
+	FMemorySharedState& GetSharedState() { return *SharedState; }
+	const FMemorySharedState& GetSharedState() const { return *SharedState; }
 
 private:
 	TSharedRef<SDockTab> SpawnTab_Toolbar(const FSpawnTabArgs& Args);
@@ -82,6 +95,9 @@ private:
 	TSharedRef<SDockTab> SpawnTab_MemTagTreeView(const FSpawnTabArgs& Args);
 	void OnMemTagTreeViewTabClosed(TSharedRef<SDockTab> TabBeingClosed);
 
+	TSharedRef<SDockTab> SpawnTab_MemAllocTableTreeView(const FSpawnTabArgs& Args, int32 TabIndex);
+	void OnMemAllocTableTreeViewTabClosed(TSharedRef<SDockTab> TabBeingClosed);
+
 	/**
 	 * Fill the main menu with menu items.
 	 *
@@ -156,11 +172,18 @@ private:
 private:
 	/** The Timing view (multi-track) widget */
 	TSharedPtr<STimingView> TimingView;
+	TArray<TSharedRef<Insights::FTimeMarker>> CustomTimeMarkers;
 	TSharedPtr<FMemorySharedState> SharedState;
 
 	/** The "LLM Tags" tree view widget */
 	TSharedPtr<SMemTagTreeView> MemTagTreeView;
 
+	/** The list of Allocations table tree view widgets */
+	TArray<TSharedPtr<Insights::SMemAllocTableTreeView>> MemAllocTableTreeViews;
+
+	const int32 MaxMemAllocTableTreeViews = 4;
+	int32 LastMemAllocTableTreeViewIndex = -1;
+
 	/** Holds the tab manager that manages the front-end's tabs. */
 	TSharedPtr<FTabManager> TabManager;
 
diff --git a/Engine/Source/Developer/TraceInsights/Private/Insights/ViewModels/TimeRulerTrack.cpp b/Engine/Source/Developer/TraceInsights/Private/Insights/ViewModels/TimeRulerTrack.cpp
--- a/Engine/Source/Developer/TraceInsights/Private/Insights/ViewModels/TimeRulerTrack.cpp
+++ b/Engine/Source/Developer/TraceInsights/Private/Insights/ViewModels/TimeRulerTrack.cpp
@@ -82,6 +82,13 @@ void FTimeRulerTrack::RemoveTimeMarker(TSharedRef<Insights::FTimeMarker> InTimeM
 
 ////////////////////////////////////////////////////////////////////////////////////////////////////
 
+void FTimeRulerTrack::RemoveAllTimeMarkers()
+{
+	TimeMarkers.Reset();
+}
+
+////////////////////////////////////////////////////////////////////////////////////////////////////
+
 TSharedPtr<Insights::FTimeMarker> FTimeRulerTrack::GetTimeMarkerByName(const FString& InTimeMarkerName)
 {
 	for (TSharedRef<Insights::FTimeMarker>& TimeMarker : TimeMarkers)
@@ -341,30 +348,51 @@ void FTimeRulerTrack::DrawTimeMarker(const ITimingTrackDrawContext& Context, con
 	DrawContext.LayerId++;
 
 	const FVector2D& MousePosition = Context.GetMousePosition();
+
+	const bool bIsMouseOverTrack =  !MousePosition.IsZero() &&
+									MousePosition.Y >= GetPosY() &&
+									MousePosition.Y < GetPosY() + GetHeight();
+
 	constexpr float FixedHalfWidth = 42.0f;
-	const bool bIsMouseOver = !MousePosition.IsZero() &&
-							  MousePosition.Y >= GetPosY() + TimeMarkerY &&
-							  MousePosition.Y < GetPosY() + TimeMarkerY + BoxHeight &&
-							  MousePosition.X >= TimeMarkerX - FixedHalfWidth &&
-							  MousePosition.X < TimeMarkerX + FixedHalfWidth;
+	const bool bIsMouseOverMarker = bIsMouseOverTrack &&
+									MousePosition.Y >= GetPosY() + TimeMarkerY &&
+									MousePosition.Y < GetPosY() + TimeMarkerY + BoxHeight &&
+									MousePosition.X >= TimeMarkerX - FixedHalfWidth &&
+									MousePosition.X < TimeMarkerX + FixedHalfWidth;
 
-	// Time at current marker
+	// Decide what text to display.
 	FString TimeMarkerText;
-	if (TimeMarker.GetName().Len() > 0)
-	{
-		TimeMarkerText = TimeMarker.GetName() + TEXT(": ");
-	}
-	const double DT = 100.0 / Viewport.GetScaleX();
-	const double Precision = FMath::Max(DT / 100.0, TimeUtils::Nanosecond);
-	if (bIsMouseOver)
+	if (bIsMouseOverTrack)
 	{
-		// If mouse is hovering the time ruler, format time with a better precision (split seconds in ms, us, ns and ps).
-		TimeMarkerText += TimeUtils::FormatTimeSplit(TimeMarker.GetTime(), Precision);
+		if (TimeMarker.GetName().Len() > 0)
+		{
+			TimeMarkerText = TimeMarker.GetName() + TEXT(": ");
+		}
+
+		// Format time value with one more digit than the time at major tick marks.
+		const double DT = 100.0 / Viewport.GetScaleX();
+		const double Precision = FMath::Max(DT / 100.0, TimeUtils::Nanosecond);
+
+		if (bIsMouseOverMarker)
+		{
+			// If mouse is hovering the time marker, format time with a better precision (split seconds in ms, us, ns and ps).
+			TimeMarkerText += TimeUtils::FormatTimeSplit(TimeMarker.GetTime(), Precision);
+		}
+		else
+		{
+			TimeMarkerText += TimeUtils::FormatTime(TimeMarker.GetTime(), Precision);
+		}
 	}
 	else
 	{
-		// Format current time with one more digit than the time at major tick marks.
-		TimeMarkerText += TimeUtils::FormatTime(TimeMarker.GetTime(), Precision);
+		if (TimeMarker.GetName().Len() > 0)
+		{
+			TimeMarkerText = TimeMarker.GetName();
+		}
+		else
+		{
+			TimeMarkerText = TEXT("T");
+		}
 	}
 
 	const TSharedRef<FSlateFontMeasure> FontMeasureService = FSlateApplication::Get().GetRenderer()->GetFontMeasureService();
@@ -378,7 +406,7 @@ void FTimeRulerTrack::DrawTimeMarker(const ITimingTrackDrawContext& Context, con
 	}
 
 	const FLinearColor TextBackgroundColor(TimeMarker.GetColor().CopyWithNewOpacity(1.0f));
-	const FLinearColor TextForegroundColor(0.1f, 0.1f, 0.1f, 1.0f);
+	const FLinearColor TextForegroundColor(0.07f, 0.07f, 0.07f, 1.0f);
 
 	// Fill the time marker box.
 	const float BoxWidth = TimeMarker.GetCrtTextWidth() + 4.0f;
diff --git a/Engine/Source/Developer/TraceInsights/Private/Insights/ViewModels/TimeRulerTrack.h b/Engine/Source/Developer/TraceInsights/Private/Insights/ViewModels/TimeRulerTrack.h
--- a/Engine/Source/Developer/TraceInsights/Private/Insights/ViewModels/TimeRulerTrack.h
+++ b/Engine/Source/Developer/TraceInsights/Private/Insights/ViewModels/TimeRulerTrack.h
@@ -74,8 +74,11 @@ public:
 
 	void SetSelection(const bool bInIsSelecting, const double InSelectionStartTime, const double InSelectionEndTime);
 
+	TArray<TSharedRef<Insights::FTimeMarker>>& GetTimeMarkers() { return TimeMarkers; }
+	const TArray<TSharedRef<Insights::FTimeMarker>>& GetTimeMarkers() const { return TimeMarkers; }
 	void AddTimeMarker(TSharedRef<Insights::FTimeMarker> InTimeMarker);
 	void RemoveTimeMarker(TSharedRef<Insights::FTimeMarker> InTimeMarker);
+	void RemoveAllTimeMarkers();
 
 	TSharedPtr<Insights::FTimeMarker> GetTimeMarkerByName(const FString& InTimeMarkerName);
 	TSharedPtr<Insights::FTimeMarker> GetTimeMarkerAtPos(const FVector2D& InPosition, const FTimingTrackViewport& InViewport);
@@ -107,7 +110,7 @@ private:
 
 	/**
 	 * The sorted list of all registered time markers. It defines the draw order of time markers.
-	 * The time marker currenly scrubbing will be moved at the end of the list in order to be displayed on top of other markers.
+	 * The time marker currently scrubbing will be moved at the end of the list in order to be displayed on top of other markers.
 	 */
 	TArray<TSharedRef<Insights::FTimeMarker>> TimeMarkers;
 
diff --git a/Engine/Source/Developer/TraceInsights/Private/Insights/Widgets/STimingView.h b/Engine/Source/Developer/TraceInsights/Private/Insights/Widgets/STimingView.h
--- a/Engine/Source/Developer/TraceInsights/Private/Insights/Widgets/STimingView.h
+++ b/Engine/Source/Developer/TraceInsights/Private/Insights/Widgets/STimingView.h
@@ -89,6 +89,14 @@ public:
 
 	void HideAllDefaultTracks();
 
+	/** Gets the time ruler track. It includes the custom time markers (ones user can drag with mouse). */
+	TSharedRef<FTimeRulerTrack> GetTimeRulerTrack() { return TimeRulerTrack; }
+	const TSharedRef<const FTimeRulerTrack> GetTimeRulerTrack() const { return TimeRulerTrack; }
+
+	/** Gets the default (custom) time marker (for backward compatibility). */
+	TSharedRef<Insights::FTimeMarker> GetDefaultTimeMarker() { return DefaultTimeMarker; }
+	const TSharedRef<const Insights::FTimeMarker> GetDefaultTimeMarker() const { return DefaultTimeMarker; }
+
 	/** Resets internal widget's data to the default one. */
 	void Reset(bool bIsFirstReset = false);
 
