[Insights] Timing View: Improved performance when having large number of frames.

diff --git a/Engine/Source/Developer/TraceInsights/Private/Insights/ViewModels/FrameTimingTrack.cpp b/Engine/Source/Developer/TraceInsights/Private/Insights/ViewModels/FrameTimingTrack.cpp
--- a/Engine/Source/Developer/TraceInsights/Private/Insights/ViewModels/FrameTimingTrack.cpp
+++ b/Engine/Source/Developer/TraceInsights/Private/Insights/ViewModels/FrameTimingTrack.cpp
@@ -187,8 +187,13 @@ void FFrameTimingTrack::BuildDrawState(ITimingEventsTrackDrawStateBuilder& Build
 		const Trace::IFrameProvider& FramesProvider = Trace::ReadFrameProvider(*Session.Get());
 		const FTimingTrackViewport& Viewport = Context.GetViewport();
 
-		uint64 StartIndex = 0;
-		uint64 EndIndex = FramesProvider.GetFrameCount(static_cast<ETraceFrameType>(FrameType));
+		const TArray64<double>& FrameStartTimes = FramesProvider.GetFrameStartTimes(static_cast<ETraceFrameType>(FrameType));
+
+		const int64 StartLowerBound = Algo::LowerBound(FrameStartTimes, Viewport.GetStartTime());
+		const uint64 StartIndex = (StartLowerBound > 0) ? StartLowerBound - 1 : 0;
+
+		const int64 EndLowerBound = Algo::LowerBound(FrameStartTimes, Viewport.GetEndTime());
+		const uint64 EndIndex = EndLowerBound;
 
 		const uint32 Color = FFrameTrackDrawHelper::GetColor32ByFrameType(FrameType);
 
@@ -216,8 +221,13 @@ void FFrameTimingTrack::BuildFilteredDrawState(ITimingEventsTrackDrawStateBuilde
 			const Trace::IFrameProvider& FramesProvider = Trace::ReadFrameProvider(*Session.Get());
 			const FTimingTrackViewport& Viewport = Context.GetViewport();
 
-			uint64 StartIndex = 0;
-			uint64 EndIndex = FramesProvider.GetFrameCount(static_cast<ETraceFrameType>(FrameType));
+			const TArray64<double>& FrameStartTimes = FramesProvider.GetFrameStartTimes(static_cast<ETraceFrameType>(FrameType));
+
+			const int64 StartLowerBound = Algo::LowerBound(FrameStartTimes, Viewport.GetStartTime());
+			const uint64 StartIndex = (StartLowerBound > 0) ? StartLowerBound - 1 : 0;
+
+			const int64 EndLowerBound = Algo::LowerBound(FrameStartTimes, Viewport.GetEndTime());
+			const uint64 EndIndex = EndLowerBound;
 
 			const uint32 Color = FFrameTrackDrawHelper::GetColor32ByFrameType(FrameType);
 
@@ -456,8 +466,13 @@ bool FFrameTimingTrack::FindFrame(const FTimingEventSearchParameters& InParamete
 				Trace::FAnalysisSessionReadScope SessionReadScope(*Session.Get());
 				const Trace::IFrameProvider& FramesProvider = Trace::ReadFrameProvider(*Session.Get());
 
-				uint64 StartIndex = 0;
-				uint64 EndIndex = FramesProvider.GetFrameCount(static_cast<ETraceFrameType>(FrameType));
+				const TArray64<double>& FrameStartTimes = FramesProvider.GetFrameStartTimes(static_cast<ETraceFrameType>(FrameType));
+
+				const int64 StartLowerBound = Algo::LowerBound(FrameStartTimes, InContext.GetParameters().StartTime);
+				const uint64 StartIndex = (StartLowerBound > 0) ? StartLowerBound - 1 : 0;
+
+				const int64 EndLowerBound = Algo::LowerBound(FrameStartTimes, InContext.GetParameters().EndTime);
+				const uint64 EndIndex = EndLowerBound;
 
 				FramesProvider.EnumerateFrames(static_cast<ETraceFrameType>(FrameType), StartIndex, EndIndex, [&InContext](const Trace::FFrame& Frame)
 				{
diff --git a/Engine/Source/Developer/TraceInsights/Private/Insights/ViewModels/TimingGraphTrack.cpp b/Engine/Source/Developer/TraceInsights/Private/Insights/ViewModels/TimingGraphTrack.cpp
--- a/Engine/Source/Developer/TraceInsights/Private/Insights/ViewModels/TimingGraphTrack.cpp
+++ b/Engine/Source/Developer/TraceInsights/Private/Insights/ViewModels/TimingGraphTrack.cpp
@@ -204,8 +204,16 @@ void FTimingGraphTrack::UpdateFrameSeries(FTimingGraphSeries& Series, const FTim
 	{
 		Trace::FAnalysisSessionReadScope SessionReadScope(*Session.Get());
 		const Trace::IFrameProvider& FramesProvider = ReadFrameProvider(*Session.Get());
-		uint64 FrameCount = FramesProvider.GetFrameCount(Series.FrameType);
-		FramesProvider.EnumerateFrames(Series.FrameType, 0, FrameCount - 1, [&Builder](const Trace::FFrame& Frame)
+
+		const TArray64<double>& FrameStartTimes = FramesProvider.GetFrameStartTimes(Series.FrameType);
+
+		const int64 StartLowerBound = Algo::LowerBound(FrameStartTimes, Viewport.GetStartTime());
+		const uint64 StartIndex = (StartLowerBound > 1) ? StartLowerBound - 2 : 0;
+
+		const int64 EndLowerBound = Algo::LowerBound(FrameStartTimes, Viewport.GetEndTime());
+		const uint64 EndIndex = EndLowerBound + 1;
+
+		FramesProvider.EnumerateFrames(Series.FrameType, StartIndex, EndIndex, [&Builder](const Trace::FFrame& Frame)
 		{
 			//TODO: add a "frame converter" (i.e. to fps, miliseconds or seconds)
 			const double Duration = Frame.EndTime - Frame.StartTime;
diff --git a/Engine/Source/Developer/TraceServices/Private/Model/Frames.cpp b/Engine/Source/Developer/TraceServices/Private/Model/Frames.cpp
--- a/Engine/Source/Developer/TraceServices/Private/Model/Frames.cpp
+++ b/Engine/Source/Developer/TraceServices/Private/Model/Frames.cpp
@@ -35,7 +35,7 @@ void FFrameProvider::EnumerateFrames(ETraceFrameType FrameType, uint64 Start, ui
 	{
 		return;
 	}
-	for (auto Iterator = Frames[FrameType].GetIteratorFromItem(Start); Iterator; ++Iterator)
+	for (auto Iterator = Frames[FrameType].GetIteratorFromItem(Start); Iterator && Iterator->Index < End; ++Iterator)
 	{
 		Callback(*Iterator);
 	}
