[Insights] Add copy to clipboard functionality to the caller and callees panels #rb Ionut.Matasaru

diff --git a/Engine/Source/Developer/TraceInsights/Private/Insights/Table/ViewModels/Table.cpp b/Engine/Source/Developer/TraceInsights/Private/Insights/Table/ViewModels/Table.cpp
--- a/Engine/Source/Developer/TraceInsights/Private/Insights/Table/ViewModels/Table.cpp
+++ b/Engine/Source/Developer/TraceInsights/Private/Insights/Table/ViewModels/Table.cpp
@@ -120,6 +120,44 @@ void FTable::GetVisibleColumns(TArray<TSharedRef<FTableColumn>>& InArray) const
 
 ////////////////////////////////////////////////////////////////////////////////////////////////////
 
+void FTable::GetVisibleColumnsData(const TArray<Insights::FBaseTreeNodePtr>& InNodes, FString& OutData) const
+{
+	constexpr TCHAR Separator = TEXT('\t');
+
+	TArray<TSharedRef<Insights::FTableColumn>> VisibleColumns;
+	GetVisibleColumns(VisibleColumns);
+
+	// Table headers
+	for (const TSharedRef<Insights::FTableColumn>& ColumnRef : VisibleColumns)
+	{
+		OutData += ColumnRef->GetShortName().ToString().ReplaceCharWithEscapedChar() + Separator;
+	}
+
+	if (OutData.Len() > 0)
+	{
+		OutData.RemoveAt(OutData.Len() - 1, 1, false);
+		OutData.AppendChar(TEXT('\n'));
+	}
+
+	// Selected items
+	for (Insights::FBaseTreeNodePtr Node : InNodes)
+	{
+		for (const TSharedRef<Insights::FTableColumn>& ColumnRef : VisibleColumns)
+		{
+			FText NodeText = ColumnRef->GetValueAsText(*Node);
+			OutData += NodeText.ToString().ReplaceCharWithEscapedChar() + Separator;
+		}
+
+		if (OutData.Len() > 0)
+		{
+			OutData.RemoveAt(OutData.Len() - 1, 1, false);
+			OutData.AppendChar(TEXT('\n'));
+		}
+	}
+}
+
+////////////////////////////////////////////////////////////////////////////////////////////////////
+
 } // namespace Insights
 
 #undef LOCTEXT_NAMESPACE
diff --git a/Engine/Source/Developer/TraceInsights/Private/Insights/Table/ViewModels/Table.h b/Engine/Source/Developer/TraceInsights/Private/Insights/Table/ViewModels/Table.h
--- a/Engine/Source/Developer/TraceInsights/Private/Insights/Table/ViewModels/Table.h
+++ b/Engine/Source/Developer/TraceInsights/Private/Insights/Table/ViewModels/Table.h
@@ -51,6 +51,8 @@ public:
 
 	int32 GetColumnPositionIndex(const FName& ColumnId) const;
 
+	void GetVisibleColumnsData(const TArray<TSharedPtr<class FBaseTreeNode>> &InNodes, FString& OutData) const;
+
 protected:
 	void ResetColumns() { Columns.Reset(); }
 	void AddColumn(TSharedRef<FTableColumn> Column);
diff --git a/Engine/Source/Developer/TraceInsights/Private/Insights/Widgets/STimerTreeView.cpp b/Engine/Source/Developer/TraceInsights/Private/Insights/Widgets/STimerTreeView.cpp
--- a/Engine/Source/Developer/TraceInsights/Private/Insights/Widgets/STimerTreeView.cpp
+++ b/Engine/Source/Developer/TraceInsights/Private/Insights/Widgets/STimerTreeView.cpp
@@ -5,6 +5,7 @@
 #include "EditorStyleSet.h"
 #include "Framework/Application/SlateApplication.h"
 #include "Framework/MultiBox/MultiBoxBuilder.h"
+#include "HAL/PlatformApplicationMisc.h"
 #include "SlateOptMacros.h"
 #include "Widgets/Layout/SScrollBox.h"
 #include "Widgets/SToolTip.h"
@@ -22,6 +23,35 @@
 
 #define LOCTEXT_NAMESPACE "STimerTreeView"
 
+////////////////////////////////////////////////////////////////////////////////////////////////////
+// FTimerTreeViewCommands
+////////////////////////////////////////////////////////////////////////////////////////////////////
+
+class FTimerTreeViewCommands : public TCommands<FTimerTreeViewCommands>
+{
+public:
+	FTimerTreeViewCommands()
+		: TCommands<FTimerTreeViewCommands>(TEXT("FTimerTreeViewCommands"), NSLOCTEXT("FTimerTreeViewCommands", "Timer Tree View Commands", "Timer Tree View Commands"), NAME_None, FEditorStyle::Get().GetStyleSetName())
+	{
+	}
+
+	virtual ~FTimerTreeViewCommands()
+	{
+	}
+
+	// UI_COMMAND takes long for the compiler to optimize
+	PRAGMA_DISABLE_OPTIMIZATION
+	virtual void RegisterCommands() override
+	{
+		UI_COMMAND(Command_CopyToClipboard, "Copy To Clipboard", "Copies selection to clipboard", EUserInterfaceActionType::Button, FInputChord(EModifierKey::Control, EKeys::C));
+	}
+	PRAGMA_ENABLE_OPTIMIZATION
+
+	TSharedPtr<FUICommandInfo> Command_CopyToClipboard;
+};
+
+////////////////////////////////////////////////////////////////////////////////////////////////////
+// STimerTreeView
 ////////////////////////////////////////////////////////////////////////////////////////////////////
 
 STimerTreeView::STimerTreeView()
@@ -49,6 +79,15 @@ STimerTreeView::~STimerTreeView()
 
 ////////////////////////////////////////////////////////////////////////////////////////////////////
 
+void STimerTreeView::InitCommandList()
+{
+	FTimerTreeViewCommands::Register();
+	CommandList = MakeShared<FUICommandList>();
+	CommandList->MapAction(FTimerTreeViewCommands::Get().Command_CopyToClipboard, FExecuteAction::CreateSP(this, &STimerTreeView::ContextMenu_CopySelectedToClipboard_Execute), FCanExecuteAction::CreateSP(this, &STimerTreeView::ContextMenu_CopySelectedToClipboard_CanExecute));
+}
+
+////////////////////////////////////////////////////////////////////////////////////////////////////
+
 BEGIN_SLATE_FUNCTION_BUILD_OPTIMIZATION
 void STimerTreeView::Construct(const FArguments& InArgs, const FText& InViewName)
 {
@@ -125,6 +164,8 @@ void STimerTreeView::Construct(const FArguments& InArgs, const FText& InViewName
 	InitializeAndShowHeaderColumns();
 
 	CreateSortings();
+
+	InitCommandList();
 }
 END_SLATE_FUNCTION_BUILD_OPTIMIZATION
 
@@ -167,7 +208,7 @@ TSharedPtr<SWidget> STimerTreeView::TreeView_GetMenuContent()
 	}
 
 	const bool bShouldCloseWindowAfterMenuSelection = true;
-	FMenuBuilder MenuBuilder(bShouldCloseWindowAfterMenuSelection, NULL);
+	FMenuBuilder MenuBuilder(bShouldCloseWindowAfterMenuSelection, CommandList.ToSharedRef());
 
 	// Selection menu
 	MenuBuilder.BeginSection("Selection", LOCTEXT("ContextMenu_Header_Selection", "Selection"));
@@ -193,6 +234,15 @@ TSharedPtr<SWidget> STimerTreeView::TreeView_GetMenuContent()
 
 	MenuBuilder.BeginSection("Misc", LOCTEXT("ContextMenu_Header_Misc", "Miscellaneous"));
 	{
+		MenuBuilder.AddMenuEntry
+		(
+			FTimerTreeViewCommands::Get().Command_CopyToClipboard,
+			NAME_None,
+			TAttribute<FText>(),
+			TAttribute<FText>(),
+			FSlateIcon(FEditorStyle::GetStyleSetName(), "Profiler.Misc.CopyToClipboard")
+		);
+
 		MenuBuilder.AddSubMenu
 		(
 			LOCTEXT("ContextMenu_Header_Misc_Sort", "Sort By"),
@@ -1026,4 +1076,55 @@ void STimerTreeView::ExpandNodesRec(FTimerNodePtr NodePtr, int32 Depth)
 
 ////////////////////////////////////////////////////////////////////////////////////////////////////
 
+bool STimerTreeView::ContextMenu_CopySelectedToClipboard_CanExecute() const
+{
+	const TArray<FTimerNodePtr> SelectedNodes = TreeView->GetSelectedItems();
+
+	return SelectedNodes.Num() > 0;
+}
+
+////////////////////////////////////////////////////////////////////////////////////////////////////
+
+void STimerTreeView::ContextMenu_CopySelectedToClipboard_Execute()
+{
+	if (!Table->IsValid())
+	{
+		return;
+	}
+
+	TArray<Insights::FBaseTreeNodePtr> SelectedNodes;
+	for (FTimerNodePtr TimerPtr : TreeView->GetSelectedItems())
+	{
+		SelectedNodes.Add(TimerPtr);
+	}
+
+	if (SelectedNodes.Num() == 0)
+	{
+		return;
+	}
+
+	FString ClipboardText;
+
+	if (CurrentSorter.IsValid())
+	{
+		CurrentSorter->Sort(SelectedNodes, ColumnSortMode == EColumnSortMode::Ascending ? Insights::ESortMode::Ascending : Insights::ESortMode::Descending);
+	}
+
+	Table->GetVisibleColumnsData(SelectedNodes, ClipboardText);
+
+	if (ClipboardText.Len() > 0)
+	{
+		FPlatformApplicationMisc::ClipboardCopy(*ClipboardText);
+	}
+}
+
+////////////////////////////////////////////////////////////////////////////////////////////////////
+
+FReply STimerTreeView::OnKeyDown(const FGeometry& MyGeometry, const FKeyEvent& InKeyEvent)
+{
+	return CommandList->ProcessCommandBindings(InKeyEvent) == true ? FReply::Handled() : FReply::Unhandled();
+}
+
+////////////////////////////////////////////////////////////////////////////////////////////////////
+
 #undef LOCTEXT_NAMESPACE
diff --git a/Engine/Source/Developer/TraceInsights/Private/Insights/Widgets/STimerTreeView.h b/Engine/Source/Developer/TraceInsights/Private/Insights/Widgets/STimerTreeView.h
--- a/Engine/Source/Developer/TraceInsights/Private/Insights/Widgets/STimerTreeView.h
+++ b/Engine/Source/Developer/TraceInsights/Private/Insights/Widgets/STimerTreeView.h
@@ -55,7 +55,10 @@ public:
 	void Reset();
 	void SetTree(const Trace::FTimingProfilerButterflyNode& Root);
 
+	virtual FReply OnKeyDown(const FGeometry& MyGeometry, const FKeyEvent& InKeyEvent) override;
+
 private:
+	void InitCommandList();
 	FTimerNodePtr CreateTimerNodeRec(const Trace::FTimingProfilerButterflyNode& Node);
 	void ExpandNodesRec(FTimerNodePtr NodePtr, int32 Depth);
 
@@ -66,6 +69,9 @@ private:
 	void TreeView_BuildSortByMenu(FMenuBuilder& MenuBuilder);
 	void TreeView_BuildViewColumnMenu(FMenuBuilder& MenuBuilder);
 
+	bool ContextMenu_CopySelectedToClipboard_CanExecute() const;
+	void ContextMenu_CopySelectedToClipboard_Execute();
+
 	////////////////////////////////////////////////////////////////////////////////////////////////////
 	// Tree View - Columns' Header
 
@@ -174,6 +180,8 @@ private:
 	/** The view name (ex.: "Callers" or "Callees"). */
 	FText ViewName;
 
+	TSharedPtr<FUICommandList> CommandList;
+
 	//////////////////////////////////////////////////
 	// Tree View, Columns
 
diff --git a/Engine/Source/Developer/TraceInsights/Private/Insights/Widgets/STimersView.cpp b/Engine/Source/Developer/TraceInsights/Private/Insights/Widgets/STimersView.cpp
--- a/Engine/Source/Developer/TraceInsights/Private/Insights/Widgets/STimersView.cpp
+++ b/Engine/Source/Developer/TraceInsights/Private/Insights/Widgets/STimersView.cpp
@@ -39,22 +39,7 @@
 #define LOCTEXT_NAMESPACE "STimersView"
 
 ////////////////////////////////////////////////////////////////////////////////////////////////////
-
-STimersView::STimersView()
-	: Table(MakeShared<Insights::FTable>())
-	, bExpansionSaved(false)
-	, bFilterOutZeroCountTimers(false)
-	, GroupingMode(ETimerGroupingMode::ByType)
-	, AvailableSorters()
-	, CurrentSorter(nullptr)
-	, ColumnBeingSorted(GetDefaultColumnBeingSorted())
-	, ColumnSortMode(GetDefaultColumnSortMode())
-	, Aggregator(MakeShared<Insights::FTimerAggregator>())
-{
-	FMemory::Memset(bTimerTypeIsVisible, 1);
-}
-
-
+// FTimersViewCommands
 ////////////////////////////////////////////////////////////////////////////////////////////////////
 
 class FTimersViewCommands : public TCommands<FTimersViewCommands>
@@ -80,6 +65,26 @@ public:
 	TSharedPtr<FUICommandInfo> Command_CopyToClipboard;
 };
 
+////////////////////////////////////////////////////////////////////////////////////////////////////
+// STimersView
+////////////////////////////////////////////////////////////////////////////////////////////////////
+
+STimersView::STimersView()
+	: Table(MakeShared<Insights::FTable>())
+	, bExpansionSaved(false)
+	, bFilterOutZeroCountTimers(false)
+	, GroupingMode(ETimerGroupingMode::ByType)
+	, AvailableSorters()
+	, CurrentSorter(nullptr)
+	, ColumnBeingSorted(GetDefaultColumnBeingSorted())
+	, ColumnSortMode(GetDefaultColumnSortMode())
+	, Aggregator(MakeShared<Insights::FTimerAggregator>())
+{
+	FMemory::Memset(bTimerTypeIsVisible, 1);
+}
+
+////////////////////////////////////////////////////////////////////////////////////////////////////
+
 void STimersView::InitCommandList()
 {
 	FTimersViewCommands::Register();
@@ -2012,44 +2017,14 @@ void STimersView::ContextMenu_CopySelectedToClipboard_Execute()
 		return;
 	}
 
-	constexpr TCHAR Separator = TEXT('\t');
 	FString ClipboardText;
 
-	TArray<TSharedRef<Insights::FTableColumn>> VisibleColumns;
-	Table->GetVisibleColumns(VisibleColumns);
-
-	// Table headers
-	for (const TSharedRef<Insights::FTableColumn>& ColumnRef : VisibleColumns)
-	{
-		ClipboardText += ColumnRef->GetShortName().ToString().ReplaceCharWithEscapedChar() + Separator;
-	}
-
-	if (ClipboardText.Len() > 0)
-	{
-		ClipboardText.RemoveAt(ClipboardText.Len() - 1, 1, false);
-		ClipboardText.AppendChar(TEXT('\n'));
-	}
-
 	if (CurrentSorter.IsValid())
 	{
 		CurrentSorter->Sort(SelectedNodes, ColumnSortMode == EColumnSortMode::Ascending ? Insights::ESortMode::Ascending : Insights::ESortMode::Descending);
 	}
 
-	// Selected items
-	for (Insights::FBaseTreeNodePtr Node : SelectedNodes)
-	{
-		for (const TSharedRef<Insights::FTableColumn>& ColumnRef : VisibleColumns)
-		{
-			FText NodeText = ColumnRef->GetValueAsText(*Node);
-			ClipboardText += NodeText.ToString().ReplaceCharWithEscapedChar() + Separator;
-		}
-
-		if (ClipboardText.Len() > 0)
-		{
-			ClipboardText.RemoveAt(ClipboardText.Len() - 1, 1, false);
-			ClipboardText.AppendChar(TEXT('\n'));
-		}
-	}
+	Table->GetVisibleColumnsData(SelectedNodes, ClipboardText);
 
 	if (ClipboardText.Len() > 0)
 	{
