[Insights] Fixed variable name for consistency.

diff --git a/Engine/Source/Developer/TraceInsights/Private/Insights/NetworkingProfiler/Widgets/SNetworkingProfilerWindow.cpp b/Engine/Source/Developer/TraceInsights/Private/Insights/NetworkingProfiler/Widgets/SNetworkingProfilerWindow.cpp
--- a/Engine/Source/Developer/TraceInsights/Private/Insights/NetworkingProfiler/Widgets/SNetworkingProfilerWindow.cpp
+++ b/Engine/Source/Developer/TraceInsights/Private/Insights/NetworkingProfiler/Widgets/SNetworkingProfilerWindow.cpp
@@ -305,7 +305,7 @@ void SNetworkingProfilerWindow::Construct(const FArguments& InArgs, const TShare
 		FName(TEXT("Menu"))
 	);
 
-	TSharedRef<SWidget> Menu = MenuBarBuilder.MakeWidget();
+	TSharedRef<SWidget> MenuWidget = MenuBarBuilder.MakeWidget();
 	ChildSlot
 		[
 			SNew(SOverlay)
@@ -332,7 +332,7 @@ void SNetworkingProfilerWindow::Construct(const FArguments& InArgs, const TShare
 					+ SVerticalBox::Slot()
 						.AutoHeight()
 						[
-							Menu
+							MenuWidget
 						]
 
 					+ SVerticalBox::Slot()
@@ -359,7 +359,7 @@ void SNetworkingProfilerWindow::Construct(const FArguments& InArgs, const TShare
 		];
 
 	// Tell tab-manager about the global menu bar.
-	TabManager->SetMenuMultiBox(MenuBarBuilder.GetMultiBox(), Menu);
+	TabManager->SetMenuMultiBox(MenuBarBuilder.GetMultiBox(), MenuWidget);
 
 	BindCommands();
 }
