[Insights] Add copy to clipboard functionality to the Counters panel
+ Fix Non Unity compiler errors

diff --git a/Engine/Source/Developer/TraceInsights/Private/Insights/Widgets/SStatsView.cpp b/Engine/Source/Developer/TraceInsights/Private/Insights/Widgets/SStatsView.cpp
--- a/Engine/Source/Developer/TraceInsights/Private/Insights/Widgets/SStatsView.cpp
+++ b/Engine/Source/Developer/TraceInsights/Private/Insights/Widgets/SStatsView.cpp
@@ -3,7 +3,10 @@
 #include "SStatsView.h"
 
 #include "EditorStyleSet.h"
+#include "Framework/Commands/Commands.h"
+#include "Framework/Commands/UICommandList.h"
 #include "Framework/MultiBox/MultiBoxBuilder.h"
+#include "HAL/PlatformApplicationMisc.h"
 #include "SlateOptMacros.h"
 #include "TraceServices/AnalysisService.h"
 #include "Widgets/Input/SCheckBox.h"
@@ -36,6 +39,35 @@
 
 #define LOCTEXT_NAMESPACE "SStatsView"
 
+////////////////////////////////////////////////////////////////////////////////////////////////////
+// FStatsViewCommands
+////////////////////////////////////////////////////////////////////////////////////////////////////
+
+class FStatsViewCommands : public TCommands<FStatsViewCommands>
+{
+public:
+	FStatsViewCommands()
+		: TCommands<FStatsViewCommands>(TEXT("FStatsViewCommands"), NSLOCTEXT("FStatsViewCommands", "Stats View Commands", "Stats View Commands"), NAME_None, FEditorStyle::Get().GetStyleSetName())
+	{
+	}
+
+	virtual ~FStatsViewCommands()
+	{
+	}
+
+	// UI_COMMAND takes long for the compiler to optimize
+	PRAGMA_DISABLE_OPTIMIZATION
+	virtual void RegisterCommands() override
+	{
+		UI_COMMAND(Command_CopyToClipboard, "Copy To Clipboard", "Copies selection to clipboard", EUserInterfaceActionType::Button, FInputChord(EModifierKey::Control, EKeys::C));
+	}
+	PRAGMA_ENABLE_OPTIMIZATION
+
+	TSharedPtr<FUICommandInfo> Command_CopyToClipboard;
+};
+
+////////////////////////////////////////////////////////////////////////////////////////////////////
+// SStatsView
 ////////////////////////////////////////////////////////////////////////////////////////////////////
 
 SStatsView::SStatsView()
@@ -66,6 +98,15 @@ SStatsView::~SStatsView()
 
 ////////////////////////////////////////////////////////////////////////////////////////////////////
 
+void SStatsView::InitCommandList()
+{
+	FStatsViewCommands::Register();
+	CommandList = MakeShared<FUICommandList>();
+	CommandList->MapAction(FStatsViewCommands::Get().Command_CopyToClipboard, FExecuteAction::CreateSP(this, &SStatsView::ContextMenu_CopySelectedToClipboard_Execute), FCanExecuteAction::CreateSP(this, &SStatsView::ContextMenu_CopySelectedToClipboard_CanExecute));
+}
+
+////////////////////////////////////////////////////////////////////////////////////////////////////
+
 BEGIN_SLATE_FUNCTION_BUILD_OPTIMIZATION
 void SStatsView::Construct(const FArguments& InArgs)
 {
@@ -263,6 +304,8 @@ void SStatsView::Construct(const FArguments& InArgs)
 	CreateGroupByOptionsSources();
 	CreateSortings();
 
+	InitCommandList();
+
 	// Register ourselves with the Insights manager.
 	FInsightsManager::Get()->GetSessionChangedEvent().AddSP(this, &SStatsView::InsightsManager_OnSessionChanged);
 
@@ -310,7 +353,7 @@ TSharedPtr<SWidget> SStatsView::TreeView_GetMenuContent()
 	}
 
 	const bool bShouldCloseWindowAfterMenuSelection = true;
-	FMenuBuilder MenuBuilder(bShouldCloseWindowAfterMenuSelection, NULL);
+	FMenuBuilder MenuBuilder(bShouldCloseWindowAfterMenuSelection, CommandList.ToSharedRef());
 
 	// Selection menu
 	MenuBuilder.BeginSection("Selection", LOCTEXT("ContextMenu_Header_Selection", "Selection"));
@@ -336,19 +379,14 @@ TSharedPtr<SWidget> SStatsView::TreeView_GetMenuContent()
 
 	MenuBuilder.BeginSection("Misc", LOCTEXT("ContextMenu_Header_Misc", "Miscellaneous"));
 	{
-		/*TODO
-		FUIAction Action_CopySelectedToClipboard
-		(
-			FExecuteAction::CreateSP(this, &SStatsView::ContextMenu_CopySelectedToClipboard_Execute),
-			FCanExecuteAction::CreateSP(this, &SStatsView::ContextMenu_CopySelectedToClipboard_CanExecute)
-		);
 		MenuBuilder.AddMenuEntry
 		(
-			LOCTEXT("ContextMenu_Header_Misc_CopySelectedToClipboard", "Copy To Clipboard"),
-			LOCTEXT("ContextMenu_Header_Misc_CopySelectedToClipboard_Desc", "Copies selection to clipboard"),
-			FSlateIcon(FEditorStyle::GetStyleSetName(), "Profiler.Misc.CopyToClipboard"), Action_CopySelectedToClipboard, NAME_None, EUserInterfaceActionType::Button
+			FStatsViewCommands::Get().Command_CopyToClipboard,
+			NAME_None,
+			TAttribute<FText>(),
+			TAttribute<FText>(),
+			FSlateIcon(FEditorStyle::GetStyleSetName(), "Profiler.Misc.CopyToClipboard")
 		);
-		*/
 
 		MenuBuilder.AddSubMenu
 		(
@@ -1942,4 +1980,55 @@ void SStatsView::SelectCounterNode(uint32 CounterId)
 
 ////////////////////////////////////////////////////////////////////////////////////////////////////
 
+bool SStatsView::ContextMenu_CopySelectedToClipboard_CanExecute() const
+{
+	const TArray<FStatsNodePtr> SelectedNodes = TreeView->GetSelectedItems();
+
+	return SelectedNodes.Num() > 0;
+}
+
+////////////////////////////////////////////////////////////////////////////////////////////////////
+
+void SStatsView::ContextMenu_CopySelectedToClipboard_Execute()
+{
+	if (!Table->IsValid())
+	{
+		return;
+	}
+
+	TArray<Insights::FBaseTreeNodePtr> SelectedNodes;
+	for (FStatsNodePtr TimerPtr : TreeView->GetSelectedItems())
+	{
+		SelectedNodes.Add(TimerPtr);
+	}
+
+	if (SelectedNodes.Num() == 0)
+	{
+		return;
+	}
+
+	FString ClipboardText;
+
+	if (CurrentSorter.IsValid())
+	{
+		CurrentSorter->Sort(SelectedNodes, ColumnSortMode == EColumnSortMode::Ascending ? Insights::ESortMode::Ascending : Insights::ESortMode::Descending);
+	}
+
+	Table->GetVisibleColumnsData(SelectedNodes, ClipboardText);
+
+	if (ClipboardText.Len() > 0)
+	{
+		FPlatformApplicationMisc::ClipboardCopy(*ClipboardText);
+	}
+}
+
+////////////////////////////////////////////////////////////////////////////////////////////////////
+
+FReply SStatsView::OnKeyDown(const FGeometry& MyGeometry, const FKeyEvent& InKeyEvent)
+{
+	return CommandList->ProcessCommandBindings(InKeyEvent) == true ? FReply::Handled() : FReply::Unhandled();
+}
+
+////////////////////////////////////////////////////////////////////////////////////////////////////
+
 #undef LOCTEXT_NAMESPACE
diff --git a/Engine/Source/Developer/TraceInsights/Private/Insights/Widgets/SStatsView.h b/Engine/Source/Developer/TraceInsights/Private/Insights/Widgets/SStatsView.h
--- a/Engine/Source/Developer/TraceInsights/Private/Insights/Widgets/SStatsView.h
+++ b/Engine/Source/Developer/TraceInsights/Private/Insights/Widgets/SStatsView.h
@@ -23,6 +23,7 @@
 
 class FMenuBuilder;
 class FTimingGraphTrack;
+class FUICommandList;
 
 namespace Trace
 {
@@ -87,7 +88,11 @@ public:
 	FStatsNodePtr GetCounterNode(uint32 CounterId) const;
 	void SelectCounterNode(uint32 CounterId);
 
+	virtual FReply OnKeyDown(const FGeometry& MyGeometry, const FKeyEvent& InKeyEvent) override;
+
 private:
+	void InitCommandList();
+
 	void UpdateTree();
 
 	void UpdateNode(FStatsNodePtr NodePtr);
@@ -113,6 +118,9 @@ private:
 	void TreeView_BuildSortByMenu(FMenuBuilder& MenuBuilder);
 	void TreeView_BuildViewColumnMenu(FMenuBuilder& MenuBuilder);
 
+	bool ContextMenu_CopySelectedToClipboard_CanExecute() const;
+	void ContextMenu_CopySelectedToClipboard_Execute();
+
 	////////////////////////////////////////////////////////////////////////////////////////////////////
 	// Tree View - Columns' Header
 
@@ -268,6 +276,8 @@ private:
 	/** A weak pointer to the profiler session used to populate this widget. */
 	TSharedPtr<const Trace::IAnalysisSession>/*Weak*/ Session;
 
+	TSharedPtr<FUICommandList> CommandList;
+
 	//////////////////////////////////////////////////
 	// Tree View, Columns
 
diff --git a/Engine/Source/Developer/TraceInsights/Private/Insights/Widgets/STimerTreeView.cpp b/Engine/Source/Developer/TraceInsights/Private/Insights/Widgets/STimerTreeView.cpp
--- a/Engine/Source/Developer/TraceInsights/Private/Insights/Widgets/STimerTreeView.cpp
+++ b/Engine/Source/Developer/TraceInsights/Private/Insights/Widgets/STimerTreeView.cpp
@@ -3,6 +3,8 @@
 #include "STimerTreeView.h"
 
 #include "EditorStyleSet.h"
+#include "Framework/Commands/Commands.h"
+#include "Framework/Commands/UICommandList.h"
 #include "Framework/Application/SlateApplication.h"
 #include "Framework/MultiBox/MultiBoxBuilder.h"
 #include "HAL/PlatformApplicationMisc.h"
diff --git a/Engine/Source/Developer/TraceInsights/Private/Insights/Widgets/STimerTreeView.h b/Engine/Source/Developer/TraceInsights/Private/Insights/Widgets/STimerTreeView.h
--- a/Engine/Source/Developer/TraceInsights/Private/Insights/Widgets/STimerTreeView.h
+++ b/Engine/Source/Developer/TraceInsights/Private/Insights/Widgets/STimerTreeView.h
@@ -20,6 +20,7 @@
 ////////////////////////////////////////////////////////////////////////////////////////////////////
 
 class FMenuBuilder;
+class FUICommandList;
 
 namespace Insights
 {
diff --git a/Engine/Source/Developer/TraceInsights/Private/Insights/Widgets/STimersView.cpp b/Engine/Source/Developer/TraceInsights/Private/Insights/Widgets/STimersView.cpp
--- a/Engine/Source/Developer/TraceInsights/Private/Insights/Widgets/STimersView.cpp
+++ b/Engine/Source/Developer/TraceInsights/Private/Insights/Widgets/STimersView.cpp
@@ -3,6 +3,8 @@
 #include "STimersView.h"
 
 #include "EditorStyleSet.h"
+#include "Framework/Commands/Commands.h"
+#include "Framework/Commands/UICommandList.h"
 #include "Framework/MultiBox/MultiBoxBuilder.h"
 #include "HAL/PlatformApplicationMisc.h"
 #include "SlateOptMacros.h"
diff --git a/Engine/Source/Developer/TraceInsights/Private/Insights/Widgets/STimersView.h b/Engine/Source/Developer/TraceInsights/Private/Insights/Widgets/STimersView.h
--- a/Engine/Source/Developer/TraceInsights/Private/Insights/Widgets/STimersView.h
+++ b/Engine/Source/Developer/TraceInsights/Private/Insights/Widgets/STimersView.h
@@ -23,6 +23,7 @@
 
 class FMenuBuilder;
 class FTimingGraphTrack;
+class FUICommandList;
 
 namespace Trace
 {
