[Insights]
  - Fixed style for toggle buttons in Timers, Counters, NetStats and MemTag views.
  - Removed icons for filter by timer type (GPU, CPU) and by counter data type (int, double).
  - Fixed auto-scroll button.
  - Fixed "Tracks" dropdown button.

diff --git a/Engine/Source/Developer/TraceInsights/Private/Insights/MemoryProfiler/Widgets/SMemTagTreeView.cpp b/Engine/Source/Developer/TraceInsights/Private/Insights/MemoryProfiler/Widgets/SMemTagTreeView.cpp
--- a/Engine/Source/Developer/TraceInsights/Private/Insights/MemoryProfiler/Widgets/SMemTagTreeView.cpp
+++ b/Engine/Source/Developer/TraceInsights/Private/Insights/MemoryProfiler/Widgets/SMemTagTreeView.cpp
@@ -441,7 +441,7 @@ TSharedRef<SWidget> SMemTagTreeView::ConstructTagsFilteringWidgetArea()
 	//.AutoWidth()
 	//[
 	//	SNew(SCheckBox)
-	//	.Style(FEditorStyle::Get(), "ToggleButtonCheckbox")
+	//	.Style(FCoreStyle::Get(), "ToggleButtonCheckbox")
 	//	.HAlign(HAlign_Center)
 	//	.Padding(2.0f)
 	//	.OnCheckStateChanged(this, &SMemTagTreeView::FilterOutZeroCountMemTags_OnCheckStateChanged)
@@ -450,7 +450,6 @@ TSharedRef<SWidget> SMemTagTreeView::ConstructTagsFilteringWidgetArea()
 	//	[
 	//		SNew(STextBlock)
 	//		.Text(LOCTEXT("FilterOutZeroCountMemTags_Button", " !0 "))
-	//		.TextStyle(FEditorStyle::Get(), TEXT("Profiler.Caption"))
 	//	]
 	//]
 
@@ -461,7 +460,7 @@ TSharedRef<SWidget> SMemTagTreeView::ConstructTagsFilteringWidgetArea()
 	.AutoWidth()
 	[
 		SNew(SCheckBox)
-		.Style(FEditorStyle::Get(), "ToggleButtonCheckbox")
+		.Style(FCoreStyle::Get(), "ToggleButtonCheckbox")
 		.HAlign(HAlign_Center)
 		.Padding(2.0f)
 		.OnCheckStateChanged(this, &SMemTagTreeView::FilterByTracker_OnCheckStateChanged)
@@ -470,7 +469,6 @@ TSharedRef<SWidget> SMemTagTreeView::ConstructTagsFilteringWidgetArea()
 		[
 			SNew(STextBlock)
 			.Text(LOCTEXT("FilterByTracker_Button", " T "))
-			.TextStyle(FEditorStyle::Get(), TEXT("Profiler.Caption"))
 		]
 	];
 
diff --git a/Engine/Source/Developer/TraceInsights/Private/Insights/NetworkingProfiler/Widgets/SNetStatsView.cpp b/Engine/Source/Developer/TraceInsights/Private/Insights/NetworkingProfiler/Widgets/SNetStatsView.cpp
--- a/Engine/Source/Developer/TraceInsights/Private/Insights/NetworkingProfiler/Widgets/SNetStatsView.cpp
+++ b/Engine/Source/Developer/TraceInsights/Private/Insights/NetworkingProfiler/Widgets/SNetStatsView.cpp
@@ -116,17 +116,15 @@ void SNetStatsView::Construct(const FArguments& InArgs, TSharedPtr<SNetworkingPr
 					.AutoWidth()
 					[
 						SNew(SCheckBox)
-						.Style(FEditorStyle::Get(), "ToggleButtonCheckbox")
+						.Style(FCoreStyle::Get(), "ToggleButtonCheckbox")
 						.HAlign(HAlign_Center)
 						.Padding(2.0f)
 						.OnCheckStateChanged(this, &SNetStatsView::FilterOutZeroCountEvents_OnCheckStateChanged)
 						.IsChecked(this, &SNetStatsView::FilterOutZeroCountEvents_IsChecked)
 						.ToolTipText(LOCTEXT("FilterOutZeroCountEvents_Tooltip", "Filter out the net event types having zero total instance count (aggregated stats)."))
 						[
-							//TODO: SNew(SImage)
 							SNew(STextBlock)
 							.Text(LOCTEXT("FilterOutZeroCountEvents_Button", " !0 "))
-							.TextStyle(FEditorStyle::Get(), TEXT("Profiler.Caption"))
 						]
 					]
 				]
@@ -746,7 +744,7 @@ void SNetStatsView::HandleItemToStringArray(const FNetEventNodePtr& FNetEventNod
 TSharedRef<SWidget> SNetStatsView::GetToggleButtonForNetEventType(const ENetEventNodeType NodeType)
 {
 	return SNew(SCheckBox)
-		.Style(FEditorStyle::Get(), "ToggleButtonCheckbox")
+		.Style(FCoreStyle::Get(), "ToggleButtonCheckbox")
 		.HAlign(HAlign_Center)
 		.Padding(2.0f)
 		.OnCheckStateChanged(this, &SNetStatsView::FilterByNetEventType_OnCheckStateChanged, NodeType)
@@ -756,21 +754,20 @@ TSharedRef<SWidget> SNetStatsView::GetToggleButtonForNetEventType(const ENetEven
 			SNew(SHorizontalBox)
 
 			+ SHorizontalBox::Slot()
-				.AutoWidth()
-				.VAlign(VAlign_Center)
-				[
-					SNew(SImage)
-						.Image(NetEventNodeTypeHelper::GetIconForNetEventNodeType(NodeType))
-				]
+			.AutoWidth()
+			.VAlign(VAlign_Center)
+			[
+				SNew(SImage)
+				.Image(NetEventNodeTypeHelper::GetIconForNetEventNodeType(NodeType))
+			]
 
 			+ SHorizontalBox::Slot()
-				.Padding(2.0f, 0.0f, 0.0f, 0.0f)
-				.VAlign(VAlign_Center)
-				[
-					SNew(STextBlock)
-						.Text(NetEventNodeTypeHelper::ToText(NodeType))
-						.TextStyle(FEditorStyle::Get(), TEXT("Profiler.Caption"))
-				]
+			.Padding(2.0f, 0.0f, 0.0f, 0.0f)
+			.VAlign(VAlign_Center)
+			[
+				SNew(STextBlock)
+				.Text(NetEventNodeTypeHelper::ToText(NodeType))
+			]
 		];
 }
 
diff --git a/Engine/Source/Developer/TraceInsights/Private/Insights/Widgets/SStatsView.cpp b/Engine/Source/Developer/TraceInsights/Private/Insights/Widgets/SStatsView.cpp
--- a/Engine/Source/Developer/TraceInsights/Private/Insights/Widgets/SStatsView.cpp
+++ b/Engine/Source/Developer/TraceInsights/Private/Insights/Widgets/SStatsView.cpp
@@ -154,17 +154,15 @@ void SStatsView::Construct(const FArguments& InArgs)
 					.AutoWidth()
 					[
 						SNew(SCheckBox)
-						.Style(FEditorStyle::Get(), "ToggleButtonCheckbox")
+						.Style(FCoreStyle::Get(), "ToggleButtonCheckbox")
 						.HAlign(HAlign_Center)
 						.Padding(2.0f)
 						.OnCheckStateChanged(this, &SStatsView::FilterOutZeroCountStats_OnCheckStateChanged)
 						.IsChecked(this, &SStatsView::FilterOutZeroCountStats_IsChecked)
 						.ToolTipText(LOCTEXT("FilterOutZeroCountStats_Tooltip", "Filter out the stats counters having zero total instance count (aggregated stats)."))
 						[
-							//TODO: SNew(SImage)
 							SNew(STextBlock)
 							.Text(LOCTEXT("FilterOutZeroCountStats_Button", " !0 "))
-							.TextStyle(FEditorStyle::Get(), TEXT("Profiler.Caption"))
 						]
 					]
 				]
@@ -859,7 +857,7 @@ void SStatsView::HandleItemToStringArray(const FStatsNodePtr& FStatsNodePtr, TAr
 TSharedRef<SWidget> SStatsView::GetToggleButtonForNodeType(const EStatsNodeType NodeType)
 {
 	return SNew(SCheckBox)
-		.Style(FEditorStyle::Get(), "ToggleButtonCheckbox")
+		.Style(FCoreStyle::Get(), "ToggleButtonCheckbox")
 		.HAlign(HAlign_Center)
 		.Padding(2.0f)
 		.OnCheckStateChanged(this, &SStatsView::FilterByStatsType_OnCheckStateChanged, NodeType)
@@ -869,21 +867,20 @@ TSharedRef<SWidget> SStatsView::GetToggleButtonForNodeType(const EStatsNodeType
 			SNew(SHorizontalBox)
 
 			+ SHorizontalBox::Slot()
-				.AutoWidth()
-				.VAlign(VAlign_Center)
-				[
-					SNew(SImage)
-						.Image(StatsNodeTypeHelper::GetIcon(NodeType))
-				]
+			.AutoWidth()
+			.VAlign(VAlign_Center)
+			[
+				SNew(SImage)
+				.Image(StatsNodeTypeHelper::GetIcon(NodeType))
+			]
 
 			+ SHorizontalBox::Slot()
-				.Padding(2.0f, 0.0f, 0.0f, 0.0f)
-				.VAlign(VAlign_Center)
-				[
-					SNew(STextBlock)
-						.Text(StatsNodeTypeHelper::ToText(NodeType))
-						.TextStyle(FEditorStyle::Get(), TEXT("Profiler.Caption"))
-				]
+			.Padding(2.0f, 0.0f, 0.0f, 0.0f)
+			.VAlign(VAlign_Center)
+			[
+				SNew(STextBlock)
+				.Text(StatsNodeTypeHelper::ToText(NodeType))
+			]
 		];
 }
 
@@ -892,7 +889,7 @@ TSharedRef<SWidget> SStatsView::GetToggleButtonForNodeType(const EStatsNodeType
 TSharedRef<SWidget> SStatsView::GetToggleButtonForDataType(const EStatsNodeDataType DataType)
 {
 	return SNew(SCheckBox)
-		.Style(FEditorStyle::Get(), "ToggleButtonCheckbox")
+		.Style(FCoreStyle::Get(), "ToggleButtonCheckbox")
 		.HAlign(HAlign_Center)
 		.Padding(2.0f)
 		.OnCheckStateChanged(this, &SStatsView::FilterByStatsDataType_OnCheckStateChanged, DataType)
@@ -901,22 +898,21 @@ TSharedRef<SWidget> SStatsView::GetToggleButtonForDataType(const EStatsNodeDataT
 		[
 			SNew(SHorizontalBox)
 
-			+ SHorizontalBox::Slot()
-				.AutoWidth()
-				.VAlign(VAlign_Center)
-				[
-					SNew(SImage)
-						.Image(StatsNodeDataTypeHelper::GetIcon(DataType))
-				]
+			//+ SHorizontalBox::Slot()
+			//.AutoWidth()
+			//.VAlign(VAlign_Center)
+			//[
+			//	SNew(SImage)
+			//	.Image(StatsNodeDataTypeHelper::GetIcon(DataType))
+			//]
 
 			+ SHorizontalBox::Slot()
-				.Padding(2.0f, 0.0f, 0.0f, 0.0f)
-				.VAlign(VAlign_Center)
-				[
-					SNew(STextBlock)
-						.Text(StatsNodeDataTypeHelper::ToText(DataType))
-						.TextStyle(FEditorStyle::Get(), TEXT("Profiler.Caption"))
-				]
+			.Padding(2.0f, 0.0f, 0.0f, 0.0f)
+			.VAlign(VAlign_Center)
+			[
+				SNew(STextBlock)
+				.Text(StatsNodeDataTypeHelper::ToText(DataType))
+			]
 		];
 }
 
diff --git a/Engine/Source/Developer/TraceInsights/Private/Insights/Widgets/STimersView.cpp b/Engine/Source/Developer/TraceInsights/Private/Insights/Widgets/STimersView.cpp
--- a/Engine/Source/Developer/TraceInsights/Private/Insights/Widgets/STimersView.cpp
+++ b/Engine/Source/Developer/TraceInsights/Private/Insights/Widgets/STimersView.cpp
@@ -156,17 +156,15 @@ void STimersView::Construct(const FArguments& InArgs)
 					.AutoWidth()
 					[
 						SNew(SCheckBox)
-						.Style(FEditorStyle::Get(), "ToggleButtonCheckbox")
+						.Style(FCoreStyle::Get(), "ToggleButtonCheckbox")
 						.HAlign(HAlign_Center)
 						.Padding(2.0f)
 						.OnCheckStateChanged(this, &STimersView::FilterOutZeroCountTimers_OnCheckStateChanged)
 						.IsChecked(this, &STimersView::FilterOutZeroCountTimers_IsChecked)
 						.ToolTipText(LOCTEXT("FilterOutZeroCountTimers_Tooltip", "Filter out the timers having zero total instance count (aggregated stats)."))
 						[
-							//TODO: SNew(SImage)
 							SNew(STextBlock)
 							.Text(LOCTEXT("FilterOutZeroCountTimers_Button", " !0 "))
-							.TextStyle(FEditorStyle::Get(), TEXT("Profiler.Caption"))
 						]
 					]
 				]
@@ -894,7 +892,7 @@ void STimersView::HandleItemToStringArray(const FTimerNodePtr& FTimerNodePtr, TA
 TSharedRef<SWidget> STimersView::GetToggleButtonForTimerType(const ETimerNodeType NodeType)
 {
 	return SNew(SCheckBox)
-		.Style(FEditorStyle::Get(), "ToggleButtonCheckbox")
+		.Style(FCoreStyle::Get(), "ToggleButtonCheckbox")
 		.HAlign(HAlign_Center)
 		.Padding(2.0f)
 		.OnCheckStateChanged(this, &STimersView::FilterByTimerType_OnCheckStateChanged, NodeType)
@@ -903,22 +901,21 @@ TSharedRef<SWidget> STimersView::GetToggleButtonForTimerType(const ETimerNodeTyp
 		[
 			SNew(SHorizontalBox)
 
-			+ SHorizontalBox::Slot()
-				.AutoWidth()
-				.VAlign(VAlign_Center)
-				[
-					SNew(SImage)
-						.Image(TimerNodeTypeHelper::GetIconForTimerNodeType(NodeType))
-				]
+			//+ SHorizontalBox::Slot()
+			//.AutoWidth()
+			//.VAlign(VAlign_Center)
+			//[
+			//	SNew(SImage)
+			//	.Image(TimerNodeTypeHelper::GetIconForTimerNodeType(NodeType))
+			//]
 
 			+ SHorizontalBox::Slot()
-				.Padding(2.0f, 0.0f, 0.0f, 0.0f)
-				.VAlign(VAlign_Center)
-				[
-					SNew(STextBlock)
-						.Text(TimerNodeTypeHelper::ToText(NodeType))
-						.TextStyle(FEditorStyle::Get(), TEXT("Profiler.Caption"))
-				]
+			.Padding(2.0f, 0.0f, 0.0f, 0.0f)
+			.VAlign(VAlign_Center)
+			[
+				SNew(STextBlock)
+				.Text(TimerNodeTypeHelper::ToText(NodeType))
+			]
 		];
 }
 
diff --git a/Engine/Source/Developer/TraceInsights/Private/Insights/Widgets/STimingView.cpp b/Engine/Source/Developer/TraceInsights/Private/Insights/Widgets/STimingView.cpp
--- a/Engine/Source/Developer/TraceInsights/Private/Insights/Widgets/STimingView.cpp
+++ b/Engine/Source/Developer/TraceInsights/Private/Insights/Widgets/STimingView.cpp
@@ -164,32 +164,30 @@ void STimingView::Construct(const FArguments& InArgs)
 		[
 			// Tracks Filter
 			SNew(SComboButton)
-			.ComboButtonStyle(FEditorStyle::Get(), "GenericFilters.ComboButtonStyle")
-			.ForegroundColor(FLinearColor::White)
+			.ComboButtonStyle(FCoreStyle::Get(), "ComboButton")
+			.ButtonColorAndOpacity(FLinearColor::Transparent)
 			.ToolTipText(LOCTEXT("TracksFilterToolTip", "Filter timing tracks."))
 			.OnGetMenuContent(this, &STimingView::MakeTracksFilterMenu)
 			.HasDownArrow(true)
-			.ContentPadding(FMargin(1.0f, 1.0f, 1.0f, 1.0f))
+			.ContentPadding(FMargin(-4.0f, 0.0f, -5.0f, 0.0f))
 			.ButtonContent()
 			[
 				SNew(SHorizontalBox)
 
 				+ SHorizontalBox::Slot()
-				.Padding(0.0f)
+				.Padding(FMargin(0.0f, 2.0f, 0.0f, 0.0f))
 				.AutoWidth()
 				[
 					SNew(STextBlock)
-					.TextStyle(FEditorStyle::Get(), "GenericFilters.TextStyle")
 					.Font(FEditorStyle::Get().GetFontStyle("FontAwesome.9"))
 					.Text(FText::FromString(FString(TEXT("\xf0b0"))) /*fa-filter*/)
 				]
 
 				+ SHorizontalBox::Slot()
-				.Padding(FMargin(2.0f, 0.0f, 0.0f, 0.0f))
+				.Padding(FMargin(3.0f, 0.0f, -7.0f, 0.0f))
 				.AutoWidth()
 				[
 					SNew(STextBlock)
-					.TextStyle(FEditorStyle::Get(), "GenericFilters.TextStyle")
 					.Text(LOCTEXT("TracksFilter", "Tracks"))
 				]
 			]
@@ -203,20 +201,19 @@ void STimingView::Construct(const FArguments& InArgs)
 			SNew(SHorizontalBox)
 
 			+ SHorizontalBox::Slot()
-			.Padding(0.0f)
+			.Padding(FMargin(0.0f, 1.0f, 0.0f, 1.0f))
 			.AutoWidth()
 			[
 				SNew(SCheckBox)
-				.Style(FEditorStyle::Get(), "ToggleButtonCheckbox")
+				.Style(FCoreStyle::Get(), "ToggleButtonCheckbox")
 				.HAlign(HAlign_Center)
-				.Padding(FMargin(0.0f, 2.0f, 1.0f, 4.0f))
+				.Padding(FMargin(4.0f, 2.0f, 4.0f, 2.0f))
 				.OnCheckStateChanged(this, &STimingView::AutoScroll_OnCheckStateChanged)
 				.IsChecked(this, &STimingView::AutoScroll_IsChecked)
 				.ToolTipText(LOCTEXT("AutoScroll_Tooltip", "Auto Scroll"))
 				[
 					SNew(STextBlock)
-					.Text(LOCTEXT("AutoScroll_Button", " >> "))
-					.TextStyle(FEditorStyle::Get(), TEXT("Profiler.Caption"))
+					.Text(LOCTEXT("AutoScroll_Button", ">>"))
 				]
 			]
 
@@ -225,12 +222,11 @@ void STimingView::Construct(const FArguments& InArgs)
 			.AutoWidth()
 			[
 				SNew(SComboButton)
-				.ComboButtonStyle(FEditorStyle::Get(), "GenericFilters.ComboButtonStyle")
-				.ForegroundColor(FLinearColor::White)
-				//.ToolTipText(LOCTEXT("AutoScrollOptions_ToolTip", "Auto-scroll options"))
+				.ComboButtonStyle(FCoreStyle::Get(), "ComboButton")
+				.ButtonColorAndOpacity(FLinearColor::Transparent)
 				.OnGetMenuContent(this, &STimingView::MakeAutoScrollOptionsMenu)
 				.HasDownArrow(true)
-				.ContentPadding(FMargin(0.0f, 0.0f, 0.0f, 0.0f))
+				.ContentPadding(FMargin(-14.0f, 0.0f, -6.0f, 0.0f))
 				.ButtonContent()
 				[
 					SNew(SBorder)
@@ -1319,10 +1315,10 @@ int32 STimingView::OnPaint(const FPaintArgs& Args, const FGeometry& AllottedGeom
 	Tooltip.Draw(DrawContext);
 
 	// Fill background for the "Tracks" filter combobox.
-	DrawContext.DrawBox(0.0f, 0.0f, 66.0f, 24.0f, WhiteBrush, FLinearColor(0.05f, 0.05f, 0.05f, 1.0f));
+	DrawContext.DrawBox(0.0f, 0.0f, 73.0f, 24.0f, WhiteBrush, FLinearColor(0.06f, 0.06f, 0.06f, 1.0f));
 
 	// Fill background for the "Auto-scroll" toggle button.
-	DrawContext.DrawBox(ViewWidth - 40.0f, 0.0f, 40.0f, 24.0f, WhiteBrush, FLinearColor(0.05f, 0.05f, 0.05f, 1.0f));
+	DrawContext.DrawBox(ViewWidth - 39.0f, 0.0f, 39.0f, 24.0f, WhiteBrush, FLinearColor(0.06f, 0.06f, 0.06f, 1.0f));
 
 	//////////////////////////////////////////////////
 	// Draw the overscroll indication lines.
