[Insights]
  - Fixed size and position of scrallbars in Timing view.
  - Fixed padding in the buttons of the LLM Tags panel's minitoolbar.

diff --git a/Engine/Source/Developer/TraceInsights/Private/Insights/MemoryProfiler/Widgets/SMemTagTreeView.cpp b/Engine/Source/Developer/TraceInsights/Private/Insights/MemoryProfiler/Widgets/SMemTagTreeView.cpp
--- a/Engine/Source/Developer/TraceInsights/Private/Insights/MemoryProfiler/Widgets/SMemTagTreeView.cpp
+++ b/Engine/Source/Developer/TraceInsights/Private/Insights/MemoryProfiler/Widgets/SMemTagTreeView.cpp
@@ -547,13 +547,17 @@ TSharedRef<SWidget> SMemTagTreeView::ConstructTracksMiniToolbar()
 	[
 		SNew(SButton)
 		.ToolTipText(LOCTEXT("HideAll_ToolTip", "Remove all memory graph tracks."))
-		.ContentPadding(FMargin(2.0f, 0.0f, 2.0f, 0.0f))
+		.ContentPadding(FMargin(-14.0f, 0.0f, -17.0f, -1.0f))
 		.OnClicked(this, &SMemTagTreeView::HideAllTracks_OnClicked)
 		.Content()
 		[
-			SNew(SImage)
-			.ColorAndOpacity(FSlateColor(FLinearColor(0.5f, 0.0f, 0.0f, 1.0f)))
-			.Image(FInsightsStyle::Get().GetBrush("Mem.Remove.Small"))
+			SNew(SBox)
+			.HAlign(HAlign_Center)
+			[
+				SNew(SImage)
+				.ColorAndOpacity(FSlateColor(FLinearColor(0.5f, 0.0f, 0.0f, 1.0f)))
+				.Image(FInsightsStyle::Get().GetBrush("Mem.Remove.Small"))
+			]
 		]
 	]
 
@@ -564,13 +568,17 @@ TSharedRef<SWidget> SMemTagTreeView::ConstructTracksMiniToolbar()
 	[
 		SNew(SButton)
 		.ToolTipText(LOCTEXT("ShowAll_ToolTip", "Create memory graph tracks for all visible (filtered) LLM tags."))
-		.ContentPadding(FMargin(2.0f, 0.0f, 2.0f, 0.0f))
+		.ContentPadding(FMargin(-14.0f, 0.0f, -17.0f, -1.0f))
 		.OnClicked(this, &SMemTagTreeView::ShowAllTracks_OnClicked)
 		.Content()
 		[
-			SNew(SImage)
-			.ColorAndOpacity(FSlateColor(FLinearColor(0.0f, 0.5f, 0.0f, 1.0f)))
-			.Image(FInsightsStyle::Get().GetBrush("Mem.Add.Small"))
+			SNew(SBox)
+			.HAlign(HAlign_Center)
+			[
+				SNew(SImage)
+				.ColorAndOpacity(FSlateColor(FLinearColor(0.0f, 0.5f, 0.0f, 1.0f)))
+				.Image(FInsightsStyle::Get().GetBrush("Mem.Add.Small"))
+			]
 		]
 	]
 
@@ -581,12 +589,16 @@ TSharedRef<SWidget> SMemTagTreeView::ConstructTracksMiniToolbar()
 	[
 		SNew(SButton)
 		.ToolTipText(LOCTEXT("LoadReportXML_ToolTip", "Load LLMReportTypes.xml"))
-		.ContentPadding(FMargin(2.0f, 2.0f, 2.0f, 2.0f))
+		.ContentPadding(FMargin(-14.0f, 2.0f, -14.0f, 2.0f))
 		.OnClicked(this, &SMemTagTreeView::LoadReportXML_OnClicked)
 		.Content()
 		[
-			SNew(STextBlock)
-			.Text(LOCTEXT("LoadReportXML_Text", "Load XML..."))
+			SNew(SBox)
+			.HAlign(HAlign_Center)
+			[
+				SNew(STextBlock)
+				.Text(LOCTEXT("LoadReportXML_Text", "Load XML..."))
+			]
 		]
 	]
 
@@ -596,12 +608,16 @@ TSharedRef<SWidget> SMemTagTreeView::ConstructTracksMiniToolbar()
 	[
 		SNew(SButton)
 		.ToolTipText(LOCTEXT("SmallHeight_ToolTip", "Change height of LLM Tag Graph tracks to Small."))
-		.ContentPadding(FMargin(2.0f, 2.0f, 2.0f, 2.0f))
+		.ContentPadding(FMargin(-18.0f, 2.0f, -18.0f, 2.0f))
 		.OnClicked(this, &SMemTagTreeView::AllTracksSmallHeight_OnClicked)
 		.Content()
 		[
-			SNew(STextBlock)
-			.Text(LOCTEXT("SmallHeight_Text", "\u2195S"))
+			SNew(SBox)
+			.HAlign(HAlign_Center)
+			[
+				SNew(STextBlock)
+				.Text(LOCTEXT("SmallHeight_Text", "\u2195S"))
+			]
 		]
 	]
 
@@ -611,12 +627,16 @@ TSharedRef<SWidget> SMemTagTreeView::ConstructTracksMiniToolbar()
 	[
 		SNew(SButton)
 		.ToolTipText(LOCTEXT("MediumHeight_ToolTip", "Change height of LLM Tag Graph tracks to Medium."))
-		.ContentPadding(FMargin(2.0f, 2.0f, 2.0f, 2.0f))
+		.ContentPadding(FMargin(-18.0f, 2.0f, -18.0f, 2.0f))
 		.OnClicked(this, &SMemTagTreeView::AllTracksMediumHeight_OnClicked)
 		.Content()
 		[
-			SNew(STextBlock)
-			.Text(LOCTEXT("MediumHeight_Text", "\u2195M"))
+			SNew(SBox)
+			.HAlign(HAlign_Center)
+			[
+				SNew(STextBlock)
+				.Text(LOCTEXT("MediumHeight_Text", "\u2195M"))
+			]
 		]
 	]
 
@@ -626,12 +646,16 @@ TSharedRef<SWidget> SMemTagTreeView::ConstructTracksMiniToolbar()
 	[
 		SNew(SButton)
 		.ToolTipText(LOCTEXT("LargeHeight_ToolTip", "Change height of LLM Tag Graph tracks to Large."))
-		.ContentPadding(FMargin(2.0f, 2.0f, 2.0f, 2.0f))
+		.ContentPadding(FMargin(-18.0f, 2.0f, -18.0f, 2.0f))
 		.OnClicked(this, &SMemTagTreeView::AllTracksLargeHeight_OnClicked)
 		.Content()
 		[
-			SNew(STextBlock)
-			.Text(LOCTEXT("LargeHeight_Text", "\u2195L"))
+			SNew(SBox)
+			.HAlign(HAlign_Center)
+			[
+				SNew(STextBlock)
+				.Text(LOCTEXT("LargeHeight_Text", "\u2195L"))
+			]
 		]
 	];
 
diff --git a/Engine/Source/Developer/TraceInsights/Private/Insights/Widgets/STimingView.cpp b/Engine/Source/Developer/TraceInsights/Private/Insights/Widgets/STimingView.cpp
--- a/Engine/Source/Developer/TraceInsights/Private/Insights/Widgets/STimingView.cpp
+++ b/Engine/Source/Developer/TraceInsights/Private/Insights/Widgets/STimingView.cpp
@@ -135,25 +135,23 @@ void STimingView::Construct(const FArguments& InArgs)
 
 		+ SOverlay::Slot()
 		.VAlign(VAlign_Bottom)
-		.Padding(FMargin(0, 0, 12, 0))
+		.Padding(FMargin(0.0f, 0.0f, 8.0f, 0.0f))
 		[
 			SAssignNew(HorizontalScrollBar, SScrollBar)
 			.Orientation(Orient_Horizontal)
 			.AlwaysShowScrollbar(false)
 			.Visibility(EVisibility::Visible)
-			.Thickness(FVector2D(5.0f, 5.0f))
 			.OnUserScrolled(this, &STimingView::HorizontalScrollBar_OnUserScrolled)
 		]
 
 		+ SOverlay::Slot()
 		.HAlign(HAlign_Right)
-		.Padding(FMargin(0.0f, 0.0f, 0.0f, 0.0f))
+		.Padding(FMargin(0.0f, 0.0f, 2.0f, 0.0f))
 		[
 			SAssignNew(VerticalScrollBar, SScrollBar)
 			.Orientation(Orient_Vertical)
 			.AlwaysShowScrollbar(false)
 			.Visibility(EVisibility::Visible)
-			.Thickness(FVector2D(5.0f, 5.0f))
 			.OnUserScrolled(this, &STimingView::VerticalScrollBar_OnUserScrolled)
 		]
 
@@ -771,7 +769,7 @@ void STimingView::Tick(const FGeometry& AllottedGeometry, const double InCurrent
 	}
 
 	// Set the VerticalScrollBar padding so it is limited to the scrollable area.
-	VerticalScrollBar->SetPadding(FMargin(0.0f, TopOffset, 0.0f, FMath::Max(BottomOffset, 12.0f)));
+	VerticalScrollBar->SetPadding(FMargin(0.0f, TopOffset + 2.0f, 0.0f, FMath::Max(BottomOffset + 2.0f, 10.0f)));
 
 	////////////////////////////////////////////////////////////////////////////////////////////////////
 
@@ -809,7 +807,7 @@ void STimingView::Tick(const FGeometry& AllottedGeometry, const double InCurrent
 	// Elastic snap to vertical scroll limits.
 	if (!bIsPanning)
 	{
-		const float DY = Viewport.GetScrollHeight() - Viewport.GetScrollableAreaHeight() + 7.0f; // +7 is to allow some space for the horizontal scrollbar
+		const float DY = Viewport.GetScrollHeight() - Viewport.GetScrollableAreaHeight();
 		const float MinY = FMath::Min(DY, 0.0f);
 		const float MaxY = DY - MinY;
 
@@ -2977,7 +2975,7 @@ float STimingView::EnforceVerticalScrollLimits(const float InScrollPosY)
 {
 	float NewScrollPosY = InScrollPosY;
 
-	const float DY = Viewport.GetScrollHeight() - Viewport.GetScrollableAreaHeight() + 7.0f; // +7 is to allow some space for the horizontal scrollbar
+	const float DY = Viewport.GetScrollHeight() - Viewport.GetScrollableAreaHeight();
 	const float MinY = FMath::Min(DY, 0.0f);
 	const float MaxY = DY - MinY;
 
