[Insights] Improved UX for the Open button and the drop down menu in Session Browser dialog.

diff --git a/Engine/Source/Developer/TraceInsights/Private/Insights/InsightsStyle.cpp b/Engine/Source/Developer/TraceInsights/Private/Insights/InsightsStyle.cpp
--- a/Engine/Source/Developer/TraceInsights/Private/Insights/InsightsStyle.cpp
+++ b/Engine/Source/Developer/TraceInsights/Private/Insights/InsightsStyle.cpp
@@ -3,9 +3,17 @@
 #include "InsightsStyle.h"
 #include "Styling/SlateStyleRegistry.h"
 
-#define IMAGE_BRUSH(RelativePath, ...)  FSlateImageBrush (FPaths::EngineContentDir() / "Editor/Slate" / RelativePath + TEXT(".png"), __VA_ARGS__)
-#define BOX_BRUSH(RelativePath, ...)    FSlateBoxBrush   (FPaths::EngineContentDir() / "Editor/Slate" / RelativePath + TEXT(".png"), __VA_ARGS__)
-#define BORDER_BRUSH(RelativePath, ...) FSlateBorderBrush(FPaths::EngineContentDir() / "Editor/Slate" / RelativePath + TEXT(".png"), __VA_ARGS__)
+#define IMAGE_BRUSH(RelativePath, ...)  FSlateImageBrush (Style.RootToContentDir(RelativePath, TEXT(".png")), __VA_ARGS__)
+#define BOX_BRUSH(RelativePath, ...)    FSlateBoxBrush   (Style.RootToContentDir(RelativePath, TEXT(".png")), __VA_ARGS__)
+#define BORDER_BRUSH(RelativePath, ...) FSlateBorderBrush(Style.RootToContentDir(RelativePath, TEXT(".png")), __VA_ARGS__)
+
+#define IMAGE_BRUSH_SVG(RelativePath, ...)  FSlateVectorImageBrush (Style.RootToContentDir(RelativePath, TEXT(".svg")), __VA_ARGS__)
+#define BOX_BRUSH_SVG(RelativePath, ...)    FSlateVectorBoxBrush   (Style.RootToContentDir(RelativePath, TEXT(".svg")), __VA_ARGS__)
+#define BORDER_BRUSH_SVG(RelativePath, ...) FSlateVectorBorderBrush(Style.RootToContentDir(RelativePath, TEXT(".svg")), __VA_ARGS__)
+
+#define CORE_IMAGE_BRUSH_SVG(RelativePath, ...)  FSlateVectorImageBrush (Style.RootToCoreContentDir(RelativePath, TEXT(".svg")), __VA_ARGS__)
+#define CORE_BOX_BRUSH_SVG(RelativePath, ...)    FSlateVectorBoxBrush   (Style.RootToCoreContentDir(RelativePath, TEXT(".svg")), __VA_ARGS__)
+#define CORE_BORDER_BRUSH_SVG(RelativePath, ...) FSlateVectorBorderBrush(Style.RootToCoreContentDir(RelativePath, TEXT(".svg")), __VA_ARGS__)
 
 TSharedPtr<FSlateStyleSet> FInsightsStyle::StyleInstance = nullptr;
 
@@ -33,13 +41,18 @@ FName FInsightsStyle::GetStyleSetName()
 
 TSharedRef<FSlateStyleSet> FInsightsStyle::Create()
 {
+	const FVector2D Icon16x16(16.0f, 16.0f);
+	const FVector2D Icon20x20(20.0f, 20.0f);
+	const FVector2D Icon24x24(24.0f, 24.0f);
+	const FVector2D Icon32x32(32.0f, 32.0f);
+
 	TSharedRef<FSlateStyleSet> StyleRef = MakeShared<FSlateStyleSet>(FInsightsStyle::GetStyleSetName());
 	StyleRef->SetContentRoot(FPaths::EngineContentDir() / TEXT("Editor/Slate"));
 	StyleRef->SetCoreContentRoot(FPaths::EngineContentDir() / TEXT("Slate"));
 
 	FSlateStyleSet& Style = StyleRef.Get();
 
-	//Style.Set("AppIcon", new IMAGE_BRUSH("Icons/Insights/AppIcon_24x", FVector2D(24.0f, 24.0f)));
+	//Style.Set("AppIcon", new IMAGE_BRUSH("Icons/Insights/AppIcon_24x", Icon24x24));
 
 	//////////////////////////////////////////////////
 
@@ -54,90 +67,94 @@ TSharedRef<FSlateStyleSet> FInsightsStyle::Create()
 	//////////////////////////////////////////////////
 	// Icons for major components
 
-	Style.Set("SessionInfo.Icon.Large", new IMAGE_BRUSH("Icons/icon_tab_Tools_16x", FVector2D(32.0f, 32.0f)));
-	Style.Set("SessionInfo.Icon.Small", new IMAGE_BRUSH("Icons/icon_tab_Tools_16x", FVector2D(16.0f, 16.0f)));
+	Style.Set("SessionInfo.Icon.Large", new IMAGE_BRUSH("Icons/icon_tab_Tools_16x", Icon32x32));
+	Style.Set("SessionInfo.Icon.Small", new IMAGE_BRUSH("Icons/icon_tab_Tools_16x", Icon16x16));
 
-	Style.Set("Toolbar.Icon.Large", new IMAGE_BRUSH("Icons/icon_tab_Tools_16x", FVector2D(32.0f, 32.0f)));
-	Style.Set("Toolbar.Icon.Small", new IMAGE_BRUSH("Icons/icon_tab_Tools_16x", FVector2D(16.0f, 16.0f)));
+	Style.Set("Toolbar.Icon.Large", new IMAGE_BRUSH("Icons/icon_tab_Tools_16x", Icon32x32));
+	Style.Set("Toolbar.Icon.Small", new IMAGE_BRUSH("Icons/icon_tab_Tools_16x", Icon16x16));
 
 	//////////////////////////////////////////////////
 	// Start Page buttons
 
-	Style.Set("StartPage.Icon.Large", new IMAGE_BRUSH("Icons/icon_tab_Tools_16x", FVector2D(32.0f, 32.0f)));
-	Style.Set("StartPage.Icon.Small", new IMAGE_BRUSH("Icons/icon_tab_Tools_16x", FVector2D(16.0f, 16.0f)));
+	Style.Set("StartPage.Icon.Large", new IMAGE_BRUSH("Icons/icon_tab_Tools_16x", Icon32x32));
+	Style.Set("StartPage.Icon.Small", new IMAGE_BRUSH("Icons/icon_tab_Tools_16x", Icon16x16));
 
-	Style.Set("Open.Icon.Large", new IMAGE_BRUSH("Icons/Profiler/Profiler_LoadMultiple_Profiler_40x", FVector2D(32.0f, 32.0f)));
-	Style.Set("Open.Icon.Small", new IMAGE_BRUSH("Icons/Profiler/Profiler_Load_Profiler_40x", FVector2D(16.0f, 16.0f)));
+	Style.Set("Open.Icon.Large", new IMAGE_BRUSH("Icons/Profiler/Profiler_LoadMultiple_Profiler_40x", Icon32x32));
+	Style.Set("Open.Icon.Small", new IMAGE_BRUSH("Icons/Profiler/Profiler_Load_Profiler_40x", Icon16x16));
 
-	Style.Set("OpenFile.Icon.Large", new IMAGE_BRUSH("Icons/LV_Load", FVector2D(32.0f, 32.0f)));
-	Style.Set("OpenFile.Icon.Small", new IMAGE_BRUSH("Icons/LV_Load", FVector2D(16.0f, 16.0f)));
+	Style.Set("OpenFile.Icon.Large", new IMAGE_BRUSH("Icons/LV_Load", Icon32x32));
+	Style.Set("OpenFile.Icon.Small", new IMAGE_BRUSH("Icons/LV_Load", Icon16x16));
+	
+
+	Style.Set("Menu.Icon.Large", new CORE_IMAGE_BRUSH_SVG("Starship/Common/menu", Icon32x32));
+	Style.Set("Menu.Icon.Small", new CORE_IMAGE_BRUSH_SVG("Starship/Common/menu", Icon16x16));
 
 	//////////////////////////////////////////////////
 	// Timing Insights
 
-	Style.Set("TimingProfiler.Icon.Large", new IMAGE_BRUSH("Icons/Profiler/profiler_stats_40x", FVector2D(32.0f, 32.0f)));
-	Style.Set("TimingProfiler.Icon.Small", new IMAGE_BRUSH("Icons/Profiler/profiler_stats_40x", FVector2D(16.0f, 16.0f)));
+	Style.Set("TimingProfiler.Icon.Large", new IMAGE_BRUSH("Icons/Profiler/profiler_stats_40x", Icon32x32));
+	Style.Set("TimingProfiler.Icon.Small", new IMAGE_BRUSH("Icons/Profiler/profiler_stats_40x", Icon16x16));
 
-	Style.Set("FramesTrack.Icon.Large", new IMAGE_BRUSH("Icons/Profiler/profiler_stats_40x", FVector2D(32.0f, 32.0f)));
-	Style.Set("FramesTrack.Icon.Small", new IMAGE_BRUSH("Icons/Profiler/profiler_stats_40x", FVector2D(16.0f, 16.0f)));
+	Style.Set("FramesTrack.Icon.Large", new IMAGE_BRUSH("Icons/Profiler/profiler_stats_40x", Icon32x32));
+	Style.Set("FramesTrack.Icon.Small", new IMAGE_BRUSH("Icons/Profiler/profiler_stats_40x", Icon16x16));
 
-	Style.Set("GraphTrack.Icon.Large", new IMAGE_BRUSH("Icons/Profiler/profiler_stats_40x", FVector2D(32.0f, 32.0f)));
-	Style.Set("GraphTrack.Icon.Small", new IMAGE_BRUSH("Icons/Profiler/profiler_stats_40x", FVector2D(16.0f, 16.0f)));
+	Style.Set("GraphTrack.Icon.Large", new IMAGE_BRUSH("Icons/Profiler/profiler_stats_40x", Icon32x32));
+	Style.Set("GraphTrack.Icon.Small", new IMAGE_BRUSH("Icons/Profiler/profiler_stats_40x", Icon16x16));
 
-	Style.Set("TimingView.Icon.Large", new IMAGE_BRUSH("Icons/Profiler/profiler_stats_40x", FVector2D(32.0f, 32.0f)));
-	Style.Set("TimingView.Icon.Small", new IMAGE_BRUSH("Icons/Profiler/profiler_stats_40x", FVector2D(16.0f, 16.0f)));
+	Style.Set("TimingView.Icon.Large", new IMAGE_BRUSH("Icons/Profiler/profiler_stats_40x", Icon32x32));
+	Style.Set("TimingView.Icon.Small", new IMAGE_BRUSH("Icons/Profiler/profiler_stats_40x", Icon16x16));
 
-	Style.Set("TimersView.Icon.Large", new IMAGE_BRUSH("Icons/Profiler/Profiler_Data_Capture_40x", FVector2D(32.0f, 32.0f)));
-	Style.Set("TimersView.Icon.Small", new IMAGE_BRUSH("Icons/Profiler/Profiler_Data_Capture_40x", FVector2D(16.0f, 16.0f)));
+	Style.Set("TimersView.Icon.Large", new IMAGE_BRUSH("Icons/Profiler/Profiler_Data_Capture_40x", Icon32x32));
+	Style.Set("TimersView.Icon.Small", new IMAGE_BRUSH("Icons/Profiler/Profiler_Data_Capture_40x", Icon16x16));
 
-	Style.Set("StatsCountersView.Icon.Large", new IMAGE_BRUSH("Icons/Profiler/Profiler_Data_Capture_40x", FVector2D(32.0f, 32.0f)));
-	Style.Set("StatsCountersView.Icon.Small", new IMAGE_BRUSH("Icons/Profiler/Profiler_Data_Capture_40x", FVector2D(16.0f, 16.0f)));
+	Style.Set("StatsCountersView.Icon.Large", new IMAGE_BRUSH("Icons/Profiler/Profiler_Data_Capture_40x", Icon32x32));
+	Style.Set("StatsCountersView.Icon.Small", new IMAGE_BRUSH("Icons/Profiler/Profiler_Data_Capture_40x", Icon16x16));
 
-	Style.Set("LogView.Icon.Large", new IMAGE_BRUSH("Icons/Profiler/profiler_CopyToClipboard_32x", FVector2D(32.0f, 32.0f)));
-	Style.Set("LogView.Icon.Small", new IMAGE_BRUSH("Icons/Profiler/profiler_CopyToClipboard_32x", FVector2D(16.0f, 16.0f)));
+	Style.Set("LogView.Icon.Large", new IMAGE_BRUSH("Icons/Profiler/profiler_CopyToClipboard_32x", Icon32x32));
+	Style.Set("LogView.Icon.Small", new IMAGE_BRUSH("Icons/Profiler/profiler_CopyToClipboard_32x", Icon16x16));
 
-	Style.Set("TableTreeView.Icon.Large", new IMAGE_BRUSH("Icons/Profiler/Profiler_Data_Capture_40x", FVector2D(32.0f, 32.0f)));
-	Style.Set("TableTreeView.Icon.Small", new IMAGE_BRUSH("Icons/Profiler/Profiler_Data_Capture_40x", FVector2D(16.0f, 16.0f)));
+	Style.Set("TableTreeView.Icon.Large", new IMAGE_BRUSH("Icons/Profiler/Profiler_Data_Capture_40x", Icon32x32));
+	Style.Set("TableTreeView.Icon.Small", new IMAGE_BRUSH("Icons/Profiler/Profiler_Data_Capture_40x", Icon16x16));
 
 	//////////////////////////////////////////////////
 	// Asset Loading Insights
 
-	Style.Set("LoadingProfiler.Icon.Large", new IMAGE_BRUSH("Icons/Profiler/profiler_stats_40x", FVector2D(32.0f, 32.0f)));
-	Style.Set("LoadingProfiler.Icon.Small", new IMAGE_BRUSH("Icons/Profiler/profiler_stats_40x", FVector2D(16.0f, 16.0f)));
+	Style.Set("LoadingProfiler.Icon.Large", new IMAGE_BRUSH("Icons/Profiler/profiler_stats_40x", Icon32x32));
+	Style.Set("LoadingProfiler.Icon.Small", new IMAGE_BRUSH("Icons/Profiler/profiler_stats_40x", Icon16x16));
 
 	//////////////////////////////////////////////////
 	// Networking Insights
 
-	Style.Set("NetworkingProfiler.Icon.Large", new IMAGE_BRUSH("Icons/Profiler/profiler_stats_40x", FVector2D(32.0f, 32.0f)));
-	Style.Set("NetworkingProfiler.Icon.Small", new IMAGE_BRUSH("Icons/Profiler/profiler_stats_40x", FVector2D(16.0f, 16.0f)));
+	Style.Set("NetworkingProfiler.Icon.Large", new IMAGE_BRUSH("Icons/Profiler/profiler_stats_40x", Icon32x32));
+	Style.Set("NetworkingProfiler.Icon.Small", new IMAGE_BRUSH("Icons/Profiler/profiler_stats_40x", Icon16x16));
 
-	Style.Set("PacketOveriew.Icon.Large", new IMAGE_BRUSH("Icons/Profiler/profiler_stats_40x", FVector2D(32.0f, 32.0f)));
-	Style.Set("PacketOveriew.Icon.Small", new IMAGE_BRUSH("Icons/Profiler/profiler_stats_40x", FVector2D(16.0f, 16.0f)));
+	Style.Set("PacketOveriew.Icon.Large", new IMAGE_BRUSH("Icons/Profiler/profiler_stats_40x", Icon32x32));
+	Style.Set("PacketOveriew.Icon.Small", new IMAGE_BRUSH("Icons/Profiler/profiler_stats_40x", Icon16x16));
 
-	Style.Set("PacketContentView.Icon.Large", new IMAGE_BRUSH("Icons/Profiler/profiler_stats_40x", FVector2D(32.0f, 32.0f)));
-	Style.Set("PacketContentView.Icon.Small", new IMAGE_BRUSH("Icons/Profiler/profiler_stats_40x", FVector2D(16.0f, 16.0f)));
+	Style.Set("PacketContentView.Icon.Large", new IMAGE_BRUSH("Icons/Profiler/profiler_stats_40x", Icon32x32));
+	Style.Set("PacketContentView.Icon.Small", new IMAGE_BRUSH("Icons/Profiler/profiler_stats_40x", Icon16x16));
 
-	Style.Set("NetStatsView.Icon.Large", new IMAGE_BRUSH("Icons/Profiler/profiler_stats_40x", FVector2D(32.0f, 32.0f)));
-	Style.Set("NetStatsView.Icon.Small", new IMAGE_BRUSH("Icons/Profiler/profiler_stats_40x", FVector2D(16.0f, 16.0f)));
+	Style.Set("NetStatsView.Icon.Large", new IMAGE_BRUSH("Icons/Profiler/profiler_stats_40x", Icon32x32));
+	Style.Set("NetStatsView.Icon.Small", new IMAGE_BRUSH("Icons/Profiler/profiler_stats_40x", Icon16x16));
 
 	//////////////////////////////////////////////////
 	// Memory Insights
 
-	Style.Set("MemoryProfiler.Icon.Large", new IMAGE_BRUSH("Icons/Profiler/profiler_stats_40x", FVector2D(32.0f, 32.0f)));
-	Style.Set("MemoryProfiler.Icon.Small", new IMAGE_BRUSH("Icons/Profiler/profiler_stats_40x", FVector2D(16.0f, 16.0f)));
+	Style.Set("MemoryProfiler.Icon.Large", new IMAGE_BRUSH("Icons/Profiler/profiler_stats_40x", Icon32x32));
+	Style.Set("MemoryProfiler.Icon.Small", new IMAGE_BRUSH("Icons/Profiler/profiler_stats_40x", Icon16x16));
 
-	Style.Set("MemTagTreeView.Icon.Large", new IMAGE_BRUSH("Icons/Profiler/profiler_stats_40x", FVector2D(32.0f, 32.0f)));
-	Style.Set("MemTagTreeView.Icon.Small", new IMAGE_BRUSH("Icons/Profiler/profiler_stats_40x", FVector2D(16.0f, 16.0f)));
+	Style.Set("MemTagTreeView.Icon.Large", new IMAGE_BRUSH("Icons/Profiler/profiler_stats_40x", Icon32x32));
+	Style.Set("MemTagTreeView.Icon.Small", new IMAGE_BRUSH("Icons/Profiler/profiler_stats_40x", Icon16x16));
 
-	Style.Set("Mem.Add.Small", new IMAGE_BRUSH("Icons/icon_Cascade_AddLOD2_40x", FVector2D(20.0f, 20.0f)));
-	Style.Set("Mem.Remove.Small", new IMAGE_BRUSH("Icons/icon_Cascade_DeleteLOD_40x", FVector2D(20.0f, 20.0f)));
+	Style.Set("Mem.Add.Small", new IMAGE_BRUSH("Icons/icon_Cascade_AddLOD2_40x", Icon20x20));
+	Style.Set("Mem.Remove.Small", new IMAGE_BRUSH("Icons/icon_Cascade_DeleteLOD_40x", Icon20x20));
 
 	//////////////////////////////////////////////////
 
-	Style.Set("FindFirst", new IMAGE_BRUSH("Animation/backward_end", FVector2D(20.0f, 20.0f)));
-	Style.Set("FindPrevious", new IMAGE_BRUSH("Animation/backward", FVector2D(20.0f, 20.0f)));
-	Style.Set("FindNext", new IMAGE_BRUSH("Animation/forward", FVector2D(20.0f, 20.0f)));
-	Style.Set("FindLast", new IMAGE_BRUSH("Animation/forward_end", FVector2D(20.0f, 20.0f)));
+	Style.Set("FindFirst", new IMAGE_BRUSH("Animation/backward_end", Icon20x20));
+	Style.Set("FindPrevious", new IMAGE_BRUSH("Animation/backward", Icon20x20));
+	Style.Set("FindNext", new IMAGE_BRUSH("Animation/forward", Icon20x20));
+	Style.Set("FindLast", new IMAGE_BRUSH("Animation/forward_end", Icon20x20));
 
 	return StyleRef;
 }
@@ -146,6 +163,14 @@ TSharedRef<FSlateStyleSet> FInsightsStyle::Create()
 #undef BOX_BRUSH
 #undef BORDER_BRUSH
 
+#undef IMAGE_BRUSH_SVG
+#undef BOX_BRUSH_SVG
+#undef BORDER_BRUSH_SVG
+
+#undef CORE_IMAGE_BRUSH_SVG
+#undef CORE_BOX_BRUSH_SVG
+#undef CORE_BORDER_BRUSH_SVG
+
 const ISlateStyle& FInsightsStyle::Get()
 {
 	return *StyleInstance;
diff --git a/Engine/Source/Developer/TraceInsights/Private/Insights/Widgets/SStartPageWindow.cpp b/Engine/Source/Developer/TraceInsights/Private/Insights/Widgets/SStartPageWindow.cpp
--- a/Engine/Source/Developer/TraceInsights/Private/Insights/Widgets/SStartPageWindow.cpp
+++ b/Engine/Source/Developer/TraceInsights/Private/Insights/Widgets/SStartPageWindow.cpp
@@ -567,7 +567,7 @@ void SStartPageWindow::Construct(const FArguments& InArgs)
 			+ SOverlay::Slot()
 			.HAlign(HAlign_Right)
 			.VAlign(VAlign_Top)
-			.Padding(0.0f, -16.0f, 0.0f, 0.0f)
+			.Padding(0.0f, -16.0f, 4.0f, 0.0f)
 			[
 				SNew(STextBlock)
 				.Clipping(EWidgetClipping::ClipToBoundsWithoutIntersecting)
@@ -800,7 +800,7 @@ TSharedRef<SWidget> SStartPageWindow::ConstructLoadPanel()
 		.IsEnabled(this, &SStartPageWindow::Open_IsEnabled)
 		.OnClicked(this, &SStartPageWindow::Open_OnClicked)
 		.ToolTipText(LOCTEXT("OpenButtonTooltip", "Start analysis for selected trace session."))
-		.ContentPadding(FMargin(4.0f, 1.0f))
+		.ContentPadding(FMargin(0.0f, 0.0f))
 		.Content()
 		[
 			SNew(SHorizontalBox)
@@ -818,18 +818,25 @@ TSharedRef<SWidget> SStartPageWindow::ConstructLoadPanel()
 				[
 					SNew(STextBlock)
 					.Text(LOCTEXT("OpenButtonText", "Open"))
+					.TextStyle(FAppStyle::Get(), "ButtonText")
 				]
 		]
 	]
 
 	+ SHorizontalBox::Slot()
+	.Padding(FMargin(6.0f, 0.0f, 0.0f, 0.0f))
 	.AutoWidth()
 	[
 		SNew(SComboButton)
 		.ToolTipText(LOCTEXT("MRU_Tooltip", "Open a trace file or choose a trace session."))
 		.OnGetMenuContent(this, &SStartPageWindow::MakeTraceListMenu)
 		.HasDownArrow(true)
-		.ContentPadding(FMargin(1.0f, 1.0f, 1.0f, 1.0f))
+		.ContentPadding(FMargin(0.0f, 0.0f, 0.0f, 0.0f))
+		.ButtonContent()
+		[
+			SNew(SImage)
+			.Image(FInsightsStyle::GetBrush("Menu.Icon.Small"))
+		]
 	];
 
 	return Widget;
