Fix UnrealInsightsLayout saved in incorrect folder

diff --git a/Engine/Source/Developer/TraceInsights/Private/Insights/TraceInsightsModule.cpp b/Engine/Source/Developer/TraceInsights/Private/Insights/TraceInsightsModule.cpp
--- a/Engine/Source/Developer/TraceInsights/Private/Insights/TraceInsightsModule.cpp
+++ b/Engine/Source/Developer/TraceInsights/Private/Insights/TraceInsightsModule.cpp
@@ -5,6 +5,8 @@
 #include "Framework/Docking/LayoutService.h"
 #include "Framework/Notifications/NotificationManager.h"
 #include "HAL/PlatformApplicationMisc.h"
+#include "Misc/ConfigCacheIni.h"
+#include "Misc/Paths.h"
 #include "Modules/ModuleManager.h"
 #include "Trace/StoreClient.h"
 #include "Trace/StoreService.h"
@@ -51,7 +53,7 @@ void FTraceInsightsModule::StartupModule()
 	RegisterComponent(FNetworkingProfilerManager::CreateInstance());
 	RegisterComponent(FMemoryProfilerManager::CreateInstance());
 
-	UnrealInsightsLayoutIni = FPaths::GetPath(GEngineIni) + "/UnrealInsightsLayout.ini";
+	UnrealInsightsLayoutIni = FConfigCacheIni::GetDestIniFilename(TEXT("UnrealInsightsLayout"), nullptr, *FPaths::GeneratedConfigDir());
 }
 
 ////////////////////////////////////////////////////////////////////////////////////////////////////
