Filter Configurator initial implementation

diff --git a/Engine/Source/Developer/TraceInsights/Private/Insights/InsightsManager.cpp b/Engine/Source/Developer/TraceInsights/Private/Insights/InsightsManager.cpp
--- a/Engine/Source/Developer/TraceInsights/Private/Insights/InsightsManager.cpp
+++ b/Engine/Source/Developer/TraceInsights/Private/Insights/InsightsManager.cpp
@@ -20,6 +20,7 @@
 #include "Insights/NetworkingProfiler/NetworkingProfilerManager.h"
 #include "Insights/Tests/InsightsTestRunner.h"
 #include "Insights/TimingProfilerManager.h"
+#include "Insights/ViewModels/Filters.h"
 #include "Insights/Widgets/SStartPageWindow.h"
 #include "Insights/Widgets/SSessionInfoWindow.h"
 #include "Insights/Widgets/STimingProfilerWindow.h"
@@ -146,6 +147,8 @@ void FInsightsManager::Initialize(IUnrealInsightsModule& InsightsModule)
 
 	InsightsMenuBuilder = MakeShared<FInsightsMenuBuilder>();
 
+	FFilterService::CreateInstance();
+
 	// Register tick functions.
 	OnTick = FTickerDelegate::CreateSP(this, &FInsightsManager::Tick);
 	OnTickHandle = FTicker::GetCoreTicker().AddTicker(OnTick, 0.0f);
diff --git a/Engine/Source/Developer/TraceInsights/Private/Insights/Table/ViewModels/BaseTreeNode.h b/Engine/Source/Developer/TraceInsights/Private/Insights/Table/ViewModels/BaseTreeNode.h
--- a/Engine/Source/Developer/TraceInsights/Private/Insights/Table/ViewModels/BaseTreeNode.h
+++ b/Engine/Source/Developer/TraceInsights/Private/Insights/Table/ViewModels/BaseTreeNode.h
@@ -150,6 +150,11 @@ protected:
 		return Children;
 	}
 
+	void SetGroupPtr(FBaseTreeNodePtr InGroupPtr)
+	{
+		GroupPtr = InGroupPtr;
+	}
+
 private:
 	/** The default sort order. Index used to optimize sorting. */
 	int32 DefaultSortOrder;
diff --git a/Engine/Source/Developer/TraceInsights/Private/Insights/ViewModels/FilterConfigurator.cpp b/Engine/Source/Developer/TraceInsights/Private/Insights/ViewModels/FilterConfigurator.cpp
new file mode 100644
--- /dev/null
+++ b/Engine/Source/Developer/TraceInsights/Private/Insights/ViewModels/FilterConfigurator.cpp
@@ -0,0 +1,44 @@
+// Copyright Epic Games, Inc. All Rights Reserved.
+
+#include "FilterConfigurator.h"
+
+////////////////////////////////////////////////////////////////////////////////////////////////////
+
+FFilterConfigurator::FFilterConfigurator()
+{
+	RootNode = MakeShared<FFilterConfiguratorNode>(TEXT(""), true);
+	AvailableFilters = MakeShared<TArray<TSharedPtr<FFilter>>>();
+	RootNode->SetAvailableFilters(AvailableFilters);
+}
+
+////////////////////////////////////////////////////////////////////////////////////////////////////
+
+FFilterConfigurator::FFilterConfigurator(const FFilterConfigurator& Other)
+{
+	*this = Other;
+}
+
+////////////////////////////////////////////////////////////////////////////////////////////////////
+
+FFilterConfigurator& FFilterConfigurator::operator=(const FFilterConfigurator& Other)
+{
+	RootNode = MakeShared<FFilterConfiguratorNode>(*Other.RootNode);
+	RootNode->SetGroupPtrForChildren();
+	AvailableFilters = Other.AvailableFilters;
+
+	if (OnChangesCommitedCallback)
+	{
+		OnChangesCommitedCallback();
+	}
+
+	return *this;
+}
+
+////////////////////////////////////////////////////////////////////////////////////////////////////
+
+bool FFilterConfigurator::ApplyFilters(const FFilterContext& Context) const
+{
+	return RootNode->ApplyFilters(Context);
+}
+
+////////////////////////////////////////////////////////////////////////////////////////////////////
diff --git a/Engine/Source/Developer/TraceInsights/Private/Insights/ViewModels/FilterConfigurator.h b/Engine/Source/Developer/TraceInsights/Private/Insights/ViewModels/FilterConfigurator.h
new file mode 100644
--- /dev/null
+++ b/Engine/Source/Developer/TraceInsights/Private/Insights/ViewModels/FilterConfigurator.h
@@ -0,0 +1,41 @@
+// Copyright Epic Games, Inc. All Rights Reserved.
+
+#pragma once
+
+#include "CoreMinimal.h"
+#include "Insights/ViewModels/Filters.h"
+#include "Insights/ViewModels/FilterConfiguratorNode.h"
+ 
+////////////////////////////////////////////////////////////////////////////////////////////////////
+
+class FFilterConfigurator
+{
+public:
+	FFilterConfigurator();
+
+	FFilterConfigurator(const FFilterConfigurator& Other);
+	FFilterConfigurator& operator=(const FFilterConfigurator& Other);
+
+	~FFilterConfigurator() {}
+
+	void AddFilter(EFilterField FilterKey)
+	{
+		AvailableFilters->Add(FFilterService::Get()->GetFilter(FilterKey));
+	}
+
+	FFilterConfiguratorNodePtr GetRootNode() { return RootNode; }
+
+	bool ApplyFilters(const FFilterContext& Context) const;
+
+	void SetOnChangesCommitedCallback(TFunction<void()> InCallback) {	OnChangesCommitedCallback = InCallback;	}
+
+private:
+
+	FFilterConfiguratorNodePtr RootNode;
+
+	TSharedPtr<TArray<TSharedPtr<struct FFilter>>> AvailableFilters;
+
+	TFunction<void()> OnChangesCommitedCallback = nullptr;
+};
+
+////////////////////////////////////////////////////////////////////////////////////////////////////
\ No newline at end of file
diff --git a/Engine/Source/Developer/TraceInsights/Private/Insights/ViewModels/FilterConfiguratorNode.cpp b/Engine/Source/Developer/TraceInsights/Private/Insights/ViewModels/FilterConfiguratorNode.cpp
new file mode 100644
--- /dev/null
+++ b/Engine/Source/Developer/TraceInsights/Private/Insights/ViewModels/FilterConfiguratorNode.cpp
@@ -0,0 +1,164 @@
+// Copyright Epic Games, Inc. All Rights Reserved.
+
+#include "FilterConfiguratorNode.h"
+#include "Insights/ViewModels/Filters.h"
+
+// Insights
+
+#define LOCTEXT_NAMESPACE "FilterConfiguratorNode"
+
+////////////////////////////////////////////////////////////////////////////////////////////////////
+
+const FName FFilterConfiguratorNode::TypeName(TEXT("FFilterConfiguratorNode"));
+
+////////////////////////////////////////////////////////////////////////////////////////////////////
+
+FFilterConfiguratorNode::FFilterConfiguratorNode(const FName InName, bool bInIsGroup)
+	: FBaseTreeNode(InName, bInIsGroup)
+{
+	if (bInIsGroup)
+	{
+		SelectedFilterGroupOperator = GetFilterGroupOperators()[0];
+	}
+}
+
+////////////////////////////////////////////////////////////////////////////////////////////////////
+
+FFilterConfiguratorNode::FFilterConfiguratorNode(const FFilterConfiguratorNode& Other)
+	: FBaseTreeNode(Other.GetName(), Other.IsGroup())
+{
+	*this = Other;
+}
+
+////////////////////////////////////////////////////////////////////////////////////////////////////
+
+FFilterConfiguratorNode& FFilterConfiguratorNode::operator=(const FFilterConfiguratorNode& Other)
+{
+	GetChildrenMutable().Empty();
+
+	AvailableFilters = Other.AvailableFilters;
+	SelectedFilter = Other.SelectedFilter;
+	SelectedFilterOperator = Other.SelectedFilterOperator;
+	SelectedFilterGroupOperator = Other.SelectedFilterGroupOperator;
+	TextBoxValue = Other.TextBoxValue;
+	SetExpansion(Other.IsExpanded());
+
+	for (Insights::FBaseTreeNodePtr Child : Other.GetChildren())
+	{
+		GetChildrenMutable().Add(MakeShared<FFilterConfiguratorNode>(*StaticCastSharedPtr<FFilterConfiguratorNode>(Child)));
+	}
+
+	return *this;
+}
+
+////////////////////////////////////////////////////////////////////////////////////////////////////
+
+TArray<TSharedPtr<struct FFilterGroupOperator>>& FFilterConfiguratorNode::GetFilterGroupOperators()
+{
+	return FFilterService::Get()->GetFilterGroupOperators();
+}
+
+////////////////////////////////////////////////////////////////////////////////////////////////////
+
+void FFilterConfiguratorNode::DeleteChildNode(FFilterConfiguratorNodePtr InNode)
+{
+	Insights::FBaseTreeNodePtr Node = StaticCastSharedPtr<Insights::FBaseTreeNode>(InNode);
+	GetChildrenMutable().Remove(Node);
+}
+
+////////////////////////////////////////////////////////////////////////////////////////////////////
+
+void FFilterConfiguratorNode::SetGroupPtrForChildren()
+{
+	for (Insights::FBaseTreeNodePtr Child : GetChildrenMutable())
+	{
+		FFilterConfiguratorNodePtr CastedChild = StaticCastSharedPtr<FFilterConfiguratorNode>(Child);
+		CastedChild->SetGroupPtrForChildren();
+		CastedChild->SetGroupPtr(SharedThis(this));
+	}
+}
+
+////////////////////////////////////////////////////////////////////////////////////////////////////
+
+void FFilterConfiguratorNode::SetAvailableFilters(TSharedPtr<TArray<TSharedPtr<struct FFilter>>> InAvailableFilters)
+{
+	AvailableFilters = InAvailableFilters;
+
+	if (AvailableFilters.IsValid() && AvailableFilters->Num() > 0)
+	{
+		SetSelectedFilter(AvailableFilters->GetData()[0]);
+	}
+}
+
+////////////////////////////////////////////////////////////////////////////////////////////////////
+
+void FFilterConfiguratorNode::SetSelectedFilter(TSharedPtr<struct FFilter> InSelectedFilter)
+{
+	SelectedFilter = InSelectedFilter;
+	if (SelectedFilter.IsValid() && SelectedFilter->GetSupportedOperators()->Num() > 0)
+	{
+		SetSelectedFilterOperator(SelectedFilter->GetSupportedOperators()->GetData()[0]);
+	}
+}
+
+////////////////////////////////////////////////////////////////////////////////////////////////////
+
+bool FFilterConfiguratorNode::ApplyFilters(const FFilterContext& Context) const
+{
+	bool ret = true;
+	if (IsGroup())
+	{
+		if (SelectedFilterGroupOperator->Type == EFilterGroupOperator::And)
+		{
+			auto& ChildrenArr = GetChildren();
+			for (int Index = 0; Index < ChildrenArr.Num(); ++Index)
+			{
+				ret &= ((FFilterConfiguratorNode*)&*ChildrenArr[Index])->ApplyFilters(Context);
+			}
+		}
+
+		if (SelectedFilterGroupOperator->Type == EFilterGroupOperator::Or)
+		{
+			ret = false;
+
+			auto& ChildrenArr = GetChildren();
+			for (int Index = 0; Index < ChildrenArr.Num(); ++Index)
+			{
+				ret |= ((FFilterConfiguratorNode*)&*ChildrenArr[Index])->ApplyFilters(Context);
+			}
+		}
+	}
+	else
+	{
+		switch (SelectedFilter->DataType)
+		{
+		case EFilterDataType::Double:
+			{
+				FFilterOperator<double>* Operator = (FFilterOperator<double>*)&*SelectedFilterOperator;
+				double Value = 0.0;
+				Context.GetFilterData<double>(SelectedFilter->Key, Value);
+
+				ret = Operator->Func(Value, FCString::Atod(*TextBoxValue));
+				break;
+			}
+
+		case EFilterDataType::Int64:
+		{
+			FFilterOperator<int64>* Operator = (FFilterOperator<int64>*) & *SelectedFilterOperator;
+			int64 Value = 0;
+			Context.GetFilterData<int64>(SelectedFilter->Key, Value);
+
+			ret = Operator->Func(Value, FCString::Atoi64(*TextBoxValue));
+			break;
+		}
+		default:
+			break;
+		}
+	}
+
+	return ret;
+}
+
+////////////////////////////////////////////////////////////////////////////////////////////////////
+
+#undef LOCTEXT_NAMESPACE
diff --git a/Engine/Source/Developer/TraceInsights/Private/Insights/ViewModels/FilterConfiguratorNode.h b/Engine/Source/Developer/TraceInsights/Private/Insights/ViewModels/FilterConfiguratorNode.h
new file mode 100644
--- /dev/null
+++ b/Engine/Source/Developer/TraceInsights/Private/Insights/ViewModels/FilterConfiguratorNode.h
@@ -0,0 +1,93 @@
+// Copyright Epic Games, Inc. All Rights Reserved.
+
+#pragma once
+
+#include "CoreMinimal.h"
+
+// Insights
+#include "Insights/Table/ViewModels/BaseTreeNode.h"
+
+////////////////////////////////////////////////////////////////////////////////////////////////////
+
+enum class EFilterConfiguratorNodeType
+{
+	/** A Filter node. */
+	Filter,
+
+	/** A group node. */
+	Group,
+
+	/** Invalid enum type, may be used as a number of enumerations. */
+	InvalidOrMax,
+};
+
+////////////////////////////////////////////////////////////////////////////////////////////////////
+
+/** Type definition for shared pointers to instances of FFilterConfiguratorNode. */
+typedef TSharedPtr<class FFilterConfiguratorNode> FFilterConfiguratorNodePtr;
+
+/** Type definition for shared references to instances of FFilterConfiguratorNode. */
+typedef TSharedRef<class FFilterConfiguratorNode> FFilterConfiguratorRef;
+
+/** Type definition for shared references to const instances of FFilterConfiguratorNode. */
+typedef TSharedRef<const class FFilterConfiguratorNode> FFilterConfiguratorRefConst;
+
+/** Type definition for weak references to instances of FFilterConfiguratorNode. */
+typedef TWeakPtr<class FFilterConfiguratorNode> FFilterConfiguratorNodeWeak;
+
+////////////////////////////////////////////////////////////////////////////////////////////////////
+/**
+ * Class used to store information about a filter node
+ */
+class FFilterConfiguratorNode : public Insights::FBaseTreeNode
+{
+	static const FName TypeName;
+
+public:
+	/** Initialization constructor for the filter configurator node. */
+	FFilterConfiguratorNode(const FName InName, bool bInIsGroup);
+
+	FFilterConfiguratorNode(const FFilterConfiguratorNode& Other);
+
+	virtual ~FFilterConfiguratorNode() {}
+
+	virtual const FName& GetTypeName() const override { return TypeName; }
+
+	void SetAvailableFilters(TSharedPtr<TArray<TSharedPtr<struct FFilter>>> InAvailableFilters);
+	TSharedPtr<TArray<TSharedPtr<struct FFilter>>> GetAvailableFilters() { return AvailableFilters; }
+
+	void SetSelectedFilter(TSharedPtr<struct FFilter> InSelectedFilter);
+	TSharedPtr<struct FFilter> GetSelectedFilter() const { return SelectedFilter; }
+
+	void SetSelectedFilterOperator(TSharedPtr<class IFilterOperator> InSelectedFilterOperator) { SelectedFilterOperator = InSelectedFilterOperator; }
+	TSharedPtr<IFilterOperator> GetSelectedFilterOperator() const { return SelectedFilterOperator;	}
+
+	TArray<TSharedPtr<struct FFilterGroupOperator>>& GetFilterGroupOperators();
+
+	void SetSelectedFilterGroupOperator(TSharedPtr<struct FFilterGroupOperator> InSelectedFilterGroupOperator) { SelectedFilterGroupOperator = InSelectedFilterGroupOperator; }
+	TSharedPtr<struct FFilterGroupOperator> GetSelectedFilterGroupOperator() const { return SelectedFilterGroupOperator; }
+
+	void DeleteChildNode(FFilterConfiguratorNodePtr InNode);
+
+	const FString& GetTextBoxValue() { return TextBoxValue; }
+	void SetTextBoxValue(const FString& InValue) { TextBoxValue = InValue; }
+
+	void SetGroupPtrForChildren();
+
+	bool ApplyFilters(const class FFilterContext& Context) const;
+
+private:
+	FFilterConfiguratorNode& operator=(const FFilterConfiguratorNode& Other);
+
+	TSharedPtr<TArray<TSharedPtr<struct FFilter>>> AvailableFilters;
+
+	TSharedPtr<struct FFilter> SelectedFilter;
+
+	TSharedPtr<class IFilterOperator> SelectedFilterOperator;
+
+	TSharedPtr<struct FFilterGroupOperator> SelectedFilterGroupOperator;
+
+	FString TextBoxValue;
+};
+
+////////////////////////////////////////////////////////////////////////////////////////////////////
diff --git a/Engine/Source/Developer/TraceInsights/Private/Insights/ViewModels/Filters.cpp b/Engine/Source/Developer/TraceInsights/Private/Insights/ViewModels/Filters.cpp
new file mode 100644
--- /dev/null
+++ b/Engine/Source/Developer/TraceInsights/Private/Insights/ViewModels/Filters.cpp
@@ -0,0 +1,71 @@
+// Copyright Epic Games, Inc. All Rights Reserved.
+
+#include "Filters.h"
+
+#include "Widgets/Docking/SDockTab.h"
+
+#include "Insights/Widgets/SFilterConfigurator.h"
+#include "Insights/ViewModels/FilterConfigurator.h"
+
+#define LOCTEXT_NAMESPACE "SFilterService"
+
+////////////////////////////////////////////////////////////////////////////////////////////////////
+
+TSharedPtr<FFilterService> FFilterService::Instance;
+
+FName const FFilterService::FilterConfiguratorTabId(TEXT("FilterConfigurator"));
+
+////////////////////////////////////////////////////////////////////////////////////////////////////
+
+FFilterService::FFilterService()
+{
+	RegisterTabSpawner();
+}
+
+////////////////////////////////////////////////////////////////////////////////////////////////////
+
+FFilterService::~FFilterService()
+{
+
+}
+
+////////////////////////////////////////////////////////////////////////////////////////////////////
+
+TSharedPtr<SWidget> FFilterService::CreateFilterConfiguratorWidget(TSharedPtr<FFilterConfigurator> FilterConfiguratorViewModel)
+{
+	SAssignNew(PendingWidget, SFilterConfigurator, FilterConfiguratorViewModel);
+
+	if (FGlobalTabmanager::Get()->HasTabSpawner(FilterConfiguratorTabId))
+	{
+		FGlobalTabmanager::Get()->TryInvokeTab(FilterConfiguratorTabId);
+	}
+
+	return PendingWidget;
+}
+
+////////////////////////////////////////////////////////////////////////////////////////////////////
+
+TSharedRef<SDockTab> FFilterService::SpawnTab(const FSpawnTabArgs& Args)
+{
+	const TSharedRef<SDockTab> DockTab = SNew(SDockTab)
+		.TabRole(ETabRole::NomadTab);
+
+	DockTab->SetContent(PendingWidget.ToSharedRef());
+	PendingWidget->SetParentTab(DockTab);
+
+	PendingWidget = nullptr;
+	return DockTab;
+}
+
+////////////////////////////////////////////////////////////////////////////////////////////////////
+
+void FFilterService::RegisterTabSpawner()
+{
+	FTabSpawnerEntry& TabSpawnerEntry = FGlobalTabmanager::Get()->RegisterNomadTabSpawner(FilterConfiguratorTabId,
+		FOnSpawnTab::CreateRaw(this, &FFilterService::SpawnTab))
+		.SetDisplayName(LOCTEXT("FilterConfiguratorTabTitle", "Filter Configurator"));
+}
+
+////////////////////////////////////////////////////////////////////////////////////////////////////
+
+#undef LOCTEXT_NAMESPACE
\ No newline at end of file
diff --git a/Engine/Source/Developer/TraceInsights/Private/Insights/ViewModels/Filters.h b/Engine/Source/Developer/TraceInsights/Private/Insights/ViewModels/Filters.h
new file mode 100644
--- /dev/null
+++ b/Engine/Source/Developer/TraceInsights/Private/Insights/ViewModels/Filters.h
@@ -0,0 +1,220 @@
+// Copyright Epic Games, Inc. All Rights Reserved.
+
+#pragma once
+
+#include <functional>
+
+#include "CoreMinimal.h"
+#include "Misc/TVariant.h"
+
+#define LOCTEXT_NAMESPACE "Filters"
+
+////////////////////////////////////////////////////////////////////////////////////////////////////
+
+enum class EFilterField : uint32
+{
+	StartTime = 0,
+	EndTime = 1,
+	Duration = 2,
+	EventType = 3
+};
+
+////////////////////////////////////////////////////////////////////////////////////////////////////
+
+enum class EFilterDataType : uint32
+{
+	Int64,
+	Double
+};
+
+////////////////////////////////////////////////////////////////////////////////////////////////////
+
+enum class EFilterOperator : uint8
+{
+	Eq = 0, // Equals
+	Lt = 1, // Less Than
+	Lte = 2, // Less than or equal to
+	Gt = 3, // Greater than
+	Gte = 4, // Greater than or equal to 
+};
+
+////////////////////////////////////////////////////////////////////////////////////////////////////
+
+class IFilterOperator
+{
+public:
+
+	virtual EFilterOperator GetKey() = 0;
+	virtual FString GetName() = 0;
+};
+
+////////////////////////////////////////////////////////////////////////////////////////////////////
+
+template<typename T>
+class FFilterOperator : public IFilterOperator
+{
+public:
+
+	typedef TFunction<bool(T, T)> OperatorFunc;
+
+	FFilterOperator(EFilterOperator InKey, FString InName, OperatorFunc InFunc)
+		: Func(InFunc)
+		, Key(InKey)
+		, Name(InName)
+	{}
+	virtual ~FFilterOperator() {};
+
+	virtual EFilterOperator GetKey() override { return Key; }
+	virtual FString GetName() override { return Name; };
+
+	OperatorFunc Func;
+private:
+	EFilterOperator Key;
+	FString Name;
+};
+
+////////////////////////////////////////////////////////////////////////////////////////////////////
+
+enum class EFilterGroupOperator
+{
+	And = 0,
+	Or = 1,
+};
+
+////////////////////////////////////////////////////////////////////////////////////////////////////
+
+struct FFilterGroupOperator
+{
+	FFilterGroupOperator(EFilterGroupOperator InType, FText InName, FText InDesc)
+		: Type(InType)
+		, Name(InName)
+		, Desc(InDesc)
+	{}
+
+	EFilterGroupOperator Type;
+	FText Name;
+	FText Desc;
+};
+
+////////////////////////////////////////////////////////////////////////////////////////////////////
+
+struct FFilter
+{
+	typedef TSharedPtr<const TArray<TSharedPtr<IFilterOperator>>> SupportedOperatorsArrayPtr;
+
+	FFilter(EFilterField InKey, FText InName, FText InDesc, EFilterDataType InDataType, SupportedOperatorsArrayPtr InSupportedOperators)
+		: Key(InKey)
+		, Name(InName)
+		, Desc(InDesc)
+		, DataType(InDataType)
+		, SupportedOperators(InSupportedOperators)
+	{}
+
+	SupportedOperatorsArrayPtr GetSupportedOperators() { return SupportedOperators;	}
+
+	EFilterField Key;
+	FText Name;
+	FText Desc;
+	EFilterDataType DataType;
+	SupportedOperatorsArrayPtr SupportedOperators;
+};
+
+////////////////////////////////////////////////////////////////////////////////////////////////////
+
+class FFilterStorage
+{
+public:
+	FFilterStorage()
+	{
+		TSharedPtr<TArray<TSharedPtr<IFilterOperator>>> DoubleOperators = MakeShared<TArray<TSharedPtr<IFilterOperator>>>();
+		DoubleOperators->Add(StaticCastSharedRef<IFilterOperator>(MakeShared<FFilterOperator<double>>(EFilterOperator::Lt, TEXT("<"), std::less<>{})));
+		DoubleOperators->Add(StaticCastSharedRef<IFilterOperator>(MakeShared<FFilterOperator<double>>(EFilterOperator::Lte, TEXT("<="), std::less_equal<>())));
+		DoubleOperators->Add(StaticCastSharedRef<IFilterOperator>(MakeShared<FFilterOperator<double>>(EFilterOperator::Eq, TEXT("="), std::equal_to<>())));
+		DoubleOperators->Add(StaticCastSharedRef<IFilterOperator>(MakeShared<FFilterOperator<double>>(EFilterOperator::Gt, TEXT(">"), std::greater<>())));
+		DoubleOperators->Add(StaticCastSharedRef<IFilterOperator>(MakeShared<FFilterOperator<double>>(EFilterOperator::Gte, TEXT(">="), std::greater_equal<>())));
+
+		TSharedPtr<TArray<TSharedPtr<IFilterOperator>>> IntegerOperators = MakeShared<TArray<TSharedPtr<IFilterOperator>>>();
+		IntegerOperators->Add(StaticCastSharedRef<IFilterOperator>(MakeShared<FFilterOperator<int64>>(EFilterOperator::Lt, TEXT("<"), std::less<>{})));
+		IntegerOperators->Add(StaticCastSharedRef<IFilterOperator>(MakeShared<FFilterOperator<int64>>(EFilterOperator::Lte, TEXT("<="), std::less_equal<>())));
+		IntegerOperators->Add(StaticCastSharedRef<IFilterOperator>(MakeShared<FFilterOperator<int64>>(EFilterOperator::Eq, TEXT("="), std::equal_to<>())));
+		IntegerOperators->Add(StaticCastSharedRef<IFilterOperator>(MakeShared<FFilterOperator<int64>>(EFilterOperator::Gt, TEXT(">"), std::greater<>())));
+		IntegerOperators->Add(StaticCastSharedRef<IFilterOperator>(MakeShared<FFilterOperator<int64>>(EFilterOperator::Gte, TEXT(">="), std::greater_equal<>())));
+
+		AllAvailableFilters.Add(EFilterField::StartTime, MakeShared<FFilter>(EFilterField::StartTime, LOCTEXT("StartTime", "Start Time"), LOCTEXT("StartTime", "Start Time"), EFilterDataType::Double, DoubleOperators));
+		AllAvailableFilters.Add(EFilterField::EndTime, MakeShared<FFilter>(EFilterField::EndTime, LOCTEXT("EndTime", "End Time"), LOCTEXT("EndTime", "End Time"), EFilterDataType::Double, DoubleOperators));
+		AllAvailableFilters.Add(EFilterField::Duration, MakeShared<FFilter>(EFilterField::Duration, LOCTEXT("Duration", "Duration"), LOCTEXT("Duration", "Duration"), EFilterDataType::Double, DoubleOperators));
+		AllAvailableFilters.Add(EFilterField::EventType, MakeShared<FFilter>(EFilterField::EventType, LOCTEXT("Type", "Type"), LOCTEXT("Type", "Type"), EFilterDataType::Int64, IntegerOperators));
+
+		FilterGroupOperators.Add(MakeShared<FFilterGroupOperator>(EFilterGroupOperator::And, LOCTEXT("AllOff", "All Off"), LOCTEXT("AllOff", "All Off")));
+		FilterGroupOperators.Add(MakeShared<FFilterGroupOperator>(EFilterGroupOperator::Or, LOCTEXT("AnyOff", "Any Off"), LOCTEXT("AnyOff", "Any Off")));
+	}
+
+	TSharedPtr<FFilter> GetFilter(EFilterField FilterKey) { return *AllAvailableFilters.Find(FilterKey); }
+	TArray<TSharedPtr<FFilterGroupOperator>>& GetFilterGroupOperators() { return FilterGroupOperators; }
+
+private:
+	TMap<EFilterField, TSharedPtr<FFilter>> AllAvailableFilters;
+	TArray<TSharedPtr<FFilterGroupOperator>> FilterGroupOperators;
+};
+
+////////////////////////////////////////////////////////////////////////////////////////////////////
+
+class FFilterService
+{
+public:
+	FFilterService();
+	~FFilterService();
+
+	static void CreateInstance() { Instance = MakeShared<FFilterService>(); }
+	static TSharedPtr<FFilterService> Get() { return Instance; }
+
+	void RegisterTabSpawner();
+	TSharedRef<class SDockTab> SpawnTab(const class FSpawnTabArgs& Args);
+
+	TSharedPtr<FFilter> GetFilter(EFilterField FilterKey) { return FilterStorage.GetFilter(FilterKey); }
+	TArray<TSharedPtr<FFilterGroupOperator>>& GetFilterGroupOperators() { return FilterStorage.GetFilterGroupOperators(); }
+
+	TSharedPtr<class SWidget> CreateFilterConfiguratorWidget(TSharedPtr<class FFilterConfigurator> FilterConfiguratorViewModel);
+
+private:
+	static const FName FilterConfiguratorTabId;
+
+	static TSharedPtr<FFilterService> Instance;
+	FFilterStorage FilterStorage;
+
+	TSharedPtr<class SFilterConfigurator> PendingWidget;
+};
+
+////////////////////////////////////////////////////////////////////////////////////////////////////
+
+
+class FFilterContext
+{
+public:
+	typedef TVariant<double, int64> ContextData;
+
+	template<typename T>
+	void AddFilterData(EFilterField Key, const T& InData)
+	{
+		ContextData VariantData;
+		VariantData.Set<T>(InData);
+		DataMap.Add(Key, VariantData);
+	}
+
+	template<typename T>
+	void GetFilterData(EFilterField Key, T& OutData) const
+	{
+		const ContextData* Data = DataMap.Find(Key);
+		check(Data);
+
+		check(Data->IsType<T>());
+		OutData = Data->Get<T>();
+	}
+
+private:
+	TMap<EFilterField, ContextData> DataMap;
+};
+
+////////////////////////////////////////////////////////////////////////////////////////////////////
+
+#undef LOCTEXT_NAMESPACE
diff --git a/Engine/Source/Developer/TraceInsights/Private/Insights/ViewModels/ThreadTimingTrack.cpp b/Engine/Source/Developer/TraceInsights/Private/Insights/ViewModels/ThreadTimingTrack.cpp
--- a/Engine/Source/Developer/TraceInsights/Private/Insights/ViewModels/ThreadTimingTrack.cpp
+++ b/Engine/Source/Developer/TraceInsights/Private/Insights/ViewModels/ThreadTimingTrack.cpp
@@ -15,6 +15,8 @@
 #include "Insights/InsightsManager.h"
 #include "Insights/ITimingViewSession.h"
 #include "Insights/TimingProfilerManager.h"
+#include "Insights/ViewModels/Filters.h"
+#include "Insights/ViewModels/FilterConfigurator.h"
 #include "Insights/ViewModels/TimerNode.h"
 #include "Insights/ViewModels/ThreadTrackEvent.h"
 #include "Insights/ViewModels/TimingEventSearch.h"
@@ -795,6 +797,73 @@ void FThreadTimingTrack::BuildFilteredDrawState(ITimingEventsTrackDrawStateBuild
 			}
 		}
 	}
+
+	if (HasCustomFilter()) // Custom filter (from the filtering widget)
+	{
+		TSharedPtr<const Trace::IAnalysisSession> Session = FInsightsManager::Get()->GetSession();
+		if (Session.IsValid() && Trace::ReadTimingProfilerProvider(*Session.Get()))
+		{
+			Trace::FAnalysisSessionReadScope SessionReadScope(*Session.Get());
+
+			const Trace::ITimingProfilerProvider& TimingProfilerProvider = *Trace::ReadTimingProfilerProvider(*Session.Get());
+
+			const Trace::ITimingProfilerTimerReader* TimerReader;
+			TimingProfilerProvider.ReadTimers([&TimerReader](const Trace::ITimingProfilerTimerReader& Out) { TimerReader = &Out; });
+
+			const FTimingTrackViewport& Viewport = Context.GetViewport();
+
+			TimingProfilerProvider.ReadTimeline(TimelineIndex,
+				[&Viewport, this, &Builder, TimerReader](const Trace::ITimingProfilerProvider::Timeline& Timeline)
+				{
+					TArray<TArray<FPendingEventInfo>> FilteredEvents;
+
+					Trace::ITimeline<Trace::FTimingProfilerEvent>::EnumerateAsyncParams Params;
+					Params.IntervalStart = Viewport.GetStartTime();
+					Params.IntervalEnd = Viewport.GetEndTime();
+
+					// Note: Enumerating events for filtering should not use downsampling.
+					Params.Resolution = 0.0;
+					Params.SetupCallback = [&FilteredEvents](uint32 NumTasks)
+					{
+						FilteredEvents.AddDefaulted(NumTasks);
+					};
+					Params.Callback = [this, &Builder, TimerReader, &FilteredEvents](double StartTime, double EndTime, uint32 Depth, const Trace::FTimingProfilerEvent& Event, uint32 TaskIndex)
+					{
+						const Trace::FTimingProfilerTimer* Timer = TimerReader->GetTimer(Event.TimerIndex);
+						if (ensure(Timer != nullptr))
+						{
+							FFilterContext Context;
+							Context.AddFilterData<double>(EFilterField::StartTime, StartTime);
+							Context.AddFilterData<double>(EFilterField::EndTime, EndTime);
+							Context.AddFilterData<double>(EFilterField::Duration, EndTime - StartTime);
+							Context.AddFilterData<int64>(EFilterField::EventType, Timer->Id);
+
+							if (FilterConfigurator->ApplyFilters(Context))
+							{
+								FPendingEventInfo TimelineEvent;
+								TimelineEvent.StartTime = StartTime;
+								TimelineEvent.EndTime = EndTime;
+								TimelineEvent.Depth = Depth;
+								TimelineEvent.TimerIndex = Event.TimerIndex;
+								FilteredEvents[TaskIndex].Add(TimelineEvent);
+							}
+						}
+						return Trace::EEventEnumerate::Continue;
+					};
+
+					Timeline.EnumerateEventsDownSampledAsync(Params);
+
+					for (TArray<FPendingEventInfo>& Array : FilteredEvents)
+					{
+						for (FPendingEventInfo& TimelineEvent : Array)
+						{
+							const Trace::FTimingProfilerTimer* Timer = TimerReader->GetTimer(TimelineEvent.TimerIndex);
+							AddTimingEventToBuilder(Builder, TimelineEvent.StartTime, TimelineEvent.EndTime, TimelineEvent.Depth, TimelineEvent.TimerIndex, Timer);
+						}
+					}
+				});
+		}
+	}
 }
 
 ////////////////////////////////////////////////////////////////////////////////////////////////////
@@ -1093,6 +1162,22 @@ void FThreadTimingTrack::BuildContextMenu(FMenuBuilder& MenuBuilder)
 {
 	if (GetGroupName() != nullptr)
 	{
+		MenuBuilder.BeginSection(TEXT("Options"));
+		{
+			FExecuteAction FilterTrackAction;
+			FilterTrackAction.BindSP(this, &FThreadTimingTrack::OnFilterTrackClicked);
+
+			MenuBuilder.AddMenuEntry(
+				(LOCTEXT("FilerTrack", "Filter Track [WIP]")),
+				FText(),
+				FSlateIcon(),
+				FUIAction(FilterTrackAction, FCanExecuteAction::CreateLambda([]() { return true; })),
+				NAME_None,
+				EUserInterfaceActionType::Button
+			);
+		}
+		MenuBuilder.EndSection();
+
 		MenuBuilder.BeginSection(TEXT("Misc"));
 		{
 			MenuBuilder.AddMenuEntry(
@@ -1218,4 +1303,40 @@ bool FThreadTimingTrack::TimerIndexToTimerId(uint32 InTimerIndex, uint32& OutTim
 
 ////////////////////////////////////////////////////////////////////////////////////////////////////
 
+void FThreadTimingTrack::OnFilterTrackClicked()
+{
+	if (!FilterConfigurator.IsValid())
+	{
+		FilterConfigurator = MakeShared<FFilterConfigurator>();
+
+		FilterConfigurator->AddFilter(EFilterField::StartTime);
+		FilterConfigurator->AddFilter(EFilterField::EndTime);
+		FilterConfigurator->AddFilter(EFilterField::Duration);
+		FilterConfigurator->AddFilter(EFilterField::EventType);
+
+		FilterConfigurator->SetOnChangesCommitedCallback([this]()
+			{
+				this->SetDirtyFlag();
+			});
+	}
+
+	FFilterService::Get()->CreateFilterConfiguratorWidget(FilterConfigurator);
+}
+
+////////////////////////////////////////////////////////////////////////////////////////////////////
+
+bool FThreadTimingTrack::HasCustomFilter() const
+{
+	if (!FilterConfigurator.IsValid())
+	{
+		return false;
+	}
+	if (FilterConfigurator->GetRootNode().IsValid() && FilterConfigurator->GetRootNode()->GetChildren().Num() > 0)
+	{
+		return true;
+	}
+
+	return false;
+}
+
 #undef LOCTEXT_NAMESPACE
diff --git a/Engine/Source/Developer/TraceInsights/Private/Insights/ViewModels/ThreadTimingTrack.h b/Engine/Source/Developer/TraceInsights/Private/Insights/ViewModels/ThreadTimingTrack.h
--- a/Engine/Source/Developer/TraceInsights/Private/Insights/ViewModels/ThreadTimingTrack.h
+++ b/Engine/Source/Developer/TraceInsights/Private/Insights/ViewModels/ThreadTimingTrack.h
@@ -141,6 +141,9 @@ public:
 	virtual void OnClipboardCopyEvent(const ITimingEvent& InSelectedEvent) const override;
 	virtual void BuildContextMenu(FMenuBuilder& MenuBuilder) override;
 
+protected:
+	virtual bool HasCustomFilter() const override;
+
 private:
 	void DrawSelectedEventInfo(const FThreadTrackEvent& SelectedEvent, const FTimingTrackViewport& Viewport, const FDrawContext& DrawContext, const FSlateBrush* WhiteBrush, const FSlateFontInfo& Font) const;
 
@@ -151,6 +154,8 @@ private:
 						  TSharedPtr<FThreadTrackEvent>& OutParentTimingEvent,
 						  TSharedPtr<FThreadTrackEvent>& OutRootTimingEvent) const;
 
+	void OnFilterTrackClicked();
+
 	static void CreateFThreadTrackEventFromInfo(const TimelineEventInfo& InEventInfo, const TSharedRef<const FBaseTimingTrack> InTrack, int32 InDepth, TSharedPtr<FThreadTrackEvent> &OutTimingEvent);
 	static bool TimerIndexToTimerId(uint32 InTimerIndex, uint32 & OutTimerId);
 
@@ -159,6 +164,8 @@ private:
 	uint32 TimelineIndex;
 	uint32 ThreadId;
 
+	TSharedPtr<class FFilterConfigurator> FilterConfigurator;
+
 	// Search cache
 	mutable TTimingEventSearchCache<Trace::FTimingProfilerEvent> SearchCache;
 };
diff --git a/Engine/Source/Developer/TraceInsights/Private/Insights/ViewModels/TimingEventsTrack.cpp b/Engine/Source/Developer/TraceInsights/Private/Insights/ViewModels/TimingEventsTrack.cpp
--- a/Engine/Source/Developer/TraceInsights/Private/Insights/ViewModels/TimingEventsTrack.cpp
+++ b/Engine/Source/Developer/TraceInsights/Private/Insights/ViewModels/TimingEventsTrack.cpp
@@ -79,18 +79,18 @@ void FTimingEventsTrack::PreUpdate(const ITimingTrackUpdateContext& Context)
 		}
 
 		const TSharedPtr<ITimingEventFilter> EventFilter = Context.GetEventFilter();
-		if (EventFilter.IsValid() && EventFilter->FilterTrack(*this))
+		if ((EventFilter.IsValid() && EventFilter->FilterTrack(*this)) || HasCustomFilter())
 		{
 			const FTimingTrackViewport& Viewport = Context.GetViewport();
 
 			const bool bFastLastBuild = FilteredDrawStateInfo.LastBuildDuration < 0.005; // LastBuildDuration < 5ms
 			const bool bFilterPointerChanged = !FilteredDrawStateInfo.LastEventFilter.HasSameObject(EventFilter.Get());
-			const bool bFilterContentChanged = FilteredDrawStateInfo.LastFilterChangeNumber != EventFilter->GetChangeNumber();
+			const bool bFilterContentChanged = EventFilter.IsValid() ? FilteredDrawStateInfo.LastFilterChangeNumber != EventFilter->GetChangeNumber() : false;
 
 			if (bFastLastBuild || bFilterPointerChanged || bFilterContentChanged)
 			{
 				FilteredDrawStateInfo.LastEventFilter = EventFilter;
-				FilteredDrawStateInfo.LastFilterChangeNumber = EventFilter->GetChangeNumber();
+				FilteredDrawStateInfo.LastFilterChangeNumber = EventFilter.IsValid() ? EventFilter->GetChangeNumber() : 0;
 				FilteredDrawStateInfo.ViewportStartTime = Context.GetViewport().GetStartTime();
 				FilteredDrawStateInfo.ViewportScaleX = Context.GetViewport().GetScaleX();
 				FilteredDrawStateInfo.Counter = 0;
@@ -225,7 +225,7 @@ void FTimingEventsTrack::DrawEvents(const ITimingTrackDrawContext& Context, cons
 {
 	const FTimingViewDrawHelper& Helper = *static_cast<const FTimingViewDrawHelper*>(&Context.GetHelper());
 
-	if (Context.GetEventFilter().IsValid())
+	if (Context.GetEventFilter().IsValid() || HasCustomFilter())
 	{
 		Helper.DrawFadedEvents(*DrawState, *this, OffsetY, 0.1f);
 
diff --git a/Engine/Source/Developer/TraceInsights/Private/Insights/Widgets/SFilterConfigurator.cpp b/Engine/Source/Developer/TraceInsights/Private/Insights/Widgets/SFilterConfigurator.cpp
new file mode 100644
--- /dev/null
+++ b/Engine/Source/Developer/TraceInsights/Private/Insights/Widgets/SFilterConfigurator.cpp
@@ -0,0 +1,223 @@
+// Copyright Epic Games, Inc. All Rights Reserved.
+
+#include "SFilterConfigurator.h"
+
+#include "EditorStyleSet.h"
+#include "SlateOptMacros.h"
+#include "Widgets/DeclarativeSyntaxSupport.h"
+#include "Widgets/Docking/SDockTab.h"
+#include "Widgets/Input/SButton.h"
+#include "Widgets/Layout/SScrollBox.h"
+#include "Widgets/Images/SImage.h"
+#include "Widgets/SBoxPanel.h"
+#include "Widgets/SOverlay.h"
+#include "Widgets/Text/STextBlock.h"
+#include "Widgets/Views/STableViewBase.h"
+
+// Insights
+#include "Insights/Widgets/SFilterConfiguratorRow.h"
+#include "Insights/ViewModels/FilterConfigurator.h"
+
+#define LOCTEXT_NAMESPACE "SFilterConfigurator"
+
+////////////////////////////////////////////////////////////////////////////////////////////////////
+// SFilterConfigurator
+////////////////////////////////////////////////////////////////////////////////////////////////////
+
+SFilterConfigurator::SFilterConfigurator()
+{
+
+}
+
+////////////////////////////////////////////////////////////////////////////////////////////////////
+
+SFilterConfigurator::~SFilterConfigurator()
+{
+}
+
+////////////////////////////////////////////////////////////////////////////////////////////////////
+
+void SFilterConfigurator::InitCommandList()
+{
+	
+}
+
+////////////////////////////////////////////////////////////////////////////////////////////////////
+
+BEGIN_SLATE_FUNCTION_BUILD_OPTIMIZATION
+
+void SFilterConfigurator::Construct(const FArguments& InArgs, TSharedPtr<FFilterConfigurator> InFilterConfiguratorViewModel)
+{
+	SAssignNew(ExternalScrollbar, SScrollBar)
+		.AlwaysShowScrollbar(true);
+
+	ChildSlot
+	[
+		SNew(SVerticalBox)
+
+		// Tree view
+		+SVerticalBox::Slot()
+		.FillHeight(1.0f)
+		.Padding(0.0f, 6.0f, 0.0f, 0.0f)
+		[
+			SNew(SHorizontalBox)
+
+			+ SHorizontalBox::Slot()
+			.FillWidth(1.0f)
+			.Padding(0.0f)
+			[
+				SNew(SScrollBox)
+				.Orientation(Orient_Horizontal)
+
+				+ SScrollBox::Slot()
+				[
+					SNew(SOverlay)
+
+					+ SOverlay::Slot()
+					.HAlign(HAlign_Fill)
+					.VAlign(VAlign_Fill)
+					[
+						SAssignNew(TreeView, STreeView<FFilterConfiguratorNodePtr>)
+						.ExternalScrollbar(ExternalScrollbar)
+						.SelectionMode(ESelectionMode::None)
+						.TreeItemsSource(&GroupNodes)
+						.OnGetChildren(this, &SFilterConfigurator::TreeView_OnGetChildren)
+						.OnGenerateRow(this, &SFilterConfigurator::TreeView_OnGenerateRow)
+						.ItemHeight(12.0f)
+						.HeaderRow
+						(
+							SAssignNew(TreeViewHeaderRow, SHeaderRow)
+							.Visibility(EVisibility::Visible)
+						)
+					]
+				]
+			]
+
+			+ SHorizontalBox::Slot()
+			.AutoWidth()
+			.Padding(0.0f)
+			[
+					SNew(SBox)
+					.WidthOverride(FOptionalSize(13.0f))
+				[
+					ExternalScrollbar.ToSharedRef()
+				]
+			]
+		]
+			+ SVerticalBox::Slot()
+			.FillHeight(1.0f)
+			.Padding(0.0f, 6.0f, 0.0f, 0.0f)
+			.HAlign(EHorizontalAlignment::HAlign_Right)
+			.VAlign(EVerticalAlignment::VAlign_Center)
+			[
+				SNew(SHorizontalBox)
+
+				+SHorizontalBox::Slot()
+				.HAlign(HAlign_Left)
+				.VAlign(VAlign_Center)
+				.AutoWidth()
+				.Padding(0.0f)
+				[
+					SNew(SButton)
+					.Text(LOCTEXT("OK", "OK"))
+					.ToolTipText(LOCTEXT("OKDESC", "Apply the changes to the filters."))
+					.OnClicked(this, &SFilterConfigurator::OK_OnClicked)
+				]
+
+				+SHorizontalBox::Slot()
+				.HAlign(HAlign_Left)
+				.VAlign(VAlign_Center)
+				.AutoWidth()
+				.Padding(0.0f)
+				[
+					SNew(SButton)
+					.Text(LOCTEXT("Cancel", "Cancel"))
+					.ToolTipText(LOCTEXT("CancelDesc", "Cancel the changes to the filter."))
+					.OnClicked(this, &SFilterConfigurator::Cancel_OnClicked)
+				]
+			]
+	];
+
+	SHeaderRow::FColumn::FArguments ColumnArgs;
+	ColumnArgs
+		.ColumnId(TEXT("Filter"))
+		.DefaultLabel(LOCTEXT("FilterColumnHeader", "Filter"))
+		.HAlignHeader(HAlign_Fill)
+		.VAlignHeader(VAlign_Fill)
+		.HeaderContentPadding(FMargin(2.0f))
+		.HAlignCell(HAlign_Fill)
+		.VAlignCell(VAlign_Fill)
+		.HeaderContent()
+		[
+			SNew(SBox)
+			.HAlign(HAlign_Fill)
+			.VAlign(VAlign_Center)
+			[
+				SNew(STextBlock)
+				.Text(LOCTEXT("FilterColumnHeader", "Filter"))
+			]
+		];
+
+	TreeViewHeaderRow->InsertColumn(ColumnArgs, 0);
+
+	OriginalFilterConfiguratorViewModel = InFilterConfiguratorViewModel;
+	FilterConfiguratorViewModel = MakeShared<FFilterConfigurator>(*OriginalFilterConfiguratorViewModel);
+
+	GroupNodes.Add(FilterConfiguratorViewModel->GetRootNode());
+}
+
+END_SLATE_FUNCTION_BUILD_OPTIMIZATION
+
+////////////////////////////////////////////////////////////////////////////////////////////////////
+
+void SFilterConfigurator::TreeView_OnGetChildren(FFilterConfiguratorNodePtr InParent, TArray<FFilterConfiguratorNodePtr>& OutChildren)
+{
+	const TArray<Insights::FBaseTreeNodePtr>& Children = InParent->GetChildren();
+	OutChildren.Reset(Children.Num());
+	for (const Insights::FBaseTreeNodePtr& Child : Children)
+	{
+		OutChildren.Add(StaticCastSharedPtr<FFilterConfiguratorNode, Insights::FBaseTreeNode>(Child));
+	}
+}
+
+////////////////////////////////////////////////////////////////////////////////////////////////////
+
+TSharedRef<ITableRow> SFilterConfigurator::TreeView_OnGenerateRow(FFilterConfiguratorNodePtr TreeNode, const TSharedRef<STableViewBase>& OwnerTable)
+{
+	TSharedRef<ITableRow> TableRow =
+		SNew(SFilterConfiguratorRow, OwnerTable)
+		.FilterConfiguratorNodePtr(TreeNode);
+
+	return TableRow;
+}
+
+////////////////////////////////////////////////////////////////////////////////////////////////////
+
+FReply SFilterConfigurator::OK_OnClicked()
+{
+	*OriginalFilterConfiguratorViewModel = *FilterConfiguratorViewModel;
+	if (ParentTab.IsValid())
+	{
+		ParentTab.Pin()->RequestCloseTab();
+		ParentTab.Reset();
+	}
+
+	return FReply::Handled();
+}
+
+////////////////////////////////////////////////////////////////////////////////////////////////////
+
+FReply SFilterConfigurator::Cancel_OnClicked()
+{
+	if (ParentTab.IsValid())
+	{
+		ParentTab.Pin()->RequestCloseTab();
+		ParentTab.Reset();
+	}
+
+	return FReply::Handled();
+}
+
+////////////////////////////////////////////////////////////////////////////////////////////////////
+
+#undef LOCTEXT_NAMESPACE
diff --git a/Engine/Source/Developer/TraceInsights/Private/Insights/Widgets/SFilterConfigurator.h b/Engine/Source/Developer/TraceInsights/Private/Insights/Widgets/SFilterConfigurator.h
new file mode 100644
--- /dev/null
+++ b/Engine/Source/Developer/TraceInsights/Private/Insights/Widgets/SFilterConfigurator.h
@@ -0,0 +1,83 @@
+// Copyright Epic Games, Inc. All Rights Reserved.
+
+#pragma once
+
+#include "CoreMinimal.h"
+#include "Widgets/SCompoundWidget.h"
+#include "Widgets/SWidget.h"
+#include "Widgets/Views/STreeView.h"
+
+// Insights
+#include "Insights/ViewModels/Filters.h"
+#include "Insights/ViewModels/FilterConfiguratorNode.h"
+
+////////////////////////////////////////////////////////////////////////////////////////////////////
+class SDockTab;
+
+////////////////////////////////////////////////////////////////////////////////////////////////////
+
+/**
+ * A custom widget used to configure custom filters.
+ */
+class SFilterConfigurator: public SCompoundWidget
+{
+public:
+	/** Default constructor. */
+	SFilterConfigurator();
+
+	/** Virtual destructor. */
+	virtual ~SFilterConfigurator();
+
+	SLATE_BEGIN_ARGS(SFilterConfigurator) {}
+	SLATE_END_ARGS()
+
+	/**
+	 * Construct this widget
+	 * @param InArgs - The declaration data for this widget
+	 */
+	void Construct(const FArguments& InArgs, TSharedPtr<class FFilterConfigurator> InFilterConfiguratorViewModel);
+
+	void Reset();
+
+	/**
+	 * Called by STreeView to retrieves the children for the specified parent item.
+	 * @param InParent    - The parent node to retrieve the children from.
+	 * @param OutChildren - List of children for the parent node.
+	 */
+	void TreeView_OnGetChildren(FFilterConfiguratorNodePtr InParent, TArray<FFilterConfiguratorNodePtr>& OutChildren);
+
+	/** Called by STreeView to generate a table row for the specified item. */
+	TSharedRef<class ITableRow> TreeView_OnGenerateRow(FFilterConfiguratorNodePtr TreeNode, const TSharedRef<class STableViewBase>& OwnerTable);
+
+	void SetParentTab(const TSharedPtr<SDockTab> InTab) { ParentTab = InTab; }
+	const TWeakPtr<SDockTab> GetParentTab() { return ParentTab; };
+
+private:
+	void InitCommandList();
+
+	FReply OK_OnClicked();
+
+	FReply Cancel_OnClicked();
+
+private:
+
+	TSharedPtr<class FFilterConfigurator> OriginalFilterConfiguratorViewModel;
+
+	TSharedPtr<class FFilterConfigurator> FilterConfiguratorViewModel;
+
+	/** The tree widget which holds the filters. */
+	TSharedPtr<STreeView<FFilterConfiguratorNodePtr>> TreeView;
+
+	/** Holds the tree view header row widget which display all columns in the tree view. */
+	TSharedPtr<class SHeaderRow> TreeViewHeaderRow;
+
+	/** External scrollbar used to synchronize tree view position. */
+	TSharedPtr<class SScrollBar> ExternalScrollbar;
+
+	/** A filtered array of group and filter nodes to be displayed in the tree widget. */
+	TArray<FFilterConfiguratorNodePtr> GroupNodes;
+
+	TWeakPtr<SDockTab> ParentTab;
+};
+
+////////////////////////////////////////////////////////////////////////////////////////////////////
diff --git a/Engine/Source/Developer/TraceInsights/Private/Insights/Widgets/SFilterConfiguratorRow.cpp b/Engine/Source/Developer/TraceInsights/Private/Insights/Widgets/SFilterConfiguratorRow.cpp
new file mode 100644
--- /dev/null
+++ b/Engine/Source/Developer/TraceInsights/Private/Insights/Widgets/SFilterConfiguratorRow.cpp
@@ -0,0 +1,338 @@
+// Copyright Epic Games, Inc. All Rights Reserved.
+
+#include "SFilterConfiguratorRow.h"
+
+#include "EditorStyleSet.h"
+#include "Framework/Application/SlateApplication.h"
+#include "SlateOptMacros.h"
+#include "Widgets/DeclarativeSyntaxSupport.h"
+#include "Widgets/Images/SImage.h"
+#include "Widgets/Input/SButton.h"
+#include "Widgets/Input/SEditableTextBox.h"
+#include "Widgets/Views/STableRow.h"
+#include "Widgets/Views/STableViewBase.h"
+#include "Widgets/Views/STreeView.h"
+
+// Insights
+#include "Insights/InsightsStyle.h"
+#include "Insights/ViewModels/Filters.h"
+
+#define LOCTEXT_NAMESPACE "SFilterConfiguratorRow"
+
+////////////////////////////////////////////////////////////////////////////////////////////////////
+
+BEGIN_SLATE_FUNCTION_BUILD_OPTIMIZATION
+
+void SFilterConfiguratorRow::Construct(const FArguments& InArgs, const TSharedRef<STableViewBase>& InOwnerTableView)
+{
+	FilterConfiguratorNodePtr = InArgs._FilterConfiguratorNodePtr;
+	SMultiColumnTableRow<FFilterConfiguratorNodePtr>::Construct(SMultiColumnTableRow<FFilterConfiguratorNodePtr>::FArguments(), InOwnerTableView);
+}
+
+////////////////////////////////////////////////////////////////////////////////////////////////////
+
+TSharedRef<SWidget> SFilterConfiguratorRow::GenerateWidgetForColumn(const FName& InColumnName)
+{
+	if (!FilterConfiguratorNodePtr->IsGroup())
+	{
+		return SNew(SHorizontalBox)
+
+			+ SHorizontalBox::Slot()
+			.AutoWidth()
+			.HAlign(HAlign_Left)
+			.VAlign(VAlign_Center)
+			[
+				SNew(SExpanderArrow, SharedThis(this))
+				.ShouldDrawWires(true)
+			]
+
+			// Filter combo box
+			+ SHorizontalBox::Slot()
+			.HAlign(HAlign_Left)
+			.VAlign(VAlign_Center)
+			.AutoWidth()
+			.Padding(0.0f)
+			[
+				SNew(SBorder)
+				.Padding(0)
+				.BorderImage(FEditorStyle::GetBrush("NoBorder"))
+				.IsEnabled(FSlateApplication::Get().GetNormalExecutionAttribute())
+				[
+					SAssignNew(FilterTypeComboBox, SComboBox<TSharedPtr<FFilter>>)
+					.OptionsSource(GetAvailableFilters())
+					.OnSelectionChanged(this, &SFilterConfiguratorRow::AvailableFilters_OnSelectionChanged)
+					.OnGenerateWidget(this, &SFilterConfiguratorRow::AvailableFilters_OnGenerateWidget)
+					[
+						SNew(STextBlock)
+						.Text(this, &SFilterConfiguratorRow::AvailableFilters_GetSelectionText)
+					]
+				]
+			]
+
+			// Operator combo box
+			+ SHorizontalBox::Slot()
+			.HAlign(HAlign_Left)
+			.VAlign(VAlign_Center)
+			.FillWidth(1.0)
+			.Padding(0.0f)
+			.AutoWidth()
+			[
+				SNew(SBorder)
+				.Padding(0)
+				.BorderImage(FEditorStyle::GetBrush("NoBorder"))
+				.IsEnabled(FSlateApplication::Get().GetNormalExecutionAttribute())
+				[
+					SAssignNew(FilterOperatorComboBox, SComboBox<TSharedPtr<IFilterOperator>>)
+					.OptionsSource(GetAvailableFilterOperators())
+					.OnSelectionChanged(this, &SFilterConfiguratorRow::AvailableFilterOperators_OnSelectionChanged)
+					.OnGenerateWidget(this, &SFilterConfiguratorRow::AvailableFilterOperators_OnGenerateWidget)
+						[
+						SNew(STextBlock)
+						.Text(this, &SFilterConfiguratorRow::AvailableFilterOperators_GetSelectionText)
+						]
+				]
+			]
+
+			+ SHorizontalBox::Slot()
+			.AutoWidth()
+			.VAlign(VAlign_Center)
+			.HAlign(HAlign_Left)
+			.FillWidth(1.0)
+			.Padding(0.0f, 0.0f, 0.0f, 0.0f)
+			[
+				SNew(SEditableTextBox)
+				.MinDesiredWidth(50.0f)
+				.OnTextCommitted(this, &SFilterConfiguratorRow::OnTextBoxValueCommitted)
+				.Text(this, &SFilterConfiguratorRow::GetTextBoxValue)
+			]
+			
+			+ SHorizontalBox::Slot()
+			.VAlign(VAlign_Center)
+			.HAlign(HAlign_Left)
+			.Padding(FMargin(2.0f, 0.0f, 2.0f, 0.0f))
+			.AutoWidth()
+			[
+				SNew(SButton)
+				.ToolTipText(LOCTEXT("DeleteFilterTooptip", "Delete Filter"))
+				.ContentPadding(FMargin(2.0f, 0.0f, 2.0f, 0.0f))
+				.OnClicked(this, &SFilterConfiguratorRow::DeleteFilter_OnClicked)
+				.Content()
+				[
+					SNew(SImage)
+					.ColorAndOpacity(FSlateColor(FLinearColor(0.5f, 0.0f, 0.0f, 1.0f)))
+					.Image(FInsightsStyle::Get().GetBrush("Mem.Remove.Small"))
+				]
+			];
+	}
+	else
+	{
+		OwnerTablePtr.Pin()->Private_SetItemExpansion(FilterConfiguratorNodePtr, true);
+
+		return SNew(SHorizontalBox)
+
+			+ SHorizontalBox::Slot()
+			.AutoWidth()
+			.HAlign(HAlign_Left)
+			.VAlign(VAlign_Center)
+			[
+				SNew(SExpanderArrow, SharedThis(this))
+				.ShouldDrawWires(true)
+			]
+
+			+ SHorizontalBox::Slot()
+			.HAlign(HAlign_Left)
+			.VAlign(VAlign_Center)
+			.AutoWidth()
+			.Padding(0.0f)
+			[
+				SNew(SBorder)
+				.Padding(0)
+				.BorderImage(FEditorStyle::GetBrush("NoBorder"))
+				.IsEnabled(FSlateApplication::Get().GetNormalExecutionAttribute())
+				[
+					SAssignNew(FilterGroupOperatorComboBox, SComboBox<TSharedPtr<FFilterGroupOperator>>)
+					.OptionsSource(GetFilterGroupOperators())
+					.OnSelectionChanged(this, &SFilterConfiguratorRow::FilterGroupOperators_OnSelectionChanged)
+					.OnGenerateWidget(this, &SFilterConfiguratorRow::FilterGroupOperators_OnGenerateWidget)
+					[
+						SNew(STextBlock)
+						.Text(this, &SFilterConfiguratorRow::FilterGroupOperators_GetSelectionText)
+					]
+				]
+			]
+		
+		+ SHorizontalBox::Slot()
+		.HAlign(HAlign_Left)
+		.VAlign(VAlign_Center)
+		.AutoWidth()
+		.Padding(0.0f)
+		[
+			SNew(SButton)
+			.Text(LOCTEXT("AddFilter", "Add Filter"))
+			.ToolTipText(LOCTEXT("AddFilterDesc", "Add a filter node as a child to this group node."))
+			.OnClicked(this, &SFilterConfiguratorRow::AddFilter_OnClicked)
+		];
+	}
+}
+
+///////////////////////////////////////////////////////////////////////////////////////////////////
+
+TSharedRef<SWidget> SFilterConfiguratorRow::AvailableFilters_OnGenerateWidget(TSharedPtr<FFilter> InFilter)
+{
+	TSharedRef<SHorizontalBox> Widget = SNew(SHorizontalBox);
+	Widget->AddSlot()
+		.AutoWidth()
+		[
+			SNew(STextBlock)
+			.Text(InFilter->Name)
+			.Margin(2.0f)
+		];
+
+	return Widget;
+}
+
+///////////////////////////////////////////////////////////////////////////////////////////////////
+
+TSharedRef<SWidget> SFilterConfiguratorRow::AvailableFilterOperators_OnGenerateWidget(TSharedPtr<IFilterOperator> InFilterOperator)
+{
+	TSharedRef<SHorizontalBox> Widget = SNew(SHorizontalBox);
+	Widget->AddSlot()
+		.AutoWidth()
+		[
+			SNew(STextBlock)
+			.Text(FText::FromString(InFilterOperator->GetName()))
+			.Margin(2.0f)
+		];
+
+	return Widget;
+}
+
+///////////////////////////////////////////////////////////////////////////////////////////////////
+
+TSharedRef<SWidget> SFilterConfiguratorRow::FilterGroupOperators_OnGenerateWidget(TSharedPtr<FFilterGroupOperator> InFilterGroupOperator)
+{
+	TSharedRef<SHorizontalBox> Widget = SNew(SHorizontalBox);
+	Widget->AddSlot()
+		.AutoWidth()
+		[
+			SNew(STextBlock)
+			.Text(InFilterGroupOperator->Name)
+			.Margin(2.0f)
+		];
+
+	return Widget;
+}
+
+END_SLATE_FUNCTION_BUILD_OPTIMIZATION
+
+////////////////////////////////////////////////////////////////////////////////////////////////////
+
+const TArray<TSharedPtr<struct FFilter>>* SFilterConfiguratorRow::GetAvailableFilters()
+{
+	return &*FilterConfiguratorNodePtr->GetAvailableFilters();
+}
+
+///////////////////////////////////////////////////////////////////////////////////////////////////
+
+void SFilterConfiguratorRow::AvailableFilters_OnSelectionChanged(TSharedPtr<FFilter> InFilter, ESelectInfo::Type SelectInfo)
+{
+	FilterConfiguratorNodePtr->SetSelectedFilter(InFilter);
+}
+
+///////////////////////////////////////////////////////////////////////////////////////////////////
+
+FText SFilterConfiguratorRow::AvailableFilters_GetSelectionText() const
+{
+	return FilterConfiguratorNodePtr->GetSelectedFilter()->Name;
+}
+
+///////////////////////////////////////////////////////////////////////////////////////////////////
+
+const TArray<TSharedPtr<IFilterOperator>>* SFilterConfiguratorRow::GetAvailableFilterOperators()
+{
+	return &*FilterConfiguratorNodePtr->GetSelectedFilter()->SupportedOperators;
+}
+
+///////////////////////////////////////////////////////////////////////////////////////////////////
+
+void SFilterConfiguratorRow::AvailableFilterOperators_OnSelectionChanged(TSharedPtr<IFilterOperator> InFilter, ESelectInfo::Type SelectInfo)
+{
+	FilterConfiguratorNodePtr->SetSelectedFilterOperator(InFilter);
+}
+
+///////////////////////////////////////////////////////////////////////////////////////////////////
+
+FText SFilterConfiguratorRow::AvailableFilterOperators_GetSelectionText() const
+{
+	return FText::FromString(FilterConfiguratorNodePtr->GetSelectedFilterOperator()->GetName());
+}
+
+///////////////////////////////////////////////////////////////////////////////////////////////////
+
+TArray<TSharedPtr<FFilterGroupOperator>>* SFilterConfiguratorRow::GetFilterGroupOperators()
+{
+	return &FilterConfiguratorNodePtr->GetFilterGroupOperators();
+}
+
+///////////////////////////////////////////////////////////////////////////////////////////////////
+
+void SFilterConfiguratorRow::FilterGroupOperators_OnSelectionChanged(TSharedPtr<FFilterGroupOperator> InFilterGroupOperator, ESelectInfo::Type SelectInfo)
+{
+	FilterConfiguratorNodePtr->SetSelectedFilterGroupOperator(InFilterGroupOperator);
+}
+
+///////////////////////////////////////////////////////////////////////////////////////////////////
+
+FText SFilterConfiguratorRow::FilterGroupOperators_GetSelectionText() const
+{
+	return FilterConfiguratorNodePtr->GetSelectedFilterGroupOperator()->Name;
+}
+
+///////////////////////////////////////////////////////////////////////////////////////////////////
+
+FReply SFilterConfiguratorRow::AddFilter_OnClicked()
+{
+	FFilterConfiguratorNodePtr ChildNode = MakeShared<FFilterConfiguratorNode>(TEXT(""), false);
+	ChildNode->SetAvailableFilters(FilterConfiguratorNodePtr->GetAvailableFilters());
+
+	FilterConfiguratorNodePtr->AddChildAndSetGroupPtr(ChildNode);
+	FilterConfiguratorNodePtr->SetExpansion(true);
+	OwnerTablePtr.Pin()->Private_SetItemExpansion(FilterConfiguratorNodePtr, true);
+
+	StaticCastSharedPtr<STreeView<FFilterConfiguratorNodePtr>>(OwnerTablePtr.Pin())->RequestTreeRefresh();
+	return FReply::Handled();
+}
+
+///////////////////////////////////////////////////////////////////////////////////////////////////
+
+FReply SFilterConfiguratorRow::DeleteFilter_OnClicked()
+{
+	Insights::FBaseTreeNodeWeak GroupWeakPtr = FilterConfiguratorNodePtr->GetGroupPtr();
+	if (GroupWeakPtr.IsValid())
+	{
+		
+		FFilterConfiguratorNodePtr GroupPtr = StaticCastSharedPtr<FFilterConfiguratorNode>(GroupWeakPtr.Pin());
+		GroupPtr->DeleteChildNode(FilterConfiguratorNodePtr);
+	}
+
+	StaticCastSharedPtr<STreeView<FFilterConfiguratorNodePtr>>(OwnerTablePtr.Pin())->RequestTreeRefresh();
+	return FReply::Handled();
+}
+
+///////////////////////////////////////////////////////////////////////////////////////////////////
+
+FText SFilterConfiguratorRow::GetTextBoxValue() const
+{
+	return FText::FromString(FilterConfiguratorNodePtr->GetTextBoxValue());
+}
+
+///////////////////////////////////////////////////////////////////////////////////////////////////
+
+void SFilterConfiguratorRow::OnTextBoxValueCommitted(const FText& InNewText, ETextCommit::Type InTextCommit)
+{
+	FilterConfiguratorNodePtr->SetTextBoxValue(InNewText.ToString());
+}
+
+///////////////////////////////////////////////////////////////////////////////////////////////////
+
+#undef LOCTEXT_NAMESPACE
diff --git a/Engine/Source/Developer/TraceInsights/Private/Insights/Widgets/SFilterConfiguratorRow.h b/Engine/Source/Developer/TraceInsights/Private/Insights/Widgets/SFilterConfiguratorRow.h
new file mode 100644
--- /dev/null
+++ b/Engine/Source/Developer/TraceInsights/Private/Insights/Widgets/SFilterConfiguratorRow.h
@@ -0,0 +1,69 @@
+// Copyright Epic Games, Inc. All Rights Reserved.
+
+#pragma once
+
+#include "CoreMinimal.h"
+
+#include "Widgets/Input/SComboBox.h"
+#include "Widgets/Views/STableRow.h"
+
+// Insights
+#include "Insights/ViewModels/FilterConfiguratorNode.h"
+
+/** Widget that represents a table row in the Filter's tree control. Generates widgets for each column on demand. */
+class SFilterConfiguratorRow : public SMultiColumnTableRow<FFilterConfiguratorNodePtr>
+{
+public:
+	SFilterConfiguratorRow() {}
+	~SFilterConfiguratorRow() {}
+
+	SLATE_BEGIN_ARGS(SFilterConfiguratorRow) {}
+		SLATE_ARGUMENT(FFilterConfiguratorNodePtr, FilterConfiguratorNodePtr)
+	SLATE_END_ARGS()
+
+public:
+	void Construct(const FArguments& InArgs, const TSharedRef<STableViewBase>& InOwnerTableView);
+
+	virtual TSharedRef<SWidget> GenerateWidgetForColumn(const FName& InColumnName) override;
+
+private:
+	const TArray<TSharedPtr<struct FFilter>>* GetAvailableFilters();
+
+	TSharedRef<SWidget> AvailableFilters_OnGenerateWidget(TSharedPtr<struct FFilter> InFilter);
+
+	void AvailableFilters_OnSelectionChanged(TSharedPtr<struct FFilter> InFilter, ESelectInfo::Type SelectInfo);
+
+	FText AvailableFilters_GetSelectionText() const;
+
+	const TArray<TSharedPtr<class IFilterOperator>>* GetAvailableFilterOperators();
+
+	TSharedRef<SWidget> AvailableFilterOperators_OnGenerateWidget(TSharedPtr<class IFilterOperator> InFilter);
+
+	void AvailableFilterOperators_OnSelectionChanged(TSharedPtr<class IFilterOperator> InFilter, ESelectInfo::Type SelectInfo);
+
+	FText AvailableFilterOperators_GetSelectionText() const;
+
+	TArray<TSharedPtr<struct FFilterGroupOperator>>* GetFilterGroupOperators();
+
+	TSharedRef<SWidget> FilterGroupOperators_OnGenerateWidget(TSharedPtr<struct FFilterGroupOperator> InFilter);
+
+	void FilterGroupOperators_OnSelectionChanged(TSharedPtr<struct FFilterGroupOperator> InFilter, ESelectInfo::Type SelectInfo);
+
+	FText FilterGroupOperators_GetSelectionText() const;
+
+	FReply AddFilter_OnClicked();
+
+	FReply DeleteFilter_OnClicked();
+
+	FText GetTextBoxValue() const;
+	void OnTextBoxValueCommitted(const FText& InNewText, ETextCommit::Type InTextCommit);
+
+private:
+	FFilterConfiguratorNodePtr FilterConfiguratorNodePtr;
+
+	TSharedPtr<SComboBox<TSharedPtr<struct FFilter>>> FilterTypeComboBox;
+
+	TSharedPtr<SComboBox<TSharedPtr<class IFilterOperator>>> FilterOperatorComboBox;
+
+	TSharedPtr<SComboBox<TSharedPtr<struct FFilterGroupOperator>>> FilterGroupOperatorComboBox;
+};
diff --git a/Engine/Source/Developer/TraceInsights/Public/Insights/ViewModels/TimingEventsTrack.h b/Engine/Source/Developer/TraceInsights/Public/Insights/ViewModels/TimingEventsTrack.h
--- a/Engine/Source/Developer/TraceInsights/Public/Insights/ViewModels/TimingEventsTrack.h
+++ b/Engine/Source/Developer/TraceInsights/Public/Insights/ViewModels/TimingEventsTrack.h
@@ -66,6 +66,8 @@ protected:
 	int32 GetHeaderBackgroundLayerId(const ITimingTrackDrawContext& Context) const;
 	int32 GetHeaderTextLayerId(const ITimingTrackDrawContext& Context) const;
 
+	virtual bool HasCustomFilter() const { return false; }
+
 private:
 	int32 NumLanes; // number of lanes (sub-tracks)
 	TSharedRef<struct FTimingEventsTrackDrawState> DrawState;
