Functionality to add groups in the Filter Configurator Widget.

diff --git a/Engine/Source/Developer/TraceInsights/Private/Insights/ViewModels/Filters.h b/Engine/Source/Developer/TraceInsights/Private/Insights/ViewModels/Filters.h
--- a/Engine/Source/Developer/TraceInsights/Private/Insights/ViewModels/Filters.h
+++ b/Engine/Source/Developer/TraceInsights/Private/Insights/ViewModels/Filters.h
@@ -145,8 +145,8 @@ public:
 		AllAvailableFilters.Add(EFilterField::Duration, MakeShared<FFilter>(EFilterField::Duration, LOCTEXT("Duration", "Duration"), LOCTEXT("Duration", "Duration"), EFilterDataType::Double, DoubleOperators));
 		AllAvailableFilters.Add(EFilterField::EventType, MakeShared<FFilter>(EFilterField::EventType, LOCTEXT("Type", "Type"), LOCTEXT("Type", "Type"), EFilterDataType::Int64, IntegerOperators));
 
-		FilterGroupOperators.Add(MakeShared<FFilterGroupOperator>(EFilterGroupOperator::And, LOCTEXT("AllOff", "All Off"), LOCTEXT("AllOff", "All Off")));
-		FilterGroupOperators.Add(MakeShared<FFilterGroupOperator>(EFilterGroupOperator::Or, LOCTEXT("AnyOff", "Any Off"), LOCTEXT("AnyOff", "Any Off")));
+		FilterGroupOperators.Add(MakeShared<FFilterGroupOperator>(EFilterGroupOperator::And, LOCTEXT("AllOf", "All Of"), LOCTEXT("AllOfDesc", "All of the children must be true for the group to return true. Equivalent to an AND operation.")));
+		FilterGroupOperators.Add(MakeShared<FFilterGroupOperator>(EFilterGroupOperator::Or, LOCTEXT("AnyOf", "Any Of"), LOCTEXT("AnyOfDesc", "Any of the children must be true for the group to return true. Equivalent to an OR operation.")));
 	}
 
 	TSharedPtr<FFilter> GetFilter(EFilterField FilterKey) { return *AllAvailableFilters.Find(FilterKey); }
diff --git a/Engine/Source/Developer/TraceInsights/Private/Insights/Widgets/SFilterConfigurator.cpp b/Engine/Source/Developer/TraceInsights/Private/Insights/Widgets/SFilterConfigurator.cpp
--- a/Engine/Source/Developer/TraceInsights/Private/Insights/Widgets/SFilterConfigurator.cpp
+++ b/Engine/Source/Developer/TraceInsights/Private/Insights/Widgets/SFilterConfigurator.cpp
@@ -105,10 +105,10 @@ void SFilterConfigurator::Construct(const FArguments& InArgs, TSharedPtr<FFilter
 			]
 		]
 			+ SVerticalBox::Slot()
-			.FillHeight(1.0f)
-			.Padding(0.0f, 6.0f, 0.0f, 0.0f)
+			.AutoHeight()
+			.Padding(0.0f, 5.0f, 0.0f, 5.0f)
 			.HAlign(EHorizontalAlignment::HAlign_Right)
-			.VAlign(EVerticalAlignment::VAlign_Center)
+			.VAlign(EVerticalAlignment::VAlign_Bottom)
 			[
 				SNew(SHorizontalBox)
 
@@ -116,7 +116,7 @@ void SFilterConfigurator::Construct(const FArguments& InArgs, TSharedPtr<FFilter
 				.HAlign(HAlign_Left)
 				.VAlign(VAlign_Center)
 				.AutoWidth()
-				.Padding(0.0f)
+				.Padding(1.0f)
 				[
 					SNew(SButton)
 					.Text(LOCTEXT("OK", "OK"))
@@ -128,7 +128,7 @@ void SFilterConfigurator::Construct(const FArguments& InArgs, TSharedPtr<FFilter
 				.HAlign(HAlign_Left)
 				.VAlign(VAlign_Center)
 				.AutoWidth()
-				.Padding(0.0f)
+				.Padding(1.0f)
 				[
 					SNew(SButton)
 					.Text(LOCTEXT("Cancel", "Cancel"))
diff --git a/Engine/Source/Developer/TraceInsights/Private/Insights/Widgets/SFilterConfiguratorRow.cpp b/Engine/Source/Developer/TraceInsights/Private/Insights/Widgets/SFilterConfiguratorRow.cpp
--- a/Engine/Source/Developer/TraceInsights/Private/Insights/Widgets/SFilterConfiguratorRow.cpp
+++ b/Engine/Source/Developer/TraceInsights/Private/Insights/Widgets/SFilterConfiguratorRow.cpp
@@ -33,25 +33,24 @@ void SFilterConfiguratorRow::Construct(const FArguments& InArgs, const TSharedRe
 
 TSharedRef<SWidget> SFilterConfiguratorRow::GenerateWidgetForColumn(const FName& InColumnName)
 {
+	TSharedRef<SHorizontalBox> GeneratedWidget = SNew(SHorizontalBox);
 	if (!FilterConfiguratorNodePtr->IsGroup())
 	{
-		return SNew(SHorizontalBox)
-
-			+ SHorizontalBox::Slot()
+		GeneratedWidget->AddSlot()
 			.AutoWidth()
 			.HAlign(HAlign_Left)
 			.VAlign(VAlign_Center)
 			[
 				SNew(SExpanderArrow, SharedThis(this))
 				.ShouldDrawWires(true)
-			]
+			];
 
-			// Filter combo box
-			+ SHorizontalBox::Slot()
+		// Filter combo box
+		GeneratedWidget->AddSlot()
 			.HAlign(HAlign_Left)
 			.VAlign(VAlign_Center)
 			.AutoWidth()
-			.Padding(0.0f)
+			.Padding(2.0f)
 			[
 				SNew(SBorder)
 				.Padding(0)
@@ -67,14 +66,14 @@ TSharedRef<SWidget> SFilterConfiguratorRow::GenerateWidgetForColumn(const FName&
 						.Text(this, &SFilterConfiguratorRow::AvailableFilters_GetSelectionText)
 					]
 				]
-			]
+			];
 
-			// Operator combo box
-			+ SHorizontalBox::Slot()
+		// Operator combo box
+		GeneratedWidget->AddSlot()
 			.HAlign(HAlign_Left)
 			.VAlign(VAlign_Center)
 			.FillWidth(1.0)
-			.Padding(0.0f)
+			.Padding(2.0f)
 			.AutoWidth()
 			[
 				SNew(SBorder)
@@ -91,30 +90,30 @@ TSharedRef<SWidget> SFilterConfiguratorRow::GenerateWidgetForColumn(const FName&
 						.Text(this, &SFilterConfiguratorRow::AvailableFilterOperators_GetSelectionText)
 						]
 				]
-			]
+			];
 
-			+ SHorizontalBox::Slot()
+		GeneratedWidget->AddSlot()
 			.AutoWidth()
 			.VAlign(VAlign_Center)
 			.HAlign(HAlign_Left)
 			.FillWidth(1.0)
-			.Padding(0.0f, 0.0f, 0.0f, 0.0f)
+			.Padding(2.0f)
 			[
 				SNew(SEditableTextBox)
 				.MinDesiredWidth(50.0f)
 				.OnTextCommitted(this, &SFilterConfiguratorRow::OnTextBoxValueCommitted)
 				.Text(this, &SFilterConfiguratorRow::GetTextBoxValue)
-			]
+			];
 			
-			+ SHorizontalBox::Slot()
+		GeneratedWidget->AddSlot()
 			.VAlign(VAlign_Center)
 			.HAlign(HAlign_Left)
-			.Padding(FMargin(2.0f, 0.0f, 2.0f, 0.0f))
+			.Padding(2.0f)
 			.AutoWidth()
 			[
 				SNew(SButton)
 				.ToolTipText(LOCTEXT("DeleteFilterTooptip", "Delete Filter"))
-				.ContentPadding(FMargin(2.0f, 0.0f, 2.0f, 0.0f))
+				.ContentPadding(FMargin(2.0f, 0.0f, 2.0f, -1.0f))
 				.OnClicked(this, &SFilterConfiguratorRow::DeleteFilter_OnClicked)
 				.Content()
 				[
@@ -128,22 +127,21 @@ TSharedRef<SWidget> SFilterConfiguratorRow::GenerateWidgetForColumn(const FName&
 	{
 		OwnerTablePtr.Pin()->Private_SetItemExpansion(FilterConfiguratorNodePtr, true);
 
-		return SNew(SHorizontalBox)
-
-			+ SHorizontalBox::Slot()
+		GeneratedWidget->AddSlot()
 			.AutoWidth()
 			.HAlign(HAlign_Left)
 			.VAlign(VAlign_Center)
+			.Padding(2.0f)
 			[
 				SNew(SExpanderArrow, SharedThis(this))
 				.ShouldDrawWires(true)
-			]
+			];
 
-			+ SHorizontalBox::Slot()
+		GeneratedWidget->AddSlot()
 			.HAlign(HAlign_Left)
 			.VAlign(VAlign_Center)
 			.AutoWidth()
-			.Padding(0.0f)
+			.Padding(2.0f)
 			[
 				SNew(SBorder)
 				.Padding(0)
@@ -159,20 +157,56 @@ TSharedRef<SWidget> SFilterConfiguratorRow::GenerateWidgetForColumn(const FName&
 						.Text(this, &SFilterConfiguratorRow::FilterGroupOperators_GetSelectionText)
 					]
 				]
-			]
+			];
 		
-		+ SHorizontalBox::Slot()
-		.HAlign(HAlign_Left)
-		.VAlign(VAlign_Center)
-		.AutoWidth()
-		.Padding(0.0f)
-		[
-			SNew(SButton)
-			.Text(LOCTEXT("AddFilter", "Add Filter"))
-			.ToolTipText(LOCTEXT("AddFilterDesc", "Add a filter node as a child to this group node."))
-			.OnClicked(this, &SFilterConfiguratorRow::AddFilter_OnClicked)
-		];
+		GeneratedWidget->AddSlot()
+			.HAlign(HAlign_Left)
+			.VAlign(VAlign_Center)
+			.AutoWidth()
+			.Padding(2.0f)
+			[
+				SNew(SButton)
+				.Text(LOCTEXT("AddFilter", "Add Filter"))
+				.ToolTipText(LOCTEXT("AddFilterDesc", "Add a filter node as a child to this group node."))
+				.OnClicked(this, &SFilterConfiguratorRow::AddFilter_OnClicked)
+			];
+
+		GeneratedWidget->AddSlot()
+			.HAlign(HAlign_Left)
+			.VAlign(VAlign_Center)
+			.AutoWidth()
+			.Padding(2.0f)
+			[
+				SNew(SButton)
+				.Text(LOCTEXT("AddGroup", "Add Group"))
+				.ToolTipText(LOCTEXT("AddGroupDesc", "Add a group node as a child to this group node."))
+				.OnClicked(this, &SFilterConfiguratorRow::AddGroup_OnClicked)
+			];
+
+		// Do not show Delete button for the Root node
+		if (FilterConfiguratorNodePtr->GetGroupPtr().IsValid())
+		{
+			GeneratedWidget->AddSlot()
+				.VAlign(VAlign_Center)
+				.HAlign(HAlign_Left)
+				.Padding(2.0f)
+				.AutoWidth()
+				[
+					SNew(SButton)
+					.ToolTipText(LOCTEXT("DeleteGroupTooptip", "Delete Group"))
+					.ContentPadding(FMargin(2.0f, 0.0f, 2.0f, -1.0f))
+					.OnClicked(this, &SFilterConfiguratorRow::DeleteGroup_OnClicked)
+					.Content()
+					[
+						SNew(SImage)
+						.ColorAndOpacity(FSlateColor(FLinearColor(0.5f, 0.0f, 0.0f, 1.0f)))
+						.Image(FInsightsStyle::Get().GetBrush("Mem.Remove.Small"))
+					]
+				];
+		}
 	}
+
+	return GeneratedWidget;
 }
 
 ///////////////////////////////////////////////////////////////////////////////////////////////////
@@ -212,6 +246,7 @@ TSharedRef<SWidget> SFilterConfiguratorRow::AvailableFilterOperators_OnGenerateW
 TSharedRef<SWidget> SFilterConfiguratorRow::FilterGroupOperators_OnGenerateWidget(TSharedPtr<FFilterGroupOperator> InFilterGroupOperator)
 {
 	TSharedRef<SHorizontalBox> Widget = SNew(SHorizontalBox);
+	Widget->SetToolTipText(InFilterGroupOperator->Desc);
 	Widget->AddSlot()
 		.AutoWidth()
 		[
@@ -321,6 +356,38 @@ FReply SFilterConfiguratorRow::DeleteFilter_OnClicked()
 
 ///////////////////////////////////////////////////////////////////////////////////////////////////
 
+FReply SFilterConfiguratorRow::AddGroup_OnClicked()
+{
+	FFilterConfiguratorNodePtr ChildNode = MakeShared<FFilterConfiguratorNode>(TEXT(""), true);
+	ChildNode->SetAvailableFilters(FilterConfiguratorNodePtr->GetAvailableFilters());
+	ChildNode->SetExpansion(true);
+
+	FilterConfiguratorNodePtr->AddChildAndSetGroupPtr(ChildNode);
+	FilterConfiguratorNodePtr->SetExpansion(true);
+	OwnerTablePtr.Pin()->Private_SetItemExpansion(FilterConfiguratorNodePtr, true);
+
+	StaticCastSharedPtr<STreeView<FFilterConfiguratorNodePtr>>(OwnerTablePtr.Pin())->RequestTreeRefresh();
+	return FReply::Handled();
+}
+
+///////////////////////////////////////////////////////////////////////////////////////////////////
+
+FReply SFilterConfiguratorRow::DeleteGroup_OnClicked()
+{
+	Insights::FBaseTreeNodeWeak GroupWeakPtr = FilterConfiguratorNodePtr->GetGroupPtr();
+	if (GroupWeakPtr.IsValid())
+	{
+
+		FFilterConfiguratorNodePtr GroupPtr = StaticCastSharedPtr<FFilterConfiguratorNode>(GroupWeakPtr.Pin());
+		GroupPtr->DeleteChildNode(FilterConfiguratorNodePtr);
+	}
+
+	StaticCastSharedPtr<STreeView<FFilterConfiguratorNodePtr>>(OwnerTablePtr.Pin())->RequestTreeRefresh();
+	return FReply::Handled();
+}
+
+///////////////////////////////////////////////////////////////////////////////////////////////////
+
 FText SFilterConfiguratorRow::GetTextBoxValue() const
 {
 	return FText::FromString(FilterConfiguratorNodePtr->GetTextBoxValue());
diff --git a/Engine/Source/Developer/TraceInsights/Private/Insights/Widgets/SFilterConfiguratorRow.h b/Engine/Source/Developer/TraceInsights/Private/Insights/Widgets/SFilterConfiguratorRow.h
--- a/Engine/Source/Developer/TraceInsights/Private/Insights/Widgets/SFilterConfiguratorRow.h
+++ b/Engine/Source/Developer/TraceInsights/Private/Insights/Widgets/SFilterConfiguratorRow.h
@@ -52,9 +52,11 @@ private:
 	FText FilterGroupOperators_GetSelectionText() const;
 
 	FReply AddFilter_OnClicked();
-
 	FReply DeleteFilter_OnClicked();
 
+	FReply AddGroup_OnClicked();
+	FReply DeleteGroup_OnClicked();
+
 	FText GetTextBoxValue() const;
 	void OnTextBoxValueCommitted(const FText& InNewText, ETextCommit::Type InTextCommit);
 
