[Insights] Filter Configurator performance optimizations

diff --git a/Engine/Source/Developer/TraceInsights/Private/Insights/ViewModels/FilterConfigurator.cpp b/Engine/Source/Developer/TraceInsights/Private/Insights/ViewModels/FilterConfigurator.cpp
--- a/Engine/Source/Developer/TraceInsights/Private/Insights/ViewModels/FilterConfigurator.cpp
+++ b/Engine/Source/Developer/TraceInsights/Private/Insights/ViewModels/FilterConfigurator.cpp
@@ -31,6 +31,8 @@ FFilterConfigurator& FFilterConfigurator::operator=(const FFilterConfigurator& O
 		OnChangesCommitedCallback();
 	}
 
+	RootNode->ProcessFilter();
+
 	return *this;
 }
 
diff --git a/Engine/Source/Developer/TraceInsights/Private/Insights/ViewModels/FilterConfiguratorNode.cpp b/Engine/Source/Developer/TraceInsights/Private/Insights/ViewModels/FilterConfiguratorNode.cpp
--- a/Engine/Source/Developer/TraceInsights/Private/Insights/ViewModels/FilterConfiguratorNode.cpp
+++ b/Engine/Source/Developer/TraceInsights/Private/Insights/ViewModels/FilterConfiguratorNode.cpp
@@ -103,29 +103,64 @@ void FFilterConfiguratorNode::SetSelectedFilter(TSharedPtr<struct FFilter> InSel
 
 ////////////////////////////////////////////////////////////////////////////////////////////////////
 
+void FFilterConfiguratorNode::ProcessFilter()
+{
+	if (IsGroup())
+	{
+		TArray<Insights::FBaseTreeNodePtr> ChildArray = GetChildrenMutable();
+		for (Insights::FBaseTreeNodePtr Child : ChildArray)
+		{
+			FFilterConfiguratorNodePtr CastedChild = StaticCastSharedPtr<FFilterConfiguratorNode>(Child);
+			CastedChild->ProcessFilter();
+		}
+	}
+	else
+	{
+		switch (SelectedFilter->DataType)
+		{
+		case EFilterDataType::Double:
+		{
+			FilterValue.Set<double>(FCString::Atod(*TextBoxValue));
+			break;
+		}
+		case EFilterDataType::Int64:
+		{
+			FilterValue.Set<double>(FCString::Atoi64(*TextBoxValue));
+			break;
+		}
+		}
+	}
+}
+
+////////////////////////////////////////////////////////////////////////////////////////////////////
+
 bool FFilterConfiguratorNode::ApplyFilters(const FFilterContext& Context) const
 {
-	bool ret = true;
+	bool Ret = true;
 	if (IsGroup())
 	{
-		if (SelectedFilterGroupOperator->Type == EFilterGroupOperator::And)
+		switch (SelectedFilterGroupOperator->Type)
+		{
+		case EFilterGroupOperator::And:
 		{
 			auto& ChildrenArr = GetChildren();
-			for (int Index = 0; Index < ChildrenArr.Num(); ++Index)
+			for (int Index = 0; Index < ChildrenArr.Num() && Ret; ++Index)
 			{
-				ret &= ((FFilterConfiguratorNode*)&*ChildrenArr[Index])->ApplyFilters(Context);
+				Ret &= ((FFilterConfiguratorNode*)&*ChildrenArr[Index])->ApplyFilters(Context);
 			}
+			break;
 		}
-
-		if (SelectedFilterGroupOperator->Type == EFilterGroupOperator::Or)
+		case EFilterGroupOperator::Or:
 		{
-			ret = false;
+			Ret = false;
 
 			auto& ChildrenArr = GetChildren();
-			for (int Index = 0; Index < ChildrenArr.Num(); ++Index)
+			for (int Index = 0; Index < ChildrenArr.Num() && !Ret; ++Index)
 			{
-				ret |= ((FFilterConfiguratorNode*)&*ChildrenArr[Index])->ApplyFilters(Context);
+				Ret |= ((FFilterConfiguratorNode*)&*ChildrenArr[Index])->ApplyFilters(Context);
 			}
+			break;
+		}
 		}
 	}
 	else
@@ -133,22 +168,21 @@ bool FFilterConfiguratorNode::ApplyFilters(const FFilterContext& Context) const
 		switch (SelectedFilter->DataType)
 		{
 		case EFilterDataType::Double:
-			{
-				FFilterOperator<double>* Operator = (FFilterOperator<double>*)&*SelectedFilterOperator;
-				double Value = 0.0;
-				Context.GetFilterData<double>(SelectedFilter->Key, Value);
-
-				ret = Operator->Func(Value, FCString::Atod(*TextBoxValue));
-				break;
-			}
+		{
+			FFilterOperator<double>* Operator = (FFilterOperator<double>*)&*SelectedFilterOperator;
+			double Value = 0.0;
+			Context.GetFilterData<double>(SelectedFilter->Key, Value);
 
+			Ret = Operator->Func(Value, FilterValue.Get<double>());
+			break;
+		}
 		case EFilterDataType::Int64:
 		{
 			FFilterOperator<int64>* Operator = (FFilterOperator<int64>*) & *SelectedFilterOperator;
 			int64 Value = 0;
 			Context.GetFilterData<int64>(SelectedFilter->Key, Value);
 
-			ret = Operator->Func(Value, FCString::Atoi64(*TextBoxValue));
+			Ret = Operator->Func(Value, FilterValue.Get<int64>());
 			break;
 		}
 		default:
@@ -156,7 +190,7 @@ bool FFilterConfiguratorNode::ApplyFilters(const FFilterContext& Context) const
 		}
 	}
 
-	return ret;
+	return Ret;
 }
 
 ////////////////////////////////////////////////////////////////////////////////////////////////////
diff --git a/Engine/Source/Developer/TraceInsights/Private/Insights/ViewModels/FilterConfiguratorNode.h b/Engine/Source/Developer/TraceInsights/Private/Insights/ViewModels/FilterConfiguratorNode.h
--- a/Engine/Source/Developer/TraceInsights/Private/Insights/ViewModels/FilterConfiguratorNode.h
+++ b/Engine/Source/Developer/TraceInsights/Private/Insights/ViewModels/FilterConfiguratorNode.h
@@ -6,6 +6,7 @@
 
 // Insights
 #include "Insights/Table/ViewModels/BaseTreeNode.h"
+#include "Insights/ViewModels/Filters.h"
 
 ////////////////////////////////////////////////////////////////////////////////////////////////////
 
@@ -76,6 +77,8 @@ public:
 
 	bool ApplyFilters(const class FFilterContext& Context) const;
 
+	void ProcessFilter();
+
 private:
 	FFilterConfiguratorNode& operator=(const FFilterConfiguratorNode& Other);
 
@@ -88,6 +91,9 @@ private:
 	TSharedPtr<struct FFilterGroupOperator> SelectedFilterGroupOperator;
 
 	FString TextBoxValue;
+
+	// The textbox value, converted to its data type
+	FFilterContext::ContextData FilterValue;
 };
 
 ////////////////////////////////////////////////////////////////////////////////////////////////////
diff --git a/Engine/Source/Developer/TraceInsights/Private/Insights/ViewModels/Filters.h b/Engine/Source/Developer/TraceInsights/Private/Insights/ViewModels/Filters.h
--- a/Engine/Source/Developer/TraceInsights/Private/Insights/ViewModels/Filters.h
+++ b/Engine/Source/Developer/TraceInsights/Private/Insights/ViewModels/Filters.h
@@ -201,6 +201,12 @@ public:
 		DataMap.Add(Key, VariantData);
 	}
 
+	template<typename T>
+	void SetFilterData(EFilterField Key, const T& InData)
+	{
+		DataMap[Key].Set<T>(InData);
+	}
+
 	template<typename T>
 	void GetFilterData(EFilterField Key, T& OutData) const
 	{
diff --git a/Engine/Source/Developer/TraceInsights/Private/Insights/ViewModels/ThreadTimingTrack.cpp b/Engine/Source/Developer/TraceInsights/Private/Insights/ViewModels/ThreadTimingTrack.cpp
--- a/Engine/Source/Developer/TraceInsights/Private/Insights/ViewModels/ThreadTimingTrack.cpp
+++ b/Engine/Source/Developer/TraceInsights/Private/Insights/ViewModels/ThreadTimingTrack.cpp
@@ -816,6 +816,7 @@ void FThreadTimingTrack::BuildFilteredDrawState(ITimingEventsTrackDrawStateBuild
 				[&Viewport, this, &Builder, TimerReader](const Trace::ITimingProfilerProvider::Timeline& Timeline)
 				{
 					TArray<TArray<FPendingEventInfo>> FilteredEvents;
+					TArray<FFilterContext> FilterContexts;
 
 					Trace::ITimeline<Trace::FTimingProfilerEvent>::EnumerateAsyncParams Params;
 					Params.IntervalStart = Viewport.GetStartTime();
@@ -823,20 +824,28 @@ void FThreadTimingTrack::BuildFilteredDrawState(ITimingEventsTrackDrawStateBuild
 
 					// Note: Enumerating events for filtering should not use downsampling.
 					Params.Resolution = 0.0;
-					Params.SetupCallback = [&FilteredEvents](uint32 NumTasks)
+					Params.SetupCallback = [&FilteredEvents, &FilterContexts](uint32 NumTasks)
 					{
 						FilteredEvents.AddDefaulted(NumTasks);
+						FilterContexts.AddDefaulted(NumTasks);
+						for(FFilterContext& Context : FilterContexts)
+						{
+							Context.AddFilterData<double>(EFilterField::StartTime, 0.0f);
+							Context.AddFilterData<double>(EFilterField::EndTime, 0.0f);
+							Context.AddFilterData<double>(EFilterField::Duration, 0.0f);
+							Context.AddFilterData<int64>(EFilterField::EventType, 0);
+						}
 					};
-					Params.Callback = [this, &Builder, TimerReader, &FilteredEvents](double StartTime, double EndTime, uint32 Depth, const Trace::FTimingProfilerEvent& Event, uint32 TaskIndex)
+					Params.Callback = [this, &Builder, TimerReader, &FilteredEvents, &FilterContexts](double StartTime, double EndTime, uint32 Depth, const Trace::FTimingProfilerEvent& Event, uint32 TaskIndex)
 					{
 						const Trace::FTimingProfilerTimer* Timer = TimerReader->GetTimer(Event.TimerIndex);
 						if (ensure(Timer != nullptr))
 						{
-							FFilterContext Context;
-							Context.AddFilterData<double>(EFilterField::StartTime, StartTime);
-							Context.AddFilterData<double>(EFilterField::EndTime, EndTime);
-							Context.AddFilterData<double>(EFilterField::Duration, EndTime - StartTime);
-							Context.AddFilterData<int64>(EFilterField::EventType, Timer->Id);
+							FFilterContext& Context = FilterContexts[TaskIndex];
+							Context.SetFilterData<double>(EFilterField::StartTime, StartTime);
+							Context.SetFilterData<double>(EFilterField::EndTime, EndTime);
+							Context.SetFilterData<double>(EFilterField::Duration, EndTime - StartTime);
+							Context.SetFilterData<int64>(EFilterField::EventType, Timer->Id);
 
 							if (FilterConfigurator->ApplyFilters(Context))
 							{
@@ -1168,7 +1177,7 @@ void FThreadTimingTrack::BuildContextMenu(FMenuBuilder& MenuBuilder)
 			FilterTrackAction.BindSP(this, &FThreadTimingTrack::OnFilterTrackClicked);
 
 			MenuBuilder.AddMenuEntry(
-				(LOCTEXT("FilerTrack", "Filter Track [WIP]")),
+				(LOCTEXT("FilterTrack", "Filter Track")),
 				FText(),
 				FSlateIcon(),
 				FUIAction(FilterTrackAction, FCanExecuteAction::CreateLambda([]() { return true; })),
