Filter Configurator tab spawner no longer appears in menus

diff --git a/Engine/Source/Developer/TraceInsights/Private/Insights/ViewModels/Filters.cpp b/Engine/Source/Developer/TraceInsights/Private/Insights/ViewModels/Filters.cpp
--- a/Engine/Source/Developer/TraceInsights/Private/Insights/ViewModels/Filters.cpp
+++ b/Engine/Source/Developer/TraceInsights/Private/Insights/ViewModels/Filters.cpp
@@ -63,7 +63,8 @@ void FFilterService::RegisterTabSpawner()
 {
 	FTabSpawnerEntry& TabSpawnerEntry = FGlobalTabmanager::Get()->RegisterNomadTabSpawner(FilterConfiguratorTabId,
 		FOnSpawnTab::CreateRaw(this, &FFilterService::SpawnTab))
-		.SetDisplayName(LOCTEXT("FilterConfiguratorTabTitle", "Filter Configurator"));
+		.SetDisplayName(LOCTEXT("FilterConfiguratorTabTitle", "Filter Configurator"))
+		.SetMenuType(ETabSpawnerMenuType::Hidden);
 }
 
 ////////////////////////////////////////////////////////////////////////////////////////////////////
