[Insights] Filter Configurator closes when parent window closes

diff --git a/Engine/Source/Developer/TraceInsights/Private/Insights/ViewModels/FilterConfigurator.cpp b/Engine/Source/Developer/TraceInsights/Private/Insights/ViewModels/FilterConfigurator.cpp
--- a/Engine/Source/Developer/TraceInsights/Private/Insights/ViewModels/FilterConfigurator.cpp
+++ b/Engine/Source/Developer/TraceInsights/Private/Insights/ViewModels/FilterConfigurator.cpp
@@ -2,6 +2,8 @@
 
 #include "FilterConfigurator.h"
 
+#include "Insights/Widgets/SFilterConfigurator.h"
+
 ////////////////////////////////////////////////////////////////////////////////////////////////////
 
 FFilterConfigurator::FFilterConfigurator()
@@ -13,6 +15,13 @@ FFilterConfigurator::FFilterConfigurator()
 
 ////////////////////////////////////////////////////////////////////////////////////////////////////
 
+FFilterConfigurator::~FFilterConfigurator()
+{
+	OnDestroyedEvent.Broadcast();
+}
+
+////////////////////////////////////////////////////////////////////////////////////////////////////
+
 FFilterConfigurator::FFilterConfigurator(const FFilterConfigurator& Other)
 {
 	*this = Other;
@@ -25,11 +34,9 @@ FFilterConfigurator& FFilterConfigurator::operator=(const FFilterConfigurator& O
 	RootNode = MakeShared<FFilterConfiguratorNode>(*Other.RootNode);
 	RootNode->SetGroupPtrForChildren();
 	AvailableFilters = Other.AvailableFilters;
+	OnDestroyedEvent = Other.OnDestroyedEvent;
 
-	if (OnChangesCommitedCallback)
-	{
-		OnChangesCommitedCallback();
-	}
+	OnChangesCommitedEvent.Broadcast();
 
 	RootNode->ProcessFilter();
 
diff --git a/Engine/Source/Developer/TraceInsights/Private/Insights/ViewModels/FilterConfigurator.h b/Engine/Source/Developer/TraceInsights/Private/Insights/ViewModels/FilterConfigurator.h
--- a/Engine/Source/Developer/TraceInsights/Private/Insights/ViewModels/FilterConfigurator.h
+++ b/Engine/Source/Developer/TraceInsights/Private/Insights/ViewModels/FilterConfigurator.h
@@ -16,7 +16,7 @@ public:
 	FFilterConfigurator(const FFilterConfigurator& Other);
 	FFilterConfigurator& operator=(const FFilterConfigurator& Other);
 
-	~FFilterConfigurator() {}
+	~FFilterConfigurator();
 
 	void AddFilter(EFilterField FilterKey)
 	{
@@ -27,15 +27,33 @@ public:
 
 	bool ApplyFilters(const FFilterContext& Context) const;
 
-	void SetOnChangesCommitedCallback(TFunction<void()> InCallback) {	OnChangesCommitedCallback = InCallback;	}
+	////////////////////////////////////////////////////////////////////////////////////////////////////
+	// OnDestroyedEvent
+public:
+	/** The event to execute when an instance is destroyed. */
+	DECLARE_MULTICAST_DELEGATE(FOnDestroyedEvent);
+	FOnDestroyedEvent& GetOnDestroyedEvent() { return OnDestroyedEvent; }
+
+private:
+	FOnDestroyedEvent OnDestroyedEvent;
+
+	////////////////////////////////////////////////////////////////////////////////////////////////////
+	// OnChangesCommitedEvent
+public:
+	/** The event to execute when the changes to the Filter Widget are saved by clicking on the OK Button. */
+	DECLARE_MULTICAST_DELEGATE(FOnChangesCommitedEvent);
+	FOnChangesCommitedEvent& GetOnChangesCommitedEvent() { return OnChangesCommitedEvent; }
+
+private:
+	FOnChangesCommitedEvent OnChangesCommitedEvent;
+
+	////////////////////////////////////////////////////////////////////////////////////////////////////
 
 private:
 
 	FFilterConfiguratorNodePtr RootNode;
 
 	TSharedPtr<TArray<TSharedPtr<struct FFilter>>> AvailableFilters;
-
-	TFunction<void()> OnChangesCommitedCallback = nullptr;
 };
 
 ////////////////////////////////////////////////////////////////////////////////////////////////////
\ No newline at end of file
diff --git a/Engine/Source/Developer/TraceInsights/Private/Insights/ViewModels/ThreadTimingTrack.cpp b/Engine/Source/Developer/TraceInsights/Private/Insights/ViewModels/ThreadTimingTrack.cpp
--- a/Engine/Source/Developer/TraceInsights/Private/Insights/ViewModels/ThreadTimingTrack.cpp
+++ b/Engine/Source/Developer/TraceInsights/Private/Insights/ViewModels/ThreadTimingTrack.cpp
@@ -613,6 +613,16 @@ INSIGHTS_IMPLEMENT_RTTI(FThreadTimingTrack)
 
 ////////////////////////////////////////////////////////////////////////////////////////////////////
 
+FThreadTimingTrack::~FThreadTimingTrack()
+{
+	if (FilterConfigurator.IsValid())
+	{
+		FilterConfigurator->GetOnChangesCommitedEvent().Remove(OnFilterChangesCommitedHandle);
+	}
+}
+
+////////////////////////////////////////////////////////////////////////////////////////////////////
+
 void FThreadTimingTrack::BuildDrawState(ITimingEventsTrackDrawStateBuilder& Builder, const ITimingTrackUpdateContext& Context)
 {
 	TSharedPtr<const Trace::IAnalysisSession> Session = FInsightsManager::Get()->GetSession();
@@ -1323,7 +1333,7 @@ void FThreadTimingTrack::OnFilterTrackClicked()
 		FilterConfigurator->AddFilter(EFilterField::Duration);
 		FilterConfigurator->AddFilter(EFilterField::EventType);
 
-		FilterConfigurator->SetOnChangesCommitedCallback([this]()
+		OnFilterChangesCommitedHandle = FilterConfigurator->GetOnChangesCommitedEvent().AddLambda([this]()
 			{
 				this->SetDirtyFlag();
 			});
diff --git a/Engine/Source/Developer/TraceInsights/Private/Insights/ViewModels/ThreadTimingTrack.h b/Engine/Source/Developer/TraceInsights/Private/Insights/ViewModels/ThreadTimingTrack.h
--- a/Engine/Source/Developer/TraceInsights/Private/Insights/ViewModels/ThreadTimingTrack.h
+++ b/Engine/Source/Developer/TraceInsights/Private/Insights/ViewModels/ThreadTimingTrack.h
@@ -116,7 +116,7 @@ public:
 	{
 	}
 
-	virtual ~FThreadTimingTrack() {}
+	virtual ~FThreadTimingTrack();
 
 	const TCHAR* GetGroupName() const { return GroupName; };
 
@@ -165,6 +165,7 @@ private:
 	uint32 ThreadId;
 
 	TSharedPtr<class FFilterConfigurator> FilterConfigurator;
+	FDelegateHandle OnFilterChangesCommitedHandle;
 
 	// Search cache
 	mutable TTimingEventSearchCache<Trace::FTimingProfilerEvent> SearchCache;
diff --git a/Engine/Source/Developer/TraceInsights/Private/Insights/Widgets/SFilterConfigurator.cpp b/Engine/Source/Developer/TraceInsights/Private/Insights/Widgets/SFilterConfigurator.cpp
--- a/Engine/Source/Developer/TraceInsights/Private/Insights/Widgets/SFilterConfigurator.cpp
+++ b/Engine/Source/Developer/TraceInsights/Private/Insights/Widgets/SFilterConfigurator.cpp
@@ -160,8 +160,9 @@ void SFilterConfigurator::Construct(const FArguments& InArgs, TSharedPtr<FFilter
 
 	TreeViewHeaderRow->InsertColumn(ColumnArgs, 0);
 
+	OnViewModelDestroyedHandle = InFilterConfiguratorViewModel->GetOnDestroyedEvent().AddSP(this, &SFilterConfigurator::RequestClose);
 	OriginalFilterConfiguratorViewModel = InFilterConfiguratorViewModel;
-	FilterConfiguratorViewModel = MakeShared<FFilterConfigurator>(*OriginalFilterConfiguratorViewModel);
+	FilterConfiguratorViewModel = MakeShared<FFilterConfigurator>(*InFilterConfiguratorViewModel);
 
 	GroupNodes.Add(FilterConfiguratorViewModel->GetRootNode());
 }
@@ -195,13 +196,14 @@ TSharedRef<ITableRow> SFilterConfigurator::TreeView_OnGenerateRow(FFilterConfigu
 
 FReply SFilterConfigurator::OK_OnClicked()
 {
-	*OriginalFilterConfiguratorViewModel = *FilterConfiguratorViewModel;
-	if (ParentTab.IsValid())
+	TSharedPtr<FFilterConfigurator> OriginalFilterVM = OriginalFilterConfiguratorViewModel.Pin();
+	if (OriginalFilterVM.IsValid())
 	{
-		ParentTab.Pin()->RequestCloseTab();
-		ParentTab.Reset();
+		*OriginalFilterVM = *FilterConfiguratorViewModel;
 	}
 
+	RequestClose();
+
 	return FReply::Handled();
 }
 
@@ -209,13 +211,27 @@ FReply SFilterConfigurator::OK_OnClicked()
 
 FReply SFilterConfigurator::Cancel_OnClicked()
 {
-	if (ParentTab.IsValid())
+	RequestClose();
+
+	return FReply::Handled();
+}
+
+////////////////////////////////////////////////////////////////////////////////////////////////////
+
+void SFilterConfigurator::RequestClose()
+{
+	TSharedPtr<SDockTab> ParentTabSharedPtr = ParentTab.Pin();
+	if (ParentTabSharedPtr.IsValid())
 	{
-		ParentTab.Pin()->RequestCloseTab();
-		ParentTab.Reset();
+		ParentTabSharedPtr->RequestCloseTab();
+		ParentTabSharedPtr.Reset();
 	}
 
-	return FReply::Handled();
+	TSharedPtr<FFilterConfigurator> OriginalFilterVM = OriginalFilterConfiguratorViewModel.Pin();
+	if (OriginalFilterVM.IsValid())
+	{
+		OriginalFilterVM->GetOnDestroyedEvent().Remove(OnViewModelDestroyedHandle);
+	}
 }
 
 ////////////////////////////////////////////////////////////////////////////////////////////////////
diff --git a/Engine/Source/Developer/TraceInsights/Private/Insights/Widgets/SFilterConfigurator.h b/Engine/Source/Developer/TraceInsights/Private/Insights/Widgets/SFilterConfigurator.h
--- a/Engine/Source/Developer/TraceInsights/Private/Insights/Widgets/SFilterConfigurator.h
+++ b/Engine/Source/Developer/TraceInsights/Private/Insights/Widgets/SFilterConfigurator.h
@@ -52,6 +52,8 @@ public:
 	void SetParentTab(const TSharedPtr<SDockTab> InTab) { ParentTab = InTab; }
 	const TWeakPtr<SDockTab> GetParentTab() { return ParentTab; };
 
+	void RequestClose();
+
 private:
 	void InitCommandList();
 
@@ -61,7 +63,7 @@ private:
 
 private:
 
-	TSharedPtr<class FFilterConfigurator> OriginalFilterConfiguratorViewModel;
+	TWeakPtr<class FFilterConfigurator> OriginalFilterConfiguratorViewModel;
 
 	TSharedPtr<class FFilterConfigurator> FilterConfiguratorViewModel;
 
@@ -78,6 +80,8 @@ private:
 	TArray<FFilterConfiguratorNodePtr> GroupNodes;
 
 	TWeakPtr<SDockTab> ParentTab;
+
+	FDelegateHandle OnViewModelDestroyedHandle;
 };
 
 ////////////////////////////////////////////////////////////////////////////////////////////////////
