Moved the TraceServices module into its own namespace.

diff --git a/Engine/Plugins/Slate/SlateInsights/Source/SlateInsights/Private/SlateTraceModule.cpp b/Engine/Plugins/Slate/SlateInsights/Source/SlateInsights/Private/SlateTraceModule.cpp
--- a/Engine/Plugins/Slate/SlateInsights/Source/SlateInsights/Private/SlateTraceModule.cpp
+++ b/Engine/Plugins/Slate/SlateInsights/Source/SlateInsights/Private/SlateTraceModule.cpp
@@ -11,13 +11,13 @@ namespace SlateInsights
 	
 FName FSlateTraceModule::ModuleName("TraceModule_Slate");
 
-void FSlateTraceModule::GetModuleInfo(Trace::FModuleInfo& OutModuleInfo)
+void FSlateTraceModule::GetModuleInfo(TraceServices::FModuleInfo& OutModuleInfo)
 {
 	OutModuleInfo.Name = ModuleName;
 	OutModuleInfo.DisplayName = TEXT("Slate");
 }
 
-void FSlateTraceModule::OnAnalysisBegin(Trace::IAnalysisSession& InSession)
+void FSlateTraceModule::OnAnalysisBegin(TraceServices::IAnalysisSession& InSession)
 {
 	FSlateProvider* SlateProvider = new FSlateProvider(InSession);
 	InSession.AddProvider(FSlateProvider::ProviderName, SlateProvider);
@@ -30,10 +30,10 @@ void FSlateTraceModule::GetLoggers(TArray<const TCHAR *>& OutLoggers)
 	OutLoggers.Add(TEXT("Slate"));
 }
 
-void FSlateTraceModule::GenerateReports(const Trace::IAnalysisSession& Session, const TCHAR* CmdLine, const TCHAR* OutputDirectory)
+void FSlateTraceModule::GenerateReports(const TraceServices::IAnalysisSession& Session, const TCHAR* CmdLine, const TCHAR* OutputDirectory)
 {
 
 }
 
 } //namespace SlateInsights
-} //namespace UE
\ No newline at end of file
+} //namespace UE
diff --git a/Engine/Plugins/Slate/SlateInsights/Source/SlateInsights/Private/SlateTraceModule.h b/Engine/Plugins/Slate/SlateInsights/Source/SlateInsights/Private/SlateTraceModule.h
--- a/Engine/Plugins/Slate/SlateInsights/Source/SlateInsights/Private/SlateTraceModule.h
+++ b/Engine/Plugins/Slate/SlateInsights/Source/SlateInsights/Private/SlateTraceModule.h
@@ -13,16 +13,16 @@ namespace UE
 namespace SlateInsights
 {
 
-class FSlateTraceModule : public Trace::IModule
+class FSlateTraceModule : public TraceServices::IModule
 {
 public:
-	//~ Begin Trace::IModule interface
-	virtual void GetModuleInfo(Trace::FModuleInfo& OutModuleInfo) override;
-	virtual void OnAnalysisBegin(Trace::IAnalysisSession& Session) override;
+	//~ Begin TraceServices::IModule interface
+	virtual void GetModuleInfo(TraceServices::FModuleInfo& OutModuleInfo) override;
+	virtual void OnAnalysisBegin(TraceServices::IAnalysisSession& Session) override;
 	virtual void GetLoggers(TArray<const TCHAR *>& OutLoggers) override;
-	virtual void GenerateReports(const Trace::IAnalysisSession& Session, const TCHAR* CmdLine, const TCHAR* OutputDirectory) override;
+	virtual void GenerateReports(const TraceServices::IAnalysisSession& Session, const TCHAR* CmdLine, const TCHAR* OutputDirectory) override;
 	virtual const TCHAR* GetCommandLineArgument() override { return TEXT("slatetrace"); }
-	//~ End Trace::IModule interface
+	//~ End TraceServices::IModule interface
 
 private:
 	static FName ModuleName;
diff --git a/Engine/Source/Developer/TraceInsights/Private/Insights/InsightsManager.cpp b/Engine/Source/Developer/TraceInsights/Private/Insights/InsightsManager.cpp
--- a/Engine/Source/Developer/TraceInsights/Private/Insights/InsightsManager.cpp
+++ b/Engine/Source/Developer/TraceInsights/Private/Insights/InsightsManager.cpp
@@ -94,8 +94,8 @@ TSharedPtr<FInsightsManager> FInsightsManager::Get()
 
 ////////////////////////////////////////////////////////////////////////////////////////////////////
 
-TSharedPtr<FInsightsManager> FInsightsManager::CreateInstance(TSharedRef<Trace::IAnalysisService> TraceAnalysisService,
-															  TSharedRef<Trace::IModuleService> TraceModuleService)
+TSharedPtr<FInsightsManager> FInsightsManager::CreateInstance(TSharedRef<TraceServices::IAnalysisService> TraceAnalysisService,
+															  TSharedRef<TraceServices::IModuleService> TraceModuleService)
 {
 	ensure(!FInsightsManager::Instance.IsValid());
 	if (FInsightsManager::Instance.IsValid())
@@ -110,8 +110,8 @@ TSharedPtr<FInsightsManager> FInsightsManager::CreateInstance(TSharedRef<Trace::
 
 ////////////////////////////////////////////////////////////////////////////////////////////////////
 
-FInsightsManager::FInsightsManager(TSharedRef<Trace::IAnalysisService> InTraceAnalysisService,
-								   TSharedRef<Trace::IModuleService> InTraceModuleService)
+FInsightsManager::FInsightsManager(TSharedRef<TraceServices::IAnalysisService> InTraceAnalysisService,
+								   TSharedRef<TraceServices::IModuleService> InTraceModuleService)
 	: bIsInitialized(false)
 	, AnalysisService(InTraceAnalysisService)
 	, ModuleService(InTraceModuleService)
@@ -321,7 +321,7 @@ bool FInsightsManager::ConnectToStore(const TCHAR* Host, uint32 Port)
 
 ////////////////////////////////////////////////////////////////////////////////////////////////////
 
-TSharedPtr<const Trace::IAnalysisSession> FInsightsManager::GetSession() const
+TSharedPtr<const TraceServices::IAnalysisSession> FInsightsManager::GetSession() const
 {
 	return Session;
 }
@@ -371,7 +371,7 @@ void FInsightsManager::UpdateSessionDuration()
 	{
 		double LocalSessionDuration = 0.0;
 		{
-			Trace::FAnalysisSessionReadScope SessionReadScope(*Session.Get());
+			TraceServices::FAnalysisSessionReadScope SessionReadScope(*Session.Get());
 			if (bIsAnalysisComplete != Session->IsAnalysisComplete())
 			{
 				bIsAnalysisComplete = Session->IsAnalysisComplete();
diff --git a/Engine/Source/Developer/TraceInsights/Private/Insights/InsightsManager.h b/Engine/Source/Developer/TraceInsights/Private/Insights/InsightsManager.h
--- a/Engine/Source/Developer/TraceInsights/Private/Insights/InsightsManager.h
+++ b/Engine/Source/Developer/TraceInsights/Private/Insights/InsightsManager.h
@@ -19,6 +19,10 @@
 namespace Trace
 {
 	class FStoreClient;
+}
+
+namespace TraceServices
+{
 	class IAnalysisService;
 	class IModuleService;
 }
@@ -60,8 +64,8 @@ class FInsightsManager : public TSharedFromThis<FInsightsManager>, public IInsig
 
 public:
 	/** Creates the main manager, only one instance can exist. */
-	FInsightsManager(TSharedRef<Trace::IAnalysisService> TraceAnalysisService,
-					 TSharedRef<Trace::IModuleService> TraceModuleService);
+	FInsightsManager(TSharedRef<TraceServices::IAnalysisService> TraceAnalysisService,
+					 TSharedRef<TraceServices::IModuleService> TraceModuleService);
 
 	/** Virtual destructor. */
 	virtual ~FInsightsManager();
@@ -71,8 +75,8 @@ public:
 	 * @param TraceAnalysisService The trace analysis service
 	 * @param TraceModuleService   The trace module service
 	 */
-	static TSharedPtr<FInsightsManager> CreateInstance(TSharedRef<Trace::IAnalysisService> TraceAnalysisService,
-													   TSharedRef<Trace::IModuleService> TraceModuleService);
+	static TSharedPtr<FInsightsManager> CreateInstance(TSharedRef<TraceServices::IAnalysisService> TraceAnalysisService,
+													   TSharedRef<TraceServices::IModuleService> TraceModuleService);
 
 	/** @return the global instance of the main manager (FInsightsManager). */
 	static TSharedPtr<FInsightsManager> Get();
@@ -87,8 +91,8 @@ public:
 
 	//////////////////////////////////////////////////
 
-	TSharedRef<Trace::IAnalysisService> GetAnalysisService() const { return AnalysisService; }
-	TSharedRef<Trace::IModuleService> GetModuleService() const { return ModuleService; }
+	TSharedRef<TraceServices::IAnalysisService> GetAnalysisService() const { return AnalysisService; }
+	TSharedRef<TraceServices::IModuleService> GetModuleService() const { return ModuleService; }
 
 	void SetStoreDir(const FString& InStoreDir) { StoreDir = InStoreDir; }
 	const FString& GetStoreDir() const { return StoreDir; }
@@ -97,7 +101,7 @@ public:
 	Trace::FStoreClient* GetStoreClient() const { return StoreClient.Get(); }
 
 	/** @return an instance of the trace analysis session. */
-	TSharedPtr<const Trace::IAnalysisSession> GetSession() const;
+	TSharedPtr<const TraceServices::IAnalysisSession> GetSession() const;
 
 	/** @return the id of the trace being analyzed. */
 	uint32 GetTraceId() const { return CurrentTraceId; }
@@ -263,8 +267,8 @@ private:
 	/** Handle to the registered OnTick. */
 	FDelegateHandle OnTickHandle;
 
-	TSharedRef<Trace::IAnalysisService> AnalysisService;
-	TSharedRef<Trace::IModuleService> ModuleService;
+	TSharedRef<TraceServices::IAnalysisService> AnalysisService;
+	TSharedRef<TraceServices::IModuleService> ModuleService;
 
 	/** The location of the trace files managed by the trace store. */
 	FString StoreDir;
@@ -273,7 +277,7 @@ private:
 	TUniquePtr<Trace::FStoreClient> StoreClient;
 
 	/** The trace analysis session. */
-	TSharedPtr<const Trace::IAnalysisSession> Session;
+	TSharedPtr<const TraceServices::IAnalysisSession> Session;
 
 	/** The id of the trace being analyzed. */
 	uint32 CurrentTraceId;
diff --git a/Engine/Source/Developer/TraceInsights/Private/Insights/LoadingProfiler/LoadingProfilerManager.cpp b/Engine/Source/Developer/TraceInsights/Private/Insights/LoadingProfiler/LoadingProfilerManager.cpp
--- a/Engine/Source/Developer/TraceInsights/Private/Insights/LoadingProfiler/LoadingProfilerManager.cpp
+++ b/Engine/Source/Developer/TraceInsights/Private/Insights/LoadingProfiler/LoadingProfilerManager.cpp
@@ -221,10 +221,10 @@ bool FLoadingProfilerManager::Tick(float DeltaTime)
 	{
 		bool bIsProviderAvailable = false;
 
-		TSharedPtr<const Trace::IAnalysisSession> Session = FInsightsManager::Get()->GetSession();
+		TSharedPtr<const TraceServices::IAnalysisSession> Session = FInsightsManager::Get()->GetSession();
 		if (Session.IsValid())
 		{
-			Trace::FAnalysisSessionReadScope SessionReadScope(*Session.Get());
+			TraceServices::FAnalysisSessionReadScope SessionReadScope(*Session.Get());
 
 			if (Session->IsAnalysisComplete())
 			{
@@ -232,7 +232,7 @@ bool FLoadingProfilerManager::Tick(float DeltaTime)
 				AvailabilityCheck.Disable();
 			}
 
-			const Trace::ILoadTimeProfilerProvider* LoadTimeProfilerProvider = Trace::ReadLoadTimeProfilerProvider(*Session.Get());
+			const TraceServices::ILoadTimeProfilerProvider* LoadTimeProfilerProvider = TraceServices::ReadLoadTimeProfilerProvider(*Session.Get());
 			if (LoadTimeProfilerProvider)
 			{
 				bIsProviderAvailable = (LoadTimeProfilerProvider->GetTimelineCount() > 0);
diff --git a/Engine/Source/Developer/TraceInsights/Private/Insights/LoadingProfiler/Widgets/SLoadingProfilerWindow.cpp b/Engine/Source/Developer/TraceInsights/Private/Insights/LoadingProfiler/Widgets/SLoadingProfilerWindow.cpp
--- a/Engine/Source/Developer/TraceInsights/Private/Insights/LoadingProfiler/Widgets/SLoadingProfilerWindow.cpp
+++ b/Engine/Source/Developer/TraceInsights/Private/Insights/LoadingProfiler/Widgets/SLoadingProfilerWindow.cpp
@@ -159,16 +159,16 @@ void SLoadingProfilerWindow::UpdateEventAggregationTreeView()
 {
 	if (EventAggregationTreeView)
 	{
-		TSharedPtr<const Trace::IAnalysisSession> Session = FInsightsManager::Get()->GetSession();
-		if (Session.IsValid() && Trace::ReadLoadTimeProfilerProvider(*Session.Get()))
+		TSharedPtr<const TraceServices::IAnalysisSession> Session = FInsightsManager::Get()->GetSession();
+		if (Session.IsValid() && TraceServices::ReadLoadTimeProfilerProvider(*Session.Get()))
 		{
-			Trace::FAnalysisSessionReadScope SessionReadScope(*Session.Get());
-			const Trace::ILoadTimeProfilerProvider& LoadTimeProfilerProvider = *Trace::ReadLoadTimeProfilerProvider(*Session.Get());
+			TraceServices::FAnalysisSessionReadScope SessionReadScope(*Session.Get());
+			const TraceServices::ILoadTimeProfilerProvider& LoadTimeProfilerProvider = *TraceServices::ReadLoadTimeProfilerProvider(*Session.Get());
 
 			const double SelectionStartTime = TimingView ? TimingView->GetSelectionStartTime() : 0.0;
 			const double SelectionEndTime = TimingView ? TimingView->GetSelectionEndTime() : 0.0;
 
-			Trace::ITable<Trace::FLoadTimeProfilerAggregatedStats>* EventAggregationTable = LoadTimeProfilerProvider.CreateEventAggregation(SelectionStartTime, SelectionEndTime);
+			TraceServices::ITable<TraceServices::FLoadTimeProfilerAggregatedStats>* EventAggregationTable = LoadTimeProfilerProvider.CreateEventAggregation(SelectionStartTime, SelectionEndTime);
 			EventAggregationTreeView->UpdateSourceTable(MakeShareable(EventAggregationTable));
 		}
 		else
@@ -184,16 +184,16 @@ void SLoadingProfilerWindow::UpdateObjectTypeAggregationTreeView()
 {
 	if (ObjectTypeAggregationTreeView)
 	{
-		TSharedPtr<const Trace::IAnalysisSession> Session = FInsightsManager::Get()->GetSession();
-		if (Session.IsValid() && Trace::ReadLoadTimeProfilerProvider(*Session.Get()))
+		TSharedPtr<const TraceServices::IAnalysisSession> Session = FInsightsManager::Get()->GetSession();
+		if (Session.IsValid() && TraceServices::ReadLoadTimeProfilerProvider(*Session.Get()))
 		{
-			Trace::FAnalysisSessionReadScope SessionReadScope(*Session.Get());
-			const Trace::ILoadTimeProfilerProvider& LoadTimeProfilerProvider = *Trace::ReadLoadTimeProfilerProvider(*Session.Get());
+			TraceServices::FAnalysisSessionReadScope SessionReadScope(*Session.Get());
+			const TraceServices::ILoadTimeProfilerProvider& LoadTimeProfilerProvider = *TraceServices::ReadLoadTimeProfilerProvider(*Session.Get());
 
 			const double SelectionStartTime = TimingView ? TimingView->GetSelectionStartTime() : 0.0;
 			const double SelectionEndTime = TimingView ? TimingView->GetSelectionEndTime() : 0.0;
 
-			Trace::ITable<Trace::FLoadTimeProfilerAggregatedStats>* ObjectTypeAggregationTable = LoadTimeProfilerProvider.CreateObjectTypeAggregation(SelectionStartTime, SelectionEndTime);
+			TraceServices::ITable<TraceServices::FLoadTimeProfilerAggregatedStats>* ObjectTypeAggregationTable = LoadTimeProfilerProvider.CreateObjectTypeAggregation(SelectionStartTime, SelectionEndTime);
 			ObjectTypeAggregationTreeView->UpdateSourceTable(MakeShareable(ObjectTypeAggregationTable));
 		}
 		else
@@ -209,16 +209,16 @@ void SLoadingProfilerWindow::UpdatePackageDetailsTreeView()
 {
 	if (PackageDetailsTreeView)
 	{
-		TSharedPtr<const Trace::IAnalysisSession> Session = FInsightsManager::Get()->GetSession();
-		if (Session.IsValid() && Trace::ReadLoadTimeProfilerProvider(*Session.Get()))
+		TSharedPtr<const TraceServices::IAnalysisSession> Session = FInsightsManager::Get()->GetSession();
+		if (Session.IsValid() && TraceServices::ReadLoadTimeProfilerProvider(*Session.Get()))
 		{
-			Trace::FAnalysisSessionReadScope SessionReadScope(*Session.Get());
-			const Trace::ILoadTimeProfilerProvider& LoadTimeProfilerProvider = *Trace::ReadLoadTimeProfilerProvider(*Session.Get());
+			TraceServices::FAnalysisSessionReadScope SessionReadScope(*Session.Get());
+			const TraceServices::ILoadTimeProfilerProvider& LoadTimeProfilerProvider = *TraceServices::ReadLoadTimeProfilerProvider(*Session.Get());
 
 			const double SelectionStartTime = TimingView ? TimingView->GetSelectionStartTime() : 0.0;
 			const double SelectionEndTime = TimingView ? TimingView->GetSelectionEndTime() : 0.0;
 
-			Trace::ITable<Trace::FPackagesTableRow>* PackageDetailsTable = LoadTimeProfilerProvider.CreatePackageDetailsTable(SelectionStartTime, SelectionEndTime);
+			TraceServices::ITable<TraceServices::FPackagesTableRow>* PackageDetailsTable = LoadTimeProfilerProvider.CreatePackageDetailsTable(SelectionStartTime, SelectionEndTime);
 			PackageDetailsTreeView->UpdateSourceTable(MakeShareable(PackageDetailsTable));
 		}
 		else
@@ -234,16 +234,16 @@ void SLoadingProfilerWindow::UpdateExportDetailsTreeView()
 {
 	if (ExportDetailsTreeView)
 	{
-		TSharedPtr<const Trace::IAnalysisSession> Session = FInsightsManager::Get()->GetSession();
-		if (Session.IsValid() && Trace::ReadLoadTimeProfilerProvider(*Session.Get()))
+		TSharedPtr<const TraceServices::IAnalysisSession> Session = FInsightsManager::Get()->GetSession();
+		if (Session.IsValid() && TraceServices::ReadLoadTimeProfilerProvider(*Session.Get()))
 		{
-			Trace::FAnalysisSessionReadScope SessionReadScope(*Session.Get());
-			const Trace::ILoadTimeProfilerProvider& LoadTimeProfilerProvider = *Trace::ReadLoadTimeProfilerProvider(*Session.Get());
+			TraceServices::FAnalysisSessionReadScope SessionReadScope(*Session.Get());
+			const TraceServices::ILoadTimeProfilerProvider& LoadTimeProfilerProvider = *TraceServices::ReadLoadTimeProfilerProvider(*Session.Get());
 
 			const double SelectionStartTime = TimingView ? TimingView->GetSelectionStartTime() : 0.0;
 			const double SelectionEndTime = TimingView ? TimingView->GetSelectionEndTime() : 0.0;
 
-			Trace::ITable<Trace::FExportsTableRow>* ExportDetailsTable = LoadTimeProfilerProvider.CreateExportDetailsTable(SelectionStartTime, SelectionEndTime);
+			TraceServices::ITable<TraceServices::FExportsTableRow>* ExportDetailsTable = LoadTimeProfilerProvider.CreateExportDetailsTable(SelectionStartTime, SelectionEndTime);
 			ExportDetailsTreeView->UpdateSourceTable(MakeShareable(ExportDetailsTable));
 		}
 		else
@@ -255,31 +255,31 @@ void SLoadingProfilerWindow::UpdateExportDetailsTreeView()
 
 ////////////////////////////////////////////////////////////////////////////////////////////////////
 
-class FProxyUntypedTable : public  Trace::IUntypedTable
+class FProxyUntypedTable : public  TraceServices::IUntypedTable
 {
 public:
-	FProxyUntypedTable(const Trace::IUntypedTable* InTable) : TablePtr(InTable) {}
+	FProxyUntypedTable(const TraceServices::IUntypedTable* InTable) : TablePtr(InTable) {}
 	virtual ~FProxyUntypedTable() = default;
 
-	virtual const Trace::ITableLayout& GetLayout() const { return TablePtr->GetLayout(); }
+	virtual const TraceServices::ITableLayout& GetLayout() const { return TablePtr->GetLayout(); }
 	virtual uint64 GetRowCount() const { return TablePtr->GetRowCount(); }
-	virtual Trace::IUntypedTableReader* CreateReader() const { return TablePtr->CreateReader(); }
+	virtual TraceServices::IUntypedTableReader* CreateReader() const { return TablePtr->CreateReader(); }
 
 private:
-	const Trace::IUntypedTable* TablePtr;
+	const TraceServices::IUntypedTable* TablePtr;
 };
 
 void SLoadingProfilerWindow::UpdateRequestsTreeView()
 {
 	if (RequestsTreeView)
 	{
-		TSharedPtr<const Trace::IAnalysisSession> Session = FInsightsManager::Get()->GetSession();
-		if (Session.IsValid() && Trace::ReadLoadTimeProfilerProvider(*Session.Get()))
+		TSharedPtr<const TraceServices::IAnalysisSession> Session = FInsightsManager::Get()->GetSession();
+		if (Session.IsValid() && TraceServices::ReadLoadTimeProfilerProvider(*Session.Get()))
 		{
-			Trace::FAnalysisSessionReadScope SessionReadScope(*Session.Get());
-			const Trace::ILoadTimeProfilerProvider& LoadTimeProfilerProvider = *Trace::ReadLoadTimeProfilerProvider(*Session.Get());
+			TraceServices::FAnalysisSessionReadScope SessionReadScope(*Session.Get());
+			const TraceServices::ILoadTimeProfilerProvider& LoadTimeProfilerProvider = *TraceServices::ReadLoadTimeProfilerProvider(*Session.Get());
 
-			const Trace::ITable<Trace::FLoadRequest>& RequestsTable = LoadTimeProfilerProvider.GetRequestsTable();
+			const TraceServices::ITable<TraceServices::FLoadRequest>& RequestsTable = LoadTimeProfilerProvider.GetRequestsTable();
 			RequestsTreeView->UpdateSourceTable(MakeShared<FProxyUntypedTable>(&RequestsTable));
 		}
 		else
diff --git a/Engine/Source/Developer/TraceInsights/Private/Insights/MemoryProfiler/MemoryProfilerManager.cpp b/Engine/Source/Developer/TraceInsights/Private/Insights/MemoryProfiler/MemoryProfilerManager.cpp
--- a/Engine/Source/Developer/TraceInsights/Private/Insights/MemoryProfiler/MemoryProfilerManager.cpp
+++ b/Engine/Source/Developer/TraceInsights/Private/Insights/MemoryProfiler/MemoryProfilerManager.cpp
@@ -234,10 +234,10 @@ bool FMemoryProfilerManager::Tick(float DeltaTime)
 	{
 		uint32 TagCount = 0;
 
-		TSharedPtr<const Trace::IAnalysisSession> Session = FInsightsManager::Get()->GetSession();
+		TSharedPtr<const TraceServices::IAnalysisSession> Session = FInsightsManager::Get()->GetSession();
 		if (Session.IsValid())
 		{
-			Trace::FAnalysisSessionReadScope SessionReadScope(*Session.Get());
+			TraceServices::FAnalysisSessionReadScope SessionReadScope(*Session.Get());
 
 			if (Session->IsAnalysisComplete())
 			{
@@ -245,7 +245,7 @@ bool FMemoryProfilerManager::Tick(float DeltaTime)
 				AvailabilityCheck.Disable();
 			}
 
-			const Trace::IMemoryProvider& MemoryProvider = Trace::ReadMemoryProvider(*Session.Get());
+			const TraceServices::IMemoryProvider& MemoryProvider = TraceServices::ReadMemoryProvider(*Session.Get());
 			TagCount = MemoryProvider.GetTagCount();
 		}
 		else
diff --git a/Engine/Source/Developer/TraceInsights/Private/Insights/MemoryProfiler/ViewModels/MemTagNode.cpp b/Engine/Source/Developer/TraceInsights/Private/Insights/MemoryProfiler/ViewModels/MemTagNode.cpp
--- a/Engine/Source/Developer/TraceInsights/Private/Insights/MemoryProfiler/ViewModels/MemTagNode.cpp
+++ b/Engine/Source/Developer/TraceInsights/Private/Insights/MemoryProfiler/ViewModels/MemTagNode.cpp
@@ -32,12 +32,12 @@ FText FMemTagNode::GetTrackerText() const
 
 void FMemTagNode::ResetAggregatedStats()
 {
-	//AggregatedStats = Trace::FMemoryProfilerAggregatedStats();
+	//AggregatedStats = TraceServices::FMemoryProfilerAggregatedStats();
 }
 
 ////////////////////////////////////////////////////////////////////////////////////////////////////
 /*
-void FMemTagNode::SetAggregatedStats(const Trace::FMemoryProfilerAggregatedStats& InAggregatedStats)
+void FMemTagNode::SetAggregatedStats(const TraceServices::FMemoryProfilerAggregatedStats& InAggregatedStats)
 {
 	AggregatedStats = InAggregatedStats;
 }
diff --git a/Engine/Source/Developer/TraceInsights/Private/Insights/MemoryProfiler/ViewModels/MemTagNode.h b/Engine/Source/Developer/TraceInsights/Private/Insights/MemoryProfiler/ViewModels/MemTagNode.h
--- a/Engine/Source/Developer/TraceInsights/Private/Insights/MemoryProfiler/ViewModels/MemTagNode.h
+++ b/Engine/Source/Developer/TraceInsights/Private/Insights/MemoryProfiler/ViewModels/MemTagNode.h
@@ -10,7 +10,7 @@
 #include "Insights/MemoryProfiler/ViewModels/MemoryTag.h"
 #include "Insights/Table/ViewModels/BaseTreeNode.h"
 
-namespace Trace
+namespace TraceServices
 {
 	struct FMemoryProfilerAggregatedStats
 	{
@@ -105,16 +105,16 @@ public:
 	/**
 	 * @return the aggregated stats for this MemTag node.
 	 */
-	const Trace::FMemoryProfilerAggregatedStats& GetAggregatedStats() const { return AggregatedStats; }
+	const TraceServices::FMemoryProfilerAggregatedStats& GetAggregatedStats() const { return AggregatedStats; }
 
 	void ResetAggregatedStats();
-	//TODO: void SetAggregatedStats(const Trace::FMemoryProfilerAggregatedStats& AggregatedStats);
+	//TODO: void SetAggregatedStats(const TraceServices::FMemoryProfilerAggregatedStats& AggregatedStats);
 
 private:
 	const EMemTagNodeType Type;
 	Insights::FMemoryTag* MemTag;
 	FMemTagNodePtr ParentTagNode;
-	Trace::FMemoryProfilerAggregatedStats AggregatedStats;
+	TraceServices::FMemoryProfilerAggregatedStats AggregatedStats;
 };
 
 ////////////////////////////////////////////////////////////////////////////////////////////////////
diff --git a/Engine/Source/Developer/TraceInsights/Private/Insights/MemoryProfiler/ViewModels/MemoryGraphTrack.cpp b/Engine/Source/Developer/TraceInsights/Private/Insights/MemoryProfiler/ViewModels/MemoryGraphTrack.cpp
--- a/Engine/Source/Developer/TraceInsights/Private/Insights/MemoryProfiler/ViewModels/MemoryGraphTrack.cpp
+++ b/Engine/Source/Developer/TraceInsights/Private/Insights/MemoryProfiler/ViewModels/MemoryGraphTrack.cpp
@@ -267,13 +267,13 @@ int32 FMemoryGraphTrack::RemoveAllMemTagSeries()
 
 void FMemoryGraphTrack::PreUpdateMemTagSeries(FMemoryGraphSeries& Series, const FTimingTrackViewport& Viewport)
 {
-	TSharedPtr<const Trace::IAnalysisSession> Session = FInsightsManager::Get()->GetSession();
+	TSharedPtr<const TraceServices::IAnalysisSession> Session = FInsightsManager::Get()->GetSession();
 	if (Session.IsValid())
 	{
-		Trace::FAnalysisSessionReadScope SessionReadScope(*Session.Get());
+		TraceServices::FAnalysisSessionReadScope SessionReadScope(*Session.Get());
 
 		// Compute Min/Max values.
-		const Trace::IMemoryProvider& MemoryProvider = Trace::ReadMemoryProvider(*Session.Get());
+		const TraceServices::IMemoryProvider& MemoryProvider = TraceServices::ReadMemoryProvider(*Session.Get());
 		const uint64 TotalSampleCount = MemoryProvider.GetTagSampleCount(Series.GetTrackerId(), Series.GetTagId());
 		if (TotalSampleCount > 0)
 		{
@@ -281,7 +281,7 @@ void FMemoryGraphTrack::PreUpdateMemTagSeries(FMemoryGraphSeries& Series, const
 			double MaxValue = -std::numeric_limits<double>::infinity();
 
 			MemoryProvider.EnumerateTagSamples(Series.GetTrackerId(), Series.GetTagId(), Viewport.GetStartTime(), Viewport.GetEndTime(), true,
-				[&MinValue, &MaxValue](double Time, double Duration, const Trace::FMemoryTagSample& Sample)
+				[&MinValue, &MaxValue](double Time, double Duration, const TraceServices::FMemoryTagSample& Sample)
 			{
 				const double Value = static_cast<double>(Sample.Value);
 				if (Value < MinValue)
@@ -305,12 +305,12 @@ void FMemoryGraphTrack::UpdateMemTagSeries(FMemoryGraphSeries& Series, const FTi
 {
 	FGraphTrackBuilder Builder(*this, Series, Viewport);
 
-	TSharedPtr<const Trace::IAnalysisSession> Session = FInsightsManager::Get()->GetSession();
+	TSharedPtr<const TraceServices::IAnalysisSession> Session = FInsightsManager::Get()->GetSession();
 	if (Session.IsValid())
 	{
-		Trace::FAnalysisSessionReadScope SessionReadScope(*Session.Get());
+		TraceServices::FAnalysisSessionReadScope SessionReadScope(*Session.Get());
 
-		const Trace::IMemoryProvider& MemoryProvider = Trace::ReadMemoryProvider(*Session.Get());
+		const TraceServices::IMemoryProvider& MemoryProvider = TraceServices::ReadMemoryProvider(*Session.Get());
 
 		const uint64 TotalSampleCount = MemoryProvider.GetTagSampleCount(Series.GetTrackerId(), Series.GetTagId());
 
@@ -325,7 +325,7 @@ void FMemoryGraphTrack::UpdateMemTagSeries(FMemoryGraphSeries& Series, const FTi
 			}
 
 			MemoryProvider.EnumerateTagSamples(Series.GetTrackerId(), Series.GetTagId(), Viewport.GetStartTime(), Viewport.GetEndTime(), true,
-				[this, &Builder](double Time, double Duration, const Trace::FMemoryTagSample& Sample)
+				[this, &Builder](double Time, double Duration, const TraceServices::FMemoryTagSample& Sample)
 			{
 				Builder.AddEvent(Time, Duration, static_cast<double>(Sample.Value));
 			});
diff --git a/Engine/Source/Developer/TraceInsights/Private/Insights/MemoryProfiler/ViewModels/MemorySharedState.cpp b/Engine/Source/Developer/TraceInsights/Private/Insights/MemoryProfiler/ViewModels/MemorySharedState.cpp
--- a/Engine/Source/Developer/TraceInsights/Private/Insights/MemoryProfiler/ViewModels/MemorySharedState.cpp
+++ b/Engine/Source/Developer/TraceInsights/Private/Insights/MemoryProfiler/ViewModels/MemorySharedState.cpp
@@ -102,7 +102,7 @@ void FMemorySharedState::OnEndSession(Insights::ITimingViewSession& InSession)
 
 ////////////////////////////////////////////////////////////////////////////////////////////////////
 
-void FMemorySharedState::Tick(Insights::ITimingViewSession& InSession, const Trace::IAnalysisSession& InAnalysisSession)
+void FMemorySharedState::Tick(Insights::ITimingViewSession& InSession, const TraceServices::IAnalysisSession& InAnalysisSession)
 {
 	if (&InSession != TimingView.Get())
 	{
@@ -230,13 +230,13 @@ void FMemorySharedState::SyncTrackers()
 	CurrentTracker = nullptr;
 	Trackers.Reset();
 
-	TSharedPtr<const Trace::IAnalysisSession> Session = FInsightsManager::Get()->GetSession();
+	TSharedPtr<const TraceServices::IAnalysisSession> Session = FInsightsManager::Get()->GetSession();
 	if (Session.IsValid())
 	{
-		Trace::FAnalysisSessionReadScope SessionReadScope(*Session.Get());
-		const Trace::IMemoryProvider& MemoryProvider = Trace::ReadMemoryProvider(*Session.Get());
+		TraceServices::FAnalysisSessionReadScope SessionReadScope(*Session.Get());
+		const TraceServices::IMemoryProvider& MemoryProvider = TraceServices::ReadMemoryProvider(*Session.Get());
 
-		MemoryProvider.EnumerateTrackers([this](const Trace::FMemoryTracker& Tracker)
+		MemoryProvider.EnumerateTrackers([this](const TraceServices::FMemoryTrackerInfo& Tracker)
 		{
 			Trackers.Add(MakeShared<Insights::FMemoryTracker>(Tracker.Id, Tracker.Name));
 		});
diff --git a/Engine/Source/Developer/TraceInsights/Private/Insights/MemoryProfiler/ViewModels/MemorySharedState.h b/Engine/Source/Developer/TraceInsights/Private/Insights/MemoryProfiler/ViewModels/MemorySharedState.h
--- a/Engine/Source/Developer/TraceInsights/Private/Insights/MemoryProfiler/ViewModels/MemorySharedState.h
+++ b/Engine/Source/Developer/TraceInsights/Private/Insights/MemoryProfiler/ViewModels/MemorySharedState.h
@@ -98,7 +98,7 @@ public:
 	// ITimingViewExtender
 	virtual void OnBeginSession(Insights::ITimingViewSession& InSession) override;
 	virtual void OnEndSession(Insights::ITimingViewSession& InSession) override;
-	virtual void Tick(Insights::ITimingViewSession& InSession, const Trace::IAnalysisSession& InAnalysisSession) override;
+	virtual void Tick(Insights::ITimingViewSession& InSession, const TraceServices::IAnalysisSession& InAnalysisSession) override;
 	virtual void ExtendFilterMenu(Insights::ITimingViewSession& InSession, FMenuBuilder& InMenuBuilder) override;
 
 	bool IsAllMemoryTracksToggleOn() const { return bShowHideAllMemoryTracks; }
diff --git a/Engine/Source/Developer/TraceInsights/Private/Insights/MemoryProfiler/ViewModels/MemoryTag.cpp b/Engine/Source/Developer/TraceInsights/Private/Insights/MemoryProfiler/ViewModels/MemoryTag.cpp
--- a/Engine/Source/Developer/TraceInsights/Private/Insights/MemoryProfiler/ViewModels/MemoryTag.cpp
+++ b/Engine/Source/Developer/TraceInsights/Private/Insights/MemoryProfiler/ViewModels/MemoryTag.cpp
@@ -125,11 +125,11 @@ void FMemoryTagList::UpdateInternal()
 {
 	bool bUpdated = false;
 
-	TSharedPtr<const Trace::IAnalysisSession> Session = FInsightsManager::Get()->GetSession();
+	TSharedPtr<const TraceServices::IAnalysisSession> Session = FInsightsManager::Get()->GetSession();
 	if (Session.IsValid())
 	{
-		Trace::FAnalysisSessionReadScope SessionReadScope(*Session.Get());
-		const Trace::IMemoryProvider& MemoryProvider = Trace::ReadMemoryProvider(*Session.Get());
+		TraceServices::FAnalysisSessionReadScope SessionReadScope(*Session.Get());
+		const TraceServices::IMemoryProvider& MemoryProvider = TraceServices::ReadMemoryProvider(*Session.Get());
 
 		const int32 TraceSerialNumber = MemoryProvider.GetTagSerial();
 		if (LastTraceSerialNumber != TraceSerialNumber)
@@ -138,9 +138,9 @@ void FMemoryTagList::UpdateInternal()
 			++SerialNumber;
 			bUpdated = true;
 
-			MemoryProvider.EnumerateTags([this](const Trace::FMemoryTag& TraceTag)
+			MemoryProvider.EnumerateTags([this](const TraceServices::FMemoryTagInfo& TraceTag)
 			{
-				static_assert(FMemoryTag::InvalidTagId == static_cast<FMemoryTagId>(Trace::FMemoryTag::InvalidTagId), "Memory TagId type mismatch!");
+				static_assert(FMemoryTag::InvalidTagId == static_cast<FMemoryTagId>(TraceServices::FMemoryTagInfo::InvalidTagId), "Memory TagId type mismatch!");
 				FMemoryTagId TagId = static_cast<FMemoryTagId>(TraceTag.Id);
 
 				FMemoryTag* TagPtr = TagIdMap.FindRef(TagId);
diff --git a/Engine/Source/Developer/TraceInsights/Private/Insights/MemoryProfiler/ViewModels/MemoryTracker.cpp b/Engine/Source/Developer/TraceInsights/Private/Insights/MemoryProfiler/ViewModels/MemoryTracker.cpp
--- a/Engine/Source/Developer/TraceInsights/Private/Insights/MemoryProfiler/ViewModels/MemoryTracker.cpp
+++ b/Engine/Source/Developer/TraceInsights/Private/Insights/MemoryProfiler/ViewModels/MemoryTracker.cpp
@@ -17,7 +17,7 @@ FMemoryTracker::FMemoryTracker(FMemoryTrackerId InTrackerId, const FString InTra
 	: Id(InTrackerId)
 	, Name(InTrackerName)
 {
-	//TODO: static_assert(FMemoryTracker::InvalidTrackerId == static_cast<FMemoryTracker>(Trace::FMemoryTrackerDesc::InvalidTrackerId), "Memory TrackerId type mismatch!");
+	//TODO: static_assert(FMemoryTracker::InvalidTrackerId == static_cast<FMemoryTracker>(FMemoryTrackerDesc::InvalidTrackerId), "Memory TrackerId type mismatch!");
 }
 
 ////////////////////////////////////////////////////////////////////////////////////////////////////
diff --git a/Engine/Source/Developer/TraceInsights/Private/Insights/MemoryProfiler/Widgets/SMemAllocTableTreeView.cpp b/Engine/Source/Developer/TraceInsights/Private/Insights/MemoryProfiler/Widgets/SMemAllocTableTreeView.cpp
--- a/Engine/Source/Developer/TraceInsights/Private/Insights/MemoryProfiler/Widgets/SMemAllocTableTreeView.cpp
+++ b/Engine/Source/Developer/TraceInsights/Private/Insights/MemoryProfiler/Widgets/SMemAllocTableTreeView.cpp
@@ -49,7 +49,7 @@ void SMemAllocTableTreeView::Reset()
 
 ////////////////////////////////////////////////////////////////////////////////////////////////////
 /*
-void SMemAllocTableTreeView::UpdateSourceTable(TSharedPtr<Trace::IMemAllocTable> SourceTable)
+void SMemAllocTableTreeView::UpdateSourceTable(TSharedPtr<TraceServices::IMemAllocTable> SourceTable)
 {
 	//check(Table->Is<Insights::FMemAllocTable>());
 	TSharedPtr<Insights::FMemAllocTable> MemAllocTable = StaticCastSharedPtr<Insights::FMemAllocTable>(Table);
@@ -198,7 +198,7 @@ void SMemAllocTableTreeView::StartQuery()
 		return;
 	}
 
-	const IAllocationsProvider* AllocationsProvider = Trace::ReadAllocationsProvider(*Session.Get());
+	const IAllocationsProvider* AllocationsProvider = TraceServices::ReadAllocationsProvider(*Session.Get());
 	if (!AllocationsProvider)
 	{
 		UE_LOG(MemoryProfiler, Warning, TEXT("[MemAlloc] Invalid allocations provider!"));
@@ -295,7 +295,7 @@ void SMemAllocTableTreeView::UpdateQuery()
 #if defined(UE_USE_ALLOCATIONS_PROVIDER)
 	if (Query != 0)
 	{
-		const IAllocationsProvider* AllocationsProvider = Trace::ReadAllocationsProvider(*Session.Get());
+		const IAllocationsProvider* AllocationsProvider = TraceServices::ReadAllocationsProvider(*Session.Get());
 		if (!AllocationsProvider)
 		{
 			UE_LOG(MemoryProfiler, Warning, TEXT("[MemAlloc] Invalid allocations provider!"));
@@ -394,7 +394,7 @@ void SMemAllocTableTreeView::CancelQuery()
 #if defined(UE_USE_ALLOCATIONS_PROVIDER)
 	if (Query != 0)
 	{
-		const IAllocationsProvider* AllocationsProvider = Trace::ReadAllocationsProvider(*Session.Get());
+		const IAllocationsProvider* AllocationsProvider = TraceServices::ReadAllocationsProvider(*Session.Get());
 		if (AllocationsProvider)
 		{
 			AllocationsProvider->CancelQuery(Query);
diff --git a/Engine/Source/Developer/TraceInsights/Private/Insights/MemoryProfiler/Widgets/SMemAllocTableTreeView.h b/Engine/Source/Developer/TraceInsights/Private/Insights/MemoryProfiler/Widgets/SMemAllocTableTreeView.h
--- a/Engine/Source/Developer/TraceInsights/Private/Insights/MemoryProfiler/Widgets/SMemAllocTableTreeView.h
+++ b/Engine/Source/Developer/TraceInsights/Private/Insights/MemoryProfiler/Widgets/SMemAllocTableTreeView.h
@@ -38,7 +38,7 @@ public:
 	TSharedPtr<FMemAllocTable> GetMemAllocTable() { return StaticCastSharedPtr<FMemAllocTable>(GetTable()); }
 	const TSharedPtr<FMemAllocTable> GetMemAllocTable() const { return StaticCastSharedPtr<FMemAllocTable>(GetTable()); }
 
-	//void UpdateSourceTable(TSharedPtr<Trace::IMemAllocTable> SourceTable);
+	//void UpdateSourceTable(TSharedPtr<TraceServices::IMemAllocTable> SourceTable);
 
 	virtual void Reset();
 
diff --git a/Engine/Source/Developer/TraceInsights/Private/Insights/MemoryProfiler/Widgets/SMemTagTreeView.cpp b/Engine/Source/Developer/TraceInsights/Private/Insights/MemoryProfiler/Widgets/SMemTagTreeView.cpp
--- a/Engine/Source/Developer/TraceInsights/Private/Insights/MemoryProfiler/Widgets/SMemTagTreeView.cpp
+++ b/Engine/Source/Developer/TraceInsights/Private/Insights/MemoryProfiler/Widgets/SMemTagTreeView.cpp
@@ -1148,7 +1148,7 @@ TSharedRef<SWidget> SMemTagTreeView::TreeViewHeaderRow_GenerateColumnMenu(const
 
 void SMemTagTreeView::InsightsManager_OnSessionChanged()
 {
-	TSharedPtr<const Trace::IAnalysisSession> NewSession = FInsightsManager::Get()->GetSession();
+	TSharedPtr<const TraceServices::IAnalysisSession> NewSession = FInsightsManager::Get()->GetSession();
 
 	if (NewSession != Session)
 	{
@@ -2231,22 +2231,22 @@ void SMemTagTreeView::UpdateStatsInternal()
 	/*
 	if (Session.IsValid())
 	{
-		TUniquePtr<Trace::ITable<Trace::FMemoryProfilerAggregatedStats>> AggregationResultTable;
+		TUniquePtr<TraceServices::ITable<TraceServices::FMemoryProfilerAggregatedStats>> AggregationResultTable;
 
 		AggregationStopwatch.Start();
 		{
-			Trace::FAnalysisSessionReadScope SessionReadScope(*Session.Get());
-			const Trace::IMemoryProfilerProvider& MemoryProfilerProvider = Trace::ReadMemoryProfilerProvider(*Session.Get());
+			TraceServices::FAnalysisSessionReadScope SessionReadScope(*Session.Get());
+			const TraceServices::IMemoryProfilerProvider& MemoryProfilerProvider = TraceServices::ReadMemoryProfilerProvider(*Session.Get());
 			AggregationResultTable.Reset(MemoryProfilerProvider.CreateAggregation(StatsStartTime, StatsEndTime));
 		}
 		AggregationStopwatch.Stop();
 
 		if (AggregationResultTable.IsValid())
 		{
-			TUniquePtr<Trace::ITableReader<Trace::FMemoryProfilerAggregatedStats>> TableReader(AggregationResultTable->CreateReader());
+			TUniquePtr<TraceServices::ITableReader<TraceServices::FMemoryProfilerAggregatedStats>> TableReader(AggregationResultTable->CreateReader());
 			while (TableReader->IsValid())
 			{
-				const Trace::FMemoryProfilerAggregatedStats* Row = TableReader->GetCurrentRow();
+				const TraceServices::FMemoryProfilerAggregatedStats* Row = TableReader->GetCurrentRow();
 				FMemTagNodePtr* MemTagNodePtrPtr = MemTagNodesIdMap.Find(static_cast<uint64>(Row->EventTypeIndex));
 				if (MemTagNodePtrPtr != nullptr)
 				{
diff --git a/Engine/Source/Developer/TraceInsights/Private/Insights/MemoryProfiler/Widgets/SMemTagTreeView.h b/Engine/Source/Developer/TraceInsights/Private/Insights/MemoryProfiler/Widgets/SMemTagTreeView.h
--- a/Engine/Source/Developer/TraceInsights/Private/Insights/MemoryProfiler/Widgets/SMemTagTreeView.h
+++ b/Engine/Source/Developer/TraceInsights/Private/Insights/MemoryProfiler/Widgets/SMemTagTreeView.h
@@ -25,7 +25,7 @@ class FMenuBuilder;
 class FMemoryGraphTrack;
 class SMemoryProfilerWindow;
 
-namespace Trace
+namespace TraceServices
 {
 	class IAnalysisSession;
 }
@@ -332,7 +332,7 @@ private:
 	TSharedPtr<Insights::FTable> Table;
 
 	/** A weak pointer to the profiler session used to populate this widget. */
-	TSharedPtr<const Trace::IAnalysisSession>/*Weak*/ Session;
+	TSharedPtr<const TraceServices::IAnalysisSession>/*Weak*/ Session;
 
 	//////////////////////////////////////////////////
 	// Tree View, Columns
diff --git a/Engine/Source/Developer/TraceInsights/Private/Insights/MemoryProfiler/Widgets/SMemTagTreeViewTooltip.cpp b/Engine/Source/Developer/TraceInsights/Private/Insights/MemoryProfiler/Widgets/SMemTagTreeViewTooltip.cpp
--- a/Engine/Source/Developer/TraceInsights/Private/Insights/MemoryProfiler/Widgets/SMemTagTreeViewTooltip.cpp
+++ b/Engine/Source/Developer/TraceInsights/Private/Insights/MemoryProfiler/Widgets/SMemTagTreeViewTooltip.cpp
@@ -91,7 +91,7 @@ TSharedPtr<SToolTip> SMemTagTreeViewTooltip::GetRowTooltip(const TSharedPtr<FMem
 {
 	const FText TrackersText = MemTagNodePtr->GetTrackerText();
 
-	const Trace::FMemoryProfilerAggregatedStats& Stats = MemTagNodePtr->GetAggregatedStats();
+	const TraceServices::FMemoryProfilerAggregatedStats& Stats = MemTagNodePtr->GetAggregatedStats();
 	const FText InstanceCountText = FText::AsNumber(Stats.InstanceCount);
 	const FText MinValueText = FText::AsNumber(Stats.Min);
 	const FText MaxValueText = FText::AsNumber(Stats.Max);
diff --git a/Engine/Source/Developer/TraceInsights/Private/Insights/MemoryProfiler/Widgets/SMemoryProfilerWindow.cpp b/Engine/Source/Developer/TraceInsights/Private/Insights/MemoryProfiler/Widgets/SMemoryProfilerWindow.cpp
--- a/Engine/Source/Developer/TraceInsights/Private/Insights/MemoryProfiler/Widgets/SMemoryProfilerWindow.cpp
+++ b/Engine/Source/Developer/TraceInsights/Private/Insights/MemoryProfiler/Widgets/SMemoryProfilerWindow.cpp
@@ -187,16 +187,16 @@ void SMemoryProfilerWindow::UpdateMemTagTreeView()
 	if (MemTagTreeView)
 	{
 		/*
-		TSharedPtr<const Trace::IAnalysisSession> Session = FInsightsManager::Get()->GetSession();
-		if (Session.IsValid() && Trace::ReadMemoryProfilerProvider(*Session.Get()))
+		TSharedPtr<const TraceServices::IAnalysisSession> Session = FInsightsManager::Get()->GetSession();
+		if (Session.IsValid() && TraceServices::ReadMemoryProfilerProvider(*Session.Get()))
 		{
-			Trace::FAnalysisSessionReadScope SessionReadScope(*Session.Get());
-			const Trace::IMemoryProfilerProvider& MemoryProfilerProvider = *Trace::ReadMemoryProfilerProvider(*Session.Get());
+			TraceServices::FAnalysisSessionReadScope SessionReadScope(*Session.Get());
+			const TraceServices::IMemoryProfilerProvider& MemoryProfilerProvider = *TraceServices::ReadMemoryProfilerProvider(*Session.Get());
 
 			const double SelectionStartTime = TimingView ? TimingView->GetSelectionStartTime() : 0.0;
 			const double SelectionEndTime = TimingView ? TimingView->GetSelectionEndTime() : 0.0;
 
-			Trace::ITable<Trace::FMemoryProfilerAggregatedStats>* EventAggregationTable = MemoryProfilerProvider.CreateEventAggregation(SelectionStartTime, SelectionEndTime);
+			TraceServices::ITable<TraceServices::FMemoryProfilerAggregatedStats>* EventAggregationTable = MemoryProfilerProvider.CreateEventAggregation(SelectionStartTime, SelectionEndTime);
 			MemTagTreeView->UpdateSourceTable(MakeShareable(EventAggregationTable));
 		}
 		else
diff --git a/Engine/Source/Developer/TraceInsights/Private/Insights/NetworkingProfiler/NetworkingProfilerManager.cpp b/Engine/Source/Developer/TraceInsights/Private/Insights/NetworkingProfiler/NetworkingProfilerManager.cpp
--- a/Engine/Source/Developer/TraceInsights/Private/Insights/NetworkingProfiler/NetworkingProfilerManager.cpp
+++ b/Engine/Source/Developer/TraceInsights/Private/Insights/NetworkingProfiler/NetworkingProfilerManager.cpp
@@ -224,10 +224,10 @@ bool FNetworkingProfilerManager::Tick(float DeltaTime)
 	{
 		uint32 NetTraceVersion = 0;
 
-		TSharedPtr<const Trace::IAnalysisSession> Session = FInsightsManager::Get()->GetSession();
+		TSharedPtr<const TraceServices::IAnalysisSession> Session = FInsightsManager::Get()->GetSession();
 		if (Session.IsValid())
 		{
-			Trace::FAnalysisSessionReadScope SessionReadScope(*Session.Get());
+			TraceServices::FAnalysisSessionReadScope SessionReadScope(*Session.Get());
 
 			if (Session->IsAnalysisComplete())
 			{
@@ -235,7 +235,7 @@ bool FNetworkingProfilerManager::Tick(float DeltaTime)
 				AvailabilityCheck.Disable();
 			}
 
-			const Trace::INetProfilerProvider& NetProfilerProvider = Trace::ReadNetProfilerProvider(*Session.Get());
+			const TraceServices::INetProfilerProvider& NetProfilerProvider = TraceServices::ReadNetProfilerProvider(*Session.Get());
 			NetTraceVersion = NetProfilerProvider.GetNetTraceVersion();
 		}
 		else
diff --git a/Engine/Source/Developer/TraceInsights/Private/Insights/NetworkingProfiler/ViewModels/NetEventNode.cpp b/Engine/Source/Developer/TraceInsights/Private/Insights/NetworkingProfiler/ViewModels/NetEventNode.cpp
--- a/Engine/Source/Developer/TraceInsights/Private/Insights/NetworkingProfiler/ViewModels/NetEventNode.cpp
+++ b/Engine/Source/Developer/TraceInsights/Private/Insights/NetworkingProfiler/ViewModels/NetEventNode.cpp
@@ -10,12 +10,12 @@ const FName FNetEventNode::TypeName(TEXT("FNetEventNode"));
 
 void FNetEventNode::ResetAggregatedStats()
 {
-	AggregatedStats = Trace::FNetProfilerAggregatedStats();
+	AggregatedStats = TraceServices::FNetProfilerAggregatedStats();
 }
 
 ////////////////////////////////////////////////////////////////////////////////////////////////////
 
-void FNetEventNode::SetAggregatedStats(const Trace::FNetProfilerAggregatedStats& InAggregatedStats)
+void FNetEventNode::SetAggregatedStats(const TraceServices::FNetProfilerAggregatedStats& InAggregatedStats)
 {
 	AggregatedStats = InAggregatedStats;
 }
diff --git a/Engine/Source/Developer/TraceInsights/Private/Insights/NetworkingProfiler/ViewModels/NetEventNode.h b/Engine/Source/Developer/TraceInsights/Private/Insights/NetworkingProfiler/ViewModels/NetEventNode.h
--- a/Engine/Source/Developer/TraceInsights/Private/Insights/NetworkingProfiler/ViewModels/NetEventNode.h
+++ b/Engine/Source/Developer/TraceInsights/Private/Insights/NetworkingProfiler/ViewModels/NetEventNode.h
@@ -87,10 +87,10 @@ public:
 	/**
 	 * @return the aggregated stats for this NetEvent node.
 	 */
-	const Trace::FNetProfilerAggregatedStats& GetAggregatedStats() const { return AggregatedStats; }
+	const TraceServices::FNetProfilerAggregatedStats& GetAggregatedStats() const { return AggregatedStats; }
 
 	void ResetAggregatedStats();
-	void SetAggregatedStats(const Trace::FNetProfilerAggregatedStats& AggregatedStats);
+	void SetAggregatedStats(const TraceServices::FNetProfilerAggregatedStats& AggregatedStats);
 
 private:
 	/** The NetEvent type index. */
@@ -102,7 +102,7 @@ private:
 	uint32 Level;
 
 	/** Aggregated stats. */
-	Trace::FNetProfilerAggregatedStats AggregatedStats;
+	TraceServices::FNetProfilerAggregatedStats AggregatedStats;
 };
 
 ////////////////////////////////////////////////////////////////////////////////////////////////////
diff --git a/Engine/Source/Developer/TraceInsights/Private/Insights/NetworkingProfiler/ViewModels/PacketContentViewDrawHelper.cpp b/Engine/Source/Developer/TraceInsights/Private/Insights/NetworkingProfiler/ViewModels/PacketContentViewDrawHelper.cpp
--- a/Engine/Source/Developer/TraceInsights/Private/Insights/NetworkingProfiler/ViewModels/PacketContentViewDrawHelper.cpp
+++ b/Engine/Source/Developer/TraceInsights/Private/Insights/NetworkingProfiler/ViewModels/PacketContentViewDrawHelper.cpp
@@ -34,7 +34,7 @@ FPacketContentViewDrawStateBuilder::FPacketContentViewDrawStateBuilder(FPacketCo
 
 ////////////////////////////////////////////////////////////////////////////////////////////////////
 
-void FPacketContentViewDrawStateBuilder::AddEvent(const Trace::FNetProfilerContentEvent& Event, const TCHAR* EventName, uint32 NetId)
+void FPacketContentViewDrawStateBuilder::AddEvent(const TraceServices::FNetProfilerContentEvent& Event, const TCHAR* EventName, uint32 NetId)
 {
 	DrawState.Events.AddUninitialized();
 	FNetworkPacketEvent& PacketEvent = DrawState.Events.Last();
@@ -186,7 +186,7 @@ void FPacketContentViewDrawStateBuilder::AddEvent(const Trace::FNetProfilerConte
 		Builder.Appendf(TEXT("%s"), EventName ? EventName : TEXT("?"));
 		if (EventW > Builder.Len() * 2.0f + 48.0f)
 		{
-			const Trace::FNetProfilerBunchInfo& Info = Event.BunchInfo;
+			const TraceServices::FNetProfilerBunchInfo& Info = Event.BunchInfo;
 			if (Info.bIsValid)
 			{
 				Builder.Appendf(TEXT(" ChannelId:%u"), Info.ChannelIndex);
@@ -201,7 +201,7 @@ void FPacketContentViewDrawStateBuilder::AddEvent(const Trace::FNetProfilerConte
 				}
 				else if (Event.BunchInfo.bClose)
 				{
-					Builder.Appendf(TEXT(" | Close: %s"), LexToString(Trace::ENetProfilerChannelCloseReason(Info.ChannelCloseReason)));		
+					Builder.Appendf(TEXT(" | Close: %s"), LexToString(TraceServices::ENetProfilerChannelCloseReason(Info.ChannelCloseReason)));		
 				}
 				if (Event.BunchInfo.bReliable)
 				{
diff --git a/Engine/Source/Developer/TraceInsights/Private/Insights/NetworkingProfiler/ViewModels/PacketContentViewDrawHelper.h b/Engine/Source/Developer/TraceInsights/Private/Insights/NetworkingProfiler/ViewModels/PacketContentViewDrawHelper.h
--- a/Engine/Source/Developer/TraceInsights/Private/Insights/NetworkingProfiler/ViewModels/PacketContentViewDrawHelper.h
+++ b/Engine/Source/Developer/TraceInsights/Private/Insights/NetworkingProfiler/ViewModels/PacketContentViewDrawHelper.h
@@ -167,7 +167,7 @@ public:
 	FPacketContentViewDrawStateBuilder(const FPacketContentViewDrawStateBuilder&) = delete;
 	FPacketContentViewDrawStateBuilder& operator=(const FPacketContentViewDrawStateBuilder&) = delete;
 
-	void AddEvent(const Trace::FNetProfilerContentEvent& Event, const TCHAR* Name, uint32 NetId);
+	void AddEvent(const TraceServices::FNetProfilerContentEvent& Event, const TCHAR* Name, uint32 NetId);
 	void Flush();
 
 	int32 GetMaxDepth() const { return MaxDepth; }
diff --git a/Engine/Source/Developer/TraceInsights/Private/Insights/NetworkingProfiler/ViewModels/PacketViewDrawHelper.cpp b/Engine/Source/Developer/TraceInsights/Private/Insights/NetworkingProfiler/ViewModels/PacketViewDrawHelper.cpp
--- a/Engine/Source/Developer/TraceInsights/Private/Insights/NetworkingProfiler/ViewModels/PacketViewDrawHelper.cpp
+++ b/Engine/Source/Developer/TraceInsights/Private/Insights/NetworkingProfiler/ViewModels/PacketViewDrawHelper.cpp
@@ -18,7 +18,7 @@
 // FNetworkPacketAggregatedSample
 ////////////////////////////////////////////////////////////////////////////////////////////////////
 
-void FNetworkPacketAggregatedSample::AddPacket(const int32 PacketIndex, const Trace::FNetProfilerPacket& Packet)
+void FNetworkPacketAggregatedSample::AddPacket(const int32 PacketIndex, const TraceServices::FNetProfilerPacket& Packet)
 {
 	NumPackets++;
 
@@ -44,17 +44,17 @@ void FNetworkPacketAggregatedSample::AddPacket(const int32 PacketIndex, const Tr
 
 	switch (AggregatedStatus)
 	{
-	case Trace::ENetProfilerDeliveryStatus::Unknown:
+	case TraceServices::ENetProfilerDeliveryStatus::Unknown:
 		AggregatedStatus = Packet.DeliveryStatus;
 		break;
 
-	case Trace::ENetProfilerDeliveryStatus::Dropped:
+	case TraceServices::ENetProfilerDeliveryStatus::Dropped:
 		break;
 
-	case Trace::ENetProfilerDeliveryStatus::Delivered:
-		if (Packet.DeliveryStatus == Trace::ENetProfilerDeliveryStatus::Dropped)
+	case TraceServices::ENetProfilerDeliveryStatus::Delivered:
+		if (Packet.DeliveryStatus == TraceServices::ENetProfilerDeliveryStatus::Dropped)
 		{
-			AggregatedStatus = Trace::ENetProfilerDeliveryStatus::Dropped;
+			AggregatedStatus = TraceServices::ENetProfilerDeliveryStatus::Dropped;
 		}
 		break;
 	}
@@ -81,7 +81,7 @@ FNetworkPacketSeriesBuilder::FNetworkPacketSeriesBuilder(FNetworkPacketSeries& I
 
 ////////////////////////////////////////////////////////////////////////////////////////////////////
 
-FNetworkPacketAggregatedSample* FNetworkPacketSeriesBuilder::AddPacket(const int32 PacketIndex, const Trace::FNetProfilerPacket& Packet)
+FNetworkPacketAggregatedSample* FNetworkPacketSeriesBuilder::AddPacket(const int32 PacketIndex, const TraceServices::FNetProfilerPacket& Packet)
 {
 	NumAddedPackets++;
 
@@ -133,18 +133,18 @@ void FPacketViewDrawHelper::DrawBackground() const
 
 ////////////////////////////////////////////////////////////////////////////////////////////////////
 
-FLinearColor FPacketViewDrawHelper::GetColorByStatus(Trace::ENetProfilerDeliveryStatus Status)
+FLinearColor FPacketViewDrawHelper::GetColorByStatus(TraceServices::ENetProfilerDeliveryStatus Status)
 {
 	constexpr float Alpha = 1.0f;
 	switch (Status)
 	{
-	case Trace::ENetProfilerDeliveryStatus::Unknown:
+	case TraceServices::ENetProfilerDeliveryStatus::Unknown:
 		return FLinearColor(0.25f, 0.25f, 0.25f, Alpha);
 
-	case Trace::ENetProfilerDeliveryStatus::Delivered:
+	case TraceServices::ENetProfilerDeliveryStatus::Delivered:
 		return FLinearColor(0.5f, 1.0f, 0.5f, Alpha);
 
-	case Trace::ENetProfilerDeliveryStatus::Dropped:
+	case TraceServices::ENetProfilerDeliveryStatus::Dropped:
 		return FLinearColor(1.0f, 0.5f, 0.5f, Alpha);
 
 	default:
diff --git a/Engine/Source/Developer/TraceInsights/Private/Insights/NetworkingProfiler/ViewModels/PacketViewDrawHelper.h b/Engine/Source/Developer/TraceInsights/Private/Insights/NetworkingProfiler/ViewModels/PacketViewDrawHelper.h
--- a/Engine/Source/Developer/TraceInsights/Private/Insights/NetworkingProfiler/ViewModels/PacketViewDrawHelper.h
+++ b/Engine/Source/Developer/TraceInsights/Private/Insights/NetworkingProfiler/ViewModels/PacketViewDrawHelper.h
@@ -27,7 +27,7 @@ struct FNetworkPacket
 	uint32 ContentSizeInBits;
 	uint32 TotalSizeInBytes;
 	double TimeStamp;
-	Trace::ENetProfilerDeliveryStatus Status;
+	TraceServices::ENetProfilerDeliveryStatus Status;
 
 	FNetworkPacket()
 		: Index(-1)
@@ -35,7 +35,7 @@ struct FNetworkPacket
 		, ContentSizeInBits(0)
 		, TotalSizeInBytes(0)
 		, TimeStamp(std::numeric_limits<double>::infinity())
-		, Status(Trace::ENetProfilerDeliveryStatus::Unknown)
+		, Status(TraceServices::ENetProfilerDeliveryStatus::Unknown)
 	{}
 
 	FNetworkPacket(const FNetworkPacket&) = default;
@@ -75,7 +75,7 @@ struct FNetworkPacketAggregatedSample
 	 *    Received --> at least one packet in the sample set is confirmed received and none are confirmed lost
 	 *    Lost     --> at least one packet in the sample set is confirmed lost
 	**/
-	Trace::ENetProfilerDeliveryStatus AggregatedStatus;
+	TraceServices::ENetProfilerDeliveryStatus AggregatedStatus;
 
 	FNetworkPacket LargestPacket;
 
@@ -85,7 +85,7 @@ struct FNetworkPacketAggregatedSample
 		: NumPackets(0)
 		, StartTime(DBL_MAX)
 		, EndTime(-DBL_MAX)
-		, AggregatedStatus(Trace::ENetProfilerDeliveryStatus::Unknown)
+		, AggregatedStatus(TraceServices::ENetProfilerDeliveryStatus::Unknown)
 		, LargestPacket()
 		, bAtLeastOnePacketMatchesFilter(true)
 	{}
@@ -96,7 +96,7 @@ struct FNetworkPacketAggregatedSample
 	FNetworkPacketAggregatedSample(FNetworkPacketAggregatedSample&&) = default;
 	FNetworkPacketAggregatedSample& operator=(FNetworkPacketAggregatedSample&&) = default;
 
-	void AddPacket(const int32 PacketIndex, const Trace::FNetProfilerPacket& Packet);
+	void AddPacket(const int32 PacketIndex, const TraceServices::FNetProfilerPacket& Packet);
 
 	bool Equals(const FNetworkPacketAggregatedSample& Other) const
 	{
@@ -147,7 +147,7 @@ public:
 	FNetworkPacketSeriesBuilder(const FNetworkPacketSeriesBuilder&) = delete;
 	FNetworkPacketSeriesBuilder& operator=(const FNetworkPacketSeriesBuilder&) = delete;
 
-	FNetworkPacketAggregatedSample* AddPacket(int32 PacketIndex, const Trace::FNetProfilerPacket& Packet);
+	FNetworkPacketAggregatedSample* AddPacket(int32 PacketIndex, const TraceServices::FNetProfilerPacket& Packet);
 
 	int32 GetNumAddedPackets() const { return NumAddedPackets; }
 
@@ -190,7 +190,7 @@ public:
 	void DrawSampleHighlight(const FNetworkPacketAggregatedSample& Sample, EHighlightMode Mode) const;
 	void DrawSelection(int32 StartPacketIndex, int32 EndPacketIndex) const;
 
-	static FLinearColor GetColorByStatus(Trace::ENetProfilerDeliveryStatus Status);
+	static FLinearColor GetColorByStatus(TraceServices::ENetProfilerDeliveryStatus Status);
 
 	int32 GetNumPackets() const { return NumPackets; }
 	int32 GetNumDrawSamples() const { return NumDrawSamples; }
diff --git a/Engine/Source/Developer/TraceInsights/Private/Insights/NetworkingProfiler/Widgets/SNetStatsView.cpp b/Engine/Source/Developer/TraceInsights/Private/Insights/NetworkingProfiler/Widgets/SNetStatsView.cpp
--- a/Engine/Source/Developer/TraceInsights/Private/Insights/NetworkingProfiler/Widgets/SNetStatsView.cpp
+++ b/Engine/Source/Developer/TraceInsights/Private/Insights/NetworkingProfiler/Widgets/SNetStatsView.cpp
@@ -45,7 +45,7 @@ SNetStatsView::SNetStatsView()
 	, ObjectsChangeCount(0)
 	, GameInstanceIndex(0)
 	, ConnectionIndex(0)
-	, ConnectionMode(Trace::ENetProfilerConnectionMode::Outgoing)
+	, ConnectionMode(TraceServices::ENetProfilerConnectionMode::Outgoing)
 	, StatsPacketStartIndex(0)
 	, StatsPacketEndIndex(0)
 	, StatsStartPosition(0)
@@ -613,7 +613,7 @@ TSharedRef<SWidget> SNetStatsView::TreeViewHeaderRow_GenerateColumnMenu(const In
 
 void SNetStatsView::InsightsManager_OnSessionChanged()
 {
-	TSharedPtr<const Trace::IAnalysisSession> NewSession = FInsightsManager::Get()->GetSession();
+	TSharedPtr<const TraceServices::IAnalysisSession> NewSession = FInsightsManager::Get()->GetSession();
 
 	if (NewSession != Session)
 	{
@@ -1539,7 +1539,7 @@ void SNetStatsView::Reset()
 {
 	GameInstanceIndex = 0;
 	ConnectionIndex = 0;
-	ConnectionMode = Trace::ENetProfilerConnectionMode::Outgoing;
+	ConnectionMode = TraceServices::ENetProfilerConnectionMode::Outgoing;
 
 	StatsPacketStartIndex = 0;
 	StatsPacketEndIndex = 0;
@@ -1582,11 +1582,11 @@ void SNetStatsView::RebuildTree(bool bResync)
 	SyncStopwatch.Start();
 	if (Session.IsValid())
 	{
-		Trace::FAnalysisSessionReadScope SessionReadScope(*Session.Get());
+		TraceServices::FAnalysisSessionReadScope SessionReadScope(*Session.Get());
 
-		const Trace::INetProfilerProvider& NetProfilerProvider = Trace::ReadNetProfilerProvider(*Session.Get());
+		const TraceServices::INetProfilerProvider& NetProfilerProvider = TraceServices::ReadNetProfilerProvider(*Session.Get());
 
-		NetProfilerProvider.ReadEventTypes([this, &bResync, &NetProfilerProvider](const Trace::FNetProfilerEventType* NetEvents, uint64 InNetEventCount)
+		NetProfilerProvider.ReadEventTypes([this, &bResync, &NetProfilerProvider](const TraceServices::FNetProfilerEventType* NetEvents, uint64 InNetEventCount)
 		{
 			const uint32 NetEventCount = static_cast<uint32>(InNetEventCount);
 			if (NetEventCount != NetEventNodes.Num())
@@ -1594,10 +1594,10 @@ void SNetStatsView::RebuildTree(bool bResync)
 				NetEventNodes.Empty(NetEventCount);
 				for (uint32 Index = 0; Index < NetEventCount; ++Index)
 				{
-					const Trace::FNetProfilerEventType& NetEvent = NetEvents[Index];
+					const TraceServices::FNetProfilerEventType& NetEvent = NetEvents[Index];
 					ensure(NetEvent.EventTypeIndex == Index);
 					const TCHAR* NamePtr = nullptr;
-					NetProfilerProvider.ReadName(NetEvent.NameIndex, [&NamePtr](const Trace::FNetProfilerName& InName)
+					NetProfilerProvider.ReadName(NetEvent.NameIndex, [&NamePtr](const TraceServices::FNetProfilerName& InName)
 					{
 						NamePtr = InName.Name;
 					});
@@ -1656,7 +1656,7 @@ void SNetStatsView::ResetStats()
 {
 	GameInstanceIndex = 0;
 	ConnectionIndex = 0;
-	ConnectionMode = Trace::ENetProfilerConnectionMode::Outgoing;
+	ConnectionMode = TraceServices::ENetProfilerConnectionMode::Outgoing;
 
 	StatsPacketStartIndex = 0;
 	StatsPacketEndIndex = 0;
@@ -1668,7 +1668,7 @@ void SNetStatsView::ResetStats()
 
 ////////////////////////////////////////////////////////////////////////////////////////////////////
 
-void SNetStatsView::UpdateStats(uint32 InGameInstanceIndex, uint32 InConnectionIndex, Trace::ENetProfilerConnectionMode InConnectionMode, uint32 InStatsPacketStartIndex, uint32 InStatsPacketEndIndex, uint32 InStatsStartPosition, uint32 InStatsEndPosition)
+void SNetStatsView::UpdateStats(uint32 InGameInstanceIndex, uint32 InConnectionIndex, TraceServices::ENetProfilerConnectionMode InConnectionMode, uint32 InStatsPacketStartIndex, uint32 InStatsPacketEndIndex, uint32 InStatsStartPosition, uint32 InStatsEndPosition)
 {
 	GameInstanceIndex = InGameInstanceIndex;
 	ConnectionIndex = InConnectionIndex;
@@ -1704,12 +1704,12 @@ void SNetStatsView::UpdateStatsInternal()
 
 	if (Session.IsValid())
 	{
-		TUniquePtr<Trace::ITable<Trace::FNetProfilerAggregatedStats>> AggregationResultTable;
+		TUniquePtr<TraceServices::ITable<TraceServices::FNetProfilerAggregatedStats>> AggregationResultTable;
 
 		AggregationStopwatch.Start();
 		{
-			Trace::FAnalysisSessionReadScope SessionReadScope(*Session.Get());
-			const Trace::INetProfilerProvider& NetProfilerProvider = Trace::ReadNetProfilerProvider(*Session.Get());
+			TraceServices::FAnalysisSessionReadScope SessionReadScope(*Session.Get());
+			const TraceServices::INetProfilerProvider& NetProfilerProvider = TraceServices::ReadNetProfilerProvider(*Session.Get());
 			// CreateAggregation requires [PacketStartIndex, PacketEndIndex] as inclusive interval and [StartPos, EndPos) as exclusive interval.
 			AggregationResultTable.Reset(NetProfilerProvider.CreateAggregation(ConnectionIndex, ConnectionMode, StatsPacketStartIndex, StatsPacketEndIndex - 1, StatsStartPosition, StatsEndPosition));
 		}
@@ -1717,10 +1717,10 @@ void SNetStatsView::UpdateStatsInternal()
 
 		if (AggregationResultTable.IsValid())
 		{
-			TUniquePtr<Trace::ITableReader<Trace::FNetProfilerAggregatedStats>> TableReader(AggregationResultTable->CreateReader());
+			TUniquePtr<TraceServices::ITableReader<TraceServices::FNetProfilerAggregatedStats>> TableReader(AggregationResultTable->CreateReader());
 			while (TableReader->IsValid())
 			{
-				const Trace::FNetProfilerAggregatedStats* Row = TableReader->GetCurrentRow();
+				const TraceServices::FNetProfilerAggregatedStats* Row = TableReader->GetCurrentRow();
 				FNetEventNodePtr NetEventNodePtr = GetNetEventNode(Row->EventTypeIndex);
 				if (NetEventNodePtr)
 				{
diff --git a/Engine/Source/Developer/TraceInsights/Private/Insights/NetworkingProfiler/Widgets/SNetStatsView.h b/Engine/Source/Developer/TraceInsights/Private/Insights/NetworkingProfiler/Widgets/SNetStatsView.h
--- a/Engine/Source/Developer/TraceInsights/Private/Insights/NetworkingProfiler/Widgets/SNetStatsView.h
+++ b/Engine/Source/Developer/TraceInsights/Private/Insights/NetworkingProfiler/Widgets/SNetStatsView.h
@@ -24,7 +24,7 @@
 class FMenuBuilder;
 class SNetworkingProfilerWindow;
 
-namespace Trace
+namespace TraceServices
 {
 	class IAnalysisSession;
 }
@@ -86,7 +86,7 @@ public:
 	void RebuildTree(bool bResync);
 
 	void ResetStats();
-	void UpdateStats(uint32 InGameInstanceIndex, uint32 InConnectionIndex, Trace::ENetProfilerConnectionMode InConnectionMode, uint32 InStatsPacketStartIndex, uint32 InStatsPacketEndIndex, uint32 InStatsStartPosition, uint32 InStatsEndPosition);
+	void UpdateStats(uint32 InGameInstanceIndex, uint32 InConnectionIndex, TraceServices::ENetProfilerConnectionMode InConnectionMode, uint32 InStatsPacketStartIndex, uint32 InStatsPacketEndIndex, uint32 InStatsStartPosition, uint32 InStatsEndPosition);
 
 	FNetEventNodePtr GetNetEventNode(uint32 EventTypeIndex) const;
 	void SelectNetEventNode(uint32 EventTypeIndex);
@@ -257,7 +257,7 @@ protected:
 	TSharedPtr<Insights::FTable> Table;
 
 	/** A weak pointer to the profiler session used to populate this widget. */
-	TSharedPtr<const Trace::IAnalysisSession>/*Weak*/ Session;
+	TSharedPtr<const TraceServices::IAnalysisSession>/*Weak*/ Session;
 
 	//////////////////////////////////////////////////
 	// Tree View, Columns
@@ -357,7 +357,7 @@ protected:
 
 	uint32 GameInstanceIndex;
 	uint32 ConnectionIndex;
-	Trace::ENetProfilerConnectionMode ConnectionMode;
+	TraceServices::ENetProfilerConnectionMode ConnectionMode;
 	uint32 StatsPacketStartIndex;
 	uint32 StatsPacketEndIndex;
 	uint32 StatsStartPosition;
diff --git a/Engine/Source/Developer/TraceInsights/Private/Insights/NetworkingProfiler/Widgets/SNetStatsViewTooltip.cpp b/Engine/Source/Developer/TraceInsights/Private/Insights/NetworkingProfiler/Widgets/SNetStatsViewTooltip.cpp
--- a/Engine/Source/Developer/TraceInsights/Private/Insights/NetworkingProfiler/Widgets/SNetStatsViewTooltip.cpp
+++ b/Engine/Source/Developer/TraceInsights/Private/Insights/NetworkingProfiler/Widgets/SNetStatsViewTooltip.cpp
@@ -89,7 +89,7 @@ TSharedPtr<SToolTip> SNetStatsViewTooltip::GetColumnTooltip(const Insights::FTab
 
 TSharedPtr<SToolTip> SNetStatsViewTooltip::GetRowTooltip(const TSharedPtr<FNetEventNode> NetEventNodePtr)
 {
-	const Trace::FNetProfilerAggregatedStats& Stats = NetEventNodePtr->GetAggregatedStats();
+	const TraceServices::FNetProfilerAggregatedStats& Stats = NetEventNodePtr->GetAggregatedStats();
 
 	const FText InstanceCountText = FText::AsNumber(Stats.InstanceCount);
 
diff --git a/Engine/Source/Developer/TraceInsights/Private/Insights/NetworkingProfiler/Widgets/SNetworkingProfilerWindow.cpp b/Engine/Source/Developer/TraceInsights/Private/Insights/NetworkingProfiler/Widgets/SNetworkingProfilerWindow.cpp
--- a/Engine/Source/Developer/TraceInsights/Private/Insights/NetworkingProfiler/Widgets/SNetworkingProfilerWindow.cpp
+++ b/Engine/Source/Developer/TraceInsights/Private/Insights/NetworkingProfiler/Widgets/SNetworkingProfilerWindow.cpp
@@ -488,13 +488,13 @@ EActiveTimerReturnType SNetworkingProfilerWindow::UpdateActiveDuration(double In
 
 void SNetworkingProfilerWindow::Tick(const FGeometry& AllottedGeometry, const double InCurrentTime, const float InDeltaTime)
 {
-	TSharedPtr<const Trace::IAnalysisSession> Session = FInsightsManager::Get()->GetSession();
+	TSharedPtr<const TraceServices::IAnalysisSession> Session = FInsightsManager::Get()->GetSession();
 
 	uint32 GameInstanceCount = 0;
 	if (Session.IsValid())
 	{
-		Trace::FAnalysisSessionReadScope SessionReadScope(*Session.Get());
-		const Trace::INetProfilerProvider& NetProfilerProvider = Trace::ReadNetProfilerProvider(*Session.Get());
+		TraceServices::FAnalysisSessionReadScope SessionReadScope(*Session.Get());
+		const TraceServices::INetProfilerProvider& NetProfilerProvider = TraceServices::ReadNetProfilerProvider(*Session.Get());
 		GameInstanceCount = NetProfilerProvider.GetGameInstanceCount();
 	}
 
@@ -506,8 +506,8 @@ void SNetworkingProfilerWindow::Tick(const FGeometry& AllottedGeometry, const do
 	uint32 ConnectionCount = 0;
 	if (Session.IsValid() && SelectedGameInstance.IsValid())
 	{
-		Trace::FAnalysisSessionReadScope SessionReadScope(*Session.Get());
-		const Trace::INetProfilerProvider& NetProfilerProvider = Trace::ReadNetProfilerProvider(*Session.Get());
+		TraceServices::FAnalysisSessionReadScope SessionReadScope(*Session.Get());
+		const TraceServices::INetProfilerProvider& NetProfilerProvider = TraceServices::ReadNetProfilerProvider(*Session.Get());
 		ConnectionCount = NetProfilerProvider.GetConnectionCount(SelectedGameInstance->GetIndex());
 	}
 
@@ -523,12 +523,12 @@ void SNetworkingProfilerWindow::UpdateAvailableGameInstances()
 {
 	AvailableGameInstances.Reset();
 
-	TSharedPtr<const Trace::IAnalysisSession> Session = FInsightsManager::Get()->GetSession();
+	TSharedPtr<const TraceServices::IAnalysisSession> Session = FInsightsManager::Get()->GetSession();
 	if (Session.IsValid())
 	{
-		Trace::FAnalysisSessionReadScope SessionReadScope(*Session.Get());
-		const Trace::INetProfilerProvider& NetProfilerProvider = Trace::ReadNetProfilerProvider(*Session.Get());
-		NetProfilerProvider.ReadGameInstances([this](const Trace::FNetProfilerGameInstance& GameInstance)
+		TraceServices::FAnalysisSessionReadScope SessionReadScope(*Session.Get());
+		const TraceServices::INetProfilerProvider& NetProfilerProvider = TraceServices::ReadNetProfilerProvider(*Session.Get());
+		NetProfilerProvider.ReadGameInstances([this](const TraceServices::FNetProfilerGameInstance& GameInstance)
 		{
 			AvailableGameInstances.Add(MakeShared<FGameInstanceItem>(GameInstance));
 		});
@@ -548,12 +548,12 @@ void SNetworkingProfilerWindow::UpdateAvailableConnections()
 {
 	AvailableConnections.Reset();
 
-	TSharedPtr<const Trace::IAnalysisSession> Session = FInsightsManager::Get()->GetSession();
+	TSharedPtr<const TraceServices::IAnalysisSession> Session = FInsightsManager::Get()->GetSession();
 	if (Session.IsValid() && SelectedGameInstance.IsValid())
 	{
-		Trace::FAnalysisSessionReadScope SessionReadScope(*Session.Get());
-		const Trace::INetProfilerProvider& NetProfilerProvider = Trace::ReadNetProfilerProvider(*Session.Get());
-		NetProfilerProvider.ReadConnections(SelectedGameInstance->GetIndex(), [this](const Trace::FNetProfilerConnection& Connection)
+		TraceServices::FAnalysisSessionReadScope SessionReadScope(*Session.Get());
+		const TraceServices::INetProfilerProvider& NetProfilerProvider = TraceServices::ReadNetProfilerProvider(*Session.Get());
+		NetProfilerProvider.ReadConnections(SelectedGameInstance->GetIndex(), [this](const TraceServices::FNetProfilerConnection& Connection)
 		{
 			AvailableConnections.Add(MakeShared<FConnectionItem>(Connection));
 		});
@@ -577,11 +577,11 @@ void SNetworkingProfilerWindow::UpdateAvailableConnectionModes()
 	{
 		if (SelectedConnection->Connection.bHasIncomingData)
 		{
-			AvailableConnectionModes.Add(MakeShared<FConnectionModeItem>(Trace::ENetProfilerConnectionMode::Incoming));
+			AvailableConnectionModes.Add(MakeShared<FConnectionModeItem>(TraceServices::ENetProfilerConnectionMode::Incoming));
 		}
 		if (SelectedConnection->Connection.bHasOutgoingData)
 		{
-			AvailableConnectionModes.Add(MakeShared<FConnectionModeItem>(Trace::ENetProfilerConnectionMode::Outgoing));
+			AvailableConnectionModes.Add(MakeShared<FConnectionModeItem>(TraceServices::ENetProfilerConnectionMode::Outgoing));
 		}
 	}
 
@@ -715,10 +715,10 @@ FText SNetworkingProfilerWindow::FConnectionModeItem::GetText() const
 {
 	switch (Mode)
 	{
-	case Trace::ENetProfilerConnectionMode::Outgoing:
+	case TraceServices::ENetProfilerConnectionMode::Outgoing:
 		return LOCTEXT("ConnectionMode_Outgoing", "Outgoing");
 
-	case Trace::ENetProfilerConnectionMode::Incoming:
+	case TraceServices::ENetProfilerConnectionMode::Incoming:
 		return LOCTEXT("ConnectionMode_Incoming", "Incoming");
 
 	default:
diff --git a/Engine/Source/Developer/TraceInsights/Private/Insights/NetworkingProfiler/Widgets/SNetworkingProfilerWindow.h b/Engine/Source/Developer/TraceInsights/Private/Insights/NetworkingProfiler/Widgets/SNetworkingProfilerWindow.h
--- a/Engine/Source/Developer/TraceInsights/Private/Insights/NetworkingProfiler/Widgets/SNetworkingProfilerWindow.h
+++ b/Engine/Source/Developer/TraceInsights/Private/Insights/NetworkingProfiler/Widgets/SNetworkingProfilerWindow.h
@@ -47,7 +47,7 @@ private:
 	struct FGameInstanceItem
 	{
 		/** Conversion constructor. */
-		FGameInstanceItem(const Trace::FNetProfilerGameInstance& InGameInstance)
+		FGameInstanceItem(const TraceServices::FNetProfilerGameInstance& InGameInstance)
 			: GameInstance(InGameInstance)
 		{}
 
@@ -55,13 +55,13 @@ private:
 		FText GetText() const;
 		FText GetTooltipText() const;
 
-		const Trace::FNetProfilerGameInstance GameInstance;
+		const TraceServices::FNetProfilerGameInstance GameInstance;
 	};
 
 	struct FConnectionItem
 	{
 		/** Conversion constructor. */
-		FConnectionItem(const Trace::FNetProfilerConnection& InConnection)
+		FConnectionItem(const TraceServices::FNetProfilerConnection& InConnection)
 			: Connection(InConnection)
 		{}
 
@@ -69,20 +69,20 @@ private:
 		FText GetText() const;
 		FText GetTooltipText() const;
 
-		const Trace::FNetProfilerConnection Connection;
+		const TraceServices::FNetProfilerConnection Connection;
 	};
 
 	struct FConnectionModeItem
 	{
 		/** Conversion constructor. */
-		FConnectionModeItem(const Trace::ENetProfilerConnectionMode& InMode)
+		FConnectionModeItem(const TraceServices::ENetProfilerConnectionMode& InMode)
 			: Mode(InMode)
 		{}
 
 		FText GetText() const;
 		FText GetTooltipText() const;
 
-		Trace::ENetProfilerConnectionMode Mode;
+		TraceServices::ENetProfilerConnectionMode Mode;
 	};
 
 public:
@@ -122,11 +122,11 @@ public:
 	const bool IsNetStatsViewVisible() const { return NetStatsView.IsValid(); }
 	void ShowOrHideNetStatsView(const bool bVisibleState) { ShowOrHideTab(FNetworkingProfilerTabs::NetStatsViewID, bVisibleState); }
 
-	const Trace::FNetProfilerGameInstance* GetSelectedGameInstance() const { return SelectedGameInstance ? &SelectedGameInstance->GameInstance : nullptr; }
+	const TraceServices::FNetProfilerGameInstance* GetSelectedGameInstance() const { return SelectedGameInstance ? &SelectedGameInstance->GameInstance : nullptr; }
 	uint32 GetSelectedGameInstanceIndex() const { return SelectedGameInstance ? SelectedGameInstance->GetIndex() : 0; }
-	const Trace::FNetProfilerConnection* GetSelectedConnection() const { return SelectedConnection ? &SelectedConnection->Connection : nullptr; }
+	const TraceServices::FNetProfilerConnection* GetSelectedConnection() const { return SelectedConnection ? &SelectedConnection->Connection : nullptr; }
 	uint32 GetSelectedConnectionIndex() const { return SelectedConnection ? SelectedConnection->GetIndex() : 0; }
-	Trace::ENetProfilerConnectionMode GetSelectedConnectionMode() const { return SelectedConnectionMode ? SelectedConnectionMode->Mode : Trace::ENetProfilerConnectionMode::Outgoing; }
+	TraceServices::ENetProfilerConnectionMode GetSelectedConnectionMode() const { return SelectedConnectionMode ? SelectedConnectionMode->Mode : TraceServices::ENetProfilerConnectionMode::Outgoing; }
 
 	TSharedRef<SWidget> CreateGameInstanceComboBox();
 	TSharedRef<SWidget> CreateConnectionComboBox();
diff --git a/Engine/Source/Developer/TraceInsights/Private/Insights/NetworkingProfiler/Widgets/SPacketContentView.cpp b/Engine/Source/Developer/TraceInsights/Private/Insights/NetworkingProfiler/Widgets/SPacketContentView.cpp
--- a/Engine/Source/Developer/TraceInsights/Private/Insights/NetworkingProfiler/Widgets/SPacketContentView.cpp
+++ b/Engine/Source/Developer/TraceInsights/Private/Insights/NetworkingProfiler/Widgets/SPacketContentView.cpp
@@ -57,7 +57,7 @@ void SPacketContentView::Reset()
 
 	GameInstanceIndex = 0;
 	ConnectionIndex = 0;
-	ConnectionMode = Trace::ENetProfilerConnectionMode::Outgoing;
+	ConnectionMode = TraceServices::ENetProfilerConnectionMode::Outgoing;
 	PacketIndex = 0;
 	PacketSequence = 0;
 	PacketBitSize = 0;
@@ -616,7 +616,7 @@ void SPacketContentView::ResetPacket()
 {
 	GameInstanceIndex = 0;
 	ConnectionIndex = 0;
-	ConnectionMode = Trace::ENetProfilerConnectionMode::Outgoing;
+	ConnectionMode = TraceServices::ENetProfilerConnectionMode::Outgoing;
 	PacketIndex = 0;
 	PacketSequence = 0;
 	PacketBitSize = 0;
@@ -636,7 +636,7 @@ void SPacketContentView::ResetPacket()
 
 ////////////////////////////////////////////////////////////////////////////////////////////////////
 
-void SPacketContentView::SetPacket(uint32 InGameInstanceIndex, uint32 InConnectionIndex, Trace::ENetProfilerConnectionMode InConnectionMode, uint32 InPacketIndex, int64 InPacketBitSize)
+void SPacketContentView::SetPacket(uint32 InGameInstanceIndex, uint32 InConnectionIndex, TraceServices::ENetProfilerConnectionMode InConnectionMode, uint32 InPacketIndex, int64 InPacketBitSize)
 {
 	GameInstanceIndex = InGameInstanceIndex;
 	ConnectionIndex = InConnectionIndex;
@@ -701,13 +701,13 @@ void SPacketContentView::EnableFilterEventType(const uint32 InEventTypeIndex)
 {
 	FText EventName;
 
-	TSharedPtr<const Trace::IAnalysisSession> Session = FInsightsManager::Get()->GetSession();
+	TSharedPtr<const TraceServices::IAnalysisSession> Session = FInsightsManager::Get()->GetSession();
 	if (Session.IsValid())
 	{
-		Trace::FAnalysisSessionReadScope SessionReadScope(*Session.Get());
-		const Trace::INetProfilerProvider& NetProfilerProvider = Trace::ReadNetProfilerProvider(*Session.Get());
+		TraceServices::FAnalysisSessionReadScope SessionReadScope(*Session.Get());
+		const TraceServices::INetProfilerProvider& NetProfilerProvider = TraceServices::ReadNetProfilerProvider(*Session.Get());
 
-		NetProfilerProvider.ReadEventType(InEventTypeIndex, [&EventName](const Trace::FNetProfilerEventType& EventType)
+		NetProfilerProvider.ReadEventType(InEventTypeIndex, [&EventName](const TraceServices::FNetProfilerEventType& EventType)
 		{
 			EventName = FText::FromString(EventType.Name);
 		});
@@ -725,13 +725,13 @@ uint32 SPacketContentView::GetPacketSequence(int32 InPacketIndex) const
 {
 	uint32 NewSequenceNumber = 0U;
 
-	TSharedPtr<const Trace::IAnalysisSession> Session = FInsightsManager::Get()->GetSession();
+	TSharedPtr<const TraceServices::IAnalysisSession> Session = FInsightsManager::Get()->GetSession();
 	if (Session.IsValid())
 	{
-		Trace::FAnalysisSessionReadScope SessionReadScope(*Session.Get());
-		const Trace::INetProfilerProvider& NetProfilerProvider = Trace::ReadNetProfilerProvider(*Session.Get());
+		TraceServices::FAnalysisSessionReadScope SessionReadScope(*Session.Get());
+		const TraceServices::INetProfilerProvider& NetProfilerProvider = TraceServices::ReadNetProfilerProvider(*Session.Get());
 
-		NetProfilerProvider.EnumeratePackets(ConnectionIndex, ConnectionMode, InPacketIndex, InPacketIndex, [&NewSequenceNumber](const Trace::FNetProfilerPacket& Packet)
+		NetProfilerProvider.EnumeratePackets(ConnectionIndex, ConnectionMode, InPacketIndex, InPacketIndex, [&NewSequenceNumber](const TraceServices::FNetProfilerPacket& Packet)
 		{
 			NewSequenceNumber = Packet.SequenceNumber;
 		});
@@ -761,11 +761,11 @@ void SPacketContentView::UpdateState()
 		FPacketContentViewDrawStateBuilder Builder(*DrawState, Viewport);
 		FPacketContentViewDrawStateBuilder FilteredDrawStateBuilder(*FilteredDrawState, Viewport);
 
-		TSharedPtr<const Trace::IAnalysisSession> Session = FInsightsManager::Get()->GetSession();
+		TSharedPtr<const TraceServices::IAnalysisSession> Session = FInsightsManager::Get()->GetSession();
 		if (Session.IsValid())
 		{
-			Trace::FAnalysisSessionReadScope SessionReadScope(*Session.Get());
-			const Trace::INetProfilerProvider& NetProfilerProvider = Trace::ReadNetProfilerProvider(*Session.Get());
+			TraceServices::FAnalysisSessionReadScope SessionReadScope(*Session.Get());
+			const TraceServices::INetProfilerProvider& NetProfilerProvider = TraceServices::ReadNetProfilerProvider(*Session.Get());
 
 			const FAxisViewportDouble& ViewportX = Viewport.GetHorizontalAxisViewport();
 
@@ -773,10 +773,10 @@ void SPacketContentView::UpdateState()
 			//const int64 EndPos = static_cast<int64>(FMath::CeilToDouble(ViewportX.GetValueAtOffset(ViewportX.GetSize())));
 			const uint32 StartPos = 0;
 			const uint32 EndPos = PacketBitSize;
-			NetProfilerProvider.EnumeratePacketContentEventsByPosition(ConnectionIndex, ConnectionMode, PacketIndex, StartPos, EndPos, [this, &Builder, &FilteredDrawStateBuilder, &NetProfilerProvider](const Trace::FNetProfilerContentEvent& Event)
+			NetProfilerProvider.EnumeratePacketContentEventsByPosition(ConnectionIndex, ConnectionMode, PacketIndex, StartPos, EndPos, [this, &Builder, &FilteredDrawStateBuilder, &NetProfilerProvider](const TraceServices::FNetProfilerContentEvent& Event)
 			{
 				const TCHAR* Name = nullptr;
-				NetProfilerProvider.ReadName(Event.NameIndex, [&Name](const Trace::FNetProfilerName& NetProfilerName)
+				NetProfilerProvider.ReadName(Event.NameIndex, [&Name](const TraceServices::FNetProfilerName& NetProfilerName)
 				{
 					Name = NetProfilerName.Name;
 				});
@@ -784,7 +784,7 @@ void SPacketContentView::UpdateState()
 				uint32 NetId = 0;
 				if (Event.ObjectInstanceIndex != 0)
 				{
-					NetProfilerProvider.ReadObject(GameInstanceIndex, Event.ObjectInstanceIndex, [&NetId](const Trace::FNetProfilerObjectInstance& ObjectInstance)
+					NetProfilerProvider.ReadObject(GameInstanceIndex, Event.ObjectInstanceIndex, [&NetId](const TraceServices::FNetProfilerObjectInstance& ObjectInstance)
 					{
 						NetId = ObjectInstance.NetId;
 					});
@@ -829,28 +829,28 @@ void SPacketContentView::UpdateHoveredEvent()
 
 		const FNetworkPacketEvent& Event = HoveredEvent.Event;
 		FString Name(TEXT("?"));
-		Trace::FNetProfilerEventType EventType;
-		Trace::FNetProfilerObjectInstance ObjectInstance;
+		TraceServices::FNetProfilerEventType EventType;
+		TraceServices::FNetProfilerObjectInstance ObjectInstance;
 
-		TSharedPtr<const Trace::IAnalysisSession> Session = FInsightsManager::Get()->GetSession();
+		TSharedPtr<const TraceServices::IAnalysisSession> Session = FInsightsManager::Get()->GetSession();
 		if (Session.IsValid())
 		{
-			Trace::FAnalysisSessionReadScope SessionReadScope(*Session.Get());
-			const Trace::INetProfilerProvider& NetProfilerProvider = Trace::ReadNetProfilerProvider(*Session.Get());
+			TraceServices::FAnalysisSessionReadScope SessionReadScope(*Session.Get());
+			const TraceServices::INetProfilerProvider& NetProfilerProvider = TraceServices::ReadNetProfilerProvider(*Session.Get());
 
-			NetProfilerProvider.ReadEventType(Event.EventTypeIndex, [&EventType](const Trace::FNetProfilerEventType& InEventType)
+			NetProfilerProvider.ReadEventType(Event.EventTypeIndex, [&EventType](const TraceServices::FNetProfilerEventType& InEventType)
 			{
 				EventType = InEventType;
 			});
 
-			NetProfilerProvider.ReadName(EventType.NameIndex, [&Name](const Trace::FNetProfilerName& NetProfilerName)
+			NetProfilerProvider.ReadName(EventType.NameIndex, [&Name](const TraceServices::FNetProfilerName& NetProfilerName)
 			{
 				Name = NetProfilerName.Name;
 			});
 
 			if (Event.ObjectInstanceIndex != 0)
 			{
-				NetProfilerProvider.ReadObject(GameInstanceIndex, Event.ObjectInstanceIndex, [&ObjectInstance](const Trace::FNetProfilerObjectInstance& InObjectInstance)
+				NetProfilerProvider.ReadObject(GameInstanceIndex, Event.ObjectInstanceIndex, [&ObjectInstance](const TraceServices::FNetProfilerObjectInstance& InObjectInstance)
 				{
 					ObjectInstance = InObjectInstance;
 				});
@@ -1092,7 +1092,7 @@ int32 SPacketContentView::OnPaint(const FPaintArgs& Args, const FGeometry& Allot
 			FString::Printf(TEXT("Game Instance %d, Connection %d (%s), Packet %d"),
 				GameInstanceIndex,
 				ConnectionIndex,
-				(ConnectionMode == Trace::ENetProfilerConnectionMode::Outgoing) ? TEXT("Outgoing") : TEXT("Incoming"),
+				(ConnectionMode == TraceServices::ENetProfilerConnectionMode::Outgoing) ? TEXT("Outgoing") : TEXT("Incoming"),
 				PacketIndex),
 			SummaryFont, DbgTextColor
 		);
diff --git a/Engine/Source/Developer/TraceInsights/Private/Insights/NetworkingProfiler/Widgets/SPacketContentView.h b/Engine/Source/Developer/TraceInsights/Private/Insights/NetworkingProfiler/Widgets/SPacketContentView.h
--- a/Engine/Source/Developer/TraceInsights/Private/Insights/NetworkingProfiler/Widgets/SPacketContentView.h
+++ b/Engine/Source/Developer/TraceInsights/Private/Insights/NetworkingProfiler/Widgets/SPacketContentView.h
@@ -137,7 +137,7 @@ public:
 	//virtual FReply OnKeyUp(const FGeometry& MyGeometry, const FKeyEvent& InKeyEvent) override;
 
 	void ResetPacket();
-	void SetPacket(uint32 InGameInstanceIndex, uint32 InConnectionIndex, Trace::ENetProfilerConnectionMode InConnectionMode, uint32 InPacketIndex, int64 InPacketBitSize);
+	void SetPacket(uint32 InGameInstanceIndex, uint32 InConnectionIndex, TraceServices::ENetProfilerConnectionMode InConnectionMode, uint32 InPacketIndex, int64 InPacketBitSize);
 
 	bool IsFilterByNetIdEnabled() const { return bFilterByNetId; }
 	uint32 GetFilterNetId() const { return FilterNetId; }
@@ -213,7 +213,7 @@ private:
 
 	uint32 GameInstanceIndex;
 	uint32 ConnectionIndex;
-	Trace::ENetProfilerConnectionMode ConnectionMode;
+	TraceServices::ENetProfilerConnectionMode ConnectionMode;
 	uint32 PacketIndex;
 	uint32 PacketSequence;
 	int64 PacketBitSize; // total number of bits; [bit]
diff --git a/Engine/Source/Developer/TraceInsights/Private/Insights/NetworkingProfiler/Widgets/SPacketView.cpp b/Engine/Source/Developer/TraceInsights/Private/Insights/NetworkingProfiler/Widgets/SPacketView.cpp
--- a/Engine/Source/Developer/TraceInsights/Private/Insights/NetworkingProfiler/Widgets/SPacketView.cpp
+++ b/Engine/Source/Developer/TraceInsights/Private/Insights/NetworkingProfiler/Widgets/SPacketView.cpp
@@ -51,7 +51,7 @@ void SPacketView::Reset()
 
 	GameInstanceIndex = 0;
 	ConnectionIndex = 0;
-	ConnectionMode = Trace::ENetProfilerConnectionMode::Outgoing;
+	ConnectionMode = TraceServices::ENetProfilerConnectionMode::Outgoing;
 
 	Viewport.Reset();
 	FAxisViewportInt32& ViewportX = Viewport.GetHorizontalAxisViewport();
@@ -113,7 +113,7 @@ void SPacketView::Reset()
 
 ////////////////////////////////////////////////////////////////////////////////////////////////////
 
-void SPacketView::SetConnection(uint32 InGameInstanceIndex, uint32 InConnectionIndex, Trace::ENetProfilerConnectionMode InConnectionMode)
+void SPacketView::SetConnection(uint32 InGameInstanceIndex, uint32 InConnectionIndex, TraceServices::ENetProfilerConnectionMode InConnectionMode)
 {
 	GameInstanceIndex = InGameInstanceIndex;
 	ConnectionIndex = InConnectionIndex;
@@ -204,12 +204,12 @@ void SPacketView::Tick(const FGeometry& AllottedGeometry, const double InCurrent
 		const uint64 WaitTime = static_cast<uint64>(0.1 / FPlatformTime::GetSecondsPerCycle64()); // 100ms
 		AnalysisSyncNextTimestamp = Time + WaitTime;
 
-		TSharedPtr<const Trace::IAnalysisSession> Session = FInsightsManager::Get()->GetSession();
+		TSharedPtr<const TraceServices::IAnalysisSession> Session = FInsightsManager::Get()->GetSession();
 		if (Session.IsValid())
 		{
-			Trace::FAnalysisSessionReadScope SessionReadScope(*Session.Get());
+			TraceServices::FAnalysisSessionReadScope SessionReadScope(*Session.Get());
 
-			const Trace::INetProfilerProvider& NetProfilerProvider = Trace::ReadNetProfilerProvider(*Session.Get());
+			const TraceServices::INetProfilerProvider& NetProfilerProvider = TraceServices::ReadNetProfilerProvider(*Session.Get());
 
 			const uint32 NewConnectionChangeCount = NetProfilerProvider.GetConnectionChangeCount();
 			if (NewConnectionChangeCount != ConnectionChangeCount)
@@ -239,7 +239,7 @@ void SPacketView::Tick(const FGeometry& AllottedGeometry, const double InCurrent
 
 ////////////////////////////////////////////////////////////////////////////////////////////////////
 
-bool SPacketView::IsConnectionValid(const Trace::INetProfilerProvider& NetProfilerProvider, const uint32 InGameInstanceIndex, const uint32 InConnectionIndex, const Trace::ENetProfilerConnectionMode InConnectionMode)
+bool SPacketView::IsConnectionValid(const TraceServices::INetProfilerProvider& NetProfilerProvider, const uint32 InGameInstanceIndex, const uint32 InConnectionIndex, const TraceServices::ENetProfilerConnectionMode InConnectionMode)
 {
 	//TODO: return NetProfilerProvider.IsConnectionValid(GameInstanceIndex, ConnectionIndex, ConnectionMode);
 
@@ -250,12 +250,12 @@ bool SPacketView::IsConnectionValid(const Trace::INetProfilerProvider& NetProfil
 	}
 
 	bool bIsValidConnection = false;
-	NetProfilerProvider.ReadConnections(InGameInstanceIndex, [InConnectionIndex, InConnectionMode, &bIsValidConnection](const Trace::FNetProfilerConnection& Connection)
+	NetProfilerProvider.ReadConnections(InGameInstanceIndex, [InConnectionIndex, InConnectionMode, &bIsValidConnection](const TraceServices::FNetProfilerConnection& Connection)
 	{
 		if (InConnectionIndex == Connection.ConnectionIndex)
 		{
-			if ((InConnectionMode == Trace::ENetProfilerConnectionMode::Outgoing && Connection.bHasOutgoingData) ||
-				(InConnectionMode == Trace::ENetProfilerConnectionMode::Incoming && Connection.bHasIncomingData))
+			if ((InConnectionMode == TraceServices::ENetProfilerConnectionMode::Outgoing && Connection.bHasOutgoingData) ||
+				(InConnectionMode == TraceServices::ENetProfilerConnectionMode::Incoming && Connection.bHasIncomingData))
 			{
 				bIsValidConnection = true;
 			}
@@ -309,11 +309,11 @@ void SPacketView::UpdateState()
 
 	FNetworkPacketSeriesBuilder Builder(*PacketSeries, Viewport);
 
-	TSharedPtr<const Trace::IAnalysisSession> Session = FInsightsManager::Get()->GetSession();
+	TSharedPtr<const TraceServices::IAnalysisSession> Session = FInsightsManager::Get()->GetSession();
 	if (Session.IsValid())
 	{
-		Trace::FAnalysisSessionReadScope SessionReadScope(*Session.Get());
-		const Trace::INetProfilerProvider& NetProfilerProvider = Trace::ReadNetProfilerProvider(*Session.Get());
+		TraceServices::FAnalysisSessionReadScope SessionReadScope(*Session.Get());
+		const TraceServices::INetProfilerProvider& NetProfilerProvider = TraceServices::ReadNetProfilerProvider(*Session.Get());
 
 		if (IsConnectionValid(NetProfilerProvider, GameInstanceIndex, ConnectionIndex, ConnectionMode))
 		{
@@ -332,7 +332,7 @@ void SPacketView::UpdateState()
 				if (PacketStartIndex <= PacketEndIndex)
 				{
 					int32 PacketIndex = PacketStartIndex;
-					NetProfilerProvider.EnumeratePackets(ConnectionIndex, ConnectionMode, PacketStartIndex, PacketEndIndex, [this, &Builder, &PacketIndex, &Filter, &NetProfilerProvider](const Trace::FNetProfilerPacket& Packet)
+					NetProfilerProvider.EnumeratePackets(ConnectionIndex, ConnectionMode, PacketStartIndex, PacketEndIndex, [this, &Builder, &PacketIndex, &Filter, &NetProfilerProvider](const TraceServices::FNetProfilerPacket& Packet)
 					{
 						FNetworkPacketAggregatedSample* SamplePtr = Builder.AddPacket(PacketIndex++, Packet);
 
@@ -349,7 +349,7 @@ void SPacketView::UpdateState()
 
 								const uint32 StartPos = 0;
 								const uint32 EndPos = Packet.ContentSizeInBits;
-								NetProfilerProvider.EnumeratePacketContentEventsByPosition(ConnectionIndex, ConnectionMode, PacketIndex - 1, StartPos, EndPos, [this, &bFilterMatch, &Filter, &NetProfilerProvider](const Trace::FNetProfilerContentEvent& Event)
+								NetProfilerProvider.EnumeratePacketContentEventsByPosition(ConnectionIndex, ConnectionMode, PacketIndex - 1, StartPos, EndPos, [this, &bFilterMatch, &Filter, &NetProfilerProvider](const TraceServices::FNetProfilerContentEvent& Event)
 								{
 									if (!bFilterMatch)
 									{
@@ -363,7 +363,7 @@ void SPacketView::UpdateState()
 											uint32 NetId = uint32(-1);
 											if (Event.ObjectInstanceIndex != 0)
 											{
-												NetProfilerProvider.ReadObject(GameInstanceIndex, Event.ObjectInstanceIndex, [&NetId](const Trace::FNetProfilerObjectInstance& ObjectInstance)
+												NetProfilerProvider.ReadObject(GameInstanceIndex, Event.ObjectInstanceIndex, [&NetId](const TraceServices::FNetProfilerObjectInstance& ObjectInstance)
 												{
 													NetId = ObjectInstance.NetId;
 												});
@@ -407,21 +407,21 @@ void SPacketView::UpdateState()
 			FRandomStream RandomStream((PacketIndex * PacketIndex * PacketIndex) ^ 0x2c2c57ed);
 			int64 Size = RandomStream.RandRange(0, 2000);
 
-			Trace::ENetProfilerDeliveryStatus Status = Trace::ENetProfilerDeliveryStatus::Unknown;
+			TraceServices::ENetProfilerDeliveryStatus Status = TraceServices::ENetProfilerDeliveryStatus::Unknown;
 			const float Fraction = RandomStream.GetFraction();
 			if (Fraction < 0.01) // 1%
 			{
-				Status = Trace::ENetProfilerDeliveryStatus::Dropped;
+				Status = TraceServices::ENetProfilerDeliveryStatus::Dropped;
 			}
 			else if (Fraction < 0.05) // 4%
 			{
-				Status = Trace::ENetProfilerDeliveryStatus::Delivered;
+				Status = TraceServices::ENetProfilerDeliveryStatus::Delivered;
 			}
 
 			const double Timestamp = ((double)PacketIndex * 100.0) / (double)NumPackets + RandomStream.GetFraction() * 0.1;
 
-			Trace::FNetProfilerPacket Packet;
-			Packet.TimeStamp = static_cast<Trace::FNetProfilerTimeStamp>(Timestamp);
+			TraceServices::FNetProfilerPacket Packet;
+			Packet.TimeStamp = static_cast<TraceServices::FNetProfilerTimeStamp>(Timestamp);
 			Packet.SequenceNumber = PacketIndex;
 			Packet.ContentSizeInBits = Size;
 			Packet.TotalPacketSizeInBytes = (Size + 7) / 8;
@@ -449,18 +449,18 @@ FNetworkPacketSampleRef SPacketView::GetSample(const int32 InPacketIndex)
 	FNetworkPacketSampleRef SampleRef;
 	SampleRef.Series = PacketSeries;
 
-	TSharedPtr<const Trace::IAnalysisSession> Session = FInsightsManager::Get()->GetSession();
+	TSharedPtr<const TraceServices::IAnalysisSession> Session = FInsightsManager::Get()->GetSession();
 	if (Session.IsValid())
 	{
-		Trace::FAnalysisSessionReadScope SessionReadScope(*Session.Get());
-		const Trace::INetProfilerProvider& NetProfilerProvider = Trace::ReadNetProfilerProvider(*Session.Get());
+		TraceServices::FAnalysisSessionReadScope SessionReadScope(*Session.Get());
+		const TraceServices::INetProfilerProvider& NetProfilerProvider = TraceServices::ReadNetProfilerProvider(*Session.Get());
 
 		if (IsConnectionValid(NetProfilerProvider, GameInstanceIndex, ConnectionIndex, ConnectionMode))
 		{
 			const uint32 NumPackets = NetProfilerProvider.GetPacketCount(ConnectionIndex, ConnectionMode);
 			if (InPacketIndex >= 0 && InPacketIndex < static_cast<int32>(NumPackets))
 			{
-				NetProfilerProvider.EnumeratePackets(ConnectionIndex, ConnectionMode, InPacketIndex, InPacketIndex, [InPacketIndex, &SampleRef](const Trace::FNetProfilerPacket& Packet)
+				NetProfilerProvider.EnumeratePackets(ConnectionIndex, ConnectionMode, InPacketIndex, InPacketIndex, [InPacketIndex, &SampleRef](const TraceServices::FNetProfilerPacket& Packet)
 				{
 					SampleRef.Sample = MakeShared<FNetworkPacketAggregatedSample>();
 					SampleRef.Sample->AddPacket(InPacketIndex, Packet);
@@ -622,25 +622,25 @@ void SPacketView::OnSelectionChanged()
 
 ////////////////////////////////////////////////////////////////////////////////////////////////////
 
-const TCHAR* StatusToString(Trace::ENetProfilerDeliveryStatus Status)
+const TCHAR* StatusToString(TraceServices::ENetProfilerDeliveryStatus Status)
 {
 	switch (Status)
 	{
-		case Trace::ENetProfilerDeliveryStatus::Delivered:  return TEXT("Delivered");
-		case Trace::ENetProfilerDeliveryStatus::Dropped:    return TEXT("Dropped");
-		case Trace::ENetProfilerDeliveryStatus::Unknown:
-		default:                                            return TEXT("Unknown");
+		case TraceServices::ENetProfilerDeliveryStatus::Delivered:  return TEXT("Delivered");
+		case TraceServices::ENetProfilerDeliveryStatus::Dropped:    return TEXT("Dropped");
+		case TraceServices::ENetProfilerDeliveryStatus::Unknown:
+		default:                                               return TEXT("Unknown");
 	};
 }
 
-const TCHAR* AggregatedStatusToString(Trace::ENetProfilerDeliveryStatus Status)
+const TCHAR* AggregatedStatusToString(TraceServices::ENetProfilerDeliveryStatus Status)
 {
 	switch (Status)
 	{
-		case Trace::ENetProfilerDeliveryStatus::Delivered:  return TEXT("all packets are Delivered");
-		case Trace::ENetProfilerDeliveryStatus::Dropped:    return TEXT("at least one Dropped packet");
-		case Trace::ENetProfilerDeliveryStatus::Unknown:
-		default:                                            return TEXT("Unknown");
+		case TraceServices::ENetProfilerDeliveryStatus::Delivered:  return TEXT("all packets are Delivered");
+		case TraceServices::ENetProfilerDeliveryStatus::Dropped:    return TEXT("at least one Dropped packet");
+		case TraceServices::ENetProfilerDeliveryStatus::Unknown:
+		default:                                               return TEXT("Unknown");
 	}
 }
 
@@ -867,7 +867,7 @@ int32 SPacketView::OnPaint(const FPaintArgs& Args, const FGeometry& AllottedGeom
 			FString::Printf(TEXT("Game Instance %d, Connection %d (%s)"),
 				GameInstanceIndex,
 				ConnectionIndex,
-				(ConnectionMode == Trace::ENetProfilerConnectionMode::Outgoing) ? TEXT("Outgoing") : TEXT("Incoming")),
+				(ConnectionMode == TraceServices::ENetProfilerConnectionMode::Outgoing) ? TEXT("Outgoing") : TEXT("Incoming")),
 			SummaryFont, DbgTextColor
 		);
 		DbgY += DbgDY;
@@ -1357,11 +1357,11 @@ void SPacketView::SetSelectedPacket(const int32 InPacketIndex)
 void SPacketView::SelectPacketBySequenceNumber(const uint32 InSequenceNumber)
 {
 	// Find the PacketIndex from sequence number
-	TSharedPtr<const Trace::IAnalysisSession> Session = FInsightsManager::Get()->GetSession();
+	TSharedPtr<const TraceServices::IAnalysisSession> Session = FInsightsManager::Get()->GetSession();
 	if (Session.IsValid())
 	{
-		Trace::FAnalysisSessionReadScope SessionReadScope(*Session.Get());
-		const Trace::INetProfilerProvider& NetProfilerProvider = Trace::ReadNetProfilerProvider(*Session.Get());
+		TraceServices::FAnalysisSessionReadScope SessionReadScope(*Session.Get());
+		const TraceServices::INetProfilerProvider& NetProfilerProvider = TraceServices::ReadNetProfilerProvider(*Session.Get());
 
 		const int32 PacketId = NetProfilerProvider.FindPacketIndexFromPacketSequence(ConnectionIndex, ConnectionMode, InSequenceNumber);
 		if (PacketId != -1)
diff --git a/Engine/Source/Developer/TraceInsights/Private/Insights/NetworkingProfiler/Widgets/SPacketView.h b/Engine/Source/Developer/TraceInsights/Private/Insights/NetworkingProfiler/Widgets/SPacketView.h
--- a/Engine/Source/Developer/TraceInsights/Private/Insights/NetworkingProfiler/Widgets/SPacketView.h
+++ b/Engine/Source/Developer/TraceInsights/Private/Insights/NetworkingProfiler/Widgets/SPacketView.h
@@ -98,7 +98,7 @@ public:
 	/** Resets internal widget's data to the default one. */
 	void Reset();
 
-	void SetConnection(uint32 GameInstanceIndex, uint32 ConnectionIndex, Trace::ENetProfilerConnectionMode ConnectionMode);
+	void SetConnection(uint32 GameInstanceIndex, uint32 ConnectionIndex, TraceServices::ENetProfilerConnectionMode ConnectionMode);
 
 	SLATE_BEGIN_ARGS(SPacketView)
 	{
@@ -147,7 +147,7 @@ public:
 private:
 	void UpdateSelectedSample();
 
-	bool IsConnectionValid(const Trace::INetProfilerProvider& NetProfilerProvider, const uint32 InGameInstanceIndex, const uint32 InConnectionIndex, const Trace::ENetProfilerConnectionMode InConnectionMode);
+	bool IsConnectionValid(const TraceServices::INetProfilerProvider& NetProfilerProvider, const uint32 InGameInstanceIndex, const uint32 InConnectionIndex, const TraceServices::ENetProfilerConnectionMode InConnectionMode);
 	void UpdateState();
 
 	void DrawHorizontalAxisGrid(FDrawContext& DrawContext, const FSlateBrush* Brush, const FSlateFontInfo& Font) const;
@@ -182,7 +182,7 @@ private:
 
 	uint32 GameInstanceIndex;
 	uint32 ConnectionIndex;
-	Trace::ENetProfilerConnectionMode ConnectionMode;
+	TraceServices::ENetProfilerConnectionMode ConnectionMode;
 
 	/** The track's viewport. Encapsulates info about position and scale. */
 	FPacketViewport Viewport;
diff --git a/Engine/Source/Developer/TraceInsights/Private/Insights/StoreService/DiagnosticsSessionAnalyzer.cpp b/Engine/Source/Developer/TraceInsights/Private/Insights/StoreService/DiagnosticsSessionAnalyzer.cpp
--- a/Engine/Source/Developer/TraceInsights/Private/Insights/StoreService/DiagnosticsSessionAnalyzer.cpp
+++ b/Engine/Source/Developer/TraceInsights/Private/Insights/StoreService/DiagnosticsSessionAnalyzer.cpp
@@ -48,8 +48,6 @@ bool FDiagnosticsSessionAnalyzer::OnEvent(uint16 RouteId, EStyle, const FOnEvent
 	}
 	case RouteId_Session2:
 	{
-		Trace::FSessionInfo SessionInfo;
-
 		EventData.GetString("Platform", Platform);
 		EventData.GetString("AppName", AppName);
 		EventData.GetString("CommandLine", CommandLine);
