Moved the TraceServices module into its own namespace.

diff --git a/Engine/Source/Developer/TraceInsights/Private/Insights/Table/ViewModels/UntypedTable.cpp b/Engine/Source/Developer/TraceInsights/Private/Insights/Table/ViewModels/UntypedTable.cpp
--- a/Engine/Source/Developer/TraceInsights/Private/Insights/Table/ViewModels/UntypedTable.cpp
+++ b/Engine/Source/Developer/TraceInsights/Private/Insights/Table/ViewModels/UntypedTable.cpp
@@ -35,7 +35,7 @@ public:
 			const TSharedPtr<FUntypedTable> UntypedTablePtr = StaticCastSharedPtr<FUntypedTable>(TablePtr);
 			if (UntypedTablePtr.IsValid())
 			{
-				TSharedPtr<Trace::IUntypedTableReader> Reader = UntypedTablePtr->GetTableReader();
+				TSharedPtr<TraceServices::IUntypedTableReader> Reader = UntypedTablePtr->GetTableReader();
 				if (Reader.IsValid() && TableTreeNode.GetRowId().HasValidIndex())
 				{
 					Reader->SetRowIndex(TableTreeNode.GetRowId().RowIndex);
@@ -100,7 +100,7 @@ void FUntypedTable::Reset()
 
 ////////////////////////////////////////////////////////////////////////////////////////////////////
 
-bool AreTableLayoutsEqual(const Trace::ITableLayout& TableLayoutA, const Trace::ITableLayout& TableLayoutB)
+bool AreTableLayoutsEqual(const TraceServices::ITableLayout& TableLayoutA, const TraceServices::ITableLayout& TableLayoutB)
 {
 	if (TableLayoutA.GetColumnCount() != TableLayoutB.GetColumnCount())
 	{
@@ -125,7 +125,7 @@ bool AreTableLayoutsEqual(const Trace::ITableLayout& TableLayoutA, const Trace::
 
 ////////////////////////////////////////////////////////////////////////////////////////////////////
 
-bool FUntypedTable::UpdateSourceTable(TSharedPtr<Trace::IUntypedTable> InSourceTable)
+bool FUntypedTable::UpdateSourceTable(TSharedPtr<TraceServices::IUntypedTable> InSourceTable)
 {
 	bool bTableLayoutChanged;
 
@@ -156,7 +156,7 @@ bool FUntypedTable::UpdateSourceTable(TSharedPtr<Trace::IUntypedTable> InSourceT
 
 ////////////////////////////////////////////////////////////////////////////////////////////////////
 
-void FUntypedTable::CreateColumns(const Trace::ITableLayout& TableLayout)
+void FUntypedTable::CreateColumns(const TraceServices::ITableLayout& TableLayout)
 {
 	ensure(GetColumnCount() == 0);
 	const int32 ColumnCount = TableLayout.GetColumnCount();
@@ -170,8 +170,8 @@ void FUntypedTable::CreateColumns(const Trace::ITableLayout& TableLayout)
 	// Look for first string column.
 	//for (int32 ColumnIndex = 0; ColumnIndex < ColumnCount; ++ColumnIndex)
 	//{
-	//	Trace::ETableColumnType ColumnType = TableLayout.GetColumnType(ColumnIndex);
-	//	if (ColumnType == Trace::TableColumnType_CString)
+	//	TraceServices::ETableColumnType ColumnType = TableLayout.GetColumnType(ColumnIndex);
+	//	if (ColumnType == TraceServices::TableColumnType_CString)
 	//	{
 	//		HierarchyColumnIndex = ColumnIndex;
 	//		HierarchyColumnName = TableLayout.GetColumnName(ColumnIndex);
@@ -185,7 +185,7 @@ void FUntypedTable::CreateColumns(const Trace::ITableLayout& TableLayout)
 
 	for (int32 ColumnIndex = 0; ColumnIndex < ColumnCount; ++ColumnIndex)
 	{
-		Trace::ETableColumnType ColumnType = TableLayout.GetColumnType(ColumnIndex);
+		TraceServices::ETableColumnType ColumnType = TableLayout.GetColumnType(ColumnIndex);
 		const TCHAR* ColumnName = TableLayout.GetColumnName(ColumnIndex);
 
 		TSharedRef<FTableColumn> ColumnRef = MakeShared<FTableColumn>(FName(ColumnName));
@@ -210,7 +210,7 @@ void FUntypedTable::CreateColumns(const Trace::ITableLayout& TableLayout)
 
 		switch (ColumnType)
 		{
-		case Trace::TableColumnType_Bool:
+		case TraceServices::TableColumnType_Bool:
 			Column.SetDataType(ETableCellDataType::Bool);
 			HorizontalAlignment = HAlign_Right;
 			InitialColumnWidth = 40.0f;
@@ -220,7 +220,7 @@ void FUntypedTable::CreateColumns(const Trace::ITableLayout& TableLayout)
 			SorterPtr = MakeShared<FSorterByBoolValue>(ColumnRef);
 			break;
 
-		case Trace::TableColumnType_Int:
+		case TraceServices::TableColumnType_Int:
 			Column.SetDataType(ETableCellDataType::Int64);
 			Aggregation = ETableColumnAggregation::Sum;
 			HorizontalAlignment = HAlign_Right;
@@ -234,7 +234,7 @@ void FUntypedTable::CreateColumns(const Trace::ITableLayout& TableLayout)
 			SorterPtr = MakeShared<FSorterByInt64Value>(ColumnRef);
 			break;
 
-		case Trace::TableColumnType_Float:
+		case TraceServices::TableColumnType_Float:
 			Column.SetDataType(ETableCellDataType::Float);
 			Aggregation = ETableColumnAggregation::Sum;
 			HorizontalAlignment = HAlign_Right;
@@ -245,7 +245,7 @@ void FUntypedTable::CreateColumns(const Trace::ITableLayout& TableLayout)
 			SorterPtr = MakeShared<FSorterByFloatValue>(ColumnRef);
 			break;
 
-		case Trace::TableColumnType_Double:
+		case TraceServices::TableColumnType_Double:
 			Column.SetDataType(ETableCellDataType::Double);
 			Aggregation = ETableColumnAggregation::Sum;
 			HorizontalAlignment = HAlign_Right;
@@ -256,7 +256,7 @@ void FUntypedTable::CreateColumns(const Trace::ITableLayout& TableLayout)
 			SorterPtr = MakeShared<FSorterByDoubleValue>(ColumnRef);
 			break;
 
-		case Trace::TableColumnType_CString:
+		case TraceServices::TableColumnType_CString:
 			Column.SetDataType(ETableCellDataType::CString);
 			HorizontalAlignment = HAlign_Left;
 			InitialColumnWidth = FMath::Max(120.0f, 6.0f * ColumnNameStr.Len());
diff --git a/Engine/Source/Developer/TraceInsights/Private/Insights/Table/ViewModels/UntypedTable.h b/Engine/Source/Developer/TraceInsights/Private/Insights/Table/ViewModels/UntypedTable.h
--- a/Engine/Source/Developer/TraceInsights/Private/Insights/Table/ViewModels/UntypedTable.h
+++ b/Engine/Source/Developer/TraceInsights/Private/Insights/Table/ViewModels/UntypedTable.h
@@ -6,11 +6,13 @@
 
 #include "Insights/Table/ViewModels/Table.h"
 
-namespace Trace
+namespace TraceServices
 {
-	class ITableLayout;
-	class IUntypedTable;
-	class IUntypedTableReader;
+
+class ITableLayout;
+class IUntypedTable;
+class IUntypedTableReader;
+
 }
 
 namespace Insights
@@ -28,18 +30,18 @@ public:
 
 	virtual void Reset();
 
-	TSharedPtr<Trace::IUntypedTable> GetSourceTable() const { return SourceTable; }
-	TSharedPtr<Trace::IUntypedTableReader> GetTableReader() const { return TableReader; }
+	TSharedPtr<TraceServices::IUntypedTable> GetSourceTable() const { return SourceTable; }
+	TSharedPtr<TraceServices::IUntypedTableReader> GetTableReader() const { return TableReader; }
 
 	/* Update table content. Returns true if the table layout has changed. */
-	bool UpdateSourceTable(TSharedPtr<Trace::IUntypedTable> InSourceTable);
+	bool UpdateSourceTable(TSharedPtr<TraceServices::IUntypedTable> InSourceTable);
 
 private:
-	void CreateColumns(const Trace::ITableLayout& TableLayout);
+	void CreateColumns(const TraceServices::ITableLayout& TableLayout);
 
 private:
-	TSharedPtr<Trace::IUntypedTable> SourceTable;
-	TSharedPtr<Trace::IUntypedTableReader> TableReader;
+	TSharedPtr<TraceServices::IUntypedTable> SourceTable;
+	TSharedPtr<TraceServices::IUntypedTableReader> TableReader;
 };
 
 ////////////////////////////////////////////////////////////////////////////////////////////////////
diff --git a/Engine/Source/Developer/TraceInsights/Private/Insights/Table/Widgets/STableTreeView.cpp b/Engine/Source/Developer/TraceInsights/Private/Insights/Table/Widgets/STableTreeView.cpp
--- a/Engine/Source/Developer/TraceInsights/Private/Insights/Table/Widgets/STableTreeView.cpp
+++ b/Engine/Source/Developer/TraceInsights/Private/Insights/Table/Widgets/STableTreeView.cpp
@@ -564,7 +564,7 @@ TSharedRef<SWidget> STableTreeView::TreeViewHeaderRow_GenerateColumnMenu(const F
 
 void STableTreeView::InsightsManager_OnSessionChanged()
 {
-	TSharedPtr<const Trace::IAnalysisSession> NewSession = FInsightsManager::Get()->GetSession();
+	TSharedPtr<const TraceServices::IAnalysisSession> NewSession = FInsightsManager::Get()->GetSession();
 
 	if (NewSession != Session)
 	{
diff --git a/Engine/Source/Developer/TraceInsights/Private/Insights/Table/Widgets/STableTreeView.h b/Engine/Source/Developer/TraceInsights/Private/Insights/Table/Widgets/STableTreeView.h
--- a/Engine/Source/Developer/TraceInsights/Private/Insights/Table/Widgets/STableTreeView.h
+++ b/Engine/Source/Developer/TraceInsights/Private/Insights/Table/Widgets/STableTreeView.h
@@ -20,7 +20,7 @@
 
 class FMenuBuilder;
 
-namespace Trace
+namespace TraceServices
 {
 	class IAnalysisSession;
 }
@@ -251,7 +251,7 @@ protected:
 	TSharedPtr<Insights::FTable> Table;
 
 	/** A weak pointer to the profiler session used to populate this widget. */
-	TSharedPtr<const Trace::IAnalysisSession>/*Weak*/ Session;
+	TSharedPtr<const TraceServices::IAnalysisSession>/*Weak*/ Session;
 
 	//////////////////////////////////////////////////
 	// Tree View, Columns
diff --git a/Engine/Source/Developer/TraceInsights/Private/Insights/Table/Widgets/SUntypedTableTreeView.cpp b/Engine/Source/Developer/TraceInsights/Private/Insights/Table/Widgets/SUntypedTableTreeView.cpp
--- a/Engine/Source/Developer/TraceInsights/Private/Insights/Table/Widgets/SUntypedTableTreeView.cpp
+++ b/Engine/Source/Developer/TraceInsights/Private/Insights/Table/Widgets/SUntypedTableTreeView.cpp
@@ -43,7 +43,7 @@ void SUntypedTableTreeView::Reset()
 
 ////////////////////////////////////////////////////////////////////////////////////////////////////
 
-void SUntypedTableTreeView::UpdateSourceTable(TSharedPtr<Trace::IUntypedTable> SourceTable)
+void SUntypedTableTreeView::UpdateSourceTable(TSharedPtr<TraceServices::IUntypedTable> SourceTable)
 {
 	//check(Table->Is<Insights::FUntypedTable>());
 	TSharedPtr<Insights::FUntypedTable> UntypedTable = StaticCastSharedPtr<Insights::FUntypedTable>(Table);
@@ -74,8 +74,8 @@ void SUntypedTableTreeView::RebuildTree(bool bResync)
 	//check(Table->Is<Insights::FUntypedTable>());
 	TSharedPtr<Insights::FUntypedTable> UntypedTable = StaticCastSharedPtr<Insights::FUntypedTable>(Table);
 
-	TSharedPtr<Trace::IUntypedTable> SourceTable = UntypedTable->GetSourceTable();
-	TSharedPtr<Trace::IUntypedTableReader> TableReader = UntypedTable->GetTableReader();
+	TSharedPtr<TraceServices::IUntypedTable> SourceTable = UntypedTable->GetSourceTable();
+	TSharedPtr<TraceServices::IUntypedTableReader> TableReader = UntypedTable->GetTableReader();
 
 	SyncStopwatch.Start();
 	if (Session.IsValid() && SourceTable.IsValid() && TableReader.IsValid())
diff --git a/Engine/Source/Developer/TraceInsights/Private/Insights/Table/Widgets/SUntypedTableTreeView.h b/Engine/Source/Developer/TraceInsights/Private/Insights/Table/Widgets/SUntypedTableTreeView.h
--- a/Engine/Source/Developer/TraceInsights/Private/Insights/Table/Widgets/SUntypedTableTreeView.h
+++ b/Engine/Source/Developer/TraceInsights/Private/Insights/Table/Widgets/SUntypedTableTreeView.h
@@ -33,7 +33,7 @@ public:
 
 	TSharedPtr<Insights::FUntypedTable> GetUntypedTable() const { return StaticCastSharedPtr<Insights::FUntypedTable>(GetTable()); }
 
-	void UpdateSourceTable(TSharedPtr<Trace::IUntypedTable> SourceTable);
+	void UpdateSourceTable(TSharedPtr<TraceServices::IUntypedTable> SourceTable);
 
 	virtual void Reset();
 
diff --git a/Engine/Source/Developer/TraceInsights/Private/Insights/Tests/TimingProfilerTests.cpp b/Engine/Source/Developer/TraceInsights/Private/Insights/Tests/TimingProfilerTests.cpp
--- a/Engine/Source/Developer/TraceInsights/Private/Insights/Tests/TimingProfilerTests.cpp
+++ b/Engine/Source/Developer/TraceInsights/Private/Insights/Tests/TimingProfilerTests.cpp
@@ -21,31 +21,31 @@ bool EnumerateEventsToFile::RunTest(const FString& Parameters)
 	double SessionTime = 0.0;
 	uint32 TimelineIndex = (uint32)-1;
 
-	TSharedPtr<const Trace::IAnalysisSession> Session = FInsightsManager::Get()->GetSession();
-	if (Session.IsValid() && Trace::ReadTimingProfilerProvider(*Session.Get()))
+	TSharedPtr<const TraceServices::IAnalysisSession> Session = FInsightsManager::Get()->GetSession();
+	if (Session.IsValid() && TraceServices::ReadTimingProfilerProvider(*Session.Get()))
 	{
-		Trace::FAnalysisSessionReadScope SessionReadScope(*Session.Get());
+		TraceServices::FAnalysisSessionReadScope SessionReadScope(*Session.Get());
 
-		const Trace::ITimingProfilerProvider& TimingProfilerProvider = *Trace::ReadTimingProfilerProvider(*Session.Get());
+		const TraceServices::ITimingProfilerProvider& TimingProfilerProvider = *TraceServices::ReadTimingProfilerProvider(*Session.Get());
 
 		TimelineIndex = FTimingProfilerTests::GetTimelineIndex(TEXT("GameThread"));
 
 		AddErrorIfFalse(TimelineIndex != (uint32)-1, TEXT("Failed to get track named GameThread"));
 
-		const Trace::ITimingProfilerTimerReader* TimerReader;
-		TimingProfilerProvider.ReadTimers([&TimerReader](const Trace::ITimingProfilerTimerReader& Out) { TimerReader = &Out; });
+		const TraceServices::ITimingProfilerTimerReader* TimerReader;
+		TimingProfilerProvider.ReadTimers([&TimerReader](const TraceServices::ITimingProfilerTimerReader& Out) { TimerReader = &Out; });
 
 		FString EnumeratedEventsStr;
 		TimingProfilerProvider.ReadTimeline(TimelineIndex,
-			[Session, &EnumeratedEventsStr, TimerReader](const Trace::ITimingProfilerProvider::Timeline& Timeline)
+			[Session, &EnumeratedEventsStr, TimerReader](const TraceServices::ITimingProfilerProvider::Timeline& Timeline)
 			{
 				Timeline.EnumerateEvents(-1.0, Session->GetDurationSeconds() + 1.0,
-					[&EnumeratedEventsStr, TimerReader](double EventStartTime, double EventEndTime, uint32 EventDepth, const Trace::FTimingProfilerEvent& Event)
+					[&EnumeratedEventsStr, TimerReader](double EventStartTime, double EventEndTime, uint32 EventDepth, const TraceServices::FTimingProfilerEvent& Event)
 					{
-						const Trace::FTimingProfilerTimer* Timer = TimerReader->GetTimer(Event.TimerIndex);
+						const TraceServices::FTimingProfilerTimer* Timer = TimerReader->GetTimer(Event.TimerIndex);
 
 						EnumeratedEventsStr.Appendf(TEXT("%s %f %f %d\n"), Timer->Name, EventStartTime, EventEndTime, EventDepth);
-						return Trace::EEventEnumerate::Continue;
+						return TraceServices::EEventEnumerate::Continue;
 					});
 			});
 
@@ -64,31 +64,31 @@ bool EnumerateScopesToFile::RunTest(const FString& Parameters)
 	double SessionTime = 0.0;
 	uint32 TimelineIndex = (uint32)-1;
 
-	TSharedPtr<const Trace::IAnalysisSession> Session = FInsightsManager::Get()->GetSession();
-	if (Session.IsValid() && Trace::ReadTimingProfilerProvider(*Session.Get()))
+	TSharedPtr<const TraceServices::IAnalysisSession> Session = FInsightsManager::Get()->GetSession();
+	if (Session.IsValid() && TraceServices::ReadTimingProfilerProvider(*Session.Get()))
 	{
-		Trace::FAnalysisSessionReadScope SessionReadScope(*Session.Get());
+		TraceServices::FAnalysisSessionReadScope SessionReadScope(*Session.Get());
 
-		const Trace::ITimingProfilerProvider& TimingProfilerProvider = *Trace::ReadTimingProfilerProvider(*Session.Get());
+		const TraceServices::ITimingProfilerProvider& TimingProfilerProvider = *TraceServices::ReadTimingProfilerProvider(*Session.Get());
 
 		TimelineIndex = FTimingProfilerTests::GetTimelineIndex(TEXT("GameThread"));
 
 		AddErrorIfFalse(TimelineIndex != (uint32)-1, TEXT("Failed to get track named GameThread"));
 
-		const Trace::ITimingProfilerTimerReader* TimerReader;
-		TimingProfilerProvider.ReadTimers([&TimerReader](const Trace::ITimingProfilerTimerReader& Out) { TimerReader = &Out; });
+		const TraceServices::ITimingProfilerTimerReader* TimerReader;
+		TimingProfilerProvider.ReadTimers([&TimerReader](const TraceServices::ITimingProfilerTimerReader& Out) { TimerReader = &Out; });
 
 		FString EnumeratedEventsStr;
 		TimingProfilerProvider.ReadTimeline(TimelineIndex,
-			[Session, &EnumeratedEventsStr, TimerReader](const Trace::ITimingProfilerProvider::Timeline& Timeline)
+			[Session, &EnumeratedEventsStr, TimerReader](const TraceServices::ITimingProfilerProvider::Timeline& Timeline)
 			{
 				Timeline.EnumerateEvents(-1.0, Session->GetDurationSeconds() + 1.0,
-					[&EnumeratedEventsStr, TimerReader](bool bStart, double Time, const Trace::FTimingProfilerEvent& Event)
+					[&EnumeratedEventsStr, TimerReader](bool bStart, double Time, const TraceServices::FTimingProfilerEvent& Event)
 					{
-						const Trace::FTimingProfilerTimer* Timer = TimerReader->GetTimer(Event.TimerIndex);
+						const TraceServices::FTimingProfilerTimer* Timer = TimerReader->GetTimer(Event.TimerIndex);
 
 						EnumeratedEventsStr.Appendf(TEXT("%s %d %f\n"), Timer->Name, (int32)bStart, Time);
-						return Trace::EEventEnumerate::Continue;
+						return TraceServices::EEventEnumerate::Continue;
 					});
 			});
 
@@ -107,7 +107,7 @@ bool EnumerateByEndTimeAsyncAllTracks::RunTest(const FString& Parameters)
 	FTimingProfilerTests::FEnumerateTestParams Params;
 	Params.Interval = 24.0 * 3600; //A day - Should be big enought to contain any valid session in [0, 0 + interval]
 	Params.NumEnumerations = 1;
-	Params.SortOrder = Trace::EEventSortOrder::ByEndTime;
+	Params.SortOrder = TraceServices::EEventSortOrder::ByEndTime;
 
 	FTimingProfilerTests::RunEnumerateSyncAsyncComparisonTest(*this, Params, false);
 
@@ -122,7 +122,7 @@ bool EnumerateByEndTimeAsyncGameThreadTrack::RunTest(const FString& Parameters)
 	FTimingProfilerTests::FEnumerateTestParams Params;
 	Params.Interval = 24.0 * 3600; //A day - Should be big enought to contain any valid session in [0, 0 + interval]
 	Params.NumEnumerations = 1;
-	Params.SortOrder = Trace::EEventSortOrder::ByEndTime;
+	Params.SortOrder = TraceServices::EEventSortOrder::ByEndTime;
 
 	FTimingProfilerTests::RunEnumerateSyncAsyncComparisonTest(*this, Params, true);
 
@@ -137,7 +137,7 @@ bool EnumerateByEndTimeAllTracks10sIntervals::RunTest(const FString& Parameters)
 	FTimingProfilerTests::FEnumerateTestParams Params;
 	Params.Interval = 10.0;
 	Params.NumEnumerations = 100;
-	Params.SortOrder = Trace::EEventSortOrder::ByEndTime;
+	Params.SortOrder = TraceServices::EEventSortOrder::ByEndTime;
 
 	FTimingProfilerTests::RunEnumerateSyncAsyncComparisonTest(*this, Params, false);
 
@@ -152,7 +152,7 @@ bool EnumerateByEndTimeAllTracks5sIntervals::RunTest(const FString& Parameters)
 	FTimingProfilerTests::FEnumerateTestParams Params;
 	Params.Interval = 5.0;
 	Params.NumEnumerations = 200;
-	Params.SortOrder = Trace::EEventSortOrder::ByEndTime;
+	Params.SortOrder = TraceServices::EEventSortOrder::ByEndTime;
 
 	FTimingProfilerTests::RunEnumerateSyncAsyncComparisonTest(*this, Params, false);
 
@@ -169,7 +169,7 @@ bool EnumerateByStartTimeAsyncAllTracks::RunTest(const FString& Parameters)
 	FTimingProfilerTests::FEnumerateTestParams Params;
 	Params.Interval = 24.0 * 3600; //A day - Should be big enought to contain any valid session in [0, 0 + interval]
 	Params.NumEnumerations = 1;
-	Params.SortOrder = Trace::EEventSortOrder::ByStartTime;
+	Params.SortOrder = TraceServices::EEventSortOrder::ByStartTime;
 
 	FTimingProfilerTests::RunEnumerateSyncAsyncComparisonTest(*this, Params, false);
 
@@ -184,7 +184,7 @@ bool EnumerateByStartTimeAsyncGameThreadTrack::RunTest(const FString& Parameters
 	FTimingProfilerTests::FEnumerateTestParams Params;
 	Params.Interval = 24.0 * 3600; //A day - Should be big enought to contain any valid session in [0, 0 + interval]
 	Params.NumEnumerations = 1;
-	Params.SortOrder = Trace::EEventSortOrder::ByStartTime;
+	Params.SortOrder = TraceServices::EEventSortOrder::ByStartTime;
 
 	FTimingProfilerTests::RunEnumerateSyncAsyncComparisonTest(*this, Params, true);
 
@@ -199,7 +199,7 @@ bool EnumerateByStartTimeAllTracks10sIntervals::RunTest(const FString& Parameter
 	FTimingProfilerTests::FEnumerateTestParams Params;
 	Params.Interval = 10.0;
 	Params.NumEnumerations = 100;
-	Params.SortOrder = Trace::EEventSortOrder::ByStartTime;
+	Params.SortOrder = TraceServices::EEventSortOrder::ByStartTime;
 
 	FTimingProfilerTests::RunEnumerateSyncAsyncComparisonTest(*this, Params, false);
 
@@ -214,7 +214,7 @@ bool EnumerateByStartTimeAllTracks5sIntervals::RunTest(const FString& Parameters
 	FTimingProfilerTests::FEnumerateTestParams Params;
 	Params.Interval = 5.0;
 	Params.NumEnumerations = 200;
-	Params.SortOrder = Trace::EEventSortOrder::ByStartTime;
+	Params.SortOrder = TraceServices::EEventSortOrder::ByStartTime;
 
 	FTimingProfilerTests::RunEnumerateSyncAsyncComparisonTest(*this, Params, false);
 
@@ -235,12 +235,12 @@ void FTimingProfilerTests::RunEnumerateBenchmark(const FEnumerateTestParams& InP
 	double SessionTime = 0.0;
 	uint32 TimelineIndex = (uint32)-1;
 
-	TSharedPtr<const Trace::IAnalysisSession> Session = FInsightsManager::Get()->GetSession();
-	if (Session.IsValid() && Trace::ReadTimingProfilerProvider(*Session.Get()))
+	TSharedPtr<const TraceServices::IAnalysisSession> Session = FInsightsManager::Get()->GetSession();
+	if (Session.IsValid() && TraceServices::ReadTimingProfilerProvider(*Session.Get()))
 	{
-		Trace::FAnalysisSessionReadScope SessionReadScope(*Session.Get());
+		TraceServices::FAnalysisSessionReadScope SessionReadScope(*Session.Get());
 
-		const Trace::ITimingProfilerProvider& TimingProfilerProvider = *Trace::ReadTimingProfilerProvider(*Session.Get());
+		const TraceServices::ITimingProfilerProvider& TimingProfilerProvider = *TraceServices::ReadTimingProfilerProvider(*Session.Get());
 
 		SessionTime = Session->GetDurationSeconds();
 		OutCheckValues.SessionDuration = SessionTime;
@@ -250,19 +250,19 @@ void FTimingProfilerTests::RunEnumerateBenchmark(const FEnumerateTestParams& InP
 		TimelineIndex = GetTimelineIndex(TEXT("GameThread"));
 
 		TimingProfilerProvider.ReadTimeline(TimelineIndex,
-			[&OutCheckValues, &InParams, TimeIncrement](const Trace::ITimingProfilerProvider::Timeline& Timeline)
+			[&OutCheckValues, &InParams, TimeIncrement](const TraceServices::ITimingProfilerProvider::Timeline& Timeline)
 			{
 				double Time = 0.0;
 				for (int32 Index = 0; Index < InParams.NumEnumerations; ++Index)
 				{
 					Timeline.EnumerateEvents(Time, Time + InParams.Interval,
-						[&OutCheckValues](double EventStartTime, double EventEndTime, uint32 EventDepth, const Trace::FTimingProfilerEvent& Event)
+						[&OutCheckValues](double EventStartTime, double EventEndTime, uint32 EventDepth, const TraceServices::FTimingProfilerEvent& Event)
 						{
 							OutCheckValues.TotalEventDuration += EventEndTime - EventStartTime;
 							++OutCheckValues.EventCount;
 							OutCheckValues.SumDepth += EventDepth;
 							OutCheckValues.SumTimerIndex += Event.TimerIndex;
-							return Trace::EEventEnumerate::Continue;
+							return TraceServices::EEventEnumerate::Continue;
 						});
 
 					Time += TimeIncrement;
@@ -289,12 +289,12 @@ void FTimingProfilerTests::RunEnumerateAsyncBenchmark(const FEnumerateTestParams
 	double SessionTime = 0.0;
 	uint32 TimelineIndex = (uint32)-1;
 
-	TSharedPtr<const Trace::IAnalysisSession> Session = FInsightsManager::Get()->GetSession();
-	if (Session.IsValid() && Trace::ReadTimingProfilerProvider(*Session.Get()))
+	TSharedPtr<const TraceServices::IAnalysisSession> Session = FInsightsManager::Get()->GetSession();
+	if (Session.IsValid() && TraceServices::ReadTimingProfilerProvider(*Session.Get()))
 	{
-		Trace::FAnalysisSessionReadScope SessionReadScope(*Session.Get());
+		TraceServices::FAnalysisSessionReadScope SessionReadScope(*Session.Get());
 
-		const Trace::ITimingProfilerProvider& TimingProfilerProvider = *Trace::ReadTimingProfilerProvider(*Session.Get());
+		const TraceServices::ITimingProfilerProvider& TimingProfilerProvider = *TraceServices::ReadTimingProfilerProvider(*Session.Get());
 
 		SessionTime = Session->GetDurationSeconds();
 		OutCheckValues.SessionDuration = SessionTime;
@@ -306,12 +306,12 @@ void FTimingProfilerTests::RunEnumerateAsyncBenchmark(const FEnumerateTestParams
 		TArray<FCheckValues> TaskCheckValues;
 
 		TimingProfilerProvider.ReadTimeline(TimelineIndex,
-			[&OutCheckValues, &InParams, TimeIncrement, &TaskCheckValues](const Trace::ITimingProfilerProvider::Timeline& Timeline)
+			[&OutCheckValues, &InParams, TimeIncrement, &TaskCheckValues](const TraceServices::ITimingProfilerProvider::Timeline& Timeline)
 			{
 				double Time = 0.0;
 				for (int32 Index = 0; Index < InParams.NumEnumerations; ++Index)
 				{
-					Trace::ITimeline<Trace::FTimingProfilerEvent>::EnumerateAsyncParams Params;
+					TraceServices::ITimeline<TraceServices::FTimingProfilerEvent>::EnumerateAsyncParams Params;
 					Params.IntervalStart = Time;
 					Params.IntervalEnd = Time + InParams.Interval;
 					Params.Resolution = 0.0;
@@ -320,13 +320,13 @@ void FTimingProfilerTests::RunEnumerateAsyncBenchmark(const FEnumerateTestParams
 					{
 						TaskCheckValues.AddDefaulted(NumTasks);
 					};
-					Params.Callback = [&TaskCheckValues](double EventStartTime, double EventEndTime, uint32 EventDepth, const Trace::FTimingProfilerEvent& Event, uint32 TaskIndex)
+					Params.Callback = [&TaskCheckValues](double EventStartTime, double EventEndTime, uint32 EventDepth, const TraceServices::FTimingProfilerEvent& Event, uint32 TaskIndex)
 					{
 						TaskCheckValues[TaskIndex].TotalEventDuration += EventEndTime - EventStartTime;
 						++TaskCheckValues[TaskIndex].EventCount;
 						TaskCheckValues[TaskIndex].SumDepth += EventDepth;
 						TaskCheckValues[TaskIndex].SumTimerIndex += Event.TimerIndex;
-						return Trace::EEventEnumerate::Continue;
+						return TraceServices::EEventEnumerate::Continue;
 					};
 
 					Timeline.EnumerateEventsDownSampledAsync(Params);
@@ -364,38 +364,38 @@ void FTimingProfilerTests::RunEnumerateAllTracksBenchmark(const FEnumerateTestPa
 	double SessionTime = 0.0;
 	uint32 TimelineIndex = (uint32)-1;
 
-	TSharedPtr<const Trace::IAnalysisSession> Session = FInsightsManager::Get()->GetSession();
-	if (Session.IsValid() && Trace::ReadTimingProfilerProvider(*Session.Get()))
+	TSharedPtr<const TraceServices::IAnalysisSession> Session = FInsightsManager::Get()->GetSession();
+	if (Session.IsValid() && TraceServices::ReadTimingProfilerProvider(*Session.Get()))
 	{
-		Trace::FAnalysisSessionReadScope SessionReadScope(*Session.Get());
+		TraceServices::FAnalysisSessionReadScope SessionReadScope(*Session.Get());
 
-		const Trace::ITimingProfilerProvider& TimingProfilerProvider = *Trace::ReadTimingProfilerProvider(*Session.Get());
+		const TraceServices::ITimingProfilerProvider& TimingProfilerProvider = *TraceServices::ReadTimingProfilerProvider(*Session.Get());
 
 		SessionTime = Session->GetDurationSeconds();
 		OutCheckValues.SessionDuration = SessionTime;
 
 		const double TimeIncrement = SessionTime / static_cast<double>(InParams.NumEnumerations);
 
-		const Trace::IThreadProvider& ThreadProvider = Trace::ReadThreadProvider(*Session.Get());
+		const TraceServices::IThreadProvider& ThreadProvider = TraceServices::ReadThreadProvider(*Session.Get());
 		ThreadProvider.EnumerateThreads(
-			[&TimelineIndex, &TimingProfilerProvider, &OutCheckValues, &InParams, TimeIncrement](const Trace::FThreadInfo& ThreadInfo)
+			[&TimelineIndex, &TimingProfilerProvider, &OutCheckValues, &InParams, TimeIncrement](const TraceServices::FThreadInfo& ThreadInfo)
 			{
 				TimingProfilerProvider.GetCpuThreadTimelineIndex(ThreadInfo.Id, TimelineIndex);
 
 				TimingProfilerProvider.ReadTimeline(TimelineIndex,
-					[&OutCheckValues, &InParams, TimeIncrement](const Trace::ITimingProfilerProvider::Timeline& Timeline)
+					[&OutCheckValues, &InParams, TimeIncrement](const TraceServices::ITimingProfilerProvider::Timeline& Timeline)
 					{
 						double Time = 0.0;
 						for (int32 Index = 0; Index < InParams.NumEnumerations; ++Index)
 						{
 							Timeline.EnumerateEvents(Time, Time + InParams.Interval,
-								[&OutCheckValues](double EventStartTime, double EventEndTime, uint32 EventDepth, const Trace::FTimingProfilerEvent& Event)
+								[&OutCheckValues](double EventStartTime, double EventEndTime, uint32 EventDepth, const TraceServices::FTimingProfilerEvent& Event)
 								{
 									OutCheckValues.TotalEventDuration += EventEndTime - EventStartTime;
 									++OutCheckValues.EventCount;
 									OutCheckValues.SumDepth += EventDepth;
 									OutCheckValues.SumTimerIndex += Event.TimerIndex;
-									return Trace::EEventEnumerate::Continue;
+									return TraceServices::EEventEnumerate::Continue;
 								});
 
 							Time += TimeIncrement;
@@ -424,12 +424,12 @@ void FTimingProfilerTests::RunEnumerateAsyncAllTracksBenchmark(const FEnumerateT
 	double SessionTime = 0.0;
 	uint32 TimelineIndex = (uint32)-1;
 
-	TSharedPtr<const Trace::IAnalysisSession> Session = FInsightsManager::Get()->GetSession();
-	if (Session.IsValid() && Trace::ReadTimingProfilerProvider(*Session.Get()))
+	TSharedPtr<const TraceServices::IAnalysisSession> Session = FInsightsManager::Get()->GetSession();
+	if (Session.IsValid() && TraceServices::ReadTimingProfilerProvider(*Session.Get()))
 	{
-		Trace::FAnalysisSessionReadScope SessionReadScope(*Session.Get());
+		TraceServices::FAnalysisSessionReadScope SessionReadScope(*Session.Get());
 
-		const Trace::ITimingProfilerProvider& TimingProfilerProvider = *Trace::ReadTimingProfilerProvider(*Session.Get());
+		const TraceServices::ITimingProfilerProvider& TimingProfilerProvider = *TraceServices::ReadTimingProfilerProvider(*Session.Get());
 
 		SessionTime = Session->GetDurationSeconds();
 		OutCheckValues.SessionDuration = SessionTime;
@@ -438,19 +438,19 @@ void FTimingProfilerTests::RunEnumerateAsyncAllTracksBenchmark(const FEnumerateT
 
 		TArray<FCheckValues> TaskCheckValues;
 
-		const Trace::IThreadProvider& ThreadProvider = Trace::ReadThreadProvider(*Session.Get());
+		const TraceServices::IThreadProvider& ThreadProvider = TraceServices::ReadThreadProvider(*Session.Get());
 		ThreadProvider.EnumerateThreads(
-			[&TimelineIndex, &TimingProfilerProvider, &OutCheckValues, &InParams, TimeIncrement, &TaskCheckValues](const Trace::FThreadInfo& ThreadInfo)
+			[&TimelineIndex, &TimingProfilerProvider, &OutCheckValues, &InParams, TimeIncrement, &TaskCheckValues](const TraceServices::FThreadInfo& ThreadInfo)
 			{
 				TimingProfilerProvider.GetCpuThreadTimelineIndex(ThreadInfo.Id, TimelineIndex);
 
 				TimingProfilerProvider.ReadTimeline(TimelineIndex,
-					[&OutCheckValues, &InParams, TimeIncrement, &TaskCheckValues](const Trace::ITimingProfilerProvider::Timeline& Timeline)
+					[&OutCheckValues, &InParams, TimeIncrement, &TaskCheckValues](const TraceServices::ITimingProfilerProvider::Timeline& Timeline)
 					{
 						double Time = 0.0;
 						for (int32 Index = 0; Index < InParams.NumEnumerations; ++Index)
 						{
-							Trace::ITimeline<Trace::FTimingProfilerEvent>::EnumerateAsyncParams Params;
+							TraceServices::ITimeline<TraceServices::FTimingProfilerEvent>::EnumerateAsyncParams Params;
 							Params.IntervalStart = Time;
 							Params.IntervalEnd = Time + InParams.Interval;
 							Params.Resolution = 0.0;
@@ -459,13 +459,13 @@ void FTimingProfilerTests::RunEnumerateAsyncAllTracksBenchmark(const FEnumerateT
 							{
 								TaskCheckValues.AddDefaulted(NumTasks);
 							};
-							Params.Callback = [&OutCheckValues, &TaskCheckValues](double EventStartTime, double EventEndTime, uint32 EventDepth, const Trace::FTimingProfilerEvent& Event, uint32 TaskIndex)
+							Params.Callback = [&OutCheckValues, &TaskCheckValues](double EventStartTime, double EventEndTime, uint32 EventDepth, const TraceServices::FTimingProfilerEvent& Event, uint32 TaskIndex)
 							{
 								TaskCheckValues[TaskIndex].TotalEventDuration += EventEndTime - EventStartTime;
 								++TaskCheckValues[TaskIndex].EventCount;
 								TaskCheckValues[TaskIndex].SumDepth += EventDepth;
 								TaskCheckValues[TaskIndex].SumTimerIndex += Event.TimerIndex;
-								return Trace::EEventEnumerate::Continue;
+								return TraceServices::EEventEnumerate::Continue;
 							};
 
 							Timeline.EnumerateEventsDownSampledAsync(Params);
@@ -497,12 +497,12 @@ void FTimingProfilerTests::RunEnumerateAsyncAllTracksBenchmark(const FEnumerateT
 
 uint32 FTimingProfilerTests::GetTimelineIndex(const TCHAR* InName)
 {
-	TSharedPtr<const Trace::IAnalysisSession> Session = FInsightsManager::Get()->GetSession();
-	const Trace::ITimingProfilerProvider& TimingProfilerProvider = *Trace::ReadTimingProfilerProvider(*Session.Get());
-	const Trace::IThreadProvider& ThreadProvider = Trace::ReadThreadProvider(*Session.Get());
+	TSharedPtr<const TraceServices::IAnalysisSession> Session = FInsightsManager::Get()->GetSession();
+	const TraceServices::ITimingProfilerProvider& TimingProfilerProvider = *TraceServices::ReadTimingProfilerProvider(*Session.Get());
+	const TraceServices::IThreadProvider& ThreadProvider = TraceServices::ReadThreadProvider(*Session.Get());
 	uint32 TimelineIndex = (uint32) -1;
 	ThreadProvider.EnumerateThreads(
-		[&TimelineIndex, &TimingProfilerProvider, InName](const Trace::FThreadInfo& ThreadInfo)
+		[&TimelineIndex, &TimingProfilerProvider, InName](const TraceServices::FThreadInfo& ThreadInfo)
 		{
 			if (!FCString::Strcmp(ThreadInfo.Name, InName))
 			{
diff --git a/Engine/Source/Developer/TraceInsights/Private/Insights/TimingProfilerManager.cpp b/Engine/Source/Developer/TraceInsights/Private/Insights/TimingProfilerManager.cpp
--- a/Engine/Source/Developer/TraceInsights/Private/Insights/TimingProfilerManager.cpp
+++ b/Engine/Source/Developer/TraceInsights/Private/Insights/TimingProfilerManager.cpp
@@ -233,7 +233,7 @@ bool FTimingProfilerManager::Tick(float DeltaTime)
 	// Check if session has Timing events (to spawn the tab), but not too often.
 	if (!bIsAvailable && AvailabilityCheck.Tick())
 	{
-		TSharedPtr<const Trace::IAnalysisSession> Session = FInsightsManager::Get()->GetSession();
+		TSharedPtr<const TraceServices::IAnalysisSession> Session = FInsightsManager::Get()->GetSession();
 		if (Session.IsValid())
 		{
 			bIsAvailable = true;
@@ -269,16 +269,16 @@ void FTimingProfilerManager::FinishTimerButterflyAggregation()
 		TSharedPtr<STimerTreeView> CallersTreeView = Wnd->GetCallersTreeView();
 		if (CallersTreeView)
 		{
-			Trace::ITimingProfilerButterfly* TimingProfilerButterfly = TimerButterflyAggregator->GetResultButterfly();
-			const Trace::FTimingProfilerButterflyNode& Callers = TimingProfilerButterfly->GenerateCallersTree(SelectedTimerId);
+			TraceServices::ITimingProfilerButterfly* TimingProfilerButterfly = TimerButterflyAggregator->GetResultButterfly();
+			const TraceServices::FTimingProfilerButterflyNode& Callers = TimingProfilerButterfly->GenerateCallersTree(SelectedTimerId);
 			CallersTreeView->SetTree(Callers);
 		}
 
 		TSharedPtr<STimerTreeView> CalleesTreeView = Wnd->GetCalleesTreeView();
 		if (CalleesTreeView)
 		{
-			Trace::ITimingProfilerButterfly* TimingProfilerButterfly = TimerButterflyAggregator->GetResultButterfly();
-			const Trace::FTimingProfilerButterflyNode& Callees = TimingProfilerButterfly->GenerateCalleesTree(SelectedTimerId);
+			TraceServices::ITimingProfilerButterfly* TimingProfilerButterfly = TimerButterflyAggregator->GetResultButterfly();
+			const TraceServices::FTimingProfilerButterflyNode& Callees = TimingProfilerButterfly->GenerateCalleesTree(SelectedTimerId);
 			CalleesTreeView->SetTree(Callees);
 		}
 	}
diff --git a/Engine/Source/Developer/TraceInsights/Private/Insights/TraceInsightsModule.cpp b/Engine/Source/Developer/TraceInsights/Private/Insights/TraceInsightsModule.cpp
--- a/Engine/Source/Developer/TraceInsights/Private/Insights/TraceInsightsModule.cpp
+++ b/Engine/Source/Developer/TraceInsights/Private/Insights/TraceInsightsModule.cpp
@@ -409,7 +409,7 @@ void FTraceInsightsModule::SetUnrealInsightsLayoutIni(const FString& InIniPath)
 
 ////////////////////////////////////////////////////////////////////////////////////////////////////
 
-TSharedPtr<const Trace::IAnalysisSession> FTraceInsightsModule::GetAnalysisSession() const
+TSharedPtr<const TraceServices::IAnalysisSession> FTraceInsightsModule::GetAnalysisSession() const
 {
 	return FInsightsManager::Get()->GetSession();
 }
diff --git a/Engine/Source/Developer/TraceInsights/Private/Insights/TraceInsightsModule.h b/Engine/Source/Developer/TraceInsights/Private/Insights/TraceInsightsModule.h
--- a/Engine/Source/Developer/TraceInsights/Private/Insights/TraceInsightsModule.h
+++ b/Engine/Source/Developer/TraceInsights/Private/Insights/TraceInsightsModule.h
@@ -9,6 +9,10 @@
 namespace Trace
 {
 	class FStoreService;
+}
+
+namespace TraceServices
+{
 	class IAnalysisService;
 	class IModuleService;
 }
@@ -39,7 +43,7 @@ public:
 	virtual void CreateSessionBrowser(bool bAllowDebugTools = false, bool bSingleProcess = false, bool bInitializeTesting = false) override;
 	virtual void CreateSessionViewer(bool bAllowDebugTools = false) override;
 
-	virtual TSharedPtr<const Trace::IAnalysisSession> GetAnalysisSession() const override;
+	virtual TSharedPtr<const TraceServices::IAnalysisSession> GetAnalysisSession() const override;
 	virtual void StartAnalysisForTrace(uint32 InTraceId, bool InAutoQuit = false) override;
 	virtual void StartAnalysisForLastLiveSession() override;
 	virtual void StartAnalysisForTraceFile(const TCHAR* InTraceFile, bool InAutoQuit = false) override;
@@ -91,8 +95,8 @@ protected:
 protected:
 	TUniquePtr<Trace::FStoreService> StoreService;
 
-	TSharedPtr<Trace::IAnalysisService> TraceAnalysisService;
-	TSharedPtr<Trace::IModuleService> TraceModuleService;
+	TSharedPtr<TraceServices::IAnalysisService> TraceAnalysisService;
+	TSharedPtr<TraceServices::IModuleService> TraceModuleService;
 
 	TMap<FName, FInsightsMajorTabConfig> TabConfigs;
 	TMap<FName, FOnRegisterMajorTabExtensions> MajorTabExtensionDelegates;
diff --git a/Engine/Source/Developer/TraceInsights/Private/Insights/ViewModels/CounterAggregation.cpp b/Engine/Source/Developer/TraceInsights/Private/Insights/ViewModels/CounterAggregation.cpp
--- a/Engine/Source/Developer/TraceInsights/Private/Insights/ViewModels/CounterAggregation.cpp
+++ b/Engine/Source/Developer/TraceInsights/Private/Insights/ViewModels/CounterAggregation.cpp
@@ -49,14 +49,14 @@ public:
 		IntervalEndTime = InIntervalEndTime;
 	}
 
-	void Update(uint32 CounterId, const Trace::ICounter& Counter)
+	void Update(uint32 CounterId, const TraceServices::ICounter& Counter)
 	{
 		EnumerateValues(CounterId, Counter, UpdateMinMax);
 	}
 
 	void PrecomputeHistograms();
 
-	void UpdateHistograms(uint32 CounterId, const Trace::ICounter& Counter)
+	void UpdateHistograms(uint32 CounterId, const TraceServices::ICounter& Counter)
 	{
 		EnumerateValues(CounterId, Counter, UpdateHistogram);
 	}
@@ -67,7 +67,7 @@ public:
 
 private:
 	template<typename CallbackType>
-	void EnumerateValues(uint32 CounterId, const Trace::ICounter& Counter, CallbackType Callback);
+	void EnumerateValues(uint32 CounterId, const TraceServices::ICounter& Counter, CallbackType Callback);
 
 	static void UpdateMinMax(TAggregatedStatsEx<Type>& Stats, Type Value);
 	static void UpdateHistogram(TAggregatedStatsEx<Type>& StatsEx, Type Value);
@@ -93,7 +93,7 @@ void TCounterAggregationHelper<Type>::Reset()
 
 template<>
 template<typename CallbackType>
-void TCounterAggregationHelper<double>::EnumerateValues(uint32 CounterId, const Trace::ICounter& Counter, CallbackType Callback)
+void TCounterAggregationHelper<double>::EnumerateValues(uint32 CounterId, const TraceServices::ICounter& Counter, CallbackType Callback)
 {
 	TAggregatedStatsEx<double>* StatsExPtr = StatsMap.Find(CounterId);
 	if (!StatsExPtr)
@@ -116,7 +116,7 @@ void TCounterAggregationHelper<double>::EnumerateValues(uint32 CounterId, const
 
 template<>
 template<typename CallbackType>
-void TCounterAggregationHelper<int64>::EnumerateValues(uint32 CounterId, const Trace::ICounter& Counter, CallbackType Callback)
+void TCounterAggregationHelper<int64>::EnumerateValues(uint32 CounterId, const TraceServices::ICounter& Counter, CallbackType Callback)
 {
 	TAggregatedStatsEx<int64>* StatsExPtr = StatsMap.Find(CounterId);
 	if (!StatsExPtr)
@@ -344,7 +344,7 @@ void TCounterAggregationHelper<int64>::ApplyResultsTo(const TMap<uint32, FStatsN
 class FCounterAggregationWorker : public IStatsAggregationWorker
 {
 public:
-	FCounterAggregationWorker(TSharedPtr<const Trace::IAnalysisSession> InSession, double InStartTime, double InEndTime)
+	FCounterAggregationWorker(TSharedPtr<const TraceServices::IAnalysisSession> InSession, double InStartTime, double InEndTime)
 		: Session(InSession)
 		, StartTime(InStartTime)
 		, EndTime(InEndTime)
@@ -359,7 +359,7 @@ public:
 	void ResetResults();
 
 private:
-	TSharedPtr<const Trace::IAnalysisSession> Session;
+	TSharedPtr<const TraceServices::IAnalysisSession> Session;
 	double StartTime;
 	double EndTime;
 	const bool bComputeMedian = true;
@@ -377,15 +377,15 @@ void FCounterAggregationWorker::DoWork()
 	if (Session.IsValid())
 	{
 		// Suspend analysis in order to avoid write locks (ones blocked by the read lock below) to further block other read locks.
-		//Trace::FAnalysisSessionSuspensionScope SessionPauseScope(*Session.Get());
+		//TraceServices::FAnalysisSessionSuspensionScope SessionPauseScope(*Session.Get());
 
-		Trace::FAnalysisSessionReadScope SessionReadScope(*Session.Get());
+		TraceServices::FAnalysisSessionReadScope SessionReadScope(*Session.Get());
 
-		const Trace::ICounterProvider& CountersProvider = Trace::ReadCounterProvider(*Session.Get());
+		const TraceServices::ICounterProvider& CountersProvider = TraceServices::ReadCounterProvider(*Session.Get());
 
 		// Compute instance count and total/min/max inclusive/exclusive times for each counter.
 		// Iterate through all counters.
-		CountersProvider.EnumerateCounters([this](uint32 CounterId, const Trace::ICounter& Counter)
+		CountersProvider.EnumerateCounters([this](uint32 CounterId, const TraceServices::ICounter& Counter)
 			{
 				if (Counter.IsFloatingPoint())
 				{
@@ -406,7 +406,7 @@ void FCounterAggregationWorker::DoWork()
 
 			// Compute histogram.
 			// Iterate again through all counters.
-			CountersProvider.EnumerateCounters([this](uint32 CounterId, const Trace::ICounter& Counter)
+			CountersProvider.EnumerateCounters([this](uint32 CounterId, const TraceServices::ICounter& Counter)
 				{
 					if (Counter.IsFloatingPoint())
 					{
@@ -445,7 +445,7 @@ void FCounterAggregationWorker::ResetResults()
 // FCounterAggregator
 ////////////////////////////////////////////////////////////////////////////////////////////////////
 
-IStatsAggregationWorker* FCounterAggregator::CreateWorker(TSharedPtr<const Trace::IAnalysisSession> InSession)
+IStatsAggregationWorker* FCounterAggregator::CreateWorker(TSharedPtr<const TraceServices::IAnalysisSession> InSession)
 {
 	return new FCounterAggregationWorker(InSession, GetIntervalStartTime(), GetIntervalEndTime());
 }
diff --git a/Engine/Source/Developer/TraceInsights/Private/Insights/ViewModels/CounterAggregation.h b/Engine/Source/Developer/TraceInsights/Private/Insights/ViewModels/CounterAggregation.h
--- a/Engine/Source/Developer/TraceInsights/Private/Insights/ViewModels/CounterAggregation.h
+++ b/Engine/Source/Developer/TraceInsights/Private/Insights/ViewModels/CounterAggregation.h
@@ -24,7 +24,7 @@ public:
 	void ResetResults();
 
 protected:
-	virtual IStatsAggregationWorker* CreateWorker(TSharedPtr<const Trace::IAnalysisSession> InSession) override;
+	virtual IStatsAggregationWorker* CreateWorker(TSharedPtr<const TraceServices::IAnalysisSession> InSession) override;
 };
 
 ////////////////////////////////////////////////////////////////////////////////////////////////////
diff --git a/Engine/Source/Developer/TraceInsights/Private/Insights/ViewModels/FileActivityTimingTrack.cpp b/Engine/Source/Developer/TraceInsights/Private/Insights/ViewModels/FileActivityTimingTrack.cpp
--- a/Engine/Source/Developer/TraceInsights/Private/Insights/ViewModels/FileActivityTimingTrack.cpp
+++ b/Engine/Source/Developer/TraceInsights/Private/Insights/ViewModels/FileActivityTimingTrack.cpp
@@ -29,13 +29,13 @@
 
 ////////////////////////////////////////////////////////////////////////////////////////////////////
 
-const TCHAR* GetFileActivityTypeName(Trace::EFileActivityType Type)
+const TCHAR* GetFileActivityTypeName(TraceServices::EFileActivityType Type)
 {
-	static_assert(Trace::FileActivityType_Open == 0, "Trace::EFileActivityType enum has changed!?");
-	static_assert(Trace::FileActivityType_Close == 1, "Trace::EFileActivityType enum has changed!?");
-	static_assert(Trace::FileActivityType_Read == 2, "Trace::EFileActivityType enum has changed!?");
-	static_assert(Trace::FileActivityType_Write == 3, "Trace::EFileActivityType enum has changed!?");
-	static_assert(Trace::FileActivityType_Count == 4, "Trace::EFileActivityType enum has changed!?");
+	static_assert(TraceServices::FileActivityType_Open == 0, "TraceServices::EFileActivityType enum has changed!?");
+	static_assert(TraceServices::FileActivityType_Close == 1, "TraceServices::EFileActivityType enum has changed!?");
+	static_assert(TraceServices::FileActivityType_Read == 2, "TraceServices::EFileActivityType enum has changed!?");
+	static_assert(TraceServices::FileActivityType_Write == 3, "TraceServices::EFileActivityType enum has changed!?");
+	static_assert(TraceServices::FileActivityType_Count == 4, "TraceServices::EFileActivityType enum has changed!?");
 	static const TCHAR* GFileActivityTypeNames[] =
 	{
 		TEXT("Open"),
@@ -50,7 +50,7 @@ const TCHAR* GetFileActivityTypeName(Trace::EFileActivityType Type)
 
 ////////////////////////////////////////////////////////////////////////////////////////////////////
 
-uint32 GetFileActivityTypeColor(Trace::EFileActivityType Type)
+uint32 GetFileActivityTypeColor(TraceServices::EFileActivityType Type)
 {
 	static const uint32 GFileActivityTypeColors[] =
 	{
@@ -110,14 +110,14 @@ void FFileActivitySharedState::OnEndSession(Insights::ITimingViewSession& InSess
 
 ////////////////////////////////////////////////////////////////////////////////////////////////////
 
-void FFileActivitySharedState::Tick(Insights::ITimingViewSession& InSession, const Trace::IAnalysisSession& InAnalysisSession)
+void FFileActivitySharedState::Tick(Insights::ITimingViewSession& InSession, const TraceServices::IAnalysisSession& InAnalysisSession)
 {
 	if (&InSession != TimingView)
 	{
 		return;
 	}
 
-	if (!Trace::ReadFileActivityProvider(InAnalysisSession))
+	if (!TraceServices::ReadFileActivityProvider(InAnalysisSession))
 	{
 		return;
 	}
@@ -151,12 +151,12 @@ void FFileActivitySharedState::Tick(Insights::ITimingViewSession& InSession, con
 
 		// Enumerate all IO events and cache them.
 		{
-			Trace::FAnalysisSessionReadScope SessionReadScope(InAnalysisSession);
-			const Trace::IFileActivityProvider& FileActivityProvider = *Trace::ReadFileActivityProvider(InAnalysisSession);
-			FileActivityProvider.EnumerateFileActivity([this](const Trace::FFileInfo& FileInfo, const Trace::IFileActivityProvider::Timeline& Timeline)
+			TraceServices::FAnalysisSessionReadScope SessionReadScope(InAnalysisSession);
+			const TraceServices::IFileActivityProvider& FileActivityProvider = *TraceServices::ReadFileActivityProvider(InAnalysisSession);
+			FileActivityProvider.EnumerateFileActivity([this](const TraceServices::FFileInfo& FileInfo, const TraceServices::IFileActivityProvider::Timeline& Timeline)
 			{
 				Timeline.EnumerateEvents(-std::numeric_limits<double>::infinity(), +std::numeric_limits<double>::infinity(),
-					[this, &FileInfo, &Timeline](double EventStartTime, double EventEndTime, uint32 EventDepth, const Trace::FFileActivity* FileActivity)
+					[this, &FileInfo, &Timeline](double EventStartTime, double EventEndTime, uint32 EventDepth, const TraceServices::FFileActivity* FileActivity)
 				{
 					//if (EventEndTime == std::numeric_limits<double>::infinity())
 					//{
@@ -182,7 +182,7 @@ void FFileActivitySharedState::Tick(Insights::ITimingViewSession& InSession, con
 					}
 					else
 					{
-						if (FileActivity->ActivityType != Trace::FileActivityType_Close)
+						if (FileActivity->ActivityType != TraceServices::FileActivityType_Close)
 						{
 							ensure(EventStartTime >= Activity->StartTime);
 							if (EventStartTime < Activity->StartTime)
@@ -218,7 +218,7 @@ void FFileActivitySharedState::Tick(Insights::ITimingViewSession& InSession, con
 					uint32 Type = ((uint32)FileActivity->ActivityType & 0x0F) | (FileActivity->Failed ? 0x80 : 0);
 
 					AllIoEvents.Add(FIoTimingEvent{ EventStartTime, EventEndTime, EventDepth, Type, FileActivity->Offset, FileActivity->Size, FileActivity->ActualSize, Activity });
-					return Trace::EEventEnumerate::Continue;
+					return TraceServices::EEventEnumerate::Continue;
 				});
 
 				return true;
@@ -552,7 +552,7 @@ void FFileActivityTimingTrack::InitTooltip(FTooltipDrawState& InOutTooltip, cons
 		FTimingEventSearchParameters SearchParameters(TooltipEvent.GetStartTime(), TooltipEvent.GetEndTime(), ETimingEventSearchFlags::StopAtFirstMatch, MatchEvent);
 		FindIoTimingEvent(SearchParameters, [this, &InOutTooltip, &TooltipEvent](double InFoundStartTime, double InFoundEndTime, uint32 InFoundDepth, const FFileActivitySharedState::FIoTimingEvent& InEvent)
 		{
-			const Trace::EFileActivityType ActivityType = static_cast<Trace::EFileActivityType>(InEvent.Type & 0x0F);
+			const TraceServices::EFileActivityType ActivityType = static_cast<TraceServices::EFileActivityType>(InEvent.Type & 0x0F);
 			const bool bHasFailed = ((InEvent.Type & 0xF0) != 0);
 
 			FString TypeStr;
@@ -583,7 +583,7 @@ void FFileActivityTimingTrack::InitTooltip(FTooltipDrawState& InOutTooltip, cons
 			const double Duration = InEvent.EndTime - InEvent.StartTime;
 			InOutTooltip.AddNameValueTextLine(TEXT("Duration:"), TimeUtils::FormatTimeAuto(Duration));
 
-			if (ActivityType == Trace::FileActivityType_Read || ActivityType == Trace::FileActivityType_Write)
+			if (ActivityType == TraceServices::FileActivityType_Read || ActivityType == TraceServices::FileActivityType_Write)
 			{
 				InOutTooltip.AddNameValueTextLine(TEXT("Offset:"), FText::AsNumber(InEvent.Offset).ToString() + TEXT(" bytes"));
 				InOutTooltip.AddNameValueTextLine(TEXT("Size:"), FText::AsNumber(InEvent.Size).ToString() + TEXT(" bytes"));
@@ -678,10 +678,10 @@ void FOverviewFileActivityTimingTrack::BuildDrawState(ITimingEventsTrackDrawStat
 
 	for (const FFileActivitySharedState::FIoTimingEvent& Event : SharedState.AllIoEvents)
 	{
-		const Trace::EFileActivityType ActivityType = static_cast<Trace::EFileActivityType>(Event.Type & 0x0F);
+		const TraceServices::EFileActivityType ActivityType = static_cast<TraceServices::EFileActivityType>(Event.Type & 0x0F);
 		const uint64 EventType = static_cast<uint64>(ActivityType);
 
-		if (ActivityType >= Trace::FileActivityType_Count)
+		if (ActivityType >= TraceServices::FileActivityType_Count)
 		{
 			// Ignore "Idle" and "NotClosed" events.
 			continue;
@@ -723,7 +723,7 @@ void FOverviewFileActivityTimingTrack::BuildDrawState(ITimingEventsTrackDrawStat
 					EventName += TEXT("Failed ");
 				}
 
-				const Trace::EFileActivityType ActivityType = static_cast<Trace::EFileActivityType>(Event.Type & 0x0F);
+				const TraceServices::EFileActivityType ActivityType = static_cast<TraceServices::EFileActivityType>(Event.Type & 0x0F);
 				EventName += GetFileActivityTypeName(ActivityType);
 
 				if (Event.ActualSize != Event.Size)
@@ -829,7 +829,7 @@ void FDetailedFileActivityTimingTrack::BuildDrawState(ITimingEventsTrackDrawStat
 		}
 
 		ensure(Event.Depth <= 10000);
-		const Trace::EFileActivityType ActivityType = static_cast<Trace::EFileActivityType>(Event.Type & 0x0F);
+		const TraceServices::EFileActivityType ActivityType = static_cast<TraceServices::EFileActivityType>(Event.Type & 0x0F);
 
 		const bool bHasFailed = ((Event.Type & 0xF0) != 0);
 
@@ -855,7 +855,7 @@ void FDetailedFileActivityTimingTrack::BuildDrawState(ITimingEventsTrackDrawStat
 					EventName += TEXT("Failed ");
 				}
 
-				const Trace::EFileActivityType ActivityType = static_cast<Trace::EFileActivityType>(Event.Type & 0x0F);
+				const TraceServices::EFileActivityType ActivityType = static_cast<TraceServices::EFileActivityType>(Event.Type & 0x0F);
 				EventName += GetFileActivityTypeName(ActivityType);
 
 				if (Event.ActualSize != Event.Size)
@@ -863,7 +863,7 @@ void FDetailedFileActivityTimingTrack::BuildDrawState(ITimingEventsTrackDrawStat
 					EventName += TEXT(" [!]");
 				}
 
-				if (ActivityType >= Trace::FileActivityType_Count)
+				if (ActivityType >= TraceServices::FileActivityType_Count)
 				{
 					EventName += " [";
 					EventName += Event.FileActivity->Path;
diff --git a/Engine/Source/Developer/TraceInsights/Private/Insights/ViewModels/FileActivityTimingTrack.h b/Engine/Source/Developer/TraceInsights/Private/Insights/ViewModels/FileActivityTimingTrack.h
--- a/Engine/Source/Developer/TraceInsights/Private/Insights/ViewModels/FileActivityTimingTrack.h
+++ b/Engine/Source/Developer/TraceInsights/Private/Insights/ViewModels/FileActivityTimingTrack.h
@@ -39,7 +39,7 @@ public:
 		double StartTime;
 		double EndTime;
 		uint32 Depth;
-		uint32 Type; // Trace::EFileActivityType + "Failed" flag
+		uint32 Type; // TraceServices::EFileActivityType + "Failed" flag
 		uint64 Offset;
 		uint64 Size;
 		uint64 ActualSize;
@@ -53,7 +53,7 @@ public:
 	// ITimingViewExtender
 	virtual void OnBeginSession(Insights::ITimingViewSession& InSession) override;
 	virtual void OnEndSession(Insights::ITimingViewSession& InSession) override;
-	virtual void Tick(Insights::ITimingViewSession& InSession, const Trace::IAnalysisSession& InAnalysisSession) override;
+	virtual void Tick(Insights::ITimingViewSession& InSession, const TraceServices::IAnalysisSession& InAnalysisSession) override;
 	virtual void ExtendFilterMenu(Insights::ITimingViewSession& InSession, FMenuBuilder& InOutMenuBuilder) override;
 
 	const TArray<FIoTimingEvent>& GetAllEvents() const { return AllIoEvents; }
diff --git a/Engine/Source/Developer/TraceInsights/Private/Insights/ViewModels/FrameTimingTrack.cpp b/Engine/Source/Developer/TraceInsights/Private/Insights/ViewModels/FrameTimingTrack.cpp
--- a/Engine/Source/Developer/TraceInsights/Private/Insights/ViewModels/FrameTimingTrack.cpp
+++ b/Engine/Source/Developer/TraceInsights/Private/Insights/ViewModels/FrameTimingTrack.cpp
@@ -68,7 +68,7 @@ void FFrameSharedState::OnEndSession(Insights::ITimingViewSession& InSession)
 
 ////////////////////////////////////////////////////////////////////////////////////////////////////
 
-void FFrameSharedState::Tick(Insights::ITimingViewSession& InSession, const Trace::IAnalysisSession& InAnalysisSession)
+void FFrameSharedState::Tick(Insights::ITimingViewSession& InSession, const TraceServices::IAnalysisSession& InAnalysisSession)
 {
 	if (&InSession != TimingView)
 	{
@@ -180,11 +180,11 @@ void FFrameTimingTrack::Reset()
 
 void FFrameTimingTrack::BuildDrawState(ITimingEventsTrackDrawStateBuilder& Builder, const ITimingTrackUpdateContext& Context)
 {
-	TSharedPtr<const Trace::IAnalysisSession> Session = FInsightsManager::Get()->GetSession();
+	TSharedPtr<const TraceServices::IAnalysisSession> Session = FInsightsManager::Get()->GetSession();
 	if (Session.IsValid())
 	{
-		Trace::FAnalysisSessionReadScope SessionReadScope(*Session.Get());
-		const Trace::IFrameProvider& FramesProvider = Trace::ReadFrameProvider(*Session.Get());
+		TraceServices::FAnalysisSessionReadScope SessionReadScope(*Session.Get());
+		const TraceServices::IFrameProvider& FramesProvider = TraceServices::ReadFrameProvider(*Session.Get());
 		const FTimingTrackViewport& Viewport = Context.GetViewport();
 
 		const TArray64<double>& FrameStartTimes = FramesProvider.GetFrameStartTimes(static_cast<ETraceFrameType>(FrameType));
@@ -197,7 +197,7 @@ void FFrameTimingTrack::BuildDrawState(ITimingEventsTrackDrawStateBuilder& Build
 
 		const uint32 Color = FFrameTrackDrawHelper::GetColor32ByFrameType(FrameType);
 
-		FramesProvider.EnumerateFrames(static_cast<ETraceFrameType>(FrameType), StartIndex, EndIndex, [this, Color, &Builder](const Trace::FFrame& Frame)
+		FramesProvider.EnumerateFrames(static_cast<ETraceFrameType>(FrameType), StartIndex, EndIndex, [this, Color, &Builder](const TraceServices::FFrame& Frame)
 		{
 			constexpr uint32 Depth = 0;
 			const FString FrameName = GetShortFrameName(Frame.Index);
@@ -214,11 +214,11 @@ void FFrameTimingTrack::BuildFilteredDrawState(ITimingEventsTrackDrawStateBuilde
 	const TSharedPtr<ITimingEventFilter> EventFilterPtr = Context.GetEventFilter();
 	if (EventFilterPtr.IsValid() && EventFilterPtr->FilterTrack(*this))
 	{
-		TSharedPtr<const Trace::IAnalysisSession> Session = FInsightsManager::Get()->GetSession();
+		TSharedPtr<const TraceServices::IAnalysisSession> Session = FInsightsManager::Get()->GetSession();
 		if (Session.IsValid())
 		{
-			Trace::FAnalysisSessionReadScope SessionReadScope(*Session.Get());
-			const Trace::IFrameProvider& FramesProvider = Trace::ReadFrameProvider(*Session.Get());
+			TraceServices::FAnalysisSessionReadScope SessionReadScope(*Session.Get());
+			const TraceServices::IFrameProvider& FramesProvider = TraceServices::ReadFrameProvider(*Session.Get());
 			const FTimingTrackViewport& Viewport = Context.GetViewport();
 
 			const TArray64<double>& FrameStartTimes = FramesProvider.GetFrameStartTimes(static_cast<ETraceFrameType>(FrameType));
@@ -231,7 +231,7 @@ void FFrameTimingTrack::BuildFilteredDrawState(ITimingEventsTrackDrawStateBuilde
 
 			const uint32 Color = FFrameTrackDrawHelper::GetColor32ByFrameType(FrameType);
 
-			FramesProvider.EnumerateFrames(static_cast<ETraceFrameType>(FrameType), StartIndex, EndIndex, [this, Color, &Builder, &EventFilterPtr](const Trace::FFrame& Frame)
+			FramesProvider.EnumerateFrames(static_cast<ETraceFrameType>(FrameType), StartIndex, EndIndex, [this, Color, &Builder, &EventFilterPtr](const TraceServices::FFrame& Frame)
 			{
 				constexpr uint32 Depth = 0;
 				const uint64 IndexAsEventType = Frame.Index;
@@ -302,7 +302,7 @@ void FFrameTimingTrack::PostDraw(const ITimingTrackDrawContext& Context) const
 
 void FFrameTimingTrack::DrawSelectedEventInfo(const FTimingEvent& SelectedEvent, const FTimingTrackViewport& Viewport, const FDrawContext& DrawContext, const FSlateBrush* WhiteBrush, const FSlateFontInfo& Font) const
 {
-	FindFrame(SelectedEvent, [this, &SelectedEvent, &Font, &Viewport, &DrawContext, &WhiteBrush](double InFoundStartTime, double InFoundEndTime, uint32 InFoundDepth, const Trace::FFrame& InFoundFrame)
+	FindFrame(SelectedEvent, [this, &SelectedEvent, &Font, &Viewport, &DrawContext, &WhiteBrush](double InFoundStartTime, double InFoundEndTime, uint32 InFoundDepth, const TraceServices::FFrame& InFoundFrame)
 	{
 		//const double Duration = SelectedEvent.GetDuration();
 		const double Duration = InFoundFrame.EndTime - InFoundFrame.StartTime;
@@ -357,7 +357,7 @@ void FFrameTimingTrack::InitTooltip(FTooltipDrawState& InOutTooltip, const ITimi
 	{
 		const FTimingEvent& TooltipEvent = InTooltipEvent.As<FTimingEvent>();
 
-		FindFrame(TooltipEvent, [this, &InOutTooltip, &TooltipEvent](double InFoundStartTime, double InFoundEndTime, uint32 InFoundDepth, const Trace::FFrame& InFoundFrame)
+		FindFrame(TooltipEvent, [this, &InOutTooltip, &TooltipEvent](double InFoundStartTime, double InFoundEndTime, uint32 InFoundDepth, const TraceServices::FFrame& InFoundFrame)
 		{
 			InOutTooltip.ResetContent();
 
@@ -379,7 +379,7 @@ const TSharedPtr<const ITimingEvent> FFrameTimingTrack::SearchEvent(const FTimin
 {
 	TSharedPtr<const ITimingEvent> FoundEvent;
 
-	FindFrame(InSearchParameters, [this, &FoundEvent](double InFoundStartTime, double InFoundEndTime, uint32 InFoundDepth, const Trace::FFrame& InFoundFrame)
+	FindFrame(InSearchParameters, [this, &FoundEvent](double InFoundStartTime, double InFoundEndTime, uint32 InFoundDepth, const TraceServices::FFrame& InFoundFrame)
 	{
 		const uint64 IndexAsEventType = InFoundFrame.Index; // storing the frame index as event type
 		FoundEvent = MakeShared<FTimingEvent>(SharedThis(this), InFoundStartTime, InFoundEndTime, InFoundDepth, IndexAsEventType);
@@ -437,7 +437,7 @@ void FFrameTimingTrack::BuildContextMenu(FMenuBuilder& MenuBuilder)
 
 ////////////////////////////////////////////////////////////////////////////////////////////////////
 
-bool FFrameTimingTrack::FindFrame(const FTimingEvent& InTimingEvent, TFunctionRef<void(double, double, uint32, const Trace::FFrame&)> InFoundPredicate) const
+bool FFrameTimingTrack::FindFrame(const FTimingEvent& InTimingEvent, TFunctionRef<void(double, double, uint32, const TraceServices::FFrame&)> InFoundPredicate) const
 {
 	auto MatchEvent = [&InTimingEvent](double InStartTime, double InEndTime, uint32 InDepth)
 	{
@@ -453,18 +453,18 @@ bool FFrameTimingTrack::FindFrame(const FTimingEvent& InTimingEvent, TFunctionRe
 
 ////////////////////////////////////////////////////////////////////////////////////////////////////
 
-bool FFrameTimingTrack::FindFrame(const FTimingEventSearchParameters& InParameters, TFunctionRef<void(double, double, uint32, const Trace::FFrame&)> InFoundPredicate) const
+bool FFrameTimingTrack::FindFrame(const FTimingEventSearchParameters& InParameters, TFunctionRef<void(double, double, uint32, const TraceServices::FFrame&)> InFoundPredicate) const
 {
-	return TTimingEventSearch<Trace::FFrame>::Search(
+	return TTimingEventSearch<TraceServices::FFrame>::Search(
 		InParameters,
 
-		[this](TTimingEventSearch<Trace::FFrame>::FContext& InContext)
+		[this](TTimingEventSearch<TraceServices::FFrame>::FContext& InContext)
 		{
-			TSharedPtr<const Trace::IAnalysisSession> Session = FInsightsManager::Get()->GetSession();
+			TSharedPtr<const TraceServices::IAnalysisSession> Session = FInsightsManager::Get()->GetSession();
 			if (Session.IsValid())
 			{
-				Trace::FAnalysisSessionReadScope SessionReadScope(*Session.Get());
-				const Trace::IFrameProvider& FramesProvider = Trace::ReadFrameProvider(*Session.Get());
+				TraceServices::FAnalysisSessionReadScope SessionReadScope(*Session.Get());
+				const TraceServices::IFrameProvider& FramesProvider = TraceServices::ReadFrameProvider(*Session.Get());
 
 				const TArray64<double>& FrameStartTimes = FramesProvider.GetFrameStartTimes(static_cast<ETraceFrameType>(FrameType));
 
@@ -474,7 +474,7 @@ bool FFrameTimingTrack::FindFrame(const FTimingEventSearchParameters& InParamete
 				const int64 EndLowerBound = Algo::LowerBound(FrameStartTimes, InContext.GetParameters().EndTime);
 				const uint64 EndIndex = EndLowerBound;
 
-				FramesProvider.EnumerateFrames(static_cast<ETraceFrameType>(FrameType), StartIndex, EndIndex, [&InContext](const Trace::FFrame& Frame)
+				FramesProvider.EnumerateFrames(static_cast<ETraceFrameType>(FrameType), StartIndex, EndIndex, [&InContext](const TraceServices::FFrame& Frame)
 				{
 					if (Frame.StartTime < InContext.GetParameters().EndTime &&
 						Frame.EndTime > InContext.GetParameters().StartTime)
@@ -486,7 +486,7 @@ bool FFrameTimingTrack::FindFrame(const FTimingEventSearchParameters& InParamete
 			}
 		},
 
-		[&InFoundPredicate](double InFoundStartTime, double InFoundEndTime, uint32 InFoundDepth, const Trace::FFrame& InEvent)
+		[&InFoundPredicate](double InFoundStartTime, double InFoundEndTime, uint32 InFoundDepth, const TraceServices::FFrame& InEvent)
 		{
 			InFoundPredicate(InFoundStartTime, InFoundEndTime, InFoundDepth, InEvent);
 		},
diff --git a/Engine/Source/Developer/TraceInsights/Private/Insights/ViewModels/FrameTimingTrack.h b/Engine/Source/Developer/TraceInsights/Private/Insights/ViewModels/FrameTimingTrack.h
--- a/Engine/Source/Developer/TraceInsights/Private/Insights/ViewModels/FrameTimingTrack.h
+++ b/Engine/Source/Developer/TraceInsights/Private/Insights/ViewModels/FrameTimingTrack.h
@@ -35,7 +35,7 @@ public:
 
 	virtual void OnBeginSession(Insights::ITimingViewSession& InSession) override;
 	virtual void OnEndSession(Insights::ITimingViewSession& InSession) override;
-	virtual void Tick(Insights::ITimingViewSession& InSession, const Trace::IAnalysisSession& InAnalysisSession) override;
+	virtual void Tick(Insights::ITimingViewSession& InSession, const TraceServices::IAnalysisSession& InAnalysisSession) override;
 	virtual void ExtendFilterMenu(Insights::ITimingViewSession& InSession, FMenuBuilder& InMenuBuilder) override;
 
 	//////////////////////////////////////////////////
@@ -116,8 +116,8 @@ public:
 private:
 	void DrawSelectedEventInfo(const FTimingEvent& SelectedEvent, const FTimingTrackViewport& Viewport, const FDrawContext& DrawContext, const FSlateBrush* WhiteBrush, const FSlateFontInfo& Font) const;
 
-	bool FindFrame(const FTimingEvent& InTimingEvent, TFunctionRef<void(double, double, uint32, const Trace::FFrame&)> InFoundPredicate) const;
-	bool FindFrame(const FTimingEventSearchParameters& InParameters, TFunctionRef<void(double, double, uint32, const Trace::FFrame&)> InFoundPredicate) const;
+	bool FindFrame(const FTimingEvent& InTimingEvent, TFunctionRef<void(double, double, uint32, const TraceServices::FFrame&)> InFoundPredicate) const;
+	bool FindFrame(const FTimingEventSearchParameters& InParameters, TFunctionRef<void(double, double, uint32, const TraceServices::FFrame&)> InFoundPredicate) const;
 
 private:
 	uint32 FrameType; // ETraceFrameType
@@ -125,7 +125,7 @@ private:
 	FTrackHeader Header;
 
 	// Search cache
-	mutable TTimingEventSearchCache<Trace::FFrame> SearchCache;
+	mutable TTimingEventSearchCache<TraceServices::FFrame> SearchCache;
 };
 
 ////////////////////////////////////////////////////////////////////////////////////////////////////
diff --git a/Engine/Source/Developer/TraceInsights/Private/Insights/ViewModels/FrameTrackHelper.cpp b/Engine/Source/Developer/TraceInsights/Private/Insights/ViewModels/FrameTrackHelper.cpp
--- a/Engine/Source/Developer/TraceInsights/Private/Insights/ViewModels/FrameTrackHelper.cpp
+++ b/Engine/Source/Developer/TraceInsights/Private/Insights/ViewModels/FrameTrackHelper.cpp
@@ -37,7 +37,7 @@ FFrameTrackSeriesBuilder::FFrameTrackSeriesBuilder(FFrameTrackSeries& InSeries,
 
 ////////////////////////////////////////////////////////////////////////////////////////////////////
 
-void FFrameTrackSeriesBuilder::AddFrame(const Trace::FFrame& Frame)
+void FFrameTrackSeriesBuilder::AddFrame(const TraceServices::FFrame& Frame)
 {
 	NumAddedFrames++;
 
diff --git a/Engine/Source/Developer/TraceInsights/Private/Insights/ViewModels/FrameTrackHelper.h b/Engine/Source/Developer/TraceInsights/Private/Insights/ViewModels/FrameTrackHelper.h
--- a/Engine/Source/Developer/TraceInsights/Private/Insights/ViewModels/FrameTrackHelper.h
+++ b/Engine/Source/Developer/TraceInsights/Private/Insights/ViewModels/FrameTrackHelper.h
@@ -14,7 +14,7 @@ struct FSlateBrush;
 class FFrameTrackViewport;
 class FSlateWindowElementList;
 
-namespace Trace
+namespace TraceServices
 {
 	struct FFrame;
 }
@@ -101,7 +101,7 @@ public:
 	FFrameTrackSeriesBuilder(const FFrameTrackSeriesBuilder&) = delete;
 	FFrameTrackSeriesBuilder& operator=(const FFrameTrackSeriesBuilder&) = delete;
 
-	void AddFrame(const Trace::FFrame& Frame);
+	void AddFrame(const TraceServices::FFrame& Frame);
 
 	int32 GetNumAddedFrames() const { return NumAddedFrames; }
 
diff --git a/Engine/Source/Developer/TraceInsights/Private/Insights/ViewModels/LoadingTimingTrack.cpp b/Engine/Source/Developer/TraceInsights/Private/Insights/ViewModels/LoadingTimingTrack.cpp
--- a/Engine/Source/Developer/TraceInsights/Private/Insights/ViewModels/LoadingTimingTrack.cpp
+++ b/Engine/Source/Developer/TraceInsights/Private/Insights/ViewModels/LoadingTimingTrack.cpp
@@ -67,17 +67,17 @@ void FLoadingSharedState::OnEndSession(Insights::ITimingViewSession& InSession)
 
 ////////////////////////////////////////////////////////////////////////////////////////////////////
 
-void FLoadingSharedState::Tick(Insights::ITimingViewSession& InSession, const Trace::IAnalysisSession& InAnalysisSession)
+void FLoadingSharedState::Tick(Insights::ITimingViewSession& InSession, const TraceServices::IAnalysisSession& InAnalysisSession)
 {
 	if (&InSession != TimingView)
 	{
 		return;
 	}
 
-	const Trace::ILoadTimeProfilerProvider* LoadTimeProfilerProvider = Trace::ReadLoadTimeProfilerProvider(InAnalysisSession);
+	const TraceServices::ILoadTimeProfilerProvider* LoadTimeProfilerProvider = TraceServices::ReadLoadTimeProfilerProvider(InAnalysisSession);
 	if (LoadTimeProfilerProvider)
 	{
-		Trace::FAnalysisSessionReadScope SessionReadScope(InAnalysisSession);
+		TraceServices::FAnalysisSessionReadScope SessionReadScope(InAnalysisSession);
 
 		const uint64 CurrentLoadTimeProfilerTimelineCount = LoadTimeProfilerProvider->GetTimelineCount();
 		if (CurrentLoadTimeProfilerTimelineCount != LoadTimeProfilerTimelineCount)
@@ -85,8 +85,8 @@ void FLoadingSharedState::Tick(Insights::ITimingViewSession& InSession, const Tr
 			LoadTimeProfilerTimelineCount = CurrentLoadTimeProfilerTimelineCount;
 
 			// Iterate through threads.
-			const Trace::IThreadProvider& ThreadProvider = Trace::ReadThreadProvider(InAnalysisSession);
-			ThreadProvider.EnumerateThreads([this, &InSession, LoadTimeProfilerProvider](const Trace::FThreadInfo& ThreadInfo)
+			const TraceServices::IThreadProvider& ThreadProvider = TraceServices::ReadThreadProvider(InAnalysisSession);
+			ThreadProvider.EnumerateThreads([this, &InSession, LoadTimeProfilerProvider](const TraceServices::FThreadInfo& ThreadInfo)
 			{
 				// Check available Asset Loading tracks.
 				uint32 LoadingTimelineIndex;
@@ -154,28 +154,28 @@ void FLoadingSharedState::SetAllLoadingTracksToggle(bool bOnOff)
 
 ////////////////////////////////////////////////////////////////////////////////////////////////////
 
-const TCHAR* FLoadingSharedState::GetEventNameByEventType(uint32 Depth, const Trace::FLoadTimeProfilerCpuEvent& Event) const
+const TCHAR* FLoadingSharedState::GetEventNameByEventType(uint32 Depth, const TraceServices::FLoadTimeProfilerCpuEvent& Event) const
 {
-	return Trace::GetLoadTimeProfilerObjectEventTypeString(Event.EventType);
+	return TraceServices::GetLoadTimeProfilerObjectEventTypeString(Event.EventType);
 }
 
 ////////////////////////////////////////////////////////////////////////////////////////////////////
 
-const TCHAR* FLoadingSharedState::GetEventNameByPackageName(uint32 Depth, const Trace::FLoadTimeProfilerCpuEvent& Event) const
+const TCHAR* FLoadingSharedState::GetEventNameByPackageName(uint32 Depth, const TraceServices::FLoadTimeProfilerCpuEvent& Event) const
 {
 	return Event.Package ? Event.Package->Name : TEXT("");
 }
 
 ////////////////////////////////////////////////////////////////////////////////////////////////////
 
-const TCHAR* FLoadingSharedState::GetEventNameByExportClassName(uint32 Depth, const Trace::FLoadTimeProfilerCpuEvent& Event) const
+const TCHAR* FLoadingSharedState::GetEventNameByExportClassName(uint32 Depth, const TraceServices::FLoadTimeProfilerCpuEvent& Event) const
 {
 	return Event.Export && Event.Export->Class ? Event.Export->Class->Name : TEXT("");
 }
 
 ////////////////////////////////////////////////////////////////////////////////////////////////////
 
-const TCHAR* FLoadingSharedState::GetEventNameByPackageAndExportClassName(uint32 Depth, const Trace::FLoadTimeProfilerCpuEvent& Event) const
+const TCHAR* FLoadingSharedState::GetEventNameByPackageAndExportClassName(uint32 Depth, const TraceServices::FLoadTimeProfilerCpuEvent& Event) const
 {
 	if (Depth == 0)
 	{
@@ -195,7 +195,7 @@ const TCHAR* FLoadingSharedState::GetEventNameByPackageAndExportClassName(uint32
 
 ////////////////////////////////////////////////////////////////////////////////////////////////////
 
-const TCHAR* FLoadingSharedState::GetEventName(uint32 Depth, const Trace::FLoadTimeProfilerCpuEvent& Event) const
+const TCHAR* FLoadingSharedState::GetEventName(uint32 Depth, const TraceServices::FLoadTimeProfilerCpuEvent& Event) const
 {
 	return GetEventNameDelegate.Execute(Depth, Event);
 }
@@ -229,38 +229,38 @@ INSIGHTS_IMPLEMENT_RTTI(FLoadingTimingTrack)
 
 void FLoadingTimingTrack::BuildDrawState(ITimingEventsTrackDrawStateBuilder& Builder, const ITimingTrackUpdateContext& Context)
 {
-	TSharedPtr<const Trace::IAnalysisSession> Session = FInsightsManager::Get()->GetSession();
-	if (Session.IsValid() && Trace::ReadLoadTimeProfilerProvider(*Session.Get()))
+	TSharedPtr<const TraceServices::IAnalysisSession> Session = FInsightsManager::Get()->GetSession();
+	if (Session.IsValid() && TraceServices::ReadLoadTimeProfilerProvider(*Session.Get()))
 	{
-		Trace::FAnalysisSessionReadScope SessionReadScope(*Session.Get());
+		TraceServices::FAnalysisSessionReadScope SessionReadScope(*Session.Get());
 
-		const Trace::ILoadTimeProfilerProvider& LoadTimeProfilerProvider = *Trace::ReadLoadTimeProfilerProvider(*Session.Get());
+		const TraceServices::ILoadTimeProfilerProvider& LoadTimeProfilerProvider = *TraceServices::ReadLoadTimeProfilerProvider(*Session.Get());
 
 		const FTimingTrackViewport& Viewport = Context.GetViewport();
 
-		LoadTimeProfilerProvider.ReadTimeline(TimelineIndex, [this, &Builder, &Viewport](const Trace::ILoadTimeProfilerProvider::CpuTimeline& Timeline)
+		LoadTimeProfilerProvider.ReadTimeline(TimelineIndex, [this, &Builder, &Viewport](const TraceServices::ILoadTimeProfilerProvider::CpuTimeline& Timeline)
 		{
 			if (FTimingEventsTrack::bUseDownSampling)
 			{
 				const double SecondsPerPixel = 1.0 / Viewport.GetScaleX();
-				Timeline.EnumerateEventsDownSampled(Viewport.GetStartTime(), Viewport.GetEndTime(), SecondsPerPixel, [this, &Builder](double StartTime, double EndTime, uint32 Depth, const Trace::FLoadTimeProfilerCpuEvent& Event)
+				Timeline.EnumerateEventsDownSampled(Viewport.GetStartTime(), Viewport.GetEndTime(), SecondsPerPixel, [this, &Builder](double StartTime, double EndTime, uint32 Depth, const TraceServices::FLoadTimeProfilerCpuEvent& Event)
 				{
 					const TCHAR* Name = SharedState.GetEventName(Depth, Event);
 					const uint64 Type = static_cast<uint64>(Event.EventType);
 					const uint32 Color = 0;
 					Builder.AddEvent(StartTime, EndTime, Depth, Name, Type, Color);
-					return Trace::EEventEnumerate::Continue;
+					return TraceServices::EEventEnumerate::Continue;
 				});
 			}
 			else
 			{
-				Timeline.EnumerateEvents(Viewport.GetStartTime(), Viewport.GetEndTime(), [this, &Builder](double StartTime, double EndTime, uint32 Depth, const Trace::FLoadTimeProfilerCpuEvent& Event)
+				Timeline.EnumerateEvents(Viewport.GetStartTime(), Viewport.GetEndTime(), [this, &Builder](double StartTime, double EndTime, uint32 Depth, const TraceServices::FLoadTimeProfilerCpuEvent& Event)
 				{
 					const TCHAR* Name = SharedState.GetEventName(Depth, Event);
 					const uint64 Type = static_cast<uint64>(Event.EventType);
 					const uint32 Color = 0;
 					Builder.AddEvent(StartTime, EndTime, Depth, Name, Type, Color);
-					return Trace::EEventEnumerate::Continue;
+					return TraceServices::EEventEnumerate::Continue;
 				});
 			}
 		});
@@ -283,14 +283,14 @@ void FLoadingTimingTrack::InitTooltip(FTooltipDrawState& InOutTooltip, const ITi
 		};
 
 		FTimingEventSearchParameters SearchParameters(TooltipEvent.GetStartTime(), TooltipEvent.GetEndTime(), ETimingEventSearchFlags::StopAtFirstMatch, MatchEvent);
-		FindLoadTimeProfilerCpuEvent(SearchParameters, [this, &InOutTooltip, &TooltipEvent](double InFoundStartTime, double InFoundEndTime, uint32 InFoundDepth, const Trace::FLoadTimeProfilerCpuEvent& InFoundEvent)
+		FindLoadTimeProfilerCpuEvent(SearchParameters, [this, &InOutTooltip, &TooltipEvent](double InFoundStartTime, double InFoundEndTime, uint32 InFoundDepth, const TraceServices::FLoadTimeProfilerCpuEvent& InFoundEvent)
 		{
 			InOutTooltip.ResetContent();
 
 			InOutTooltip.AddTitle(SharedState.GetEventName(TooltipEvent.GetDepth(), InFoundEvent));
 
-			const Trace::FPackageInfo* Package = InFoundEvent.Package;
-			const Trace::FPackageExportInfo* Export = InFoundEvent.Export;
+			const TraceServices::FPackageInfo* Package = InFoundEvent.Package;
+			const TraceServices::FPackageExportInfo* Export = InFoundEvent.Export;
 
 			InOutTooltip.AddNameValueTextLine(TEXT("Duration:"), TimeUtils::FormatTimeAuto(TooltipEvent.GetDuration()));
 			InOutTooltip.AddNameValueTextLine(TEXT("Depth:"), FString::Printf(TEXT("%d"), TooltipEvent.GetDepth()));
@@ -302,7 +302,7 @@ void FLoadingTimingTrack::InitTooltip(FTooltipDrawState& InOutTooltip, const ITi
 				InOutTooltip.AddNameValueTextLine(TEXT("Package Summary:"), FString::Printf(TEXT("%d imports, %d exports"), Package->Summary.ImportCount, Package->Summary.ExportCount));
 			}
 
-			InOutTooltip.AddNameValueTextLine(TEXT("Export Event:"), FString::Printf(TEXT("%s"), Trace::GetLoadTimeProfilerObjectEventTypeString(InFoundEvent.EventType)));
+			InOutTooltip.AddNameValueTextLine(TEXT("Export Event:"), FString::Printf(TEXT("%s"), TraceServices::GetLoadTimeProfilerObjectEventTypeString(InFoundEvent.EventType)));
 
 			if (Export)
 			{
@@ -321,7 +321,7 @@ const TSharedPtr<const ITimingEvent> FLoadingTimingTrack::SearchEvent(const FTim
 {
 	TSharedPtr<const ITimingEvent> FoundEvent;
 
-	FindLoadTimeProfilerCpuEvent(InSearchParameters, [this, &FoundEvent](double InFoundStartTime, double InFoundEndTime, uint32 InFoundDepth, const Trace::FLoadTimeProfilerCpuEvent& InFoundEvent)
+	FindLoadTimeProfilerCpuEvent(InSearchParameters, [this, &FoundEvent](double InFoundStartTime, double InFoundEndTime, uint32 InFoundDepth, const TraceServices::FLoadTimeProfilerCpuEvent& InFoundEvent)
 	{
 		FoundEvent = MakeShared<FTimingEvent>(SharedThis(this), InFoundStartTime, InFoundEndTime, InFoundDepth);
 	});
@@ -331,35 +331,35 @@ const TSharedPtr<const ITimingEvent> FLoadingTimingTrack::SearchEvent(const FTim
 
 ////////////////////////////////////////////////////////////////////////////////////////////////////
 
-bool FLoadingTimingTrack::FindLoadTimeProfilerCpuEvent(const FTimingEventSearchParameters& InParameters, TFunctionRef<void(double, double, uint32, const Trace::FLoadTimeProfilerCpuEvent&)> InFoundPredicate) const
+bool FLoadingTimingTrack::FindLoadTimeProfilerCpuEvent(const FTimingEventSearchParameters& InParameters, TFunctionRef<void(double, double, uint32, const TraceServices::FLoadTimeProfilerCpuEvent&)> InFoundPredicate) const
 {
-	return TTimingEventSearch<Trace::FLoadTimeProfilerCpuEvent>::Search(
+	return TTimingEventSearch<TraceServices::FLoadTimeProfilerCpuEvent>::Search(
 		InParameters,
 
-		[this](TTimingEventSearch<Trace::FLoadTimeProfilerCpuEvent>::FContext& InContext)
+		[this](TTimingEventSearch<TraceServices::FLoadTimeProfilerCpuEvent>::FContext& InContext)
 		{
-			TSharedPtr<const Trace::IAnalysisSession> Session = FInsightsManager::Get()->GetSession();
+			TSharedPtr<const TraceServices::IAnalysisSession> Session = FInsightsManager::Get()->GetSession();
 			if (Session.IsValid())
 			{
-				Trace::FAnalysisSessionReadScope SessionReadScope(*Session.Get());
+				TraceServices::FAnalysisSessionReadScope SessionReadScope(*Session.Get());
 
-				if (Trace::ReadLoadTimeProfilerProvider(*Session.Get()))
+				if (TraceServices::ReadLoadTimeProfilerProvider(*Session.Get()))
 				{
-					const Trace::ILoadTimeProfilerProvider& LoadTimeProfilerProvider = *Trace::ReadLoadTimeProfilerProvider(*Session.Get());
+					const TraceServices::ILoadTimeProfilerProvider& LoadTimeProfilerProvider = *TraceServices::ReadLoadTimeProfilerProvider(*Session.Get());
 
-					LoadTimeProfilerProvider.ReadTimeline(TimelineIndex, [&InContext](const Trace::ILoadTimeProfilerProvider::CpuTimeline& Timeline)
+					LoadTimeProfilerProvider.ReadTimeline(TimelineIndex, [&InContext](const TraceServices::ILoadTimeProfilerProvider::CpuTimeline& Timeline)
 					{
-						Timeline.EnumerateEvents(InContext.GetParameters().StartTime, InContext.GetParameters().EndTime, [&InContext](double EventStartTime, double EventEndTime, uint32 EventDepth, const Trace::FLoadTimeProfilerCpuEvent& Event)
+						Timeline.EnumerateEvents(InContext.GetParameters().StartTime, InContext.GetParameters().EndTime, [&InContext](double EventStartTime, double EventEndTime, uint32 EventDepth, const TraceServices::FLoadTimeProfilerCpuEvent& Event)
 						{
 							InContext.Check(EventStartTime, EventEndTime, EventDepth, Event);
-							return InContext.ShouldContinueSearching() ? Trace::EEventEnumerate::Continue : Trace::EEventEnumerate::Stop;
+							return InContext.ShouldContinueSearching() ? TraceServices::EEventEnumerate::Continue : TraceServices::EEventEnumerate::Stop;
 						});
 					});
 				}
 			}
 		},
 
-		[&InFoundPredicate](double InFoundStartTime, double InFoundEndTime, uint32 InFoundDepth, const Trace::FLoadTimeProfilerCpuEvent& InEvent)
+		[&InFoundPredicate](double InFoundStartTime, double InFoundEndTime, uint32 InFoundDepth, const TraceServices::FLoadTimeProfilerCpuEvent& InEvent)
 		{
 			InFoundPredicate(InFoundStartTime, InFoundEndTime, InFoundDepth, InEvent);
 		});
diff --git a/Engine/Source/Developer/TraceInsights/Private/Insights/ViewModels/LoadingTimingTrack.h b/Engine/Source/Developer/TraceInsights/Private/Insights/ViewModels/LoadingTimingTrack.h
--- a/Engine/Source/Developer/TraceInsights/Private/Insights/ViewModels/LoadingTimingTrack.h
+++ b/Engine/Source/Developer/TraceInsights/Private/Insights/ViewModels/LoadingTimingTrack.h
@@ -14,7 +14,7 @@ class STimingView;
 class FLoadingTimingTrack;
 
 /** Defines FLoadingTrackGetEventNameDelegate delegate interface. Returns the name for a timing event in a Loading track. */
-DECLARE_DELEGATE_RetVal_TwoParams(const TCHAR*, FLoadingTrackGetEventNameDelegate, uint32 /*Depth*/, const Trace::FLoadTimeProfilerCpuEvent& /*Event*/);
+DECLARE_DELEGATE_RetVal_TwoParams(const TCHAR*, FLoadingTrackGetEventNameDelegate, uint32 /*Depth*/, const TraceServices::FLoadTimeProfilerCpuEvent& /*Event*/);
 
 ////////////////////////////////////////////////////////////////////////////////////////////////////
 
@@ -27,10 +27,10 @@ public:
 	// ITimingViewExtender
 	virtual void OnBeginSession(Insights::ITimingViewSession& InSession) override;
 	virtual void OnEndSession(Insights::ITimingViewSession& InSession) override;
-	virtual void Tick(Insights::ITimingViewSession& InSession, const Trace::IAnalysisSession& InAnalysisSession) override;
+	virtual void Tick(Insights::ITimingViewSession& InSession, const TraceServices::IAnalysisSession& InAnalysisSession) override;
 	virtual void ExtendFilterMenu(Insights::ITimingViewSession& InSession, FMenuBuilder& InMenuBuilder) override;
 
-	const TCHAR* GetEventName(uint32 Depth, const Trace::FLoadTimeProfilerCpuEvent& Event) const;
+	const TCHAR* GetEventName(uint32 Depth, const TraceServices::FLoadTimeProfilerCpuEvent& Event) const;
 	void SetColorSchema(int32 Schema);
 
 	TSharedPtr<FLoadingTimingTrack> GetLoadingTrack(uint32 InThreadId)
@@ -46,10 +46,10 @@ public:
 	void ShowHideAllLoadingTracks() { SetAllLoadingTracksToggle(!IsAllLoadingTracksToggleOn()); }
 
 private:
-	const TCHAR* GetEventNameByEventType(uint32 Depth, const Trace::FLoadTimeProfilerCpuEvent& Event) const;
-	const TCHAR* GetEventNameByPackageName(uint32 Depth, const Trace::FLoadTimeProfilerCpuEvent& Event) const;
-	const TCHAR* GetEventNameByExportClassName(uint32 Depth, const Trace::FLoadTimeProfilerCpuEvent& Event) const;
-	const TCHAR* GetEventNameByPackageAndExportClassName(uint32 Depth, const Trace::FLoadTimeProfilerCpuEvent& Event) const;
+	const TCHAR* GetEventNameByEventType(uint32 Depth, const TraceServices::FLoadTimeProfilerCpuEvent& Event) const;
+	const TCHAR* GetEventNameByPackageName(uint32 Depth, const TraceServices::FLoadTimeProfilerCpuEvent& Event) const;
+	const TCHAR* GetEventNameByExportClassName(uint32 Depth, const TraceServices::FLoadTimeProfilerCpuEvent& Event) const;
+	const TCHAR* GetEventNameByPackageAndExportClassName(uint32 Depth, const TraceServices::FLoadTimeProfilerCpuEvent& Event) const;
 
 private:
 	STimingView* TimingView;
@@ -87,7 +87,7 @@ public:
 
 protected:
 	// Helper function to find an event given search parameters
-	bool FindLoadTimeProfilerCpuEvent(const FTimingEventSearchParameters& InParameters, TFunctionRef<void(double, double, uint32, const Trace::FLoadTimeProfilerCpuEvent&)> InFoundPredicate) const;
+	bool FindLoadTimeProfilerCpuEvent(const FTimingEventSearchParameters& InParameters, TFunctionRef<void(double, double, uint32, const TraceServices::FLoadTimeProfilerCpuEvent&)> InFoundPredicate) const;
 
 protected:
 	FLoadingSharedState& SharedState;
diff --git a/Engine/Source/Developer/TraceInsights/Private/Insights/ViewModels/LogMessage.cpp b/Engine/Source/Developer/TraceInsights/Private/Insights/ViewModels/LogMessage.cpp
--- a/Engine/Source/Developer/TraceInsights/Private/Insights/ViewModels/LogMessage.cpp
+++ b/Engine/Source/Developer/TraceInsights/Private/Insights/ViewModels/LogMessage.cpp
@@ -26,7 +26,7 @@ FLogMessageRecord::FLogMessageRecord()
 
 ////////////////////////////////////////////////////////////////////////////////////////////////////
 
-FLogMessageRecord::FLogMessageRecord(const Trace::FLogMessage& TraceLogMessage)
+FLogMessageRecord::FLogMessageRecord(const TraceServices::FLogMessage& TraceLogMessage)
 	: Index(static_cast<int32>(TraceLogMessage.Index))
 	, Time(TraceLogMessage.Time)
 	, Verbosity(TraceLogMessage.Verbosity)
@@ -119,7 +119,7 @@ FLogMessageCache::FLogMessageCache()
 
 ////////////////////////////////////////////////////////////////////////////////////////////////////
 
-void FLogMessageCache::SetSession(TSharedPtr<const Trace::IAnalysisSession> InSession)
+void FLogMessageCache::SetSession(TSharedPtr<const TraceServices::IAnalysisSession> InSession)
 {
 	if (Session != InSession)
 	{
@@ -156,9 +156,9 @@ FLogMessageRecord& FLogMessageCache::Get(uint64 Index)
 
 	if (Session.IsValid())
 	{
-		Trace::FAnalysisSessionReadScope SessionReadScope(*Session.Get());
-		const Trace::ILogProvider& LogProvider = Trace::ReadLogProvider(*Session.Get());
-		LogProvider.ReadMessage(Index, [this, Index](const Trace::FLogMessage& Message)
+		TraceServices::FAnalysisSessionReadScope SessionReadScope(*Session.Get());
+		const TraceServices::ILogProvider& LogProvider = TraceServices::ReadLogProvider(*Session.Get());
+		LogProvider.ReadMessage(Index, [this, Index](const TraceServices::FLogMessage& Message)
 		{
 			FScopeLock Lock(&CriticalSection);
 			FLogMessageRecord Entry(Message);
@@ -185,9 +185,9 @@ TSharedPtr<FLogMessageRecord> FLogMessageCache::GetUncached(uint64 Index) const
 
 	if (Session.IsValid())
 	{
-		Trace::FAnalysisSessionReadScope SessionReadScope(*Session.Get());
-		const Trace::ILogProvider& LogProvider = Trace::ReadLogProvider(*Session.Get());
-		LogProvider.ReadMessage(Index, [&EntryPtr](const Trace::FLogMessage& Message)
+		TraceServices::FAnalysisSessionReadScope SessionReadScope(*Session.Get());
+		const TraceServices::ILogProvider& LogProvider = TraceServices::ReadLogProvider(*Session.Get());
+		LogProvider.ReadMessage(Index, [&EntryPtr](const TraceServices::FLogMessage& Message)
 		{
 			EntryPtr = MakeShared<FLogMessageRecord>(Message);
 		});
diff --git a/Engine/Source/Developer/TraceInsights/Private/Insights/ViewModels/LogMessage.h b/Engine/Source/Developer/TraceInsights/Private/Insights/ViewModels/LogMessage.h
--- a/Engine/Source/Developer/TraceInsights/Private/Insights/ViewModels/LogMessage.h
+++ b/Engine/Source/Developer/TraceInsights/Private/Insights/ViewModels/LogMessage.h
@@ -6,7 +6,7 @@
 
 ////////////////////////////////////////////////////////////////////////////////////////////////////
 
-namespace Trace
+namespace TraceServices
 {
 	struct FLogMessage;
 	class IAnalysisSession;
@@ -18,7 +18,7 @@ class FLogMessageRecord
 {
 public:
 	FLogMessageRecord();
-	FLogMessageRecord(const Trace::FLogMessage& Message);
+	FLogMessageRecord(const TraceServices::FLogMessage& Message);
 
 	FText GetIndexAsText() const;
 	FText GetTimeAsText() const;
@@ -47,7 +47,7 @@ class FLogMessageCache
 public:
 	FLogMessageCache();
 
-	void SetSession(TSharedPtr<const Trace::IAnalysisSession> InSession);
+	void SetSession(TSharedPtr<const TraceServices::IAnalysisSession> InSession);
 	void Reset();
 
 	FLogMessageRecord& Get(uint64 Index);
@@ -55,7 +55,7 @@ public:
 
 private:
 	FCriticalSection CriticalSection;
-	TSharedPtr<const Trace::IAnalysisSession> Session;
+	TSharedPtr<const TraceServices::IAnalysisSession> Session;
 	TMap<uint64, FLogMessageRecord> Map;
 	FLogMessageRecord InvalidEntry;
 };
diff --git a/Engine/Source/Developer/TraceInsights/Private/Insights/ViewModels/MarkersTimingTrack.cpp b/Engine/Source/Developer/TraceInsights/Private/Insights/ViewModels/MarkersTimingTrack.cpp
--- a/Engine/Source/Developer/TraceInsights/Private/Insights/ViewModels/MarkersTimingTrack.cpp
+++ b/Engine/Source/Developer/TraceInsights/Private/Insights/ViewModels/MarkersTimingTrack.cpp
@@ -120,18 +120,18 @@ void FMarkersTimingTrack::UpdateDrawState(const FTimingTrackViewport& InViewport
 {
 	FTimeMarkerTrackBuilder Builder(*this, InViewport);
 
-	TSharedPtr<const Trace::IAnalysisSession> Session = FInsightsManager::Get()->GetSession();
+	TSharedPtr<const TraceServices::IAnalysisSession> Session = FInsightsManager::Get()->GetSession();
 	if (Session.IsValid())
 	{
-		Trace::FAnalysisSessionReadScope SessionReadScope(*Session.Get());
+		TraceServices::FAnalysisSessionReadScope SessionReadScope(*Session.Get());
 
-		const Trace::ILogProvider& LogProvider = Trace::ReadLogProvider(*Session.Get());
+		const TraceServices::ILogProvider& LogProvider = TraceServices::ReadLogProvider(*Session.Get());
 		Builder.BeginLog(LogProvider);
 
 		LogProvider.EnumerateMessages(
 			Builder.GetViewport().GetStartTime(),
 			Builder.GetViewport().GetEndTime(),
-			[&Builder](const Trace::FLogMessage& Message) { Builder.AddLogMessage(Message); });
+			[&Builder](const TraceServices::FLogMessage& Message) { Builder.AddLogMessage(Message); });
 
 		Builder.EndLog();
 	}
@@ -298,7 +298,7 @@ FTimeMarkerTrackBuilder::FTimeMarkerTrackBuilder(FMarkersTimingTrack& InTrack, c
 
 ////////////////////////////////////////////////////////////////////////////////////////////////////
 
-void FTimeMarkerTrackBuilder::BeginLog(const Trace::ILogProvider& LogProvider)
+void FTimeMarkerTrackBuilder::BeginLog(const TraceServices::ILogProvider& LogProvider)
 {
 	LogProviderPtr = &LogProvider;
 
@@ -312,7 +312,7 @@ void FTimeMarkerTrackBuilder::BeginLog(const Trace::ILogProvider& LogProvider)
 
 ////////////////////////////////////////////////////////////////////////////////////////////////////
 
-void FTimeMarkerTrackBuilder::AddLogMessage(const Trace::FLogMessage& Message)
+void FTimeMarkerTrackBuilder::AddLogMessage(const TraceServices::FLogMessage& Message)
 {
 	Track.NumLogMessages++;
 
@@ -323,7 +323,7 @@ void FTimeMarkerTrackBuilder::AddLogMessage(const Trace::FLogMessage& Message)
 		//TODO: Search API like: LogProviderPtr->SearchMessage(StartIndex, ESearchDirection::Backward, LambdaPredicate, bResolveFormatString);
 		LogProviderPtr->ReadMessage(
 			Message.Index - 1,
-			[this](const Trace::FLogMessage& Message) { AddLogMessage(Message); });
+			[this](const TraceServices::FLogMessage& Message) { AddLogMessage(Message); });
 	}
 
 	check(Message.Category != nullptr);
diff --git a/Engine/Source/Developer/TraceInsights/Private/Insights/ViewModels/MarkersTimingTrack.h b/Engine/Source/Developer/TraceInsights/Private/Insights/ViewModels/MarkersTimingTrack.h
--- a/Engine/Source/Developer/TraceInsights/Private/Insights/ViewModels/MarkersTimingTrack.h
+++ b/Engine/Source/Developer/TraceInsights/Private/Insights/ViewModels/MarkersTimingTrack.h
@@ -9,7 +9,7 @@
 #include "Insights/ViewModels/BaseTimingTrack.h"
 #include "Insights/ViewModels/TrackHeader.h"
 
-namespace Trace
+namespace TraceServices
 {
 	struct FLogMessage;
 	class ILogProvider;
@@ -124,8 +124,8 @@ public:
 
 	const FTimingTrackViewport& GetViewport() { return Viewport; }
 
-	void BeginLog(const Trace::ILogProvider& LogProvider);
-	void AddLogMessage(const Trace::FLogMessage& Message);
+	void BeginLog(const TraceServices::ILogProvider& LogProvider);
+	void AddLogMessage(const TraceServices::FLogMessage& Message);
 	void EndLog();
 
 	static FLinearColor GetColorByCategory(const TCHAR* const Category);
@@ -142,7 +142,7 @@ private:
 	const TSharedRef<class FSlateFontMeasure> FontMeasureService;
 	const FSlateFontInfo Font;
 
-	const Trace::ILogProvider* LogProviderPtr; // valid only between BeginLog() and EndLog()
+	const TraceServices::ILogProvider* LogProviderPtr; // valid only between BeginLog() and EndLog()
 
 	float LastX1;
 	float LastX2;
diff --git a/Engine/Source/Developer/TraceInsights/Private/Insights/ViewModels/StatsAggregator.cpp b/Engine/Source/Developer/TraceInsights/Private/Insights/ViewModels/StatsAggregator.cpp
--- a/Engine/Source/Developer/TraceInsights/Private/Insights/ViewModels/StatsAggregator.cpp
+++ b/Engine/Source/Developer/TraceInsights/Private/Insights/ViewModels/StatsAggregator.cpp
@@ -97,7 +97,7 @@ void FStatsAggregator::Start()
 
 ////////////////////////////////////////////////////////////////////////////////////////////////////
 
-void FStatsAggregator::Tick(TSharedPtr<const Trace::IAnalysisSession> InSession, const double InCurrentTime, const float InDeltaTime, TFunctionRef<void()> OnFinishedCallback)
+void FStatsAggregator::Tick(TSharedPtr<const TraceServices::IAnalysisSession> InSession, const double InCurrentTime, const float InDeltaTime, TFunctionRef<void()> OnFinishedCallback)
 {
 	if (AsyncTask && AsyncTask->IsDone())
 	{
diff --git a/Engine/Source/Developer/TraceInsights/Private/Insights/ViewModels/StatsAggregator.h b/Engine/Source/Developer/TraceInsights/Private/Insights/ViewModels/StatsAggregator.h
--- a/Engine/Source/Developer/TraceInsights/Private/Insights/ViewModels/StatsAggregator.h
+++ b/Engine/Source/Developer/TraceInsights/Private/Insights/ViewModels/StatsAggregator.h
@@ -8,7 +8,7 @@
 
 ////////////////////////////////////////////////////////////////////////////////////////////////////
 
-namespace Trace
+namespace TraceServices
 {
 	class IAnalysisSession;
 }
@@ -63,7 +63,7 @@ public:
 		IntervalEndTime = InEndTime;
 	}
 
-	void Tick(TSharedPtr<const Trace::IAnalysisSession> InSession, const double InCurrentTime, const float InDeltaTime, TFunctionRef<void()> OnFinishedCallback);
+	void Tick(TSharedPtr<const TraceServices::IAnalysisSession> InSession, const double InCurrentTime, const float InDeltaTime, TFunctionRef<void()> OnFinishedCallback);
 
 	//////////////////////////////////////////////////
 	// IStatsAggregator
@@ -81,7 +81,7 @@ public:
 	//////////////////////////////////////////////////
 
 protected:
-	virtual IStatsAggregationWorker* CreateWorker(TSharedPtr<const Trace::IAnalysisSession> InSession) = 0;
+	virtual IStatsAggregationWorker* CreateWorker(TSharedPtr<const TraceServices::IAnalysisSession> InSession) = 0;
 
 	// Returns true only when it is called from OnFinishedCallback.
 	bool IsFinished() const { return bIsFinished; }
