Moved the TraceServices module into its own namespace.

diff --git a/Engine/Source/Developer/TraceInsights/Private/Insights/ViewModels/ThreadTimingTrack.cpp b/Engine/Source/Developer/TraceInsights/Private/Insights/ViewModels/ThreadTimingTrack.cpp
--- a/Engine/Source/Developer/TraceInsights/Private/Insights/ViewModels/ThreadTimingTrack.cpp
+++ b/Engine/Source/Developer/TraceInsights/Private/Insights/ViewModels/ThreadTimingTrack.cpp
@@ -198,7 +198,7 @@ static void AppendMetadataToString(FString& Str, TArrayView<const uint8>& Metada
 
 static void AddTimingEventToBuilder(ITimingEventsTrackDrawStateBuilder& Builder,
 									double EventStartTime, double EventEndTime, uint32 EventDepth,
-									uint32 TimerIndex, const Trace::FTimingProfilerTimer* Timer)
+									uint32 TimerIndex, const TraceServices::FTimingProfilerTimer* Timer)
 {
 	//const uint32 EventColor = FTimingEvent::ComputeEventColor(Timer->Id);
 	const uint32 EventColor = FTimingEvent::ComputeEventColor(Timer->Name);
@@ -219,15 +219,15 @@ static void AddTimingEventToBuilder(ITimingEventsTrackDrawStateBuilder& Builder,
 				{
 					//EventName = TEXT("!") + EventName; // for debugging
 
-					TSharedPtr<const Trace::IAnalysisSession> Session = FInsightsManager::Get()->GetSession();
+					TSharedPtr<const TraceServices::IAnalysisSession> Session = FInsightsManager::Get()->GetSession();
 					check(Session.IsValid());
 
-					//Trace::FAnalysisSessionReadScope SessionReadScope(*Session.Get());
+					//TraceServices::FAnalysisSessionReadScope SessionReadScope(*Session.Get());
 
-					const Trace::ITimingProfilerProvider& TimingProfilerProvider = *Trace::ReadTimingProfilerProvider(*Session.Get());
+					const TraceServices::ITimingProfilerProvider& TimingProfilerProvider = *TraceServices::ReadTimingProfilerProvider(*Session.Get());
 
-					const Trace::ITimingProfilerTimerReader* TimerReader;
-					TimingProfilerProvider.ReadTimers([&TimerReader](const Trace::ITimingProfilerTimerReader& Out) { TimerReader = &Out; });
+					const TraceServices::ITimingProfilerTimerReader* TimerReader;
+					TimingProfilerProvider.ReadTimers([&TimerReader](const TraceServices::ITimingProfilerTimerReader& Out) { TimerReader = &Out; });
 
 					TArrayView<const uint8> Metadata = TimerReader->GetMetadata(TimerIndex);
 					if (Metadata.Num() > 0)
@@ -331,19 +331,19 @@ void FThreadTimingSharedState::OnEndSession(Insights::ITimingViewSession& InSess
 
 ////////////////////////////////////////////////////////////////////////////////////////////////////
 
-void FThreadTimingSharedState::Tick(Insights::ITimingViewSession& InSession, const Trace::IAnalysisSession& InAnalysisSession)
+void FThreadTimingSharedState::Tick(Insights::ITimingViewSession& InSession, const TraceServices::IAnalysisSession& InAnalysisSession)
 {
 	if (&InSession != TimingView)
 	{
 		return;
 	}
 
-	const Trace::ITimingProfilerProvider* TimingProfilerProvider = Trace::ReadTimingProfilerProvider(InAnalysisSession);
-	const Trace::ILoadTimeProfilerProvider* LoadTimeProfilerProvider = Trace::ReadLoadTimeProfilerProvider(InAnalysisSession);
+	const TraceServices::ITimingProfilerProvider* TimingProfilerProvider = TraceServices::ReadTimingProfilerProvider(InAnalysisSession);
+	const TraceServices::ILoadTimeProfilerProvider* LoadTimeProfilerProvider = TraceServices::ReadLoadTimeProfilerProvider(InAnalysisSession);
 
 	if (TimingProfilerProvider)
 	{
-		Trace::FAnalysisSessionReadScope SessionReadScope(InAnalysisSession);
+		TraceServices::FAnalysisSessionReadScope SessionReadScope(InAnalysisSession);
 
 		const uint64 CurrentTimingProfilerTimelineCount = TimingProfilerProvider->GetTimelineCount();
 		const uint64 CurrentLoadTimeProfilerTimelineCount = (LoadTimeProfilerProvider) ? LoadTimeProfilerProvider->GetTimelineCount() : 0;
@@ -371,8 +371,8 @@ void FThreadTimingSharedState::Tick(Insights::ITimingViewSession& InSession, con
 			int32 Order = FTimingTrackOrder::Cpu;
 
 			// Iterate through threads.
-			const Trace::IThreadProvider& ThreadProvider = Trace::ReadThreadProvider(InAnalysisSession);
-			ThreadProvider.EnumerateThreads([this, &InSession, &bTracksOrderChanged, &Order, TimingProfilerProvider, LoadTimeProfilerProvider](const Trace::FThreadInfo& ThreadInfo)
+			const TraceServices::IThreadProvider& ThreadProvider = TraceServices::ReadThreadProvider(InAnalysisSession);
+			ThreadProvider.EnumerateThreads([this, &InSession, &bTracksOrderChanged, &Order, TimingProfilerProvider, LoadTimeProfilerProvider](const TraceServices::FThreadInfo& ThreadInfo)
 			{
 				// Check if this thread is part of a group?
 				bool bIsGroupVisible = bShowHideAllCpuTracks;
@@ -625,28 +625,28 @@ FThreadTimingTrack::~FThreadTimingTrack()
 
 void FThreadTimingTrack::BuildDrawState(ITimingEventsTrackDrawStateBuilder& Builder, const ITimingTrackUpdateContext& Context)
 {
-	TSharedPtr<const Trace::IAnalysisSession> Session = FInsightsManager::Get()->GetSession();
-	if (Session.IsValid() && Trace::ReadTimingProfilerProvider(*Session.Get()))
+	TSharedPtr<const TraceServices::IAnalysisSession> Session = FInsightsManager::Get()->GetSession();
+	if (Session.IsValid() && TraceServices::ReadTimingProfilerProvider(*Session.Get()))
 	{
-		Trace::FAnalysisSessionReadScope SessionReadScope(*Session.Get());
+		TraceServices::FAnalysisSessionReadScope SessionReadScope(*Session.Get());
 
-		const Trace::ITimingProfilerProvider& TimingProfilerProvider = *Trace::ReadTimingProfilerProvider(*Session.Get());
+		const TraceServices::ITimingProfilerProvider& TimingProfilerProvider = *TraceServices::ReadTimingProfilerProvider(*Session.Get());
 
-		const Trace::ITimingProfilerTimerReader* TimerReader;
-		TimingProfilerProvider.ReadTimers([&TimerReader](const Trace::ITimingProfilerTimerReader& Out) { TimerReader = &Out; });
+		const TraceServices::ITimingProfilerTimerReader* TimerReader;
+		TimingProfilerProvider.ReadTimers([&TimerReader](const TraceServices::ITimingProfilerTimerReader& Out) { TimerReader = &Out; });
 
 		const FTimingTrackViewport& Viewport = Context.GetViewport();
 
 		TimingProfilerProvider.ReadTimeline(TimelineIndex,
-			[&Viewport, this, &Builder, TimerReader](const Trace::ITimingProfilerProvider::Timeline& Timeline)
+			[&Viewport, this, &Builder, TimerReader](const TraceServices::ITimingProfilerProvider::Timeline& Timeline)
 			{
 				if (FTimingEventsTrack::bUseDownSampling)
 				{
 					const double SecondsPerPixel = 1.0 / Viewport.GetScaleX();
 					Timeline.EnumerateEventsDownSampled(Viewport.GetStartTime(), Viewport.GetEndTime(), SecondsPerPixel,
-						[this, &Builder, TimerReader](double StartTime, double EndTime, uint32 Depth, const Trace::FTimingProfilerEvent& Event)
+						[this, &Builder, TimerReader](double StartTime, double EndTime, uint32 Depth, const TraceServices::FTimingProfilerEvent& Event)
 						{
-							const Trace::FTimingProfilerTimer* Timer = TimerReader->GetTimer(Event.TimerIndex);
+							const TraceServices::FTimingProfilerTimer* Timer = TimerReader->GetTimer(Event.TimerIndex);
 							if (ensure(Timer != nullptr))
 							{
 								AddTimingEventToBuilder(Builder, StartTime, EndTime, Depth, Event.TimerIndex, Timer);
@@ -655,15 +655,15 @@ void FThreadTimingTrack::BuildDrawState(ITimingEventsTrackDrawStateBuilder& Buil
 							{
 								Builder.AddEvent(StartTime, EndTime, Depth, 0xFF000000, [&Event](float) { return FString::Printf(TEXT("[%u]"), Event.TimerIndex); });
 							}
-							return Trace::EEventEnumerate::Continue;
+							return TraceServices::EEventEnumerate::Continue;
 						});
 				}
 				else
 				{
 					Timeline.EnumerateEvents(Viewport.GetStartTime(), Viewport.GetEndTime(),
-						[this, &Builder, TimerReader](double StartTime, double EndTime, uint32 Depth, const Trace::FTimingProfilerEvent& Event)
+						[this, &Builder, TimerReader](double StartTime, double EndTime, uint32 Depth, const TraceServices::FTimingProfilerEvent& Event)
 						{
-							const Trace::FTimingProfilerTimer* Timer = TimerReader->GetTimer(Event.TimerIndex);
+							const TraceServices::FTimingProfilerTimer* Timer = TimerReader->GetTimer(Event.TimerIndex);
 							if (ensure(Timer != nullptr))
 							{
 								AddTimingEventToBuilder(Builder, StartTime, EndTime, Depth, Event.TimerIndex, Timer);
@@ -672,7 +672,7 @@ void FThreadTimingTrack::BuildDrawState(ITimingEventsTrackDrawStateBuilder& Buil
 							{
 								Builder.AddEvent(StartTime, EndTime, Depth, 0xFF000000, [&Event](float) { return FString::Printf(TEXT("[%u]"), Event.TimerIndex); });
 							}
-							return Trace::EEventEnumerate::Continue;
+							return TraceServices::EEventEnumerate::Continue;
 						});
 				}
 			});
@@ -695,15 +695,15 @@ void FThreadTimingTrack::BuildFilteredDrawState(ITimingEventsTrackDrawStateBuild
 			FilterEventType = EventFilter.GetEventType();
 		}
 
-		TSharedPtr<const Trace::IAnalysisSession> Session = FInsightsManager::Get()->GetSession();
-		if (Session.IsValid() && Trace::ReadTimingProfilerProvider(*Session.Get()))
+		TSharedPtr<const TraceServices::IAnalysisSession> Session = FInsightsManager::Get()->GetSession();
+		if (Session.IsValid() && TraceServices::ReadTimingProfilerProvider(*Session.Get()))
 		{
-			Trace::FAnalysisSessionReadScope SessionReadScope(*Session.Get());
+			TraceServices::FAnalysisSessionReadScope SessionReadScope(*Session.Get());
 
-			const Trace::ITimingProfilerProvider& TimingProfilerProvider = *Trace::ReadTimingProfilerProvider(*Session.Get());
+			const TraceServices::ITimingProfilerProvider& TimingProfilerProvider = *TraceServices::ReadTimingProfilerProvider(*Session.Get());
 
-			const Trace::ITimingProfilerTimerReader* TimerReader;
-			TimingProfilerProvider.ReadTimers([&TimerReader](const Trace::ITimingProfilerTimerReader& Out) { TimerReader = &Out; });
+			const TraceServices::ITimingProfilerTimerReader* TimerReader;
+			TimingProfilerProvider.ReadTimers([&TimerReader](const TraceServices::ITimingProfilerTimerReader& Out) { TimerReader = &Out; });
 
 			const FTimingTrackViewport& Viewport = Context.GetViewport();
 
@@ -713,11 +713,11 @@ void FThreadTimingTrack::BuildFilteredDrawState(ITimingEventsTrackDrawStateBuild
 				if (true)
 				{
 					TimingProfilerProvider.ReadTimeline(TimelineIndex,
-						[&Viewport, this, &Builder, TimerReader, FilterEventType](const Trace::ITimingProfilerProvider::Timeline& Timeline)
+						[&Viewport, this, &Builder, TimerReader, FilterEventType](const TraceServices::ITimingProfilerProvider::Timeline& Timeline)
 						{
 							TArray<TArray<FPendingEventInfo>> FilteredEvents;
 
-							Trace::ITimeline<Trace::FTimingProfilerEvent>::EnumerateAsyncParams Params;
+							TraceServices::ITimeline<TraceServices::FTimingProfilerEvent>::EnumerateAsyncParams Params;
 							Params.IntervalStart = Viewport.GetStartTime();
 							Params.IntervalEnd = Viewport.GetEndTime();
 							Params.Resolution = 0.0;
@@ -725,9 +725,9 @@ void FThreadTimingTrack::BuildFilteredDrawState(ITimingEventsTrackDrawStateBuild
 							{
 								FilteredEvents.AddDefaulted(NumTasks);
 							};
-							Params.Callback = [this, &Builder, TimerReader, FilterEventType, &FilteredEvents](double StartTime, double EndTime, uint32 Depth, const Trace::FTimingProfilerEvent& Event, uint32 TaskIndex)
+							Params.Callback = [this, &Builder, TimerReader, FilterEventType, &FilteredEvents](double StartTime, double EndTime, uint32 Depth, const TraceServices::FTimingProfilerEvent& Event, uint32 TaskIndex)
 							{
-								const Trace::FTimingProfilerTimer* Timer = TimerReader->GetTimer(Event.TimerIndex);
+								const TraceServices::FTimingProfilerTimer* Timer = TimerReader->GetTimer(Event.TimerIndex);
 								if (ensure(Timer != nullptr))
 								{
 									if (Timer->Id == FilterEventType)
@@ -740,7 +740,7 @@ void FThreadTimingTrack::BuildFilteredDrawState(ITimingEventsTrackDrawStateBuild
 										FilteredEvents[TaskIndex].Add(TimelineEvent);
 									}
 								}
-								return Trace::EEventEnumerate::Continue;
+								return TraceServices::EEventEnumerate::Continue;
 							};
 
 							// Note: Enumerating events for filtering should not use downsampling.
@@ -750,7 +750,7 @@ void FThreadTimingTrack::BuildFilteredDrawState(ITimingEventsTrackDrawStateBuild
 							{
 								for (FPendingEventInfo& TimelineEvent : Array)
 								{
-									const Trace::FTimingProfilerTimer* Timer = TimerReader->GetTimer(TimelineEvent.TimerIndex);
+									const TraceServices::FTimingProfilerTimer* Timer = TimerReader->GetTimer(TimelineEvent.TimerIndex);
 									AddTimingEventToBuilder(Builder, TimelineEvent.StartTime, TimelineEvent.EndTime, TimelineEvent.Depth, TimelineEvent.TimerIndex, Timer);
 								}
 							}
@@ -759,13 +759,13 @@ void FThreadTimingTrack::BuildFilteredDrawState(ITimingEventsTrackDrawStateBuild
 				else
 				{
 					TimingProfilerProvider.ReadTimeline(TimelineIndex,
-						[&Viewport, this, &Builder, TimerReader, FilterEventType](const Trace::ITimingProfilerProvider::Timeline& Timeline)
+						[&Viewport, this, &Builder, TimerReader, FilterEventType](const TraceServices::ITimingProfilerProvider::Timeline& Timeline)
 						{
 							// Note: Enumerating events for filtering should not use downsampling.
 							Timeline.EnumerateEventsDownSampled(Viewport.GetStartTime(), Viewport.GetEndTime(), 0,
-								[this, &Builder, TimerReader, FilterEventType](double StartTime, double EndTime, uint32 Depth, const Trace::FTimingProfilerEvent& Event)
+								[this, &Builder, TimerReader, FilterEventType](double StartTime, double EndTime, uint32 Depth, const TraceServices::FTimingProfilerEvent& Event)
 								{
-									const Trace::FTimingProfilerTimer* Timer = TimerReader->GetTimer(Event.TimerIndex);
+									const TraceServices::FTimingProfilerTimer* Timer = TimerReader->GetTimer(Event.TimerIndex);
 									if (ensure(Timer != nullptr))
 									{
 										if (Timer->Id == FilterEventType)
@@ -773,7 +773,7 @@ void FThreadTimingTrack::BuildFilteredDrawState(ITimingEventsTrackDrawStateBuild
 											AddTimingEventToBuilder(Builder, StartTime, EndTime, Depth, Event.TimerIndex, Timer);
 										}
 									}
-									return Trace::EEventEnumerate::Continue;
+									return TraceServices::EEventEnumerate::Continue;
 								});
 						});
 				}
@@ -781,15 +781,15 @@ void FThreadTimingTrack::BuildFilteredDrawState(ITimingEventsTrackDrawStateBuild
 			else // generic filter
 			{
 				TimingProfilerProvider.ReadTimeline(TimelineIndex,
-					[&Viewport, this, &Builder, TimerReader, &EventFilterPtr](const Trace::ITimingProfilerProvider::Timeline& Timeline)
+					[&Viewport, this, &Builder, TimerReader, &EventFilterPtr](const TraceServices::ITimingProfilerProvider::Timeline& Timeline)
 					{
 						// Note: Enumerating events for filtering should not use downsampling.
 						//const double SecondsPerPixel = 1.0 / Viewport.GetScaleX();
 						//Timeline.EnumerateEventsDownSampled(Viewport.GetStartTime(), Viewport.GetEndTime(), SecondsPerPixel,
 						Timeline.EnumerateEvents(Viewport.GetStartTime(), Viewport.GetEndTime(),
-							[this, &Builder, TimerReader, &EventFilterPtr](double StartTime, double EndTime, uint32 Depth, const Trace::FTimingProfilerEvent& Event)
+							[this, &Builder, TimerReader, &EventFilterPtr](double StartTime, double EndTime, uint32 Depth, const TraceServices::FTimingProfilerEvent& Event)
 							{
-								const Trace::FTimingProfilerTimer* Timer = TimerReader->GetTimer(Event.TimerIndex);
+								const TraceServices::FTimingProfilerTimer* Timer = TimerReader->GetTimer(Event.TimerIndex);
 								if (ensure(Timer != nullptr))
 								{
 									FThreadTrackEvent TimingEvent(SharedThis(this), StartTime, EndTime, Depth);
@@ -801,7 +801,7 @@ void FThreadTimingTrack::BuildFilteredDrawState(ITimingEventsTrackDrawStateBuild
 										AddTimingEventToBuilder(Builder, StartTime, EndTime, Depth, Event.TimerIndex, Timer);
 									}
 								}
-								return Trace::EEventEnumerate::Continue;
+								return TraceServices::EEventEnumerate::Continue;
 							});
 					});
 			}
@@ -810,25 +810,25 @@ void FThreadTimingTrack::BuildFilteredDrawState(ITimingEventsTrackDrawStateBuild
 
 	if (HasCustomFilter()) // Custom filter (from the filtering widget)
 	{
-		TSharedPtr<const Trace::IAnalysisSession> Session = FInsightsManager::Get()->GetSession();
-		if (Session.IsValid() && Trace::ReadTimingProfilerProvider(*Session.Get()))
+		TSharedPtr<const TraceServices::IAnalysisSession> Session = FInsightsManager::Get()->GetSession();
+		if (Session.IsValid() && TraceServices::ReadTimingProfilerProvider(*Session.Get()))
 		{
-			Trace::FAnalysisSessionReadScope SessionReadScope(*Session.Get());
+			TraceServices::FAnalysisSessionReadScope SessionReadScope(*Session.Get());
 
-			const Trace::ITimingProfilerProvider& TimingProfilerProvider = *Trace::ReadTimingProfilerProvider(*Session.Get());
+			const TraceServices::ITimingProfilerProvider& TimingProfilerProvider = *TraceServices::ReadTimingProfilerProvider(*Session.Get());
 
-			const Trace::ITimingProfilerTimerReader* TimerReader;
-			TimingProfilerProvider.ReadTimers([&TimerReader](const Trace::ITimingProfilerTimerReader& Out) { TimerReader = &Out; });
+			const TraceServices::ITimingProfilerTimerReader* TimerReader;
+			TimingProfilerProvider.ReadTimers([&TimerReader](const TraceServices::ITimingProfilerTimerReader& Out) { TimerReader = &Out; });
 
 			const FTimingTrackViewport& Viewport = Context.GetViewport();
 
 			TimingProfilerProvider.ReadTimeline(TimelineIndex,
-				[&Viewport, this, &Builder, TimerReader](const Trace::ITimingProfilerProvider::Timeline& Timeline)
+				[&Viewport, this, &Builder, TimerReader](const TraceServices::ITimingProfilerProvider::Timeline& Timeline)
 				{
 					TArray<TArray<FPendingEventInfo>> FilteredEvents;
 					TArray<FFilterContext> FilterContexts;
 
-					Trace::ITimeline<Trace::FTimingProfilerEvent>::EnumerateAsyncParams Params;
+					TraceServices::ITimeline<TraceServices::FTimingProfilerEvent>::EnumerateAsyncParams Params;
 					Params.IntervalStart = Viewport.GetStartTime();
 					Params.IntervalEnd = Viewport.GetEndTime();
 
@@ -846,9 +846,9 @@ void FThreadTimingTrack::BuildFilteredDrawState(ITimingEventsTrackDrawStateBuild
 							Context.AddFilterData<int64>(EFilterField::EventType, 0);
 						}
 					};
-					Params.Callback = [this, &Builder, TimerReader, &FilteredEvents, &FilterContexts](double StartTime, double EndTime, uint32 Depth, const Trace::FTimingProfilerEvent& Event, uint32 TaskIndex)
+					Params.Callback = [this, &Builder, TimerReader, &FilteredEvents, &FilterContexts](double StartTime, double EndTime, uint32 Depth, const TraceServices::FTimingProfilerEvent& Event, uint32 TaskIndex)
 					{
-						const Trace::FTimingProfilerTimer* Timer = TimerReader->GetTimer(Event.TimerIndex);
+						const TraceServices::FTimingProfilerTimer* Timer = TimerReader->GetTimer(Event.TimerIndex);
 						if (ensure(Timer != nullptr))
 						{
 							FFilterContext& Context = FilterContexts[TaskIndex];
@@ -867,7 +867,7 @@ void FThreadTimingTrack::BuildFilteredDrawState(ITimingEventsTrackDrawStateBuild
 								FilteredEvents[TaskIndex].Add(TimelineEvent);
 							}
 						}
-						return Trace::EEventEnumerate::Continue;
+						return TraceServices::EEventEnumerate::Continue;
 					};
 
 					Timeline.EnumerateEventsDownSampledAsync(Params);
@@ -876,7 +876,7 @@ void FThreadTimingTrack::BuildFilteredDrawState(ITimingEventsTrackDrawStateBuild
 					{
 						for (FPendingEventInfo& TimelineEvent : Array)
 						{
-							const Trace::FTimingProfilerTimer* Timer = TimerReader->GetTimer(TimelineEvent.TimerIndex);
+							const TraceServices::FTimingProfilerTimer* Timer = TimerReader->GetTimer(TimelineEvent.TimerIndex);
 							AddTimingEventToBuilder(Builder, TimelineEvent.StartTime, TimelineEvent.EndTime, TimelineEvent.Depth, TimelineEvent.TimerIndex, Timer);
 						}
 					}
@@ -904,17 +904,17 @@ void FThreadTimingTrack::PostDraw(const ITimingTrackDrawContext& Context) const
 
 void FThreadTimingTrack::DrawSelectedEventInfo(const FThreadTrackEvent& SelectedEvent, const FTimingTrackViewport& Viewport, const FDrawContext& DrawContext, const FSlateBrush* WhiteBrush, const FSlateFontInfo& Font) const
 {
-	TSharedPtr<const Trace::IAnalysisSession> Session = FInsightsManager::Get()->GetSession();
+	TSharedPtr<const TraceServices::IAnalysisSession> Session = FInsightsManager::Get()->GetSession();
 	check(Session.IsValid());
 
-	Trace::FAnalysisSessionReadScope SessionReadScope(*Session.Get());
+	TraceServices::FAnalysisSessionReadScope SessionReadScope(*Session.Get());
 
-	const Trace::ITimingProfilerProvider& TimingProfilerProvider = *Trace::ReadTimingProfilerProvider(*Session.Get());
+	const TraceServices::ITimingProfilerProvider& TimingProfilerProvider = *TraceServices::ReadTimingProfilerProvider(*Session.Get());
 
-	const Trace::ITimingProfilerTimerReader* TimerReader;
-	TimingProfilerProvider.ReadTimers([&TimerReader](const Trace::ITimingProfilerTimerReader& Out) { TimerReader = &Out; });
+	const TraceServices::ITimingProfilerTimerReader* TimerReader;
+	TimingProfilerProvider.ReadTimers([&TimerReader](const TraceServices::ITimingProfilerTimerReader& Out) { TimerReader = &Out; });
 
-	const Trace::FTimingProfilerTimer* Timer = TimerReader->GetTimer(SelectedEvent.GetTimerIndex());
+	const TraceServices::FTimingProfilerTimer* Timer = TimerReader->GetTimer(SelectedEvent.GetTimerIndex());
 	if (Timer != nullptr)
 	{
 		FString Str = FString::Printf(TEXT("%s (Incl.: %s, Excl.: %s)"),
@@ -952,17 +952,17 @@ void FThreadTimingTrack::InitTooltip(FTooltipDrawState& InOutTooltip, const ITim
 		TSharedPtr<FThreadTrackEvent> RootTimingEvent;
 		GetParentAndRoot(TooltipEvent, ParentTimingEvent, RootTimingEvent);
 
-		TSharedPtr<const Trace::IAnalysisSession> Session = FInsightsManager::Get()->GetSession();
+		TSharedPtr<const TraceServices::IAnalysisSession> Session = FInsightsManager::Get()->GetSession();
 		check(Session.IsValid());
 
-		Trace::FAnalysisSessionReadScope SessionReadScope(*Session.Get());
+		TraceServices::FAnalysisSessionReadScope SessionReadScope(*Session.Get());
 
-		const Trace::ITimingProfilerProvider& TimingProfilerProvider = *Trace::ReadTimingProfilerProvider(*Session.Get());
+		const TraceServices::ITimingProfilerProvider& TimingProfilerProvider = *TraceServices::ReadTimingProfilerProvider(*Session.Get());
 
-		const Trace::ITimingProfilerTimerReader* TimerReader;
-		TimingProfilerProvider.ReadTimers([&TimerReader](const Trace::ITimingProfilerTimerReader& Out) { TimerReader = &Out; });
+		const TraceServices::ITimingProfilerTimerReader* TimerReader;
+		TimingProfilerProvider.ReadTimers([&TimerReader](const TraceServices::ITimingProfilerTimerReader& Out) { TimerReader = &Out; });
 
-		const Trace::FTimingProfilerTimer* Timer = TimerReader->GetTimer(TooltipEvent.GetTimerIndex());
+		const TraceServices::FTimingProfilerTimer* Timer = TimerReader->GetTimer(TooltipEvent.GetTimerIndex());
 		FString TimerName = (Timer != nullptr) ? Timer->Name : TEXT("N/A");
 		InOutTooltip.AddTitle(TimerName);
 
@@ -1013,16 +1013,16 @@ void FThreadTimingTrack::GetParentAndRoot(const FThreadTrackEvent& TimingEvent,
 {
 	if (TimingEvent.GetDepth() > 0)
 	{
-		TSharedPtr<const Trace::IAnalysisSession> Session = FInsightsManager::Get()->GetSession();
+		TSharedPtr<const TraceServices::IAnalysisSession> Session = FInsightsManager::Get()->GetSession();
 		if (Session.IsValid())
 		{
-			Trace::FAnalysisSessionReadScope SessionReadScope(*Session.Get());
+			TraceServices::FAnalysisSessionReadScope SessionReadScope(*Session.Get());
 
-			if (Trace::ReadTimingProfilerProvider(*Session.Get()))
+			if (TraceServices::ReadTimingProfilerProvider(*Session.Get()))
 			{
-				const Trace::ITimingProfilerProvider& TimingProfilerProvider = *Trace::ReadTimingProfilerProvider(*Session.Get());
+				const TraceServices::ITimingProfilerProvider& TimingProfilerProvider = *TraceServices::ReadTimingProfilerProvider(*Session.Get());
 
-				TimingProfilerProvider.ReadTimeline(GetTimelineIndex(), [&TimingEvent, &OutParentTimingEvent, &OutRootTimingEvent](const Trace::ITimingProfilerProvider::Timeline& Timeline)
+				TimingProfilerProvider.ReadTimeline(GetTimelineIndex(), [&TimingEvent, &OutParentTimingEvent, &OutRootTimingEvent](const TraceServices::ITimingProfilerProvider::Timeline& Timeline)
 					{
 						const double Time = (TimingEvent.GetStartTime() + TimingEvent.GetEndTime()) / 2;
 						TimelineEventInfo EventInfo;
@@ -1062,16 +1062,16 @@ const TSharedPtr<const ITimingEvent> FThreadTimingTrack::GetEvent(float InPosX,
 
 		const double EventTime = Viewport.SlateUnitsToTime(InPosX);
 
-		TSharedPtr<const Trace::IAnalysisSession> Session = FInsightsManager::Get()->GetSession();
+		TSharedPtr<const TraceServices::IAnalysisSession> Session = FInsightsManager::Get()->GetSession();
 		if (Session.IsValid())
 		{
-			Trace::FAnalysisSessionReadScope SessionReadScope(*Session.Get());
+			TraceServices::FAnalysisSessionReadScope SessionReadScope(*Session.Get());
 
-			if (Trace::ReadTimingProfilerProvider(*Session.Get()))
+			if (TraceServices::ReadTimingProfilerProvider(*Session.Get()))
 			{
-				const Trace::ITimingProfilerProvider& TimingProfilerProvider = *Trace::ReadTimingProfilerProvider(*Session.Get());
+				const TraceServices::ITimingProfilerProvider& TimingProfilerProvider = *TraceServices::ReadTimingProfilerProvider(*Session.Get());
 
-				TimingProfilerProvider.ReadTimeline(GetTimelineIndex(), [this, &EventTime, &Depth, &TimingEvent, &SecondsPerPixel](const Trace::ITimingProfilerProvider::Timeline& Timeline)
+				TimingProfilerProvider.ReadTimeline(GetTimelineIndex(), [this, &EventTime, &Depth, &TimingEvent, &SecondsPerPixel](const TraceServices::ITimingProfilerProvider::Timeline& Timeline)
 					{
 						TimelineEventInfo EventInfo;
 						bool IsFound = Timeline.GetEventInfo(EventTime, 2 * SecondsPerPixel, Depth, EventInfo);
@@ -1092,7 +1092,7 @@ const TSharedPtr<const ITimingEvent> FThreadTimingTrack::GetEvent(float InPosX,
 const TSharedPtr<const ITimingEvent> FThreadTimingTrack::SearchEvent(const FTimingEventSearchParameters& InSearchParameters) const
 {
 	TSharedPtr<FThreadTrackEvent> FoundEvent;
-	FindTimingProfilerEvent(InSearchParameters, [this, &FoundEvent](double InFoundStartTime, double InFoundEndTime, uint32 InFoundDepth, const Trace::FTimingProfilerEvent& InFoundEvent)
+	FindTimingProfilerEvent(InSearchParameters, [this, &FoundEvent](double InFoundStartTime, double InFoundEndTime, uint32 InFoundDepth, const TraceServices::FTimingProfilerEvent& InFoundEvent)
 	{
 		FoundEvent = MakeShared<FThreadTrackEvent>(SharedThis(this), InFoundStartTime, InFoundEndTime, InFoundDepth);
 		FoundEvent->SetTimerIndex(InFoundEvent.TimerIndex);
@@ -1120,17 +1120,17 @@ void FThreadTimingTrack::UpdateEventStats(ITimingEvent& InOutEvent) const
 			return;
 		}
 
-		TSharedPtr<const Trace::IAnalysisSession> Session = FInsightsManager::Get()->GetSession();
+		TSharedPtr<const TraceServices::IAnalysisSession> Session = FInsightsManager::Get()->GetSession();
 		if (Session.IsValid())
 		{
-			Trace::FAnalysisSessionReadScope SessionReadScope(*Session.Get());
+			TraceServices::FAnalysisSessionReadScope SessionReadScope(*Session.Get());
 
-			if (Trace::ReadTimingProfilerProvider(*Session.Get()))
+			if (TraceServices::ReadTimingProfilerProvider(*Session.Get()))
 			{
-				const Trace::ITimingProfilerProvider& TimingProfilerProvider = *Trace::ReadTimingProfilerProvider(*Session.Get());
+				const TraceServices::ITimingProfilerProvider& TimingProfilerProvider = *TraceServices::ReadTimingProfilerProvider(*Session.Get());
 
 				// Get Exclusive Time.
-				TimingProfilerProvider.ReadTimeline(GetTimelineIndex(), [&TrackEvent](const Trace::ITimingProfilerProvider::Timeline& Timeline)
+				TimingProfilerProvider.ReadTimeline(GetTimelineIndex(), [&TrackEvent](const TraceServices::ITimingProfilerProvider::Timeline& Timeline)
 				{
 					TimelineEventInfo EventInfo;
 					bool bIsFound = Timeline.GetEventInfo(TrackEvent.GetStartTime(), 0.0, TrackEvent.GetDepth(), EventInfo);
@@ -1224,7 +1224,7 @@ void FThreadTimingTrack::BuildContextMenu(FMenuBuilder& MenuBuilder)
 
 ////////////////////////////////////////////////////////////////////////////////////////////////////
 
-bool FThreadTimingTrack::FindTimingProfilerEvent(const FThreadTrackEvent& InTimingEvent, TFunctionRef<void(double, double, uint32, const Trace::FTimingProfilerEvent&)> InFoundPredicate) const
+bool FThreadTimingTrack::FindTimingProfilerEvent(const FThreadTrackEvent& InTimingEvent, TFunctionRef<void(double, double, uint32, const TraceServices::FTimingProfilerEvent&)> InFoundPredicate) const
 {
 	auto MatchEvent = [&InTimingEvent](double InStartTime, double InEndTime, uint32 InDepth)
 	{
@@ -1241,35 +1241,35 @@ bool FThreadTimingTrack::FindTimingProfilerEvent(const FThreadTrackEvent& InTimi
 
 ////////////////////////////////////////////////////////////////////////////////////////////////////
 
-bool FThreadTimingTrack::FindTimingProfilerEvent(const FTimingEventSearchParameters& InParameters, TFunctionRef<void(double, double, uint32, const Trace::FTimingProfilerEvent&)> InFoundPredicate) const
+bool FThreadTimingTrack::FindTimingProfilerEvent(const FTimingEventSearchParameters& InParameters, TFunctionRef<void(double, double, uint32, const TraceServices::FTimingProfilerEvent&)> InFoundPredicate) const
 {
-	return TTimingEventSearch<Trace::FTimingProfilerEvent>::Search(
+	return TTimingEventSearch<TraceServices::FTimingProfilerEvent>::Search(
 		InParameters,
 
-		[this](TTimingEventSearch<Trace::FTimingProfilerEvent>::FContext& InContext)
+		[this](TTimingEventSearch<TraceServices::FTimingProfilerEvent>::FContext& InContext)
 		{
-			TSharedPtr<const Trace::IAnalysisSession> Session = FInsightsManager::Get()->GetSession();
+			TSharedPtr<const TraceServices::IAnalysisSession> Session = FInsightsManager::Get()->GetSession();
 			if (Session.IsValid())
 			{
-				Trace::FAnalysisSessionReadScope SessionReadScope(*Session.Get());
+				TraceServices::FAnalysisSessionReadScope SessionReadScope(*Session.Get());
 
-				if (Trace::ReadTimingProfilerProvider(*Session.Get()))
+				if (TraceServices::ReadTimingProfilerProvider(*Session.Get()))
 				{
-					const Trace::ITimingProfilerProvider& TimingProfilerProvider = *Trace::ReadTimingProfilerProvider(*Session.Get());
+					const TraceServices::ITimingProfilerProvider& TimingProfilerProvider = *TraceServices::ReadTimingProfilerProvider(*Session.Get());
 
-					TimingProfilerProvider.ReadTimeline(GetTimelineIndex(), [&InContext](const Trace::ITimingProfilerProvider::Timeline& Timeline)
+					TimingProfilerProvider.ReadTimeline(GetTimelineIndex(), [&InContext](const TraceServices::ITimingProfilerProvider::Timeline& Timeline)
 					{
-						Timeline.EnumerateEvents(InContext.GetParameters().StartTime, InContext.GetParameters().EndTime, [&InContext](double EventStartTime, double EventEndTime, uint32 EventDepth, const Trace::FTimingProfilerEvent& Event)
+						Timeline.EnumerateEvents(InContext.GetParameters().StartTime, InContext.GetParameters().EndTime, [&InContext](double EventStartTime, double EventEndTime, uint32 EventDepth, const TraceServices::FTimingProfilerEvent& Event)
 						{
 							InContext.Check(EventStartTime, EventEndTime, EventDepth, Event);
-							return InContext.ShouldContinueSearching() ? Trace::EEventEnumerate::Continue : Trace::EEventEnumerate::Stop;
+							return InContext.ShouldContinueSearching() ? TraceServices::EEventEnumerate::Continue : TraceServices::EEventEnumerate::Stop;
 						});
 					});
 				}
 			}
 		},
 
-		[&InFoundPredicate](double InFoundStartTime, double InFoundEndTime, uint32 InFoundDepth, const Trace::FTimingProfilerEvent& InEvent)
+		[&InFoundPredicate](double InFoundStartTime, double InFoundEndTime, uint32 InFoundDepth, const TraceServices::FTimingProfilerEvent& InEvent)
 		{
 			InFoundPredicate(InFoundStartTime, InFoundEndTime, InFoundDepth, InEvent);
 		},
@@ -1300,17 +1300,17 @@ void FThreadTimingTrack::CreateFThreadTrackEventFromInfo(const TimelineEventInfo
 
 bool FThreadTimingTrack::TimerIndexToTimerId(uint32 InTimerIndex, uint32& OutTimerId)
 {
-	TSharedPtr<const Trace::IAnalysisSession> Session = FInsightsManager::Get()->GetSession();
+	TSharedPtr<const TraceServices::IAnalysisSession> Session = FInsightsManager::Get()->GetSession();
 	check(Session.IsValid())
 
-	Trace::FAnalysisSessionReadScope SessionReadScope(*Session.Get());
+	TraceServices::FAnalysisSessionReadScope SessionReadScope(*Session.Get());
 
-	const Trace::ITimingProfilerProvider& TimingProfilerProvider = *Trace::ReadTimingProfilerProvider(*Session.Get());
+	const TraceServices::ITimingProfilerProvider& TimingProfilerProvider = *TraceServices::ReadTimingProfilerProvider(*Session.Get());
 
-	const Trace::ITimingProfilerTimerReader* TimerReader;
-	TimingProfilerProvider.ReadTimers([&TimerReader](const Trace::ITimingProfilerTimerReader& Out) { TimerReader = &Out; });
+	const TraceServices::ITimingProfilerTimerReader* TimerReader;
+	TimingProfilerProvider.ReadTimers([&TimerReader](const TraceServices::ITimingProfilerTimerReader& Out) { TimerReader = &Out; });
 
-	const Trace::FTimingProfilerTimer* Timer = TimerReader->GetTimer(InTimerIndex);
+	const TraceServices::FTimingProfilerTimer* Timer = TimerReader->GetTimer(InTimerIndex);
 	if (Timer == nullptr)
 	{
 		return false;
diff --git a/Engine/Source/Developer/TraceInsights/Private/Insights/ViewModels/ThreadTimingTrack.h b/Engine/Source/Developer/TraceInsights/Private/Insights/ViewModels/ThreadTimingTrack.h
--- a/Engine/Source/Developer/TraceInsights/Private/Insights/ViewModels/ThreadTimingTrack.h
+++ b/Engine/Source/Developer/TraceInsights/Private/Insights/ViewModels/ThreadTimingTrack.h
@@ -50,7 +50,7 @@ public:
 
 	virtual void OnBeginSession(Insights::ITimingViewSession& InSession) override;
 	virtual void OnEndSession(Insights::ITimingViewSession& InSession) override;
-	virtual void Tick(Insights::ITimingViewSession& InSession, const Trace::IAnalysisSession& InAnalysisSession) override;
+	virtual void Tick(Insights::ITimingViewSession& InSession, const TraceServices::IAnalysisSession& InAnalysisSession) override;
 	virtual void ExtendFilterMenu(Insights::ITimingViewSession& InSession, FMenuBuilder& InMenuBuilder) override;
 
 	//////////////////////////////////////////////////
@@ -98,7 +98,7 @@ class FThreadTimingTrack : public FTimingEventsTrack
 	INSIGHTS_DECLARE_RTTI(FThreadTimingTrack, FTimingEventsTrack)
 
 public:
-	typedef typename Trace::ITimeline<Trace::FTimingProfilerEvent>::FTimelineEventInfo TimelineEventInfo;
+	typedef typename TraceServices::ITimeline<TraceServices::FTimingProfilerEvent>::FTimelineEventInfo TimelineEventInfo;
 
 	struct FPendingEventInfo
 	{
@@ -147,8 +147,8 @@ protected:
 private:
 	void DrawSelectedEventInfo(const FThreadTrackEvent& SelectedEvent, const FTimingTrackViewport& Viewport, const FDrawContext& DrawContext, const FSlateBrush* WhiteBrush, const FSlateFontInfo& Font) const;
 
-	bool FindTimingProfilerEvent(const FThreadTrackEvent& InTimingEvent, TFunctionRef<void(double, double, uint32, const Trace::FTimingProfilerEvent&)> InFoundPredicate) const;
-	bool FindTimingProfilerEvent(const FTimingEventSearchParameters& InParameters, TFunctionRef<void(double, double, uint32, const Trace::FTimingProfilerEvent&)> InFoundPredicate) const;
+	bool FindTimingProfilerEvent(const FThreadTrackEvent& InTimingEvent, TFunctionRef<void(double, double, uint32, const TraceServices::FTimingProfilerEvent&)> InFoundPredicate) const;
+	bool FindTimingProfilerEvent(const FTimingEventSearchParameters& InParameters, TFunctionRef<void(double, double, uint32, const TraceServices::FTimingProfilerEvent&)> InFoundPredicate) const;
 
 	void GetParentAndRoot(const FThreadTrackEvent& TimingEvent,
 						  TSharedPtr<FThreadTrackEvent>& OutParentTimingEvent,
@@ -168,7 +168,7 @@ private:
 	FDelegateHandle OnFilterChangesCommitedHandle;
 
 	// Search cache
-	mutable TTimingEventSearchCache<Trace::FTimingProfilerEvent> SearchCache;
+	mutable TTimingEventSearchCache<TraceServices::FTimingProfilerEvent> SearchCache;
 };
 
 ////////////////////////////////////////////////////////////////////////////////////////////////////
diff --git a/Engine/Source/Developer/TraceInsights/Private/Insights/ViewModels/TimerAggregation.cpp b/Engine/Source/Developer/TraceInsights/Private/Insights/ViewModels/TimerAggregation.cpp
--- a/Engine/Source/Developer/TraceInsights/Private/Insights/ViewModels/TimerAggregation.cpp
+++ b/Engine/Source/Developer/TraceInsights/Private/Insights/ViewModels/TimerAggregation.cpp
@@ -17,7 +17,7 @@ namespace Insights
 class FTimerAggregationWorker : public IStatsAggregationWorker
 {
 public:
-	FTimerAggregationWorker(TSharedPtr<const Trace::IAnalysisSession> InSession, double InStartTime, double InEndTime, const TSet<uint32>& InCpuThreads, bool bInIncludeGpuThread)
+	FTimerAggregationWorker(TSharedPtr<const TraceServices::IAnalysisSession> InSession, double InStartTime, double InEndTime, const TSet<uint32>& InCpuThreads, bool bInIncludeGpuThread)
 		: Session(InSession)
 		, StartTime(InStartTime)
 		, EndTime(InEndTime)
@@ -31,29 +31,29 @@ public:
 
 	virtual void DoWork() override;
 
-	Trace::ITable<Trace::FTimingProfilerAggregatedStats>* GetResultTable() const { return ResultTable.Get(); }
+	TraceServices::ITable<TraceServices::FTimingProfilerAggregatedStats>* GetResultTable() const { return ResultTable.Get(); }
 	void ResetResults() { ResultTable.Reset(); }
 
 private:
-	TSharedPtr<const Trace::IAnalysisSession> Session;
+	TSharedPtr<const TraceServices::IAnalysisSession> Session;
 	double StartTime;
 	double EndTime;
 	TSet<uint32> CpuThreads;
 	bool bIncludeGpuThread;
-	TUniquePtr<Trace::ITable<Trace::FTimingProfilerAggregatedStats>> ResultTable;
+	TUniquePtr<TraceServices::ITable<TraceServices::FTimingProfilerAggregatedStats>> ResultTable;
 };
 
 ////////////////////////////////////////////////////////////////////////////////////////////////////
 
 void FTimerAggregationWorker::DoWork()
 {
-	if (Session.IsValid() && Trace::ReadTimingProfilerProvider(*Session.Get()))
+	if (Session.IsValid() && TraceServices::ReadTimingProfilerProvider(*Session.Get()))
 	{
 		// Suspend analysis in order to avoid write locks (ones blocked by the read lock below) to further block other read locks.
-		//Trace::FAnalysisSessionSuspensionScope SessionPauseScope(*Session.Get());
+		//FAnalysisSessionSuspensionScope SessionPauseScope(*Session.Get());
 
-		Trace::FAnalysisSessionReadScope SessionReadScope(*Session.Get());
-		const Trace::ITimingProfilerProvider& TimingProfilerProvider = *Trace::ReadTimingProfilerProvider(*Session.Get());
+		TraceServices::FAnalysisSessionReadScope SessionReadScope(*Session.Get());
+		const TraceServices::ITimingProfilerProvider& TimingProfilerProvider = *TraceServices::ReadTimingProfilerProvider(*Session.Get());
 
 		auto CpuThreadFilter = [this](uint32 ThreadId)
 		{
@@ -68,7 +68,7 @@ void FTimerAggregationWorker::DoWork()
 // FTimerAggregator
 ////////////////////////////////////////////////////////////////////////////////////////////////////
 
-IStatsAggregationWorker* FTimerAggregator::CreateWorker(TSharedPtr<const Trace::IAnalysisSession> InSession)
+IStatsAggregationWorker* FTimerAggregator::CreateWorker(TSharedPtr<const TraceServices::IAnalysisSession> InSession)
 {
 	bool bIsGpuTrackVisible = false;
 	TSet<uint32> CpuThreads;
@@ -94,7 +94,7 @@ IStatsAggregationWorker* FTimerAggregator::CreateWorker(TSharedPtr<const Trace::
 
 ////////////////////////////////////////////////////////////////////////////////////////////////////
 
-Trace::ITable<Trace::FTimingProfilerAggregatedStats>* FTimerAggregator::GetResultTable() const
+TraceServices::ITable<TraceServices::FTimingProfilerAggregatedStats>* FTimerAggregator::GetResultTable() const
 {
 	// It can only be called from OnFinishedCallback.
 	check(IsFinished());
diff --git a/Engine/Source/Developer/TraceInsights/Private/Insights/ViewModels/TimerAggregation.h b/Engine/Source/Developer/TraceInsights/Private/Insights/ViewModels/TimerAggregation.h
--- a/Engine/Source/Developer/TraceInsights/Private/Insights/ViewModels/TimerAggregation.h
+++ b/Engine/Source/Developer/TraceInsights/Private/Insights/ViewModels/TimerAggregation.h
@@ -21,11 +21,11 @@ public:
 	FTimerAggregator() : FStatsAggregator(TEXT("Timers")) {}
 	virtual ~FTimerAggregator() {}
 
-	Trace::ITable<Trace::FTimingProfilerAggregatedStats>* GetResultTable() const;
+	TraceServices::ITable<TraceServices::FTimingProfilerAggregatedStats>* GetResultTable() const;
 	void ResetResults();
 
 protected:
-	virtual IStatsAggregationWorker* CreateWorker(TSharedPtr<const Trace::IAnalysisSession> InSession) override;
+	virtual IStatsAggregationWorker* CreateWorker(TSharedPtr<const TraceServices::IAnalysisSession> InSession) override;
 };
 
 ////////////////////////////////////////////////////////////////////////////////////////////////////
diff --git a/Engine/Source/Developer/TraceInsights/Private/Insights/ViewModels/TimerButterflyAggregation.cpp b/Engine/Source/Developer/TraceInsights/Private/Insights/ViewModels/TimerButterflyAggregation.cpp
--- a/Engine/Source/Developer/TraceInsights/Private/Insights/ViewModels/TimerButterflyAggregation.cpp
+++ b/Engine/Source/Developer/TraceInsights/Private/Insights/ViewModels/TimerButterflyAggregation.cpp
@@ -17,7 +17,7 @@ namespace Insights
 class FTimerButterflyAggregationWorker : public IStatsAggregationWorker
 {
 public:
-	FTimerButterflyAggregationWorker(TSharedPtr<const Trace::IAnalysisSession> InSession, double InStartTime, double InEndTime, const TSet<uint32>& InCpuThreads, bool bInIncludeGpuThread)
+	FTimerButterflyAggregationWorker(TSharedPtr<const TraceServices::IAnalysisSession> InSession, double InStartTime, double InEndTime, const TSet<uint32>& InCpuThreads, bool bInIncludeGpuThread)
 		: Session(InSession)
 		, StartTime(InStartTime)
 		, EndTime(InEndTime)
@@ -31,29 +31,29 @@ public:
 
 	virtual void DoWork() override;
 
-	Trace::ITimingProfilerButterfly* GetResultButterfly() const { return ResultButterfly.Get(); }
+	TraceServices::ITimingProfilerButterfly* GetResultButterfly() const { return ResultButterfly.Get(); }
 	void ResetResults() { ResultButterfly.Reset(); }
 
 private:
-	TSharedPtr<const Trace::IAnalysisSession> Session;
+	TSharedPtr<const TraceServices::IAnalysisSession> Session;
 	double StartTime;
 	double EndTime;
 	TSet<uint32> CpuThreads;
 	bool bIncludeGpuThread;
-	TUniquePtr<Trace::ITimingProfilerButterfly> ResultButterfly;
+	TUniquePtr<TraceServices::ITimingProfilerButterfly> ResultButterfly;
 };
 
 ////////////////////////////////////////////////////////////////////////////////////////////////////
 
 void FTimerButterflyAggregationWorker::DoWork()
 {
-	if (Session.IsValid() && Trace::ReadTimingProfilerProvider(*Session.Get()))
+	if (Session.IsValid() && TraceServices::ReadTimingProfilerProvider(*Session.Get()))
 	{
 		// Suspend analysis in order to avoid write locks (ones blocked by the read lock below) to further block other read locks.
-		//Trace::FAnalysisSessionSuspensionScope SessionPauseScope(*Session.Get());
+		//FAnalysisSessionSuspensionScope SessionPauseScope(*Session.Get());
 
-		Trace::FAnalysisSessionReadScope SessionReadScope(*Session.Get());
-		const Trace::ITimingProfilerProvider& TimingProfilerProvider = *Trace::ReadTimingProfilerProvider(*Session.Get());
+		TraceServices::FAnalysisSessionReadScope SessionReadScope(*Session.Get());
+		const TraceServices::ITimingProfilerProvider& TimingProfilerProvider = *TraceServices::ReadTimingProfilerProvider(*Session.Get());
 
 		auto ThreadFilter = [this](uint32 ThreadId)
 		{
@@ -68,7 +68,7 @@ void FTimerButterflyAggregationWorker::DoWork()
 // FTimerButterflyAggregator
 ////////////////////////////////////////////////////////////////////////////////////////////////////
 
-IStatsAggregationWorker* FTimerButterflyAggregator::CreateWorker(TSharedPtr<const Trace::IAnalysisSession> InSession)
+IStatsAggregationWorker* FTimerButterflyAggregator::CreateWorker(TSharedPtr<const TraceServices::IAnalysisSession> InSession)
 {
 	bool bIsGpuTrackVisible = false;
 	TSet<uint32> CpuThreads;
@@ -94,7 +94,7 @@ IStatsAggregationWorker* FTimerButterflyAggregator::CreateWorker(TSharedPtr<cons
 
 ////////////////////////////////////////////////////////////////////////////////////////////////////
 
-Trace::ITimingProfilerButterfly* FTimerButterflyAggregator::GetResultButterfly() const
+TraceServices::ITimingProfilerButterfly* FTimerButterflyAggregator::GetResultButterfly() const
 {
 	// It can only be called from OnFinishedCallback.
 	check(IsFinished());
diff --git a/Engine/Source/Developer/TraceInsights/Private/Insights/ViewModels/TimerButterflyAggregation.h b/Engine/Source/Developer/TraceInsights/Private/Insights/ViewModels/TimerButterflyAggregation.h
--- a/Engine/Source/Developer/TraceInsights/Private/Insights/ViewModels/TimerButterflyAggregation.h
+++ b/Engine/Source/Developer/TraceInsights/Private/Insights/ViewModels/TimerButterflyAggregation.h
@@ -20,11 +20,11 @@ public:
 	FTimerButterflyAggregator() : FStatsAggregator(TEXT("Butterfly")) {}
 	virtual ~FTimerButterflyAggregator() {}
 
-	Trace::ITimingProfilerButterfly* GetResultButterfly() const;
+	TraceServices::ITimingProfilerButterfly* GetResultButterfly() const;
 	void ResetResults();
 
 protected:
-	virtual IStatsAggregationWorker* CreateWorker(TSharedPtr<const Trace::IAnalysisSession> InSession) override;
+	virtual IStatsAggregationWorker* CreateWorker(TSharedPtr<const TraceServices::IAnalysisSession> InSession) override;
 };
 
 ////////////////////////////////////////////////////////////////////////////////////////////////////
diff --git a/Engine/Source/Developer/TraceInsights/Private/Insights/ViewModels/TimerNode.cpp b/Engine/Source/Developer/TraceInsights/Private/Insights/ViewModels/TimerNode.cpp
--- a/Engine/Source/Developer/TraceInsights/Private/Insights/ViewModels/TimerNode.cpp
+++ b/Engine/Source/Developer/TraceInsights/Private/Insights/ViewModels/TimerNode.cpp
@@ -54,12 +54,12 @@ FTimerNode::~FTimerNode()
 
 void FTimerNode::ResetAggregatedStats()
 {
-	AggregatedStats = Trace::FTimingProfilerAggregatedStats();
+	AggregatedStats = TraceServices::FTimingProfilerAggregatedStats();
 }
 
 ////////////////////////////////////////////////////////////////////////////////////////////////////
 
-void FTimerNode::SetAggregatedStats(const Trace::FTimingProfilerAggregatedStats& InAggregatedStats)
+void FTimerNode::SetAggregatedStats(const TraceServices::FTimingProfilerAggregatedStats& InAggregatedStats)
 {
 	AggregatedStats = InAggregatedStats;
 }
diff --git a/Engine/Source/Developer/TraceInsights/Private/Insights/ViewModels/TimerNode.h b/Engine/Source/Developer/TraceInsights/Private/Insights/ViewModels/TimerNode.h
--- a/Engine/Source/Developer/TraceInsights/Private/Insights/ViewModels/TimerNode.h
+++ b/Engine/Source/Developer/TraceInsights/Private/Insights/ViewModels/TimerNode.h
@@ -105,11 +105,11 @@ public:
 	/**
 	 * @return the aggregated stats for this timer.
 	 */
-	const Trace::FTimingProfilerAggregatedStats& GetAggregatedStats() const { return AggregatedStats; }
-	Trace::FTimingProfilerAggregatedStats& GetAggregatedStats() { return AggregatedStats; }
+	const TraceServices::FTimingProfilerAggregatedStats& GetAggregatedStats() const { return AggregatedStats; }
+	TraceServices::FTimingProfilerAggregatedStats& GetAggregatedStats() { return AggregatedStats; }
 
 	void ResetAggregatedStats();
-	void SetAggregatedStats(const Trace::FTimingProfilerAggregatedStats& AggregatedStats);
+	void SetAggregatedStats(const TraceServices::FTimingProfilerAggregatedStats& AggregatedStats);
 
 	bool IsHotPath() const { return bIsHotPath; }
 	void SetIsHotPath(bool bOnOff) { bIsHotPath = bOnOff; }
@@ -133,7 +133,7 @@ private:
 	bool bIsHotPath;
 
 	/** Aggregated stats. */
-	Trace::FTimingProfilerAggregatedStats AggregatedStats;
+	TraceServices::FTimingProfilerAggregatedStats AggregatedStats;
 };
 
 ////////////////////////////////////////////////////////////////////////////////////////////////////
diff --git a/Engine/Source/Developer/TraceInsights/Private/Insights/ViewModels/TimingGraphTrack.cpp b/Engine/Source/Developer/TraceInsights/Private/Insights/ViewModels/TimingGraphTrack.cpp
--- a/Engine/Source/Developer/TraceInsights/Private/Insights/ViewModels/TimingGraphTrack.cpp
+++ b/Engine/Source/Developer/TraceInsights/Private/Insights/ViewModels/TimingGraphTrack.cpp
@@ -199,11 +199,11 @@ void FTimingGraphTrack::AddDefaultFrameSeries()
 void FTimingGraphTrack::UpdateFrameSeries(FTimingGraphSeries& Series, const FTimingTrackViewport& Viewport)
 {
 	FGraphTrackBuilder Builder(*this, Series, Viewport);
-	TSharedPtr<const Trace::IAnalysisSession> Session = FInsightsManager::Get()->GetSession();
+	TSharedPtr<const TraceServices::IAnalysisSession> Session = FInsightsManager::Get()->GetSession();
 	if (Session.IsValid())
 	{
-		Trace::FAnalysisSessionReadScope SessionReadScope(*Session.Get());
-		const Trace::IFrameProvider& FramesProvider = ReadFrameProvider(*Session.Get());
+		TraceServices::FAnalysisSessionReadScope SessionReadScope(*Session.Get());
+		const TraceServices::IFrameProvider& FramesProvider = ReadFrameProvider(*Session.Get());
 
 		const TArray64<double>& FrameStartTimes = FramesProvider.GetFrameStartTimes(Series.FrameType);
 
@@ -213,7 +213,7 @@ void FTimingGraphTrack::UpdateFrameSeries(FTimingGraphSeries& Series, const FTim
 		const int64 EndLowerBound = Algo::LowerBound(FrameStartTimes, Viewport.GetEndTime());
 		const uint64 EndIndex = EndLowerBound + 1;
 
-		FramesProvider.EnumerateFrames(Series.FrameType, StartIndex, EndIndex, [&Builder](const Trace::FFrame& Frame)
+		FramesProvider.EnumerateFrames(Series.FrameType, StartIndex, EndIndex, [&Builder](const TraceServices::FFrame& Frame)
 		{
 			//TODO: add a "frame converter" (i.e. to fps, miliseconds or seconds)
 			const double Duration = Frame.EndTime - Frame.StartTime;
@@ -279,30 +279,30 @@ void FTimingGraphTrack::RemoveTimerSeries(uint32 TimerId)
 void FTimingGraphTrack::UpdateTimerSeries(FTimingGraphSeries& Series, const FTimingTrackViewport& Viewport)
 {
 	FGraphTrackBuilder Builder(*this, Series, Viewport);
-	TSharedPtr<const Trace::IAnalysisSession> Session = FInsightsManager::Get()->GetSession();
+	TSharedPtr<const TraceServices::IAnalysisSession> Session = FInsightsManager::Get()->GetSession();
 	if (Session.IsValid())
 	{
-		Trace::FAnalysisSessionReadScope SessionReadScope(*Session.Get());
+		TraceServices::FAnalysisSessionReadScope SessionReadScope(*Session.Get());
 
 		const double SessionDuration = Session->GetDurationSeconds();
 		if (Series.CachedSessionDuration != SessionDuration)
 		{
 			Series.CachedSessionDuration = SessionDuration;
 
-			const Trace::ITimingProfilerProvider& TimingProfilerProvider = *Trace::ReadTimingProfilerProvider(*Session.Get());
+			const TraceServices::ITimingProfilerProvider& TimingProfilerProvider = *TraceServices::ReadTimingProfilerProvider(*Session.Get());
 
-			const Trace::ITimingProfilerTimerReader* TimerReader;
-			TimingProfilerProvider.ReadTimers([&TimerReader](const Trace::ITimingProfilerTimerReader& Out) { TimerReader = &Out; });
+			const TraceServices::ITimingProfilerTimerReader* TimerReader;
+			TimingProfilerProvider.ReadTimers([&TimerReader](const TraceServices::ITimingProfilerTimerReader& Out) { TimerReader = &Out; });
 
 			const uint32 TimelineCount = TimingProfilerProvider.GetTimelineCount();
 			uint32 NumTimelinesContainingEvent = 0;
 			for (uint32 TimelineIndex = 0; TimelineIndex < TimelineCount; ++TimelineIndex)
 			{
 				TimingProfilerProvider.ReadTimeline(TimelineIndex,
-					[SessionDuration, &Series, TimerReader, &Viewport, &NumTimelinesContainingEvent](const Trace::ITimingProfilerProvider::Timeline& Timeline)
+					[SessionDuration, &Series, TimerReader, &Viewport, &NumTimelinesContainingEvent](const TraceServices::ITimingProfilerProvider::Timeline& Timeline)
 					{
 						TArray<TArray<FTimingGraphSeries::FSimpleTimingEvent>> Events;
-						Trace::ITimeline<Trace::FTimingProfilerEvent>::EnumerateAsyncParams Params;
+						TraceServices::ITimeline<TraceServices::FTimingProfilerEvent>::EnumerateAsyncParams Params;
 						Params.IntervalStart = 0;
 						Params.IntervalEnd = SessionDuration;
 						Params.Resolution = 0.0;
@@ -310,9 +310,9 @@ void FTimingGraphTrack::UpdateTimerSeries(FTimingGraphSeries& Series, const FTim
 						{
 							Events.AddDefaulted(NumTasks);
 						};
-						Params.Callback = [&Events, TimerReader, &Series](double StartTime, double EndTime, uint32 Depth, const Trace::FTimingProfilerEvent& Event, uint32 TaskIndex)
+						Params.Callback = [&Events, TimerReader, &Series](double StartTime, double EndTime, uint32 Depth, const TraceServices::FTimingProfilerEvent& Event, uint32 TaskIndex)
 						{
-							const Trace::FTimingProfilerTimer* Timer = TimerReader->GetTimer(Event.TimerIndex);
+							const TraceServices::FTimingProfilerTimer* Timer = TimerReader->GetTimer(Event.TimerIndex);
 							if (ensure(Timer != nullptr))
 							{
 								if (Timer->Id == Series.TimerId)
@@ -321,7 +321,7 @@ void FTimingGraphTrack::UpdateTimerSeries(FTimingGraphSeries& Series, const FTim
 									Events[TaskIndex].Add({ StartTime, Duration });
 								}
 							}
-							return Trace::EEventEnumerate::Continue;
+							return TraceServices::EEventEnumerate::Continue;
 						};
 
 						Timeline.EnumerateEventsDownSampledAsync(Params);
@@ -395,18 +395,18 @@ TSharedPtr<FTimingGraphSeries> FTimingGraphTrack::AddStatsCounterSeries(uint32 C
 	bool bIsMemory = false;
 	bool bIsFloatingPoint = false;
 
-	TSharedPtr<const Trace::IAnalysisSession> Session = FInsightsManager::Get()->GetSession();
+	TSharedPtr<const TraceServices::IAnalysisSession> Session = FInsightsManager::Get()->GetSession();
 	if (Session.IsValid())
 	{
-		Trace::FAnalysisSessionReadScope SessionReadScope(*Session.Get());
-		const Trace::ICounterProvider& CountersProvider = Trace::ReadCounterProvider(*Session.Get());
+		TraceServices::FAnalysisSessionReadScope SessionReadScope(*Session.Get());
+		const TraceServices::ICounterProvider& CountersProvider = TraceServices::ReadCounterProvider(*Session.Get());
 		if (CounterId < CountersProvider.GetCounterCount())
 		{
-			CountersProvider.ReadCounter(CounterId, [&](const Trace::ICounter& Counter)
+			CountersProvider.ReadCounter(CounterId, [&](const TraceServices::ICounter& Counter)
 			{
 				CounterName = Counter.GetName();
-				//bIsTime = (Counter.GetDisplayHint() == Trace::CounterDisplayHint_Time);
-				bIsMemory = (Counter.GetDisplayHint() == Trace::CounterDisplayHint_Memory);
+				//bIsTime = (Counter.GetDisplayHint() == TraceServices::CounterDisplayHint_Time);
+				bIsMemory = (Counter.GetDisplayHint() == TraceServices::CounterDisplayHint_Memory);
 				bIsFloatingPoint = Counter.IsFloatingPoint();
 			});
 		}
@@ -449,12 +449,12 @@ void FTimingGraphTrack::UpdateStatsCounterSeries(FTimingGraphSeries& Series, con
 {
 	FGraphTrackBuilder Builder(*this, Series, Viewport);
 
-	TSharedPtr<const Trace::IAnalysisSession> Session = FInsightsManager::Get()->GetSession();
+	TSharedPtr<const TraceServices::IAnalysisSession> Session = FInsightsManager::Get()->GetSession();
 	if (Session.IsValid())
 	{
-		Trace::FAnalysisSessionReadScope SessionReadScope(*Session.Get());
-		const Trace::ICounterProvider& CounterProvider = Trace::ReadCounterProvider(*Session.Get());
-		CounterProvider.ReadCounter(Series.CounterId, [this, &Viewport, &Builder, &Series](const Trace::ICounter& Counter)
+		TraceServices::FAnalysisSessionReadScope SessionReadScope(*Session.Get());
+		const TraceServices::ICounterProvider& CounterProvider = TraceServices::ReadCounterProvider(*Session.Get());
+		CounterProvider.ReadCounter(Series.CounterId, [this, &Viewport, &Builder, &Series](const TraceServices::ICounter& Counter)
 		{
 			const float TopY = 4.0f;
 			const float BottomY = GetHeight() - 4.0f;
diff --git a/Engine/Source/Developer/TraceInsights/Private/Insights/Widgets/SFrameTrack.cpp b/Engine/Source/Developer/TraceInsights/Private/Insights/Widgets/SFrameTrack.cpp
--- a/Engine/Source/Developer/TraceInsights/Private/Insights/Widgets/SFrameTrack.cpp
+++ b/Engine/Source/Developer/TraceInsights/Private/Insights/Widgets/SFrameTrack.cpp
@@ -170,12 +170,12 @@ void SFrameTrack::Tick(const FGeometry& AllottedGeometry, const double InCurrent
 		const uint64 WaitTime = static_cast<uint64>(0.1 / FPlatformTime::GetSecondsPerCycle64()); // 100ms
 		AnalysisSyncNextTimestamp = Time + WaitTime;
 
-		TSharedPtr<const Trace::IAnalysisSession> Session = FInsightsManager::Get()->GetSession();
+		TSharedPtr<const TraceServices::IAnalysisSession> Session = FInsightsManager::Get()->GetSession();
 		if (Session.IsValid())
 		{
-			Trace::FAnalysisSessionReadScope SessionReadScope(*Session.Get());
+			TraceServices::FAnalysisSessionReadScope SessionReadScope(*Session.Get());
 
-			const Trace::IFrameProvider& FramesProvider = Trace::ReadFrameProvider(*Session.Get());
+			const TraceServices::IFrameProvider& FramesProvider = TraceServices::ReadFrameProvider(*Session.Get());
 
 			for (int32 FrameType = 0; FrameType < TraceFrameType_Count; ++FrameType)
 			{
@@ -258,12 +258,12 @@ void SFrameTrack::UpdateState()
 	}
 	NumUpdatedFrames = 0;
 
-	TSharedPtr<const Trace::IAnalysisSession> Session = FInsightsManager::Get()->GetSession();
+	TSharedPtr<const TraceServices::IAnalysisSession> Session = FInsightsManager::Get()->GetSession();
 	if (Session.IsValid())
 	{
-		Trace::FAnalysisSessionReadScope SessionReadScope(*Session.Get());
+		TraceServices::FAnalysisSessionReadScope SessionReadScope(*Session.Get());
 
-		const Trace::IFrameProvider& FramesProvider = Trace::ReadFrameProvider(*Session.Get());
+		const TraceServices::IFrameProvider& FramesProvider = TraceServices::ReadFrameProvider(*Session.Get());
 
 		const FAxisViewportInt32& ViewportX = Viewport.GetHorizontalAxisViewport();
 
@@ -276,7 +276,7 @@ void SFrameTrack::UpdateState()
 
 			FFrameTrackSeriesBuilder Builder(*SeriesPtr, Viewport);
 
-			FramesProvider.EnumerateFrames(static_cast<ETraceFrameType>(FrameType), StartIndex, EndIndex, [&Builder](const Trace::FFrame& Frame)
+			FramesProvider.EnumerateFrames(static_cast<ETraceFrameType>(FrameType), StartIndex, EndIndex, [&Builder](const TraceServices::FFrame& Frame)
 			{
 				Builder.AddFrame(Frame);
 			});
diff --git a/Engine/Source/Developer/TraceInsights/Private/Insights/Widgets/SLogView.cpp b/Engine/Source/Developer/TraceInsights/Private/Insights/Widgets/SLogView.cpp
--- a/Engine/Source/Developer/TraceInsights/Private/Insights/Widgets/SLogView.cpp
+++ b/Engine/Source/Developer/TraceInsights/Private/Insights/Widgets/SLogView.cpp
@@ -709,14 +709,14 @@ void SLogView::Tick(const FGeometry& AllottedGeometry, const double InCurrentTim
 {
 	int32 NewMessageCount = 0;
 
-	TSharedPtr<const Trace::IAnalysisSession> Session = FInsightsManager::Get()->GetSession();
+	TSharedPtr<const TraceServices::IAnalysisSession> Session = FInsightsManager::Get()->GetSession();
 
 	Cache.SetSession(Session);
 
 	if (Session.IsValid())
 	{
-		Trace::FAnalysisSessionReadScope SessionReadScope(*Session.Get());
-		const Trace::ILogProvider& LogProvider = Trace::ReadLogProvider(*Session.Get());
+		TraceServices::FAnalysisSessionReadScope SessionReadScope(*Session.Get());
+		const TraceServices::ILogProvider& LogProvider = TraceServices::ReadLogProvider(*Session.Get());
 
 		NewMessageCount = static_cast<int32>(LogProvider.GetMessageCount());
 
@@ -730,7 +730,7 @@ void SLogView::Tick(const FGeometry& AllottedGeometry, const double InCurrentTim
 			UE_LOG(TraceInsights, Log, TEXT("[LogView] Total Log Categories: %d"), TotalNumCategories);
 
 			TSet<FName> Categories;
-			LogProvider.EnumerateCategories([&Categories](const Trace::FLogCategory& Category)
+			LogProvider.EnumerateCategories([&Categories](const TraceServices::FLogCategory& Category)
 			{
 				FString CategoryStr(Category.Name);
 				if (CategoryStr.StartsWith(TEXT("Log")))
diff --git a/Engine/Source/Developer/TraceInsights/Private/Insights/Widgets/SSessionInfoWindow.cpp b/Engine/Source/Developer/TraceInsights/Private/Insights/Widgets/SSessionInfoWindow.cpp
--- a/Engine/Source/Developer/TraceInsights/Private/Insights/Widgets/SSessionInfoWindow.cpp
+++ b/Engine/Source/Developer/TraceInsights/Private/Insights/Widgets/SSessionInfoWindow.cpp
@@ -262,16 +262,16 @@ void SSessionInfoWindow::Tick(const FGeometry& AllottedGeometry, const double In
 		return;
 	}
 
-	TSharedPtr<const Trace::IAnalysisSession> Session = FInsightsManager::Get()->GetSession();
+	TSharedPtr<const TraceServices::IAnalysisSession> Session = FInsightsManager::Get()->GetSession();
 
 	if (Session.IsValid())
 	{
-		Trace::FAnalysisSessionReadScope SessionReadScope(*Session.Get());
-		const Trace::IDiagnosticsProvider& DiagnosticsProvider = Trace::ReadDiagnosticsProvider(*Session.Get());
+		TraceServices::FAnalysisSessionReadScope SessionReadScope(*Session.Get());
+		const TraceServices::IDiagnosticsProvider& DiagnosticsProvider = TraceServices::ReadDiagnosticsProvider(*Session.Get());
 
 		if (DiagnosticsProvider.IsSessionInfoAvailable())
 		{
-			Trace::FSessionInfo SessionInfo = DiagnosticsProvider.GetSessionInfo();
+			TraceServices::FSessionInfo SessionInfo = DiagnosticsProvider.GetSessionInfo();
 			PlatformText = FText::FromString(SessionInfo.Platform);
 			AppNameText = FText::FromString(SessionInfo.AppName);
 			CommandLineText = FText::FromString(SessionInfo.CommandLine);
@@ -343,10 +343,10 @@ FReply SSessionInfoWindow::OnDrop(const FGeometry& MyGeometry, const FDragDropEv
 FText SSessionInfoWindow::GetSessionNameText() const
 {
 	FText SessionName;
-	TSharedPtr<const Trace::IAnalysisSession> Session = FInsightsManager::Get()->GetSession();
+	TSharedPtr<const TraceServices::IAnalysisSession> Session = FInsightsManager::Get()->GetSession();
 	if (Session.IsValid())
 	{
-		Trace::FAnalysisSessionReadScope SessionReadScope(*Session.Get());
+		TraceServices::FAnalysisSessionReadScope SessionReadScope(*Session.Get());
 		SessionName = FText::FromString(FPaths::GetBaseFilename(Session->GetName()));
 	}
 	return SessionName;
@@ -358,10 +358,10 @@ FText SSessionInfoWindow::GetUriText() const
 {
 	//TODO: update code to use a SessionInfo provider instead
 	FText LocalUri;
-	TSharedPtr<const Trace::IAnalysisSession> Session = FInsightsManager::Get()->GetSession();
+	TSharedPtr<const TraceServices::IAnalysisSession> Session = FInsightsManager::Get()->GetSession();
 	if (Session.IsValid())
 	{
-		Trace::FAnalysisSessionReadScope SessionReadScope(*Session.Get());
+		TraceServices::FAnalysisSessionReadScope SessionReadScope(*Session.Get());
 		LocalUri = FText::FromString(FString(Session->GetName()));
 	}
 	return LocalUri;
@@ -433,15 +433,15 @@ FText SSessionInfoWindow::GetStatusText() const
 FText SSessionInfoWindow::GetModulesText() const
 {
 	FString ModulesStr;
-	TArray<Trace::FModuleInfo> Modules;
+	TArray<TraceServices::FModuleInfo> Modules;
 
-	TSharedPtr<Trace::IModuleService> ModuleService = FInsightsManager::Get()->GetModuleService();
+	TSharedPtr<TraceServices::IModuleService> ModuleService = FInsightsManager::Get()->GetModuleService();
 	if (ModuleService)
 	{
 		ModuleService->GetAvailableModules(Modules);
 	}
 
-	for (const Trace::FModuleInfo& Module : Modules)
+	for (const TraceServices::FModuleInfo& Module : Modules)
 	{
 		ModulesStr += Module.DisplayName;
 		ModulesStr += TEXT(", ");
diff --git a/Engine/Source/Developer/TraceInsights/Private/Insights/Widgets/SStatsView.cpp b/Engine/Source/Developer/TraceInsights/Private/Insights/Widgets/SStatsView.cpp
--- a/Engine/Source/Developer/TraceInsights/Private/Insights/Widgets/SStatsView.cpp
+++ b/Engine/Source/Developer/TraceInsights/Private/Insights/Widgets/SStatsView.cpp
@@ -654,7 +654,7 @@ TSharedRef<SWidget> SStatsView::TreeViewHeaderRow_GenerateColumnMenu(const Insig
 
 void SStatsView::InsightsManager_OnSessionChanged()
 {
-	TSharedPtr<const Trace::IAnalysisSession> NewSession = FInsightsManager::Get()->GetSession();
+	TSharedPtr<const TraceServices::IAnalysisSession> NewSession = FInsightsManager::Get()->GetSession();
 
 	if (NewSession != Session)
 	{
@@ -1793,9 +1793,9 @@ void SStatsView::RebuildTree(bool bResync)
 	SyncStopwatch.Start();
 	if (Session.IsValid())
 	{
-		Trace::FAnalysisSessionReadScope SessionReadScope(*Session.Get());
+		TraceServices::FAnalysisSessionReadScope SessionReadScope(*Session.Get());
 
-		const Trace::ICounterProvider& CountersProvider = Trace::ReadCounterProvider(*Session.Get());
+		const TraceServices::ICounterProvider& CountersProvider = TraceServices::ReadCounterProvider(*Session.Get());
 
 		const uint32 CounterCount = CountersProvider.GetCounterCount();
 		if (CounterCount != PreviousNodeCount)
@@ -1810,13 +1810,13 @@ void SStatsView::RebuildTree(bool bResync)
 			const FName MiscInt64Group(TEXT("Misc_int64"));
 
 			// Add nodes only for new counters.
-			CountersProvider.EnumerateCounters([this, MemoryGroup, MiscFloatGroup, MiscInt64Group](uint32 CounterId, const Trace::ICounter& Counter)
+			CountersProvider.EnumerateCounters([this, MemoryGroup, MiscFloatGroup, MiscInt64Group](uint32 CounterId, const TraceServices::ICounter& Counter)
 			{
 				FStatsNodePtr NodePtr = StatsNodesIdMap.FindRef(CounterId);
 				if (!NodePtr)
 				{
 					FName Name(Counter.GetName());
-					const FName Group = ((Counter.GetDisplayHint() == Trace::CounterDisplayHint_Memory) ? MemoryGroup :
+					const FName Group = ((Counter.GetDisplayHint() == TraceServices::CounterDisplayHint_Memory) ? MemoryGroup :
 										  Counter.IsFloatingPoint() ? MiscFloatGroup : MiscInt64Group);
 					const EStatsNodeType Type = EStatsNodeType::Counter;
 					const EStatsNodeDataType DataType = Counter.IsFloatingPoint() ? EStatsNodeDataType::Double : EStatsNodeDataType::Int64;
diff --git a/Engine/Source/Developer/TraceInsights/Private/Insights/Widgets/STimerTreeView.cpp b/Engine/Source/Developer/TraceInsights/Private/Insights/Widgets/STimerTreeView.cpp
--- a/Engine/Source/Developer/TraceInsights/Private/Insights/Widgets/STimerTreeView.cpp
+++ b/Engine/Source/Developer/TraceInsights/Private/Insights/Widgets/STimerTreeView.cpp
@@ -987,7 +987,7 @@ void STimerTreeView::Reset()
 
 ////////////////////////////////////////////////////////////////////////////////////////////////////
 
-void STimerTreeView::SetTree(const Trace::FTimingProfilerButterflyNode& Root)
+void STimerTreeView::SetTree(const TraceServices::FTimingProfilerButterflyNode& Root)
 {
 	TreeNodes.Reset();
 
@@ -1018,7 +1018,7 @@ void STimerTreeView::SetTree(const Trace::FTimingProfilerButterflyNode& Root)
 
 ////////////////////////////////////////////////////////////////////////////////////////////////////
 
-FTimerNodePtr STimerTreeView::CreateTimerNodeRec(const Trace::FTimingProfilerButterflyNode& Node)
+FTimerNodePtr STimerTreeView::CreateTimerNodeRec(const TraceServices::FTimingProfilerButterflyNode& Node)
 {
 	if (Node.Timer == nullptr)
 	{
@@ -1029,13 +1029,13 @@ FTimerNodePtr STimerTreeView::CreateTimerNodeRec(const Trace::FTimingProfilerBut
 	const ETimerNodeType Type = Node.Timer->IsGpuTimer ? ETimerNodeType::GpuScope : ETimerNodeType::CpuScope;
 	FTimerNodePtr TimerNodePtr = MakeShared<FTimerNode>(Node.Timer->Id, Node.Timer->Name, Type);
 
-	Trace::FTimingProfilerAggregatedStats AggregatedStats;
+	TraceServices::FTimingProfilerAggregatedStats AggregatedStats;
 	AggregatedStats.InstanceCount = Node.Count;
 	AggregatedStats.TotalInclusiveTime = Node.InclusiveTime;
 	AggregatedStats.TotalExclusiveTime = Node.ExclusiveTime;
 	TimerNodePtr->SetAggregatedStats(AggregatedStats);
 
-	for (const Trace::FTimingProfilerButterflyNode* ChildNodePtr : Node.Children)
+	for (const TraceServices::FTimingProfilerButterflyNode* ChildNodePtr : Node.Children)
 	{
 		if (ChildNodePtr != nullptr)
 		{
diff --git a/Engine/Source/Developer/TraceInsights/Private/Insights/Widgets/STimerTreeView.h b/Engine/Source/Developer/TraceInsights/Private/Insights/Widgets/STimerTreeView.h
--- a/Engine/Source/Developer/TraceInsights/Private/Insights/Widgets/STimerTreeView.h
+++ b/Engine/Source/Developer/TraceInsights/Private/Insights/Widgets/STimerTreeView.h
@@ -54,13 +54,13 @@ public:
 	void Construct(const FArguments& InArgs, const FText& InViewName);
 
 	void Reset();
-	void SetTree(const Trace::FTimingProfilerButterflyNode& Root);
+	void SetTree(const TraceServices::FTimingProfilerButterflyNode& Root);
 
 	virtual FReply OnKeyDown(const FGeometry& MyGeometry, const FKeyEvent& InKeyEvent) override;
 
 private:
 	void InitCommandList();
-	FTimerNodePtr CreateTimerNodeRec(const Trace::FTimingProfilerButterflyNode& Node);
+	FTimerNodePtr CreateTimerNodeRec(const TraceServices::FTimingProfilerButterflyNode& Node);
 	void ExpandNodesRec(FTimerNodePtr NodePtr, int32 Depth);
 
 	////////////////////////////////////////////////////////////////////////////////////////////////////
diff --git a/Engine/Source/Developer/TraceInsights/Private/Insights/Widgets/STimersView.cpp b/Engine/Source/Developer/TraceInsights/Private/Insights/Widgets/STimersView.cpp
--- a/Engine/Source/Developer/TraceInsights/Private/Insights/Widgets/STimersView.cpp
+++ b/Engine/Source/Developer/TraceInsights/Private/Insights/Widgets/STimersView.cpp
@@ -730,7 +730,7 @@ TSharedRef<SWidget> STimersView::TreeViewHeaderRow_GenerateColumnMenu(const Insi
 
 void STimersView::InsightsManager_OnSessionChanged()
 {
-	TSharedPtr<const Trace::IAnalysisSession> NewSession = FInsightsManager::Get()->GetSession();
+	TSharedPtr<const TraceServices::IAnalysisSession> NewSession = FInsightsManager::Get()->GetSession();
 
 	if (NewSession != Session)
 	{
@@ -847,7 +847,7 @@ void STimersView::ApplyFiltering()
 	// Update aggregations for groups.
 	for (FTimerNodePtr& GroupPtr : FilteredGroupNodes)
 	{
-		Trace::FTimingProfilerAggregatedStats& AggregatedStats = GroupPtr->GetAggregatedStats();
+		TraceServices::FTimingProfilerAggregatedStats& AggregatedStats = GroupPtr->GetAggregatedStats();
 
 		GroupPtr->ResetAggregatedStats();
 
@@ -861,7 +861,7 @@ void STimersView::ApplyFiltering()
 		for (const Insights::FBaseTreeNodePtr& ChildPtr : GroupChildren)
 		{
 			const FTimerNodePtr& NodePtr = StaticCastSharedPtr<FTimerNode, Insights::FBaseTreeNode>(ChildPtr);
-			Trace::FTimingProfilerAggregatedStats& NodeAggregatedStats = NodePtr->GetAggregatedStats();
+			TraceServices::FTimingProfilerAggregatedStats& NodeAggregatedStats = NodePtr->GetAggregatedStats();
 
 			if (NodeAggregatedStats.InstanceCount > 0)
 			{
@@ -1748,14 +1748,14 @@ void STimersView::RebuildTree(bool bResync)
 	const uint32 PreviousNodeCount = TimerNodes.Num();
 
 	SyncStopwatch.Start();
-	if (Session.IsValid() && Trace::ReadTimingProfilerProvider(*Session.Get()))
+	if (Session.IsValid() && TraceServices::ReadTimingProfilerProvider(*Session.Get()))
 	{
-		Trace::FAnalysisSessionReadScope SessionReadScope(*Session.Get());
+		TraceServices::FAnalysisSessionReadScope SessionReadScope(*Session.Get());
 
-		const Trace::ITimingProfilerProvider& TimingProfilerProvider = *Trace::ReadTimingProfilerProvider(*Session.Get());
+		const TraceServices::ITimingProfilerProvider& TimingProfilerProvider = *TraceServices::ReadTimingProfilerProvider(*Session.Get());
 
-		const Trace::ITimingProfilerTimerReader* TimerReader;
-		TimingProfilerProvider.ReadTimers([&TimerReader](const Trace::ITimingProfilerTimerReader& Out) { TimerReader = &Out; });
+		const TraceServices::ITimingProfilerTimerReader* TimerReader;
+		TimingProfilerProvider.ReadTimers([&TimerReader](const TraceServices::ITimingProfilerTimerReader& Out) { TimerReader = &Out; });
 
 		const uint32 TimerCount = TimerReader->GetTimerCount();
 		if (TimerCount != PreviousNodeCount)
@@ -1766,7 +1766,7 @@ void STimersView::RebuildTree(bool bResync)
 			// Add nodes only for new timers.
 			for (uint32 TimerIndex = PreviousNodeCount; TimerIndex < TimerCount; ++TimerIndex)
 			{
-				const Trace::FTimingProfilerTimer& Timer = *(TimerReader->GetTimer(TimerIndex));
+				const TraceServices::FTimingProfilerTimer& Timer = *(TimerReader->GetTimer(TimerIndex));
 				ensure(Timer.Id == TimerIndex);
 				const ETimerNodeType Type = Timer.IsGpuTimer ? ETimerNodeType::GpuScope : ETimerNodeType::CpuScope;
 				FTimerNodePtr TimerNodePtr = MakeShared<FTimerNode>(Timer.Id, Timer.Name, Type);
@@ -1876,14 +1876,14 @@ void STimersView::FinishAggregation()
 
 ////////////////////////////////////////////////////////////////////////////////////////////////////
 
-void STimersView::ApplyAggregation(Trace::ITable<Trace::FTimingProfilerAggregatedStats>* AggregatedStatsTable)
+void STimersView::ApplyAggregation(TraceServices::ITable<TraceServices::FTimingProfilerAggregatedStats>* AggregatedStatsTable)
 {
 	if (AggregatedStatsTable)
 	{
-		TUniquePtr<Trace::ITableReader<Trace::FTimingProfilerAggregatedStats>> TableReader(AggregatedStatsTable->CreateReader());
+		TUniquePtr<TraceServices::ITableReader<TraceServices::FTimingProfilerAggregatedStats>> TableReader(AggregatedStatsTable->CreateReader());
 		while (TableReader->IsValid())
 		{
-			const Trace::FTimingProfilerAggregatedStats* Row = TableReader->GetCurrentRow();
+			const TraceServices::FTimingProfilerAggregatedStats* Row = TableReader->GetCurrentRow();
 			FTimerNodePtr TimerNodePtr = GetTimerNode(Row->Timer->Id);
 			if (TimerNodePtr)
 			{
diff --git a/Engine/Source/Developer/TraceInsights/Private/Insights/Widgets/STimersView.h b/Engine/Source/Developer/TraceInsights/Private/Insights/Widgets/STimersView.h
--- a/Engine/Source/Developer/TraceInsights/Private/Insights/Widgets/STimersView.h
+++ b/Engine/Source/Developer/TraceInsights/Private/Insights/Widgets/STimersView.h
@@ -25,7 +25,7 @@ class FMenuBuilder;
 class FTimingGraphTrack;
 class FUICommandList;
 
-namespace Trace
+namespace TraceServices
 {
 	class IAnalysisSession;
 }
@@ -94,7 +94,7 @@ private:
 	void UpdateTree();
 
 	void FinishAggregation();
-	void ApplyAggregation(Trace::ITable<Trace::FTimingProfilerAggregatedStats>* AggregatedStatsTable);
+	void ApplyAggregation(TraceServices::ITable<TraceServices::FTimingProfilerAggregatedStats>* AggregatedStatsTable);
 
 	/** Called when the analysis session has changed. */
 	void InsightsManager_OnSessionChanged();
@@ -276,7 +276,7 @@ private:
 	TSharedPtr<Insights::FTable> Table;
 
 	/** A weak pointer to the profiler session used to populate this widget. */
-	TSharedPtr<const Trace::IAnalysisSession>/*Weak*/ Session;
+	TSharedPtr<const TraceServices::IAnalysisSession>/*Weak*/ Session;
 
 	TSharedPtr<FUICommandList> CommandList;
 
diff --git a/Engine/Source/Developer/TraceInsights/Private/Insights/Widgets/STimersViewTooltip.cpp b/Engine/Source/Developer/TraceInsights/Private/Insights/Widgets/STimersViewTooltip.cpp
--- a/Engine/Source/Developer/TraceInsights/Private/Insights/Widgets/STimersViewTooltip.cpp
+++ b/Engine/Source/Developer/TraceInsights/Private/Insights/Widgets/STimersViewTooltip.cpp
@@ -90,7 +90,7 @@ TSharedPtr<SToolTip> STimersViewTooltip::GetColumnTooltip(const Insights::FTable
 
 TSharedPtr<SToolTip> STimersViewTooltip::GetRowTooltip(const TSharedPtr<FTimerNode> TimerNodePtr)
 {
-	const Trace::FTimingProfilerAggregatedStats& Stats = TimerNodePtr->GetAggregatedStats();
+	const TraceServices::FTimingProfilerAggregatedStats& Stats = TimerNodePtr->GetAggregatedStats();
 
 	const FText InstanceCountText = FText::AsNumber(Stats.InstanceCount);
 
diff --git a/Engine/Source/Developer/TraceInsights/Private/Insights/Widgets/STimingView.cpp b/Engine/Source/Developer/TraceInsights/Private/Insights/Widgets/STimingView.cpp
--- a/Engine/Source/Developer/TraceInsights/Private/Insights/Widgets/STimingView.cpp
+++ b/Engine/Source/Developer/TraceInsights/Private/Insights/Widgets/STimingView.cpp
@@ -69,8 +69,8 @@
 // start auto generated ids from a big number (MSB set to 1) to avoid collisions with ids for gpu/cpu tracks based on 32bit timeline index
 uint64 FBaseTimingTrack::IdGenerator = (1ULL << 63);
 
-const TCHAR* GetFileActivityTypeName(Trace::EFileActivityType Type);
-uint32 GetFileActivityTypeColor(Trace::EFileActivityType Type);
+const TCHAR* GetFileActivityTypeName(TraceServices::EFileActivityType Type);
+uint32 GetFileActivityTypeColor(TraceServices::EFileActivityType Type);
 
 namespace Insights { const FName TimingViewExtenderFeatureName(TEXT("TimingViewExtender")); }
 
@@ -505,12 +505,12 @@ void STimingView::Tick(const FGeometry& AllottedGeometry, const double InCurrent
 	////////////////////////////////////////////////////////////////////////////////////////////////////
 	// Check the analysis session time.
 
-	TSharedPtr<const Trace::IAnalysisSession> Session = FInsightsManager::Get()->GetSession();
+	TSharedPtr<const TraceServices::IAnalysisSession> Session = FInsightsManager::Get()->GetSession();
 	if (Session)
 	{
 		double SessionTime = 0.0;
 		{
-			Trace::FAnalysisSessionReadScope SessionReadScope(*Session.Get());
+			TraceServices::FAnalysisSessionReadScope SessionReadScope(*Session.Get());
 			SessionTime = Session->GetDurationSeconds();
 		}
 
@@ -555,8 +555,8 @@ void STimingView::Tick(const FGeometry& AllottedGeometry, const double InCurrent
 			{
 				if (Session.IsValid())
 				{
-					Trace::FAnalysisSessionReadScope SessionReadScope(*Session.Get());
-					const Trace::IFrameProvider& FramesProvider = Trace::ReadFrameProvider(*Session.Get());
+					TraceServices::FAnalysisSessionReadScope SessionReadScope(*Session.Get());
+					const TraceServices::IFrameProvider& FramesProvider = TraceServices::ReadFrameProvider(*Session.Get());
 
 					const uint64 FrameCount = FramesProvider.GetFrameCount(AutoScrollFrameType);
 
@@ -566,7 +566,7 @@ void STimingView::Tick(const FGeometry& AllottedGeometry, const double InCurrent
 						uint64 FrameIndex = FrameCount;
 						while (FrameIndex > 0)
 						{
-							const Trace::FFrame* FramePtr = FramesProvider.GetFrame(AutoScrollFrameType, --FrameIndex);
+							const TraceServices::FFrame* FramePtr = FramesProvider.GetFrame(AutoScrollFrameType, --FrameIndex);
 							if (FramePtr && FramePtr->EndTime <= Viewport.GetMaxValidTime())
 							{
 								// Align the start time of the frame with the right side of the viewport.
@@ -576,7 +576,7 @@ void STimingView::Tick(const FGeometry& AllottedGeometry, const double InCurrent
 						}
 
 						// Get the frame at the center of the viewport.
-						Trace::FFrame Frame;
+						TraceServices::FFrame Frame;
 						const double ViewportCenter = MinStartTime + ViewportDuration / 2;
 						if (FramesProvider.GetFrameFromTime(AutoScrollFrameType, ViewportCenter, Frame))
 						{
diff --git a/Engine/Source/Developer/TraceInsights/Public/Insights/ITimingViewExtender.h b/Engine/Source/Developer/TraceInsights/Public/Insights/ITimingViewExtender.h
--- a/Engine/Source/Developer/TraceInsights/Public/Insights/ITimingViewExtender.h
+++ b/Engine/Source/Developer/TraceInsights/Public/Insights/ITimingViewExtender.h
@@ -5,7 +5,7 @@
 #include "CoreMinimal.h"
 #include "Features/IModularFeature.h"
 
-namespace Trace { class IAnalysisSession; }
+namespace TraceServices { class IAnalysisSession; }
 class FMenuBuilder;
 
 namespace Insights
@@ -26,7 +26,7 @@ public:
 	virtual void OnEndSession(ITimingViewSession& InSession) = 0;
 
 	/** Called each frame. If any new tracks are created they can be added via ITimingViewSession::Add*Track() */
-	virtual void Tick(ITimingViewSession& InSession, const Trace::IAnalysisSession& InAnalysisSession) = 0;
+	virtual void Tick(ITimingViewSession& InSession, const TraceServices::IAnalysisSession& InAnalysisSession) = 0;
 
 	/** Extension hook for the 'quick filter' menu */
 	virtual void ExtendFilterMenu(ITimingViewSession& InSession, FMenuBuilder& InMenuBuilder) = 0;
diff --git a/Engine/Source/Developer/TraceInsights/Public/Insights/IUnrealInsightsModule.h b/Engine/Source/Developer/TraceInsights/Public/Insights/IUnrealInsightsModule.h
--- a/Engine/Source/Developer/TraceInsights/Public/Insights/IUnrealInsightsModule.h
+++ b/Engine/Source/Developer/TraceInsights/Public/Insights/IUnrealInsightsModule.h
@@ -13,6 +13,10 @@ class FExtender;
 namespace Trace
 {
 	class FStoreClient;
+}
+
+namespace TraceServices
+{
 	class IAnalysisSession;
 }
 
@@ -189,7 +193,7 @@ public:
 	/**
 	 * Gets the current analysis session.
 	 */
-	virtual TSharedPtr<const Trace::IAnalysisSession> GetAnalysisSession() const = 0;
+	virtual TSharedPtr<const TraceServices::IAnalysisSession> GetAnalysisSession() const = 0;
 
 	/**
 	 * Starts analysis of the specified trace. Called when the application starts in "Viewer" mode.
diff --git a/Engine/Source/Developer/TraceInsights/Public/Insights/Tests/TimingProfilerTests.h b/Engine/Source/Developer/TraceInsights/Public/Insights/Tests/TimingProfilerTests.h
--- a/Engine/Source/Developer/TraceInsights/Public/Insights/Tests/TimingProfilerTests.h
+++ b/Engine/Source/Developer/TraceInsights/Public/Insights/Tests/TimingProfilerTests.h
@@ -34,7 +34,7 @@ public:
 	{
 		double Interval = 0.01;
 		int32 NumEnumerations = 10000;
-		Trace::EEventSortOrder SortOrder = Trace::EEventSortOrder::ByEndTime;
+		TraceServices::EEventSortOrder SortOrder = TraceServices::EEventSortOrder::ByEndTime;
 	};
 
 	static void RunEnumerateBenchmark(const FEnumerateTestParams& InParams, FCheckValues& OutCheckValues);
diff --git a/Engine/Source/Developer/TraceInsights/Public/Insights/ViewModels/BaseTimingTrack.h b/Engine/Source/Developer/TraceInsights/Public/Insights/ViewModels/BaseTimingTrack.h
--- a/Engine/Source/Developer/TraceInsights/Public/Insights/ViewModels/BaseTimingTrack.h
+++ b/Engine/Source/Developer/TraceInsights/Public/Insights/ViewModels/BaseTimingTrack.h
@@ -12,7 +12,7 @@
 ////////////////////////////////////////////////////////////////////////////////////////////////////
 
 class FMenuBuilder;
-namespace Trace { class IAnalysisSession; };
+namespace TraceServices { class IAnalysisSession; };
 
 class FTimingEventSearchParameters;
 class FTimingTrackViewport;
diff --git a/Engine/Source/Developer/TraceServices/Private/AnalysisService.cpp b/Engine/Source/Developer/TraceServices/Private/AnalysisService.cpp
--- a/Engine/Source/Developer/TraceServices/Private/AnalysisService.cpp
+++ b/Engine/Source/Developer/TraceServices/Private/AnalysisService.cpp
@@ -20,7 +20,7 @@
 #include "Model/Channel.h"
 #include "Model/DiagnosticsPrivate.h"
 
-namespace Trace
+namespace TraceServices
 {
 
 thread_local FAnalysisSessionLock* GThreadCurrentSessionLock;
@@ -102,7 +102,7 @@ FAnalysisSession::~FAnalysisSession()
 
 void FAnalysisSession::Start()
 {
-	FAnalysisContext Context;
+	Trace::FAnalysisContext Context;
 	for (Trace::IAnalyzer* Analyzer : ReadAnalyzers())
 	{
 		Context.AddAnalyzer(*Analyzer);
@@ -124,7 +124,7 @@ void FAnalysisSession::Wait() const
 	Processor.Wait();
 }
 
-void FAnalysisSession::AddAnalyzer(IAnalyzer* Analyzer)
+void FAnalysisSession::AddAnalyzer(Trace::IAnalyzer* Analyzer)
 {
 	Analyzers.Add(Analyzer);
 }
@@ -180,7 +180,7 @@ TSharedPtr<const IAnalysisSession> FAnalysisService::Analyze(const TCHAR* Sessio
 TSharedPtr<const IAnalysisSession> FAnalysisService::StartAnalysis(const TCHAR* SessionUri)
 {
 	struct FFileDataStream
-		: public IInDataStream
+		: public Trace::IInDataStream
 	{
 		virtual int32 Read(void* Data, uint32 Size) override
 		{
@@ -209,7 +209,7 @@ TSharedPtr<const IAnalysisSession> FAnalysisService::StartAnalysis(const TCHAR*
 	FileStream->Handle = TUniquePtr<IFileHandle>(Handle);
 	FileStream->Remaining = Handle->Size();
 
-	TUniquePtr<IInDataStream> DataStream(FileStream);
+	TUniquePtr<Trace::IInDataStream> DataStream(FileStream);
 	return StartAnalysis(SessionUri, MoveTemp(DataStream));
 }
 
@@ -217,7 +217,7 @@ TSharedPtr<const IAnalysisSession> FAnalysisService::StartAnalysis(const TCHAR*
 {
 	TSharedRef<FAnalysisSession> Session = MakeShared<FAnalysisSession>(SessionName, MoveTemp(DataStream));
 
-	Trace::FAnalysisSessionEditScope _(*Session);
+	FAnalysisSessionEditScope _(*Session);
 
 	FBookmarkProvider* BookmarkProvider = new FBookmarkProvider(*Session);
 	Session->AddProvider(FBookmarkProvider::ProviderName, BookmarkProvider);
@@ -255,4 +255,4 @@ TSharedPtr<const IAnalysisSession> FAnalysisService::StartAnalysis(const TCHAR*
 	return Session;
 }
 
-}
+} // namespace TraceServices
diff --git a/Engine/Source/Developer/TraceServices/Private/AnalysisServicePrivate.h b/Engine/Source/Developer/TraceServices/Private/AnalysisServicePrivate.h
--- a/Engine/Source/Developer/TraceServices/Private/AnalysisServicePrivate.h
+++ b/Engine/Source/Developer/TraceServices/Private/AnalysisServicePrivate.h
@@ -9,11 +9,10 @@
 #include "Common/StringStore.h"
 #include "Trace/Analysis.h"
 
-namespace Trace
+namespace TraceServices
 {
 
 class FModuleService;
-class IInDataStream;
 
 class FAnalysisSessionLock
 {
@@ -59,10 +58,10 @@ public:
 	virtual void ReadAccessCheck() const override { return Lock.ReadAccessCheck(); }
 	virtual void WriteAccessCheck() override { return Lock.WriteAccessCheck(); }
 
-	virtual void AddAnalyzer(IAnalyzer* Analyzer) override;
+	virtual void AddAnalyzer(Trace::IAnalyzer* Analyzer) override;
 	virtual void AddProvider(const FName& Name, IProvider* Provider) override;
 
-	const TArray<IAnalyzer*> ReadAnalyzers() { return Analyzers; }
+	const TArray<Trace::IAnalyzer*> ReadAnalyzers() { return Analyzers; }
 
 private:
 	virtual const IProvider* ReadProviderPrivate(const FName& Name) const override;
@@ -74,7 +73,7 @@ private:
 	double DurationSeconds = 0.0;
 	FSlabAllocator Allocator;
 	FStringStore StringStore;
-	TArray<IAnalyzer*> Analyzers;
+	TArray<Trace::IAnalyzer*> Analyzers;
 	TArray<IProvider*> Providers;
 	TMap<FName, IProvider*> ProvidersMap;
 	mutable TUniquePtr<Trace::IInDataStream> DataStream;
@@ -95,4 +94,4 @@ private:
 	FModuleService& ModuleService;
 };
 
-}
+} // namespace TraceServices
diff --git a/Engine/Source/Developer/TraceServices/Private/Analyzers/AllocationsAnalysis.cpp b/Engine/Source/Developer/TraceServices/Private/Analyzers/AllocationsAnalysis.cpp
--- a/Engine/Source/Developer/TraceServices/Private/Analyzers/AllocationsAnalysis.cpp
+++ b/Engine/Source/Developer/TraceServices/Private/Analyzers/AllocationsAnalysis.cpp
@@ -4,6 +4,9 @@
 
 #if 0
 
+//namespace TraceServices
+//{
+//
 ////////////////////////////////////////////////////////////////////////////////
 enum
 {
@@ -14,7 +17,7 @@ enum
 };
 
 ////////////////////////////////////////////////////////////////////////////////
-FAllocationsAnalyzer::FAllocationsAnalyzer(Trace::IAnalysisSession& Session)
+FAllocationsAnalyzer::FAllocationsAnalyzer(IAnalysisSession& Session)
 {
 }
 
@@ -30,4 +33,6 @@ void FAllocationsAnalyzer::OnAnalysisBegin(const FOnAnalysisContext& Context)
 	*/
 }
 
+} // namespace TraceServices
+
 #endif // 0
diff --git a/Engine/Source/Developer/TraceServices/Private/Analyzers/AllocationsAnalysis.h b/Engine/Source/Developer/TraceServices/Private/Analyzers/AllocationsAnalysis.h
--- a/Engine/Source/Developer/TraceServices/Private/Analyzers/AllocationsAnalysis.h
+++ b/Engine/Source/Developer/TraceServices/Private/Analyzers/AllocationsAnalysis.h
@@ -6,7 +6,7 @@
 
 #if 0
 
-namespace Trace
+namespace TraceServices
 {
 
 class IAnalysisSession;
@@ -18,8 +18,10 @@ class FAllocationsAnalyzer
 	: public Trace::IAnalyzer
 {
 public:
-					FAllocationsAnalyzer(Trace::IAnalysisSession& Session);
+					FAllocationsAnalyzer(IAnalysisSession& Session);
 	virtual void	OnAnalysisBegin(const FOnAnalysisContext& Context) override;
 };
 
+} // namespace UE
+
 #endif // 0
diff --git a/Engine/Source/Developer/TraceServices/Private/Analyzers/CallstacksAnalysis.cpp b/Engine/Source/Developer/TraceServices/Private/Analyzers/CallstacksAnalysis.cpp
--- a/Engine/Source/Developer/TraceServices/Private/Analyzers/CallstacksAnalysis.cpp
+++ b/Engine/Source/Developer/TraceServices/Private/Analyzers/CallstacksAnalysis.cpp
@@ -4,8 +4,11 @@
 #include "TraceServices/Model/AnalysisSession.h"
 #include "Model/CallstacksProvider.h"
 
+namespace TraceServices
+{
+
 ////////////////////////////////////////////////////////////////////////////////
-FCallstacksAnalyzer::FCallstacksAnalyzer(Trace::IAnalysisSession& Session, Trace::FCallstacksProvider* InProvider)
+FCallstacksAnalyzer::FCallstacksAnalyzer(IAnalysisSession& Session, FCallstacksProvider* InProvider)
 	: Provider(InProvider)
 {
 }
@@ -32,3 +35,4 @@ bool FCallstacksAnalyzer::OnEvent(uint16 RouteId, EStyle Style, const FOnEventCo
 	return true;
 }
 
+} // namespace TraceServices
diff --git a/Engine/Source/Developer/TraceServices/Private/Analyzers/CallstacksAnalysis.h b/Engine/Source/Developer/TraceServices/Private/Analyzers/CallstacksAnalysis.h
--- a/Engine/Source/Developer/TraceServices/Private/Analyzers/CallstacksAnalysis.h
+++ b/Engine/Source/Developer/TraceServices/Private/Analyzers/CallstacksAnalysis.h
@@ -5,18 +5,18 @@
 #include "Trace/Analyzer.h"
 
 ////////////////////////////////////////////////////////////////////////////////
-namespace Trace
+namespace TraceServices
 {
-	class IAnalysisSession;
-	class FCallstacksProvider;
-};
+
+class IAnalysisSession;
+class FCallstacksProvider;
 
 ////////////////////////////////////////////////////////////////////////////////
 class FCallstacksAnalyzer
 	: public Trace::IAnalyzer
 {
 public:
-					FCallstacksAnalyzer(Trace::IAnalysisSession& Session, Trace::FCallstacksProvider* Provider);
+					FCallstacksAnalyzer(IAnalysisSession& Session, FCallstacksProvider* Provider);
 	virtual void	OnAnalysisBegin(const FOnAnalysisContext& Context) override;
 	virtual bool 	OnEvent(uint16 RouteId, EStyle Style, const FOnEventContext& Context) override;
 
@@ -26,6 +26,7 @@ private:
 		RouteId_Callstack,
 	};
 
-	Trace::FCallstacksProvider* Provider;
+	FCallstacksProvider* Provider;
 };
 
+} // namespace TraceServices
diff --git a/Engine/Source/Developer/TraceServices/Private/Analyzers/CountersTraceAnalysis.cpp b/Engine/Source/Developer/TraceServices/Private/Analyzers/CountersTraceAnalysis.cpp
--- a/Engine/Source/Developer/TraceServices/Private/Analyzers/CountersTraceAnalysis.cpp
+++ b/Engine/Source/Developer/TraceServices/Private/Analyzers/CountersTraceAnalysis.cpp
@@ -4,7 +4,10 @@
 #include "TraceServices/Model/Counters.h"
 #include "ProfilingDebugging/CountersTrace.h"
 
-FCountersAnalyzer::FCountersAnalyzer(Trace::IAnalysisSession& InSession, Trace::ICounterProvider& InCounterProvider)
+namespace TraceServices
+{
+
+FCountersAnalyzer::FCountersAnalyzer(IAnalysisSession& InSession, ICounterProvider& InCounterProvider)
 	: Session(InSession)
 	, CounterProvider(InCounterProvider)
 {
@@ -22,7 +25,7 @@ void FCountersAnalyzer::OnAnalysisBegin(const FOnAnalysisContext& Context)
 
 bool FCountersAnalyzer::OnEvent(uint16 RouteId, EStyle Style, const FOnEventContext& Context)
 {
-	Trace::FAnalysisSessionEditScope _(Session);
+	FAnalysisSessionEditScope _(Session);
 
 	const auto& EventData = Context.EventData;
 	switch (RouteId)
@@ -32,14 +35,14 @@ bool FCountersAnalyzer::OnEvent(uint16 RouteId, EStyle Style, const FOnEventCont
 		uint16 CounterId = EventData.GetValue<uint16>("Id");
 		ETraceCounterType CounterType = static_cast<ETraceCounterType>(EventData.GetValue<uint8>("Type"));
 		ETraceCounterDisplayHint CounterDisplayHint = static_cast<ETraceCounterDisplayHint>(EventData.GetValue<uint8>("DisplayHint"));
-		Trace::IEditableCounter* Counter = CounterProvider.CreateCounter();
+		IEditableCounter* Counter = CounterProvider.CreateCounter();
 		if (CounterType == TraceCounterType_Float)
 		{
 			Counter->SetIsFloatingPoint(true);
 		}
 		if (CounterDisplayHint == TraceCounterDisplayHint_Memory)
 		{
-			Counter->SetDisplayHint(Trace::CounterDisplayHint_Memory);
+			Counter->SetDisplayHint(CounterDisplayHint_Memory);
 		}
 		FString Name = FTraceAnalyzerUtils::LegacyAttachmentString<TCHAR>("Name", Context);
 		const TCHAR* NamePtr = Session.StoreString(*Name);
@@ -52,7 +55,7 @@ bool FCountersAnalyzer::OnEvent(uint16 RouteId, EStyle Style, const FOnEventCont
 		uint16 CounterId = EventData.GetValue<uint16>("CounterId");
 		int64 Value = EventData.GetValue<int64>("Value");
 		double Timestamp = Context.EventTime.AsSeconds(EventData.GetValue<uint64>("Cycle"));
-		Trace::IEditableCounter** FindCounter = CountersMap.Find(CounterId);
+		IEditableCounter** FindCounter = CountersMap.Find(CounterId);
 		if (FindCounter)
 		{
 			(*FindCounter)->SetValue(Timestamp, Value);
@@ -64,7 +67,7 @@ bool FCountersAnalyzer::OnEvent(uint16 RouteId, EStyle Style, const FOnEventCont
 		uint16 CounterId = EventData.GetValue<uint16>("CounterId");
 		float Value = EventData.GetValue<float>("Value");
 		double Timestamp = Context.EventTime.AsSeconds(EventData.GetValue<uint64>("Cycle"));
-		Trace::IEditableCounter* FindCounter = CountersMap.FindRef(CounterId);
+		IEditableCounter* FindCounter = CountersMap.FindRef(CounterId);
 		if (FindCounter)
 		{
 			FindCounter->SetValue(Timestamp, Value);
@@ -75,3 +78,5 @@ bool FCountersAnalyzer::OnEvent(uint16 RouteId, EStyle Style, const FOnEventCont
 
 	return true;
 }
+
+} // namespace TraceServices
diff --git a/Engine/Source/Developer/TraceServices/Private/Analyzers/CountersTraceAnalysis.h b/Engine/Source/Developer/TraceServices/Private/Analyzers/CountersTraceAnalysis.h
--- a/Engine/Source/Developer/TraceServices/Private/Analyzers/CountersTraceAnalysis.h
+++ b/Engine/Source/Developer/TraceServices/Private/Analyzers/CountersTraceAnalysis.h
@@ -5,18 +5,18 @@
 #include "Trace/Analyzer.h"
 #include "Containers/Map.h"
 
-namespace Trace
+namespace TraceServices
 {
-	class IAnalysisSession;
-	class IEditableCounter;
-	class ICounterProvider;
-}
+
+class IAnalysisSession;
+class IEditableCounter;
+class ICounterProvider;
 
 class FCountersAnalyzer
 	: public Trace::IAnalyzer
 {
 public:
-	FCountersAnalyzer(Trace::IAnalysisSession& Session, Trace::ICounterProvider& CounterProvider);
+	FCountersAnalyzer(IAnalysisSession& Session, ICounterProvider& CounterProvider);
 	virtual void OnAnalysisBegin(const FOnAnalysisContext& Context) override;
 	virtual bool OnEvent(uint16 RouteId, EStyle Style, const FOnEventContext& Context) override;
 
@@ -28,7 +28,9 @@ private:
 		RouteId_SetValueFloat,
 	};
 
-	Trace::IAnalysisSession& Session;
-	Trace::ICounterProvider& CounterProvider;
-	TMap<uint16, Trace::IEditableCounter*> CountersMap;
+	IAnalysisSession& Session;
+	ICounterProvider& CounterProvider;
+	TMap<uint16, IEditableCounter*> CountersMap;
 };
+
+} // namespace TraceServices
diff --git a/Engine/Source/Developer/TraceServices/Private/Analyzers/CpuProfilerTraceAnalysis.cpp b/Engine/Source/Developer/TraceServices/Private/Analyzers/CpuProfilerTraceAnalysis.cpp
--- a/Engine/Source/Developer/TraceServices/Private/Analyzers/CpuProfilerTraceAnalysis.cpp
+++ b/Engine/Source/Developer/TraceServices/Private/Analyzers/CpuProfilerTraceAnalysis.cpp
@@ -4,7 +4,10 @@
 #include "Common/Utils.h"
 #include "Model/ThreadsPrivate.h"
 
-FCpuProfilerAnalyzer::FCpuProfilerAnalyzer(Trace::IAnalysisSession& InSession, Trace::FTimingProfilerProvider& InTimingProfilerProvider, Trace::FThreadProvider& InThreadProvider)
+namespace TraceServices
+{
+
+FCpuProfilerAnalyzer::FCpuProfilerAnalyzer(IAnalysisSession& InSession, FTimingProfilerProvider& InTimingProfilerProvider, FThreadProvider& InThreadProvider)
 	: Session(InSession)
 	, TimingProfilerProvider(InTimingProfilerProvider)
 	, ThreadProvider(InThreadProvider)
@@ -34,7 +37,7 @@ void FCpuProfilerAnalyzer::OnAnalysisBegin(const FOnAnalysisContext& Context)
 
 bool FCpuProfilerAnalyzer::OnEvent(uint16 RouteId, EStyle Style, const FOnEventContext& Context)
 {
-	Trace::FAnalysisSessionEditScope _(Session);
+	FAnalysisSessionEditScope _(Session);
 
 	const auto& EventData = Context.EventData;
 	switch (RouteId)
@@ -166,7 +169,7 @@ uint64 FCpuProfilerAnalyzer::ProcessBuffer(const FEventTime& EventTime, uint32 T
 			double Time = EventTime.AsSeconds(PendingCycle);
 			if (bEnter)
 			{
-				Trace::FTimingProfilerEvent Event;
+				FTimingProfilerEvent Event;
 				Event.TimerIndex = PendingCursor->TimerId;
 				ThreadState.Timeline->AppendBeginEvent(Time, Event);
 			}
@@ -194,7 +197,7 @@ uint64 FCpuProfilerAnalyzer::ProcessBuffer(const FEventTime& EventTime, uint32 T
 			ScopeState.StartCycle = ActualCycle;
 			ScopeState.EventTypeId = TimerId;
 
-			Trace::FTimingProfilerEvent Event;
+			FTimingProfilerEvent Event;
 			Event.TimerIndex = TimerId;
 			double ActualTime = EventTime.AsSeconds(ActualCycle);
 			ThreadState.Timeline->AppendBeginEvent(ActualTime, Event);
@@ -228,7 +231,7 @@ uint64 FCpuProfilerAnalyzer::ProcessBuffer(const FEventTime& EventTime, uint32 T
 		}
 		else
 		{
-			Trace::FTimingProfilerEvent Event;
+			FTimingProfilerEvent Event;
 			Event.TimerIndex = PendingCursor->TimerId;
 			double Time = EventTime.AsSeconds(PendingCycle);
 			ThreadState.Timeline->AppendBeginEvent(Time, Event);
@@ -322,3 +325,5 @@ FCpuProfilerAnalyzer::FThreadState& FCpuProfilerAnalyzer::GetThreadState(uint32
 	}
 	return *ThreadState;
 }
+
+} // namespace TraceServices
diff --git a/Engine/Source/Developer/TraceServices/Private/Analyzers/CpuProfilerTraceAnalysis.h b/Engine/Source/Developer/TraceServices/Private/Analyzers/CpuProfilerTraceAnalysis.h
--- a/Engine/Source/Developer/TraceServices/Private/Analyzers/CpuProfilerTraceAnalysis.h
+++ b/Engine/Source/Developer/TraceServices/Private/Analyzers/CpuProfilerTraceAnalysis.h
@@ -6,17 +6,17 @@
 #include "Containers/UnrealString.h"
 #include "Model/TimingProfilerPrivate.h"
 
-namespace Trace
+namespace TraceServices
 {
-	class IAnalysisSession;
-	class FThreadProvider;
-}
+
+class IAnalysisSession;
+class FThreadProvider;
 
 class FCpuProfilerAnalyzer
 	: public Trace::IAnalyzer
 {
 public:
-	FCpuProfilerAnalyzer(Trace::IAnalysisSession& Session, Trace::FTimingProfilerProvider& TimingProfilerProvider, Trace::FThreadProvider& InThreadProvider);
+	FCpuProfilerAnalyzer(IAnalysisSession& Session, FTimingProfilerProvider& TimingProfilerProvider, FThreadProvider& InThreadProvider);
 	~FCpuProfilerAnalyzer();
 	virtual void OnAnalysisBegin(const FOnAnalysisContext& Context) override;
 	virtual bool OnEvent(uint16 RouteId, EStyle Style, const FOnEventContext& Context) override;
@@ -38,7 +38,7 @@ private:
 	{
 		TArray<FEventScopeState> ScopeStack;
 		TArray<FPendingEvent> PendingEvents;
-		Trace::FTimingProfilerProvider::TimelineInternal* Timeline;
+		FTimingProfilerProvider::TimelineInternal* Timeline;
 		uint64 LastCycle = 0;
 	};
 
@@ -59,9 +59,9 @@ private:
 		RouteId_ChannelToggle,
 	};
 
-	Trace::IAnalysisSession& Session;
-	Trace::FTimingProfilerProvider& TimingProfilerProvider;
-	Trace::FThreadProvider& ThreadProvider;
+	IAnalysisSession& Session;
+	FTimingProfilerProvider& TimingProfilerProvider;
+	FThreadProvider& ThreadProvider;
 	TMap<uint32, FThreadState*> ThreadStatesMap;
 	TMap<uint32, uint32> SpecIdToTimerIdMap;
 	TMap<const TCHAR*, uint32> ScopeNameToTimerIdMap;
@@ -69,3 +69,5 @@ private:
 	uint64 TotalScopeCount = 0;
 	double BytesPerScope = 0.0;
 };
+
+} // namespace TraceServices
