[Insights] Initial work to allow AllocationTreeView filtering, grouping and sorting to be executed asyncronosly

diff --git a/Engine/Source/Developer/TraceInsights/Private/Insights/MemoryProfiler/Widgets/SMemAllocTableTreeView.cpp b/Engine/Source/Developer/TraceInsights/Private/Insights/MemoryProfiler/Widgets/SMemAllocTableTreeView.cpp
--- a/Engine/Source/Developer/TraceInsights/Private/Insights/MemoryProfiler/Widgets/SMemAllocTableTreeView.cpp
+++ b/Engine/Source/Developer/TraceInsights/Private/Insights/MemoryProfiler/Widgets/SMemAllocTableTreeView.cpp
@@ -11,12 +11,17 @@
 #include "Insights/MemoryProfiler/ViewModels/MemAllocNode.h"
 #include "Insights/MemoryProfiler/ViewModels/MemAllocTable.h"
 #include "Insights/MemoryProfiler/ViewModels/MemorySharedState.h"
+#include "Insights/MemoryProfiler/Widgets/SMemoryProfilerWindow.h"
+#include "Insights/Table/ViewModels/TableCellValueFormatter.h"
+#include "Insights/Table/ViewModels/TableColumn.h"
 #include "Insights/TimingProfilerCommon.h"
 
 #include <limits>
 
 #define LOCTEXT_NAMESPACE "SMemAllocTableTreeView"
 
+using namespace TraceServices;
+
 namespace Insights
 {
 
@@ -24,6 +29,7 @@ namespace Insights
 
 SMemAllocTableTreeView::SMemAllocTableTreeView()
 {
+	bRunInAsyncMode = true;
 }
 
 ////////////////////////////////////////////////////////////////////////////////////////////////////
@@ -66,11 +72,12 @@ void SMemAllocTableTreeView::UpdateSourceTable(TSharedPtr<TraceServices::IMemAll
 
 void SMemAllocTableTreeView::Tick(const FGeometry& AllottedGeometry, const double InCurrentTime, const float InDeltaTime)
 {
+	STableTreeView::Tick(AllottedGeometry, InCurrentTime, InDeltaTime);
 	// We need to check if the list of LLM tags has changed.
 	// But, ensure we do not check too often.
 	static uint64 NextTimestamp = 0;
 	uint64 Time = FPlatformTime::Cycles64();
-	if (Time > NextTimestamp)
+	if (!bIsUpdateRunning && Time > NextTimestamp)
 	{
 		RebuildTree(false);
 
@@ -198,7 +205,7 @@ void SMemAllocTableTreeView::StartQuery()
 		return;
 	}
 
-	const IAllocationsProvider* AllocationsProvider = TraceServices::ReadAllocationsProvider(*Session.Get());
+	const IAllocationsProvider* AllocationsProvider = ReadAllocationsProvider(*Session.Get());
 	if (!AllocationsProvider)
 	{
 		UE_LOG(MemoryProfiler, Warning, TEXT("[MemAlloc] Invalid allocations provider!"));
@@ -295,7 +302,7 @@ void SMemAllocTableTreeView::UpdateQuery()
 #if defined(UE_USE_ALLOCATIONS_PROVIDER)
 	if (Query != 0)
 	{
-		const IAllocationsProvider* AllocationsProvider = TraceServices::ReadAllocationsProvider(*Session.Get());
+		const IAllocationsProvider* AllocationsProvider = ReadAllocationsProvider(*Session.Get());
 		if (!AllocationsProvider)
 		{
 			UE_LOG(MemoryProfiler, Warning, TEXT("[MemAlloc] Invalid allocations provider!"));
@@ -394,7 +401,7 @@ void SMemAllocTableTreeView::CancelQuery()
 #if defined(UE_USE_ALLOCATIONS_PROVIDER)
 	if (Query != 0)
 	{
-		const IAllocationsProvider* AllocationsProvider = TraceServices::ReadAllocationsProvider(*Session.Get());
+		const IAllocationsProvider* AllocationsProvider = ReadAllocationsProvider(*Session.Get());
 		if (AllocationsProvider)
 		{
 			AllocationsProvider->CancelQuery(Query);
@@ -407,6 +414,95 @@ void SMemAllocTableTreeView::CancelQuery()
 
 ////////////////////////////////////////////////////////////////////////////////////////////////////
 
+void SMemAllocTableTreeView::OnPreAsyncUpdate()
+{
+	auto LLMColumn = Table->FindColumn(TEXT("AllocLlmTag"));
+	OriginalLlmTagValueFormatter = LLMColumn->GetValueFormatter();
+	LLMColumn->SetValueFormatter(CreateCachedLlmTagValueFormatter());
+
+	STableTreeView::OnPreAsyncUpdate();
+}
+
+////////////////////////////////////////////////////////////////////////////////////////////////////
+
+void SMemAllocTableTreeView::OnPostAsyncUpdate()
+{
+	if (OriginalLlmTagValueFormatter.IsValid())
+	{
+		auto LLMColumn = Table->FindColumn(TEXT("AllocLlmTag"));
+		LLMColumn->SetValueFormatter(OriginalLlmTagValueFormatter.ToSharedRef());
+		OriginalLlmTagValueFormatter = nullptr;
+	}
+
+	STableTreeView::OnPostAsyncUpdate();
+}
+
+////////////////////////////////////////////////////////////////////////////////////////////////////
+
+TSharedRef<ITableCellValueFormatter> SMemAllocTableTreeView::CreateCachedLlmTagValueFormatter()
+{
+	class FCachedMemAllocLlmTagValueFormatter : public FTableCellValueFormatter
+	{
+	public:
+		virtual FText FormatValue(const TOptional<FTableCellValue>& InValue) const override
+		{
+			if (InValue.IsSet())
+			{
+				const FMemoryTagId MemTagId = static_cast<FMemoryTagId>(InValue.GetValue().Int64);
+
+				if (TagIdToStatNameMap.Contains(MemTagId))
+				{
+					FText::FromString(TagIdToStatNameMap[MemTagId]);
+				}
+				return FText::FromString(FString());
+			}
+			return FText::GetEmpty();
+		}
+		virtual FText FormatValueForTooltip(const TOptional<FTableCellValue>& InValue) const override
+		{
+			if (InValue.IsSet())
+			{
+				const FMemoryTagId MemTagId = static_cast<FMemoryTagId>(InValue.GetValue().Int64);
+				if (TagIdToStatNameMap.Contains(MemTagId))
+				{
+					return FText::FromString(FString::Printf(TEXT("%lli (%s)"), MemTagId, *TagIdToStatNameMap[MemTagId]));
+				}
+				
+				return FText::FromString(FString::Printf(TEXT("%lli ()"), MemTagId));
+			}
+			return FText::GetEmpty();
+		}
+		virtual FText FormatValue(const FTableColumn& Column, const FBaseTreeNode& Node) const override { return FormatValue(Column.GetValue(Node)); }
+		virtual FText FormatValueForTooltip(const FTableColumn& Column, const FBaseTreeNode& Node) const override { return FormatValueForTooltip(Column.GetValue(Node)); }
+
+		void AddTagToCache(FMemoryTagId TagId, const FString& StatName)
+		{
+			TagIdToStatNameMap.Add(TagId, StatName);
+		}
+
+	private:
+		TMap<FMemoryTagId, FString> TagIdToStatNameMap;
+	};
+
+	TSharedRef<FCachedMemAllocLlmTagValueFormatter> Formatter = MakeShared<FCachedMemAllocLlmTagValueFormatter>();
+
+	TSharedPtr<SMemoryProfilerWindow> ProfilerWindow = FMemoryProfilerManager::Get()->GetProfilerWindow();
+	if (ProfilerWindow)
+	{
+		auto& SharedState = ProfilerWindow->GetSharedState();
+		const Insights::FMemoryTagList& TagList = SharedState.GetTagList();
+		for (FMemoryTag* CurrentTag : TagList.GetTags())
+		{
+			Formatter->AddTagToCache(CurrentTag->GetId(), CurrentTag->GetStatName());
+		}
+
+	}
+
+	return Formatter;
+}
+
+////////////////////////////////////////////////////////////////////////////////////////////////////
+
 } // namespace Insights
 
 #undef LOCTEXT_NAMESPACE
diff --git a/Engine/Source/Developer/TraceInsights/Private/Insights/MemoryProfiler/Widgets/SMemAllocTableTreeView.h b/Engine/Source/Developer/TraceInsights/Private/Insights/MemoryProfiler/Widgets/SMemAllocTableTreeView.h
--- a/Engine/Source/Developer/TraceInsights/Private/Insights/MemoryProfiler/Widgets/SMemAllocTableTreeView.h
+++ b/Engine/Source/Developer/TraceInsights/Private/Insights/MemoryProfiler/Widgets/SMemAllocTableTreeView.h
@@ -70,19 +70,27 @@ public:
 	int32 GetTabIndex() const { return TabIndex; }
 	void SetTabIndex(int32 InTabIndex) { TabIndex = InTabIndex; }
 
+protected:
+	virtual void OnPreAsyncUpdate() override;
+	virtual void OnPostAsyncUpdate() override;
+
 private:
 	void OnQueryInvalidated();
 	void StartQuery();
 	void UpdateQuery();
 	void CancelQuery();
 
+	TSharedRef<class ITableCellValueFormatter> CreateCachedLlmTagValueFormatter();
+
 private:
 	int32 TabIndex = -1;
 	TSharedPtr<FMemoryRuleSpec> Rule = nullptr;
 	double TimeMarkers[4];
 #if defined(UE_USE_ALLOCATIONS_PROVIDER)
-	IAllocationsProvider::QueryHandle Query = 0;
+	TraceServices::IAllocationsProvider::QueryHandle Query = 0;
 #endif
+
+	TSharedPtr<class ITableCellValueFormatter> OriginalLlmTagValueFormatter;
 };
 
 ////////////////////////////////////////////////////////////////////////////////////////////////////
diff --git a/Engine/Source/Developer/TraceInsights/Private/Insights/Table/Widgets/STableTreeView.cpp b/Engine/Source/Developer/TraceInsights/Private/Insights/Table/Widgets/STableTreeView.cpp
--- a/Engine/Source/Developer/TraceInsights/Private/Insights/Table/Widgets/STableTreeView.cpp
+++ b/Engine/Source/Developer/TraceInsights/Private/Insights/Table/Widgets/STableTreeView.cpp
@@ -213,7 +213,7 @@ void STableTreeView::ConstructWidget(TSharedPtr<FTable> InTablePtr)
 	];
 
 	// Create the search filters: text based, type based etc.
-	TextFilter = MakeShared<FTableTreeNodeTextFilter>(FTableTreeNodeTextFilter::FItemToStringArray::CreateSP(this, &STableTreeView::HandleItemToStringArray));
+	TextFilter = MakeShared<FTableTreeNodeTextFilter>(FTableTreeNodeTextFilter::FItemToStringArray::CreateStatic(&STableTreeView::HandleItemToStringArray));
 	Filters = MakeShared<FTableTreeNodeFilterCollection>();
 	Filters->Add(TextFilter);
 
@@ -579,14 +579,47 @@ void STableTreeView::InsightsManager_OnSessionChanged()
 
 ////////////////////////////////////////////////////////////////////////////////////////////////////
 
+void STableTreeView::Tick(const FGeometry& AllottedGeometry, const double InCurrentTime, const float InDeltaTime)
+{
+	if (bRunInAsyncMode && bIsUpdateRunning)
+	{
+		check(PendingAsyncOperationEvent.IsValid());
+		if (PendingAsyncOperationEvent->IsComplete())
+		{
+			OnPostAsyncUpdate();
+		}
+	}
+}
+
+////////////////////////////////////////////////////////////////////////////////////////////////////
+
 void STableTreeView::UpdateTree()
 {
 	FStopwatch Stopwatch;
 	Stopwatch.Start();
 
-	CreateGroups();
-	SortTreeNodes();
-	ApplyFiltering();
+	if (bRunInAsyncMode)
+	{
+		OnPreAsyncUpdate();
+
+		FGraphEventRef GroupingCompletedEvent = TGraphTask<FTableTreeViewGroupAsyncTask>::CreateTask().ConstructAndDispatchWhenReady(SharedThis(this));
+		
+		FGraphEventArray Prerequistes;
+		Prerequistes.Add(GroupingCompletedEvent);
+		FGraphEventRef SortingCompletedEvent = TGraphTask<FTableTreeViewSortAsyncTask>::CreateTask(&Prerequistes).ConstructAndDispatchWhenReady(SharedThis(this));
+		
+		Prerequistes.Empty();
+		Prerequistes.Add(SortingCompletedEvent);
+		FGraphEventRef FilteringCompletedEvent = TGraphTask<FTableTreeViewFilterAsyncTask>::CreateTask(&Prerequistes).ConstructAndDispatchWhenReady(SharedThis(this));
+		
+		PendingAsyncOperationEvent = FilteringCompletedEvent;
+	}
+	else
+	{
+		CreateGroups();
+		SortTreeNodes();
+		ApplyFiltering();
+	}
 
 	Stopwatch.Stop();
 	const double TotalTime = Stopwatch.GetAccumulatedTime();
@@ -619,39 +652,43 @@ void STableTreeView::ApplyFiltering()
 		}
 	}
 
-	// Only expand nodes if we have a text filter.
-	const bool bNonEmptyTextFilter = !TextFilter->GetRawFilterText().IsEmpty();
-	if (bNonEmptyTextFilter)
+	// Cannot call TreeView functions from other threads than MainThread and SlateThread. 
+	if (!bRunInAsyncMode)
 	{
-		if (!bExpansionSaved)
+		// Only expand nodes if we have a text filter.
+		const bool bNonEmptyTextFilter = !TextFilter->GetRawFilterText().IsEmpty();
+		if (bNonEmptyTextFilter)
 		{
-			ExpandedNodes.Empty();
-			TreeView->GetExpandedItems(ExpandedNodes);
-			bExpansionSaved = true;
-		}
+			if (!bExpansionSaved)
+			{
+				ExpandedNodes.Empty();
+				TreeView->GetExpandedItems(ExpandedNodes);
+				bExpansionSaved = true;
+			}
 
-		for (int32 Fx = 0; Fx < FilteredGroupNodes.Num(); Fx++)
-		{
-			const FTableTreeNodePtr& GroupPtr = FilteredGroupNodes[Fx];
-			TreeView->SetItemExpansion(GroupPtr, GroupPtr->IsExpanded());
+			for (int32 Fx = 0; Fx < FilteredGroupNodes.Num(); Fx++)
+			{
+				const FTableTreeNodePtr& GroupPtr = FilteredGroupNodes[Fx];
+				TreeView->SetItemExpansion(GroupPtr, GroupPtr->IsExpanded());
+			}
 		}
-	}
-	else
-	{
-		if (bExpansionSaved)
+		else
 		{
-			// Restore previously expanded nodes when the text filter is disabled.
-			TreeView->ClearExpandedItems();
-			for (auto It = ExpandedNodes.CreateConstIterator(); It; ++It)
+			if (bExpansionSaved)
 			{
-				TreeView->SetItemExpansion(*It, true);
+				// Restore previously expanded nodes when the text filter is disabled.
+				TreeView->ClearExpandedItems();
+				for (auto It = ExpandedNodes.CreateConstIterator(); It; ++It)
+				{
+					TreeView->SetItemExpansion(*It, true);
+				}
+				bExpansionSaved = false;
 			}
-			bExpansionSaved = false;
 		}
-	}
 
-	// Request tree refresh
-	TreeView->RequestTreeRefresh();
+		// Request tree refresh
+		TreeView->RequestTreeRefresh();
+	}
 
 	Stopwatch.Stop();
 	const double TotalTime = Stopwatch.GetAccumulatedTime();
@@ -707,7 +744,7 @@ bool STableTreeView::ApplyFilteringForNode(FTableTreeNodePtr NodePtr)
 
 ////////////////////////////////////////////////////////////////////////////////////////////////////
 
-void STableTreeView::HandleItemToStringArray(const FTableTreeNodePtr& FTableTreeNodePtr, TArray<FString>& OutSearchStrings) const
+void STableTreeView::HandleItemToStringArray(const FTableTreeNodePtr& FTableTreeNodePtr, TArray<FString>& OutSearchStrings)
 {
 	OutSearchStrings.Add(FTableTreeNodePtr->GetName().GetPlainNameString());
 }
@@ -837,13 +874,30 @@ void STableTreeView::SearchBox_OnTextChanged(const FText& InFilterText)
 {
 	TextFilter->SetRawFilterText(InFilterText);
 	SearchBox->SetError(TextFilter->GetFilterErrorText());
-	ApplyFiltering();
+
+	if(bRunInAsyncMode)
+	{
+		OnPreAsyncUpdate();
+
+		FGraphEventRef FilteringCompletedEvent = TGraphTask<FTableTreeViewFilterAsyncTask>::CreateTask().ConstructAndDispatchWhenReady(SharedThis(this));
+
+		PendingAsyncOperationEvent = FilteringCompletedEvent;
+	}
+	else
+	{
+		ApplyFiltering();
+	}
 }
 
 ////////////////////////////////////////////////////////////////////////////////////////////////////
 
 bool STableTreeView::SearchBox_IsEnabled() const
 {
+	if (bIsUpdateRunning)
+	{
+		return false;
+	}
+
 	return TableTreeNodes.Num() > 0;
 }
 
@@ -906,7 +960,10 @@ void STableTreeView::GroupNodesRec(const TArray<FTableTreeNodePtr>& Nodes, FTabl
 	{
 		ensure(GroupPtr->IsGroup());
 		FTableTreeNodePtr TableTreeGroupPtr = StaticCastSharedPtr<FTableTreeNode>(GroupPtr);
-		TreeView->SetItemExpansion(TableTreeGroupPtr, TableTreeGroupPtr->IsExpanded());
+		if (!bRunInAsyncMode)
+		{
+			TreeView->SetItemExpansion(TableTreeGroupPtr, TableTreeGroupPtr->IsExpanded());
+		}
 	}
 
 	if (GroupingDepth < CurrentGroupings.Num() - 1)
@@ -1135,11 +1192,30 @@ void STableTreeView::PostChangeGroupings()
 
 	TreeViewHeaderRow->RefreshColumns();
 
-	CreateGroups();
-	SortTreeNodes();
-	ApplyFiltering();
+	if (bRunInAsyncMode)
+	{
+		OnPreAsyncUpdate();
+
+		FGraphEventRef GroupingCompletedEvent = TGraphTask<FTableTreeViewGroupAsyncTask>::CreateTask().ConstructAndDispatchWhenReady(SharedThis(this));
 
-	RebuildGroupingCrumbs();
+		FGraphEventArray Prerequistes;
+		Prerequistes.Add(GroupingCompletedEvent);
+		FGraphEventRef SortingCompletedEvent = TGraphTask<FTableTreeViewSortAsyncTask>::CreateTask(&Prerequistes).ConstructAndDispatchWhenReady(SharedThis(this));
+
+		Prerequistes.Empty();
+		Prerequistes.Add(SortingCompletedEvent);
+		FGraphEventRef FilteringCompletedEvent = TGraphTask<FTableTreeViewFilterAsyncTask>::CreateTask(&Prerequistes).ConstructAndDispatchWhenReady(SharedThis(this));
+
+		PendingAsyncOperationEvent = FilteringCompletedEvent;
+	}
+	else
+	{
+		CreateGroups();
+		SortTreeNodes();
+		ApplyFiltering();
+
+		RebuildGroupingCrumbs();
+	}
 }
 
 ////////////////////////////////////////////////////////////////////////////////////////////////////
@@ -1256,6 +1332,11 @@ TSharedRef<SWidget> STableTreeView::GetGroupingCrumbMenuContent(const TSharedPtr
 	}
 	MenuBuilder.EndSection();
 
+	auto CanExecute = [this]()
+	{
+		return !this->IsRunningAsyncUpdate();
+	};
+
 	if (CrumbGroupingDepth >= 0)
 	{
 		MenuBuilder.BeginSection("CrumbGrouping", CrumbGrouping->GetTitleName());
@@ -1274,7 +1355,7 @@ TSharedRef<SWidget> STableTreeView::GetGroupingCrumbMenuContent(const TSharedPtr
 				FUIAction Action_MoveLeft
 				(
 					FExecuteAction::CreateSP(this, &STableTreeView::GroupingCrumbMenu_MoveLeft_Execute, CrumbGrouping),
-					FCanExecuteAction()
+					FCanExecuteAction::CreateLambda(CanExecute)
 				);
 				MenuBuilder.AddMenuEntry
 				(
@@ -1289,7 +1370,7 @@ TSharedRef<SWidget> STableTreeView::GetGroupingCrumbMenuContent(const TSharedPtr
 				FUIAction Action_MoveRight
 				(
 					FExecuteAction::CreateSP(this, &STableTreeView::GroupingCrumbMenu_MoveRight_Execute, CrumbGrouping),
-					FCanExecuteAction()
+					FCanExecuteAction::CreateLambda(CanExecute)
 				);
 				MenuBuilder.AddMenuEntry
 				(
@@ -1304,7 +1385,7 @@ TSharedRef<SWidget> STableTreeView::GetGroupingCrumbMenuContent(const TSharedPtr
 				FUIAction Action_Remove
 				(
 					FExecuteAction::CreateSP(this, &STableTreeView::GroupingCrumbMenu_Remove_Execute, CrumbGrouping),
-					FCanExecuteAction()
+					FCanExecuteAction::CreateLambda(CanExecute)
 				);
 				MenuBuilder.AddMenuEntry
 				(
@@ -1324,7 +1405,7 @@ TSharedRef<SWidget> STableTreeView::GetGroupingCrumbMenuContent(const TSharedPtr
 			FUIAction Action_Reset
 			(
 				FExecuteAction::CreateSP(this, &STableTreeView::GroupingCrumbMenu_Reset_Execute),
-				FCanExecuteAction()
+				FCanExecuteAction::CreateLambda(CanExecute)
 			);
 			MenuBuilder.AddMenuEntry
 			(
@@ -1435,6 +1516,11 @@ void STableTreeView::GroupingCrumbMenu_Change_Execute(const TSharedPtr<FTreeNode
 
 bool STableTreeView::GroupingCrumbMenu_Change_CanExecute(const TSharedPtr<FTreeNodeGrouping> OldGrouping, const TSharedPtr<FTreeNodeGrouping> NewGrouping) const
 {
+	if (bIsUpdateRunning)
+	{
+		return false;
+	}
+
 	return NewGrouping != OldGrouping;
 }
 
@@ -1482,6 +1568,11 @@ void STableTreeView::GroupingCrumbMenu_Add_Execute(const TSharedPtr<FTreeNodeGro
 
 bool STableTreeView::GroupingCrumbMenu_Add_CanExecute(const TSharedPtr<FTreeNodeGrouping> Grouping, const TSharedPtr<FTreeNodeGrouping> AfterGrouping) const
 {
+	if (bIsUpdateRunning)
+	{
+		return false;
+	}
+
 	if (AfterGrouping.IsValid())
 	{
 		const int32 AfterGroupingDepth = GetGroupingDepth(AfterGrouping);
@@ -1607,14 +1698,34 @@ void STableTreeView::SetSortModeForColumn(const FName& ColumnId, const EColumnSo
 	ColumnSortMode = SortMode;
 	UpdateCurrentSortingByColumn();
 
-	SortTreeNodes();
-	ApplyFiltering();
+	if (bRunInAsyncMode)
+	{
+		OnPreAsyncUpdate();
+
+		FGraphEventRef SortingCompletedEvent = TGraphTask<FTableTreeViewSortAsyncTask>::CreateTask().ConstructAndDispatchWhenReady(SharedThis(this));
+
+		FGraphEventArray Prerequistes;
+		Prerequistes.Add(SortingCompletedEvent);
+		FGraphEventRef FilteringCompletedEvent = TGraphTask<FTableTreeViewFilterAsyncTask>::CreateTask(&Prerequistes).ConstructAndDispatchWhenReady(SharedThis(this));
+
+		PendingAsyncOperationEvent = FilteringCompletedEvent;
+	}
+	else
+	{
+		SortTreeNodes();
+		ApplyFiltering();
+	}
 }
 
 ////////////////////////////////////////////////////////////////////////////////////////////////////
 
 void STableTreeView::OnSortModeChanged(const EColumnSortPriority::Type SortPriority, const FName& ColumnId, const EColumnSortMode::Type SortMode)
 {
+	if (bIsUpdateRunning)
+	{
+		return;
+	}
+
 	SetSortModeForColumn(ColumnId, SortMode);
 	TreeView_Refresh();
 }
@@ -1632,6 +1743,11 @@ bool STableTreeView::HeaderMenu_SortMode_IsChecked(const FName ColumnId, const E
 
 bool STableTreeView::HeaderMenu_SortMode_CanExecute(const FName ColumnId, const EColumnSortMode::Type InSortMode) const
 {
+	if (bIsUpdateRunning)
+	{
+		return false;
+	}
+
 	const FTableColumn& Column = *Table->FindColumnChecked(ColumnId);
 	return Column.CanBeSorted();
 }
@@ -1657,6 +1773,11 @@ bool STableTreeView::ContextMenu_SortMode_IsChecked(const EColumnSortMode::Type
 
 bool STableTreeView::ContextMenu_SortMode_CanExecute(const EColumnSortMode::Type InSortMode) const
 {
+	if (bIsUpdateRunning)
+	{
+		return false;
+	}
+
 	return true; //ColumnSortMode != InSortMode;
 }
 
@@ -1681,6 +1802,11 @@ bool STableTreeView::ContextMenu_SortByColumn_IsChecked(const FName ColumnId)
 
 bool STableTreeView::ContextMenu_SortByColumn_CanExecute(const FName ColumnId) const
 {
+	if (bIsUpdateRunning)
+	{
+		return false;
+	}
+
 	return true; //ColumnId != ColumnBeingSorted;
 }
 
@@ -1759,6 +1885,11 @@ void STableTreeView::ShowColumn(const FName ColumnId)
 
 bool STableTreeView::CanHideColumn(const FName ColumnId) const
 {
+	if (bIsUpdateRunning)
+	{
+		return false;
+	}
+
 	const FTableColumn& Column = *Table->FindColumnChecked(ColumnId);
 	return Column.CanBeHidden();
 }
@@ -1787,6 +1918,11 @@ bool STableTreeView::IsColumnVisible(const FName ColumnId)
 
 bool STableTreeView::CanToggleColumnVisibility(const FName ColumnId) const
 {
+	if (bIsUpdateRunning)
+	{
+		return false;
+	}
+
 	const FTableColumn& Column = *Table->FindColumnChecked(ColumnId);
 	return !Column.IsVisible() || Column.CanBeHidden();
 }
@@ -1812,6 +1948,11 @@ void STableTreeView::ToggleColumnVisibility(const FName ColumnId)
 
 bool STableTreeView::ContextMenu_ShowAllColumns_CanExecute() const
 {
+	if (bIsUpdateRunning)
+	{
+		return false;
+	}
+
 	return true;
 }
 
@@ -1840,6 +1981,11 @@ void STableTreeView::ContextMenu_ShowAllColumns_Execute()
 
 bool STableTreeView::ContextMenu_ResetColumns_CanExecute() const
 {
+	if (bIsUpdateRunning)
+	{
+		return false;
+	}
+
 	return true;
 }
 
@@ -1907,6 +2053,31 @@ void STableTreeView::SelectNodeByTableRowIndex(int32 RowIndex)
 
 ////////////////////////////////////////////////////////////////////////////////////////////////////
 
+void STableTreeView::OnPreAsyncUpdate()
+{
+	check(!bIsUpdateRunning);
+
+	bIsUpdateRunning = true;
+	TreeView->SetTreeItemsSource(&DummyGroupNodes);
+	TreeView_Refresh();
+}
+
+////////////////////////////////////////////////////////////////////////////////////////////////////
+
+void STableTreeView::OnPostAsyncUpdate()
+{
+	check(bIsUpdateRunning);
+
+	bIsUpdateRunning = false;
+
+	RebuildGroupingCrumbs();
+
+	TreeView->SetTreeItemsSource(&FilteredGroupNodes);
+	TreeView_Refresh();
+}
+
+////////////////////////////////////////////////////////////////////////////////////////////////////
+
 } // namespace Insights
 
 #undef LOCTEXT_NAMESPACE
diff --git a/Engine/Source/Developer/TraceInsights/Private/Insights/Table/Widgets/STableTreeView.h b/Engine/Source/Developer/TraceInsights/Private/Insights/Table/Widgets/STableTreeView.h
--- a/Engine/Source/Developer/TraceInsights/Private/Insights/Table/Widgets/STableTreeView.h
+++ b/Engine/Source/Developer/TraceInsights/Private/Insights/Table/Widgets/STableTreeView.h
@@ -47,6 +47,10 @@ typedef TTextFilter<const FTableTreeNodePtr&> FTableTreeNodeTextFilter;
  */
 class STableTreeView : public SCompoundWidget
 {
+	friend class FTableTreeViewFilterAsyncTask;
+	friend class FTableTreeViewSortAsyncTask;
+	friend class FTableTreeViewGroupAsyncTask;
+
 public:
 	/** Default constructor. */
 	STableTreeView();
@@ -70,6 +74,15 @@ public:
 
 	void RebuildColumns();
 
+	/**
+	 * Ticks this widget.  Override in derived classes, but always call the parent implementation.
+	 *
+	 * @param  AllottedGeometry The space allotted for this widget
+	 * @param  InCurrentTime  Current absolute real time
+	 * @param  InDeltaTime  Real time passed since last tick
+	 */
+	virtual void Tick(const FGeometry& AllottedGeometry, const double InCurrentTime, const float InDeltaTime) override;
+
 	/**
 	 * Rebuilds the tree (if necessary).
 	 * @param bResync - If true, it forces a resync even if the list did not changed since last sync.
@@ -78,6 +91,7 @@ public:
 
 	FTableTreeNodePtr GetNodeByTableRowIndex(int32 RowIndex) const;
 	void SelectNodeByTableRowIndex(int32 RowIndex);
+	bool IsRunningAsyncUpdate() { return bIsUpdateRunning;  }
 
 protected:
 	void ConstructWidget(TSharedPtr<FTable> InTablePtr);
@@ -93,7 +107,7 @@ protected:
 	 * @param OutSearchStrings   - an array of strings to use in searching.
 	 *
 	 */
-	void HandleItemToStringArray(const FTableTreeNodePtr& GroupOrStatNodePtr, TArray<FString>& OutSearchStrings) const;
+	static void HandleItemToStringArray(const FTableTreeNodePtr& GroupOrStatNodePtr, TArray<FString>& OutSearchStrings);
 
 	////////////////////////////////////////////////////////////////////////////////////////////////////
 	// Tree View - Context Menu
@@ -244,6 +258,10 @@ protected:
 	bool ContextMenu_ResetColumns_CanExecute() const;
 	void ContextMenu_ResetColumns_Execute();
 
+	//Async
+	virtual void OnPreAsyncUpdate();
+	virtual void OnPostAsyncUpdate();
+
 	////////////////////////////////////////////////////////////////////////////////////////////////////
 
 protected:
@@ -334,6 +352,13 @@ protected:
 	/** How we sort the nodes? Ascending or Descending. */
 	EColumnSortMode::Type ColumnSortMode;
 
+	//////////////////////////////////////////////////
+	// Async
+	bool bRunInAsyncMode = false;
+	FGraphEventRef PendingAsyncOperationEvent;;
+	bool bIsUpdateRunning = false;
+	TArray<FTableTreeNodePtr> DummyGroupNodes;
+
 	//////////////////////////////////////////////////
 
 	double StatsStartTime;
@@ -342,4 +367,80 @@ protected:
 
 ////////////////////////////////////////////////////////////////////////////////////////////////////
 
+class FTableTreeViewFilterAsyncTask
+{
+public:
+	FTableTreeViewFilterAsyncTask(TSharedPtr<STableTreeView> InPtr)
+	{
+		TableTreeViewPtr = InPtr;
+	}
+
+	FORCEINLINE TStatId GetStatId() const { RETURN_QUICK_DECLARE_CYCLE_STAT(FTableTreeViewFilterAsyncTask, STATGROUP_TaskGraphTasks); }
+	ENamedThreads::Type GetDesiredThread() { return ENamedThreads::Type::AnyThread; }
+	static ESubsequentsMode::Type GetSubsequentsMode() { return ESubsequentsMode::TrackSubsequents; }
+
+	void DoTask(ENamedThreads::Type CurrentThread, const FGraphEventRef& MyCompletionGraphEvent)
+	{
+		if (TableTreeViewPtr.IsValid())
+		{
+			TableTreeViewPtr->ApplyFiltering();
+		}
+	}
+
+private:
+	TSharedPtr<STableTreeView> TableTreeViewPtr;
+};
+
+////////////////////////////////////////////////////////////////////////////////////////////////////
+
+class FTableTreeViewSortAsyncTask
+{
+public:
+	FTableTreeViewSortAsyncTask(TSharedPtr<STableTreeView> InPtr)
+	{
+		TableTreeViewPtr = InPtr;
+	}
+
+	FORCEINLINE TStatId GetStatId() const { RETURN_QUICK_DECLARE_CYCLE_STAT(FTableTreeViewSortAsyncTask, STATGROUP_TaskGraphTasks); }
+	ENamedThreads::Type GetDesiredThread() { return ENamedThreads::Type::AnyThread; }
+	static ESubsequentsMode::Type GetSubsequentsMode() { return ESubsequentsMode::TrackSubsequents; }
+
+	void DoTask(ENamedThreads::Type CurrentThread, const FGraphEventRef& MyCompletionGraphEvent)
+	{
+		if (TableTreeViewPtr.IsValid())
+		{
+			TableTreeViewPtr->SortTreeNodes();
+		}
+	}
+
+private:
+	TSharedPtr<STableTreeView> TableTreeViewPtr;
+};
+
+////////////////////////////////////////////////////////////////////////////////////////////////////
+
+class FTableTreeViewGroupAsyncTask
+{
+public:
+	FTableTreeViewGroupAsyncTask(TSharedPtr<STableTreeView> InPtr)
+	{
+		TableTreeViewPtr = InPtr;
+	}
+
+	FORCEINLINE TStatId GetStatId() const { RETURN_QUICK_DECLARE_CYCLE_STAT(FTableTreeViewGroupAsyncTask, STATGROUP_TaskGraphTasks); }
+	ENamedThreads::Type GetDesiredThread() { return ENamedThreads::Type::AnyThread; }
+	static ESubsequentsMode::Type GetSubsequentsMode() { return ESubsequentsMode::TrackSubsequents; }
+
+	void DoTask(ENamedThreads::Type CurrentThread, const FGraphEventRef& MyCompletionGraphEvent)
+	{
+		if (TableTreeViewPtr.IsValid())
+		{
+			TableTreeViewPtr->CreateGroups();
+		}
+	}
+
+private:
+	TSharedPtr<STableTreeView> TableTreeViewPtr;
+};
+
 } // namespace Insights
