[Insights] AllocTreeView can be safely closed during an async operation

diff --git a/Engine/Source/Developer/TraceInsights/Private/Insights/MemoryProfiler/Widgets/SMemoryProfilerWindow.cpp b/Engine/Source/Developer/TraceInsights/Private/Insights/MemoryProfiler/Widgets/SMemoryProfilerWindow.cpp
--- a/Engine/Source/Developer/TraceInsights/Private/Insights/MemoryProfiler/Widgets/SMemoryProfilerWindow.cpp
+++ b/Engine/Source/Developer/TraceInsights/Private/Insights/MemoryProfiler/Widgets/SMemoryProfilerWindow.cpp
@@ -341,6 +341,7 @@ void SMemoryProfilerWindow::OnMemAllocTableTreeViewTabClosed(TSharedRef<SDockTab
 
 	//FMemoryProfilerManager::Get()->SetMemAllocTableTreeViewVisible(MemAllocTableTreeView->GetTabIndex(), false);
 
+	MemAllocTableTreeView->OnClose();
 	MemAllocTableTreeViews.Remove(MemAllocTableTreeView);
 }
 
diff --git a/Engine/Source/Developer/TraceInsights/Private/Insights/Table/Widgets/STableTreeView.cpp b/Engine/Source/Developer/TraceInsights/Private/Insights/Table/Widgets/STableTreeView.cpp
--- a/Engine/Source/Developer/TraceInsights/Private/Insights/Table/Widgets/STableTreeView.cpp
+++ b/Engine/Source/Developer/TraceInsights/Private/Insights/Table/Widgets/STableTreeView.cpp
@@ -581,7 +581,7 @@ void STableTreeView::InsightsManager_OnSessionChanged()
 
 void STableTreeView::Tick(const FGeometry& AllottedGeometry, const double InCurrentTime, const float InDeltaTime)
 {
-	if (bRunInAsyncMode && bIsUpdateRunning)
+	if (bRunInAsyncMode && bIsUpdateRunning && !bIsCloseScheduled)
 	{
 		check(PendingAsyncOperationEvent.IsValid());
 		if (PendingAsyncOperationEvent->IsComplete())
@@ -602,15 +602,15 @@ void STableTreeView::UpdateTree()
 	{
 		OnPreAsyncUpdate();
 
-		FGraphEventRef GroupingCompletedEvent = TGraphTask<FTableTreeViewGroupAsyncTask>::CreateTask().ConstructAndDispatchWhenReady(SharedThis(this));
+		FGraphEventRef GroupingCompletedEvent = TGraphTask<FTableTreeViewGroupAsyncTask>::CreateTask().ConstructAndDispatchWhenReady(this);
 		
-		FGraphEventArray Prerequistes;
-		Prerequistes.Add(GroupingCompletedEvent);
-		FGraphEventRef SortingCompletedEvent = TGraphTask<FTableTreeViewSortAsyncTask>::CreateTask(&Prerequistes).ConstructAndDispatchWhenReady(SharedThis(this));
+		FGraphEventArray Prerequisites;
+		Prerequisites.Add(GroupingCompletedEvent);
+		FGraphEventRef SortingCompletedEvent = TGraphTask<FTableTreeViewSortAsyncTask>::CreateTask(&Prerequisites).ConstructAndDispatchWhenReady(this);
 		
-		Prerequistes.Empty();
-		Prerequistes.Add(SortingCompletedEvent);
-		FGraphEventRef FilteringCompletedEvent = TGraphTask<FTableTreeViewFilterAsyncTask>::CreateTask(&Prerequistes).ConstructAndDispatchWhenReady(SharedThis(this));
+		Prerequisites.Empty();
+		Prerequisites.Add(SortingCompletedEvent);
+		FGraphEventRef FilteringCompletedEvent = TGraphTask<FTableTreeViewFilterAsyncTask>::CreateTask(&Prerequisites).ConstructAndDispatchWhenReady(this);
 		
 		PendingAsyncOperationEvent = FilteringCompletedEvent;
 	}
@@ -879,7 +879,7 @@ void STableTreeView::SearchBox_OnTextChanged(const FText& InFilterText)
 	{
 		OnPreAsyncUpdate();
 
-		FGraphEventRef FilteringCompletedEvent = TGraphTask<FTableTreeViewFilterAsyncTask>::CreateTask().ConstructAndDispatchWhenReady(SharedThis(this));
+		FGraphEventRef FilteringCompletedEvent = TGraphTask<FTableTreeViewFilterAsyncTask>::CreateTask().ConstructAndDispatchWhenReady(this);
 
 		PendingAsyncOperationEvent = FilteringCompletedEvent;
 	}
@@ -1196,15 +1196,15 @@ void STableTreeView::PostChangeGroupings()
 	{
 		OnPreAsyncUpdate();
 
-		FGraphEventRef GroupingCompletedEvent = TGraphTask<FTableTreeViewGroupAsyncTask>::CreateTask().ConstructAndDispatchWhenReady(SharedThis(this));
+		FGraphEventRef GroupingCompletedEvent = TGraphTask<FTableTreeViewGroupAsyncTask>::CreateTask().ConstructAndDispatchWhenReady(this);
 
-		FGraphEventArray Prerequistes;
-		Prerequistes.Add(GroupingCompletedEvent);
-		FGraphEventRef SortingCompletedEvent = TGraphTask<FTableTreeViewSortAsyncTask>::CreateTask(&Prerequistes).ConstructAndDispatchWhenReady(SharedThis(this));
+		FGraphEventArray Prerequisites;
+		Prerequisites.Add(GroupingCompletedEvent);
+		FGraphEventRef SortingCompletedEvent = TGraphTask<FTableTreeViewSortAsyncTask>::CreateTask(&Prerequisites).ConstructAndDispatchWhenReady(this);
 
-		Prerequistes.Empty();
-		Prerequistes.Add(SortingCompletedEvent);
-		FGraphEventRef FilteringCompletedEvent = TGraphTask<FTableTreeViewFilterAsyncTask>::CreateTask(&Prerequistes).ConstructAndDispatchWhenReady(SharedThis(this));
+		Prerequisites.Empty();
+		Prerequisites.Add(SortingCompletedEvent);
+		FGraphEventRef FilteringCompletedEvent = TGraphTask<FTableTreeViewFilterAsyncTask>::CreateTask(&Prerequisites).ConstructAndDispatchWhenReady(this);
 
 		PendingAsyncOperationEvent = FilteringCompletedEvent;
 	}
@@ -1702,11 +1702,11 @@ void STableTreeView::SetSortModeForColumn(const FName& ColumnId, const EColumnSo
 	{
 		OnPreAsyncUpdate();
 
-		FGraphEventRef SortingCompletedEvent = TGraphTask<FTableTreeViewSortAsyncTask>::CreateTask().ConstructAndDispatchWhenReady(SharedThis(this));
+		FGraphEventRef SortingCompletedEvent = TGraphTask<FTableTreeViewSortAsyncTask>::CreateTask().ConstructAndDispatchWhenReady(this);
 
-		FGraphEventArray Prerequistes;
-		Prerequistes.Add(SortingCompletedEvent);
-		FGraphEventRef FilteringCompletedEvent = TGraphTask<FTableTreeViewFilterAsyncTask>::CreateTask(&Prerequistes).ConstructAndDispatchWhenReady(SharedThis(this));
+		FGraphEventArray Prerequisites;
+		Prerequisites.Add(SortingCompletedEvent);
+		FGraphEventRef FilteringCompletedEvent = TGraphTask<FTableTreeViewFilterAsyncTask>::CreateTask(&Prerequisites).ConstructAndDispatchWhenReady(this);
 
 		PendingAsyncOperationEvent = FilteringCompletedEvent;
 	}
@@ -2078,6 +2078,18 @@ void STableTreeView::OnPostAsyncUpdate()
 
 ////////////////////////////////////////////////////////////////////////////////////////////////////
 
+void STableTreeView::OnClose()
+{
+	if (bIsUpdateRunning && PendingAsyncOperationEvent.IsValid() && !PendingAsyncOperationEvent->IsComplete())
+	{
+		bIsCloseScheduled = true;
+
+		FGraphEventArray Prerequisites;
+		Prerequisites.Add(PendingAsyncOperationEvent);
+		TGraphTask<FTableTreeViewAsyncCompleteTask>::CreateTask(&Prerequisites).ConstructAndDispatchWhenReady(SharedThis(this));
+	}
+}
+
 } // namespace Insights
 
 #undef LOCTEXT_NAMESPACE
diff --git a/Engine/Source/Developer/TraceInsights/Private/Insights/Table/Widgets/STableTreeView.h b/Engine/Source/Developer/TraceInsights/Private/Insights/Table/Widgets/STableTreeView.h
--- a/Engine/Source/Developer/TraceInsights/Private/Insights/Table/Widgets/STableTreeView.h
+++ b/Engine/Source/Developer/TraceInsights/Private/Insights/Table/Widgets/STableTreeView.h
@@ -93,6 +93,8 @@ public:
 	void SelectNodeByTableRowIndex(int32 RowIndex);
 	bool IsRunningAsyncUpdate() { return bIsUpdateRunning;  }
 
+	void OnClose();
+
 protected:
 	void ConstructWidget(TSharedPtr<FTable> InTablePtr);
 	void UpdateTree();
@@ -358,6 +360,7 @@ protected:
 	FGraphEventRef PendingAsyncOperationEvent;;
 	bool bIsUpdateRunning = false;
 	TArray<FTableTreeNodePtr> DummyGroupNodes;
+	bool bIsCloseScheduled = false;
 
 	//////////////////////////////////////////////////
 
@@ -370,7 +373,7 @@ protected:
 class FTableTreeViewFilterAsyncTask
 {
 public:
-	FTableTreeViewFilterAsyncTask(TSharedPtr<STableTreeView> InPtr)
+	FTableTreeViewFilterAsyncTask(STableTreeView* InPtr)
 	{
 		TableTreeViewPtr = InPtr;
 	}
@@ -381,14 +384,14 @@ public:
 
 	void DoTask(ENamedThreads::Type CurrentThread, const FGraphEventRef& MyCompletionGraphEvent)
 	{
-		if (TableTreeViewPtr.IsValid())
+		if (TableTreeViewPtr)
 		{
 			TableTreeViewPtr->ApplyFiltering();
 		}
 	}
 
 private:
-	TSharedPtr<STableTreeView> TableTreeViewPtr;
+	STableTreeView* TableTreeViewPtr = nullptr;
 };
 
 ////////////////////////////////////////////////////////////////////////////////////////////////////
@@ -396,7 +399,7 @@ private:
 class FTableTreeViewSortAsyncTask
 {
 public:
-	FTableTreeViewSortAsyncTask(TSharedPtr<STableTreeView> InPtr)
+	FTableTreeViewSortAsyncTask(STableTreeView* InPtr)
 	{
 		TableTreeViewPtr = InPtr;
 	}
@@ -407,14 +410,14 @@ public:
 
 	void DoTask(ENamedThreads::Type CurrentThread, const FGraphEventRef& MyCompletionGraphEvent)
 	{
-		if (TableTreeViewPtr.IsValid())
+		if (TableTreeViewPtr)
 		{
 			TableTreeViewPtr->SortTreeNodes();
 		}
 	}
 
 private:
-	TSharedPtr<STableTreeView> TableTreeViewPtr;
+	STableTreeView* TableTreeViewPtr = nullptr;
 };
 
 ////////////////////////////////////////////////////////////////////////////////////////////////////
@@ -422,7 +425,7 @@ private:
 class FTableTreeViewGroupAsyncTask
 {
 public:
-	FTableTreeViewGroupAsyncTask(TSharedPtr<STableTreeView> InPtr)
+	FTableTreeViewGroupAsyncTask(STableTreeView* InPtr)
 	{
 		TableTreeViewPtr = InPtr;
 	}
@@ -433,12 +436,39 @@ public:
 
 	void DoTask(ENamedThreads::Type CurrentThread, const FGraphEventRef& MyCompletionGraphEvent)
 	{
-		if (TableTreeViewPtr.IsValid())
+		if (TableTreeViewPtr)
 		{
 			TableTreeViewPtr->CreateGroups();
 		}
 	}
 
+private:
+	STableTreeView* TableTreeViewPtr = nullptr;
+};
+
+////////////////////////////////////////////////////////////////////////////////////////////////////
+
+class FTableTreeViewAsyncCompleteTask
+{
+public:
+	FTableTreeViewAsyncCompleteTask(TSharedPtr<STableTreeView> InPtr)
+	{
+		TableTreeViewPtr = InPtr;
+	}
+
+	FORCEINLINE TStatId GetStatId() const { RETURN_QUICK_DECLARE_CYCLE_STAT(FTableTreeViewAsyncCompleteTask, STATGROUP_TaskGraphTasks); }
+	ENamedThreads::Type GetDesiredThread() { return ENamedThreads::Type::GameThread; }
+	static ESubsequentsMode::Type GetSubsequentsMode() { return ESubsequentsMode::TrackSubsequents; }
+
+	void DoTask(ENamedThreads::Type CurrentThread, const FGraphEventRef& MyCompletionGraphEvent)
+	{
+		// The role of this task is to keep the STableTreeView object alive until the task and it's prerequisits are completed and to destroy it on the game thread.
+		if (TableTreeViewPtr.IsValid())
+		{
+			TableTreeViewPtr.Reset();
+		}
+	}
+
 private:
 	TSharedPtr<STableTreeView> TableTreeViewPtr;
 };
