[Insights] Preserve Allocation Tree View nodes expanded status after async fitering and sorting operations

diff --git a/Engine/Source/Developer/TraceInsights/Private/Insights/Table/Widgets/STableTreeView.cpp b/Engine/Source/Developer/TraceInsights/Private/Insights/Table/Widgets/STableTreeView.cpp
--- a/Engine/Source/Developer/TraceInsights/Private/Insights/Table/Widgets/STableTreeView.cpp
+++ b/Engine/Source/Developer/TraceInsights/Private/Insights/Table/Widgets/STableTreeView.cpp
@@ -602,17 +602,9 @@ void STableTreeView::UpdateTree()
 	{
 		OnPreAsyncUpdate();
 
-		FGraphEventRef GroupingCompletedEvent = TGraphTask<FTableTreeViewGroupAsyncTask>::CreateTask().ConstructAndDispatchWhenReady(this);
-		
-		FGraphEventArray Prerequisites;
-		Prerequisites.Add(GroupingCompletedEvent);
-		FGraphEventRef SortingCompletedEvent = TGraphTask<FTableTreeViewSortAsyncTask>::CreateTask(&Prerequisites).ConstructAndDispatchWhenReady(this);
-		
-		Prerequisites.Empty();
-		Prerequisites.Add(SortingCompletedEvent);
-		FGraphEventRef FilteringCompletedEvent = TGraphTask<FTableTreeViewFilterAsyncTask>::CreateTask(&Prerequisites).ConstructAndDispatchWhenReady(this);
-		
-		PendingAsyncOperationEvent = FilteringCompletedEvent;
+		FGraphEventRef CompletedEvent = StartCreateGroupsTask();
+		CompletedEvent = StartSortTreeNodesTask(CompletedEvent);
+		PendingAsyncOperationEvent = StartApplyFiltersTask(CompletedEvent);
 	}
 	else
 	{
@@ -879,9 +871,7 @@ void STableTreeView::SearchBox_OnTextChanged(const FText& InFilterText)
 	{
 		OnPreAsyncUpdate();
 
-		FGraphEventRef FilteringCompletedEvent = TGraphTask<FTableTreeViewFilterAsyncTask>::CreateTask().ConstructAndDispatchWhenReady(this);
-
-		PendingAsyncOperationEvent = FilteringCompletedEvent;
+		PendingAsyncOperationEvent = StartApplyFiltersTask();
 	}
 	else
 	{
@@ -1196,17 +1186,9 @@ void STableTreeView::PostChangeGroupings()
 	{
 		OnPreAsyncUpdate();
 
-		FGraphEventRef GroupingCompletedEvent = TGraphTask<FTableTreeViewGroupAsyncTask>::CreateTask().ConstructAndDispatchWhenReady(this);
-
-		FGraphEventArray Prerequisites;
-		Prerequisites.Add(GroupingCompletedEvent);
-		FGraphEventRef SortingCompletedEvent = TGraphTask<FTableTreeViewSortAsyncTask>::CreateTask(&Prerequisites).ConstructAndDispatchWhenReady(this);
-
-		Prerequisites.Empty();
-		Prerequisites.Add(SortingCompletedEvent);
-		FGraphEventRef FilteringCompletedEvent = TGraphTask<FTableTreeViewFilterAsyncTask>::CreateTask(&Prerequisites).ConstructAndDispatchWhenReady(this);
-
-		PendingAsyncOperationEvent = FilteringCompletedEvent;
+		FGraphEventRef CompletedEvent = StartCreateGroupsTask();
+		CompletedEvent = StartSortTreeNodesTask(CompletedEvent);
+		PendingAsyncOperationEvent = StartApplyFiltersTask(CompletedEvent);
 	}
 	else
 	{
@@ -1702,13 +1684,8 @@ void STableTreeView::SetSortModeForColumn(const FName& ColumnId, const EColumnSo
 	{
 		OnPreAsyncUpdate();
 
-		FGraphEventRef SortingCompletedEvent = TGraphTask<FTableTreeViewSortAsyncTask>::CreateTask().ConstructAndDispatchWhenReady(this);
-
-		FGraphEventArray Prerequisites;
-		Prerequisites.Add(SortingCompletedEvent);
-		FGraphEventRef FilteringCompletedEvent = TGraphTask<FTableTreeViewFilterAsyncTask>::CreateTask(&Prerequisites).ConstructAndDispatchWhenReady(this);
-
-		PendingAsyncOperationEvent = FilteringCompletedEvent;
+		FGraphEventRef CompletedEvent = StartSortTreeNodesTask();
+		PendingAsyncOperationEvent = StartApplyFiltersTask(CompletedEvent);
 	}
 	else
 	{
@@ -2058,6 +2035,10 @@ void STableTreeView::OnPreAsyncUpdate()
 	check(!bIsUpdateRunning);
 
 	bIsUpdateRunning = true;
+
+	ExpandedNodes.Empty();
+	TreeView->GetExpandedItems(ExpandedNodes);
+
 	TreeView->SetTreeItemsSource(&DummyGroupNodes);
 	TreeView_Refresh();
 }
@@ -2073,11 +2054,78 @@ void STableTreeView::OnPostAsyncUpdate()
 	RebuildGroupingCrumbs();
 
 	TreeView->SetTreeItemsSource(&FilteredGroupNodes);
+
+	TreeView->ClearExpandedItems();
+
+	// Grouping can result in old group nodes no longer existing in the tree, so we don't keep the expanded list.
+	if (!HasPendingAsyncOperation(EAsyncOperationType::GroupingOp))
+	{
+		for (auto It = ExpandedNodes.CreateConstIterator(); It; ++It)
+		{
+			TreeView->SetItemExpansion(*It, true);
+		}
+	}
+
+	ClearPendingAsyncOperations();
 	TreeView_Refresh();
 }
 
 ////////////////////////////////////////////////////////////////////////////////////////////////////
 
+FGraphEventRef STableTreeView::StartSortTreeNodesTask(FGraphEventRef Prerequisite)
+{
+	AddPendingAsyncOperation(EAsyncOperationType::SortingOp);
+
+	if (Prerequisite.IsValid())
+	{
+		FGraphEventArray Prerequisites;
+		Prerequisites.Add(Prerequisite);
+		return TGraphTask<FTableTreeViewSortAsyncTask>::CreateTask(&Prerequisites).ConstructAndDispatchWhenReady(this);
+	}
+	else
+	{
+		return TGraphTask<FTableTreeViewSortAsyncTask>::CreateTask().ConstructAndDispatchWhenReady(this);
+	}
+}
+
+////////////////////////////////////////////////////////////////////////////////////////////////////
+
+FGraphEventRef STableTreeView::StartCreateGroupsTask(FGraphEventRef Prerequisite)
+{
+	AddPendingAsyncOperation(EAsyncOperationType::GroupingOp);
+
+	if (Prerequisite.IsValid())
+	{
+		FGraphEventArray Prerequisites;
+		Prerequisites.Add(Prerequisite);
+		return TGraphTask<FTableTreeViewGroupAsyncTask>::CreateTask(&Prerequisites).ConstructAndDispatchWhenReady(this);
+	}
+	else
+	{
+		return TGraphTask<FTableTreeViewGroupAsyncTask>::CreateTask().ConstructAndDispatchWhenReady(this);
+	}
+}
+
+////////////////////////////////////////////////////////////////////////////////////////////////////
+
+FGraphEventRef STableTreeView::StartApplyFiltersTask(FGraphEventRef Prerequisite)
+{
+	AddPendingAsyncOperation(EAsyncOperationType::FilteringOp);
+
+	if (Prerequisite.IsValid())
+	{
+		FGraphEventArray Prerequisites;
+		Prerequisites.Add(Prerequisite);
+		return TGraphTask<FTableTreeViewFilterAsyncTask>::CreateTask(&Prerequisites).ConstructAndDispatchWhenReady(this);
+	}
+	else
+	{
+		return TGraphTask<FTableTreeViewFilterAsyncTask>::CreateTask().ConstructAndDispatchWhenReady(this);
+	}
+}
+
+////////////////////////////////////////////////////////////////////////////////////////////////////
+
 void STableTreeView::OnClose()
 {
 	if (bIsUpdateRunning && PendingAsyncOperationEvent.IsValid() && !PendingAsyncOperationEvent->IsComplete())
@@ -2090,6 +2138,8 @@ void STableTreeView::OnClose()
 	}
 }
 
+////////////////////////////////////////////////////////////////////////////////////////////////////
+
 } // namespace Insights
 
 #undef LOCTEXT_NAMESPACE
diff --git a/Engine/Source/Developer/TraceInsights/Private/Insights/Table/Widgets/STableTreeView.h b/Engine/Source/Developer/TraceInsights/Private/Insights/Table/Widgets/STableTreeView.h
--- a/Engine/Source/Developer/TraceInsights/Private/Insights/Table/Widgets/STableTreeView.h
+++ b/Engine/Source/Developer/TraceInsights/Private/Insights/Table/Widgets/STableTreeView.h
@@ -3,6 +3,7 @@
 #pragma once
 
 #include "CoreMinimal.h"
+#include "Misc/EnumClassFlags.h"
 #include "Misc/FilterCollection.h"
 #include "Misc/TextFilter.h"
 #include "SlateFwd.h"
@@ -41,6 +42,17 @@ typedef TFilterCollection<const FTableTreeNodePtr&> FTableTreeNodeFilterCollecti
 /** The text based filter - used for updating the list of tree nodes. */
 typedef TTextFilter<const FTableTreeNodePtr&> FTableTreeNodeTextFilter;
 
+////////////////////////////////////////////////////////////////////////////////////////////////////
+
+enum class EAsyncOperationType : uint32
+{
+	FilteringOp = 1,
+	SortingOp = 1 << 1,
+	GroupingOp = 1 << 2,
+};
+
+ENUM_CLASS_FLAGS(EAsyncOperationType);
+
 ////////////////////////////////////////////////////////////////////////////////////////////////////
 /**
  * A custom widget used to display the list of tree nodes.
@@ -264,6 +276,14 @@ protected:
 	virtual void OnPreAsyncUpdate();
 	virtual void OnPostAsyncUpdate();
 
+	FGraphEventRef StartSortTreeNodesTask(FGraphEventRef Prerequisite = nullptr);
+	FGraphEventRef StartCreateGroupsTask(FGraphEventRef Prerequisite = nullptr);
+	FGraphEventRef StartApplyFiltersTask(FGraphEventRef Prerequisite = nullptr);
+
+	void AddPendingAsyncOperation(EAsyncOperationType InType) { EnumAddFlags(PendingAsyncOperations, InType); }
+	bool HasPendingAsyncOperation(EAsyncOperationType InType) const { return EnumHasAnyFlags(PendingAsyncOperations, InType); }
+	void ClearPendingAsyncOperations() { PendingAsyncOperations = static_cast<EAsyncOperationType>(0); }
+
 	////////////////////////////////////////////////////////////////////////////////////////////////////
 
 protected:
@@ -361,7 +381,7 @@ protected:
 	bool bIsUpdateRunning = false;
 	TArray<FTableTreeNodePtr> DummyGroupNodes;
 	bool bIsCloseScheduled = false;
-
+	EAsyncOperationType PendingAsyncOperations = static_cast<EAsyncOperationType>(0);
 	//////////////////////////////////////////////////
 
 	double StatsStartTime;
