[Insights] Forward -stompmalloc cmd line paramater to the session viewer
Remove option to run Session Viewer in the same process as the SessionBrowser

diff --git a/Engine/Source/Developer/TraceInsights/Private/Insights/InsightsManager.cpp b/Engine/Source/Developer/TraceInsights/Private/Insights/InsightsManager.cpp
--- a/Engine/Source/Developer/TraceInsights/Private/Insights/InsightsManager.cpp
+++ b/Engine/Source/Developer/TraceInsights/Private/Insights/InsightsManager.cpp
@@ -121,11 +121,6 @@ FInsightsManager::FInsightsManager(TSharedRef<TraceServices::IAnalysisService> I
 	, ActionManager(this)
 	, Settings()
 	, bIsDebugInfoEnabled(false)
-#if WITH_EDITOR
-	, bShouldOpenAnalysisInSeparateProcess(false)
-#else
-	, bShouldOpenAnalysisInSeparateProcess(true)
-#endif
 	, AnalysisStopwatch()
 	, bIsAnalysisComplete(false)
 	, SessionDuration(0.0)
diff --git a/Engine/Source/Developer/TraceInsights/Private/Insights/InsightsManager.h b/Engine/Source/Developer/TraceInsights/Private/Insights/InsightsManager.h
--- a/Engine/Source/Developer/TraceInsights/Private/Insights/InsightsManager.h
+++ b/Engine/Source/Developer/TraceInsights/Private/Insights/InsightsManager.h
@@ -177,9 +177,6 @@ public:
 
 	////////////////////////////////////////////////////////////////////////////////////////////////////
 
-	bool ShouldOpenAnalysisInSeparateProcess() const { return bShouldOpenAnalysisInSeparateProcess; }
-	void SetOpenAnalysisInSeparateProcess(bool bOnOff) { bShouldOpenAnalysisInSeparateProcess = bOnOff; }
-
 	/** Creates a new analysis session instance and loads the latest available trace that is live. */
 	void LoadLastLiveSession();
 
@@ -306,8 +303,6 @@ private:
 	/** If enabled, UI can display additional info for debugging purposes. */
 	bool bIsDebugInfoEnabled;
 
-	bool bShouldOpenAnalysisInSeparateProcess;
-
 	FStopwatch AnalysisStopwatch;
 	bool bIsAnalysisComplete;
 	double SessionDuration;
diff --git a/Engine/Source/Developer/TraceInsights/Private/Insights/TraceInsightsModule.cpp b/Engine/Source/Developer/TraceInsights/Private/Insights/TraceInsightsModule.cpp
--- a/Engine/Source/Developer/TraceInsights/Private/Insights/TraceInsightsModule.cpp
+++ b/Engine/Source/Developer/TraceInsights/Private/Insights/TraceInsightsModule.cpp
@@ -158,10 +158,8 @@ void FTraceInsightsModule::UnregisterTabSpawners()
 
 ////////////////////////////////////////////////////////////////////////////////////////////////////
 
-void FTraceInsightsModule::CreateSessionBrowser(bool bAllowDebugTools, bool bSingleProcess, bool bInitializeTesting)
+void FTraceInsightsModule::CreateSessionBrowser(const FCreateSessionBrowserParams& Params)
 {
-	FInsightsManager::Get()->SetOpenAnalysisInSeparateProcess(!bSingleProcess);
-
 	RegisterTabSpawners();
 
 	const float DPIScaleFactor = FPlatformApplicationMisc::GetDPIScaleFactorAtPoint(10.0f, 10.0f);
@@ -170,43 +168,19 @@ void FTraceInsightsModule::CreateSessionBrowser(bool bAllowDebugTools, bool bSin
 	float WindowWidth = 0.0f;
 	float WindowHeight = 0.0f;
 
-	if (!bSingleProcess)
-	{
-		WindowWidth = 920.0f;
-		WindowHeight = 665.0f;
+	WindowWidth = 920.0f;
+	WindowHeight = 665.0f;
 
-		DefaultLayout->AddArea
-		(
-			FTabManager::NewPrimaryArea()
-			->Split
-			(
-				FTabManager::NewStack()
-				->AddTab(FInsightsManagerTabs::StartPageTabId, ETabState::OpenedTab)
-			)
-		);
-	}
-	else
-	{
-		WindowWidth = 1280.0f;
-		WindowHeight = 720.0f;
-
-		DefaultLayout->AddArea
+	DefaultLayout->AddArea
+	(
+		FTabManager::NewPrimaryArea()
+		->Split
 		(
-			FTabManager::NewPrimaryArea()
-			->Split
-			(
-				FTabManager::NewStack()
-				->AddTab(FInsightsManagerTabs::StartPageTabId, ETabState::OpenedTab)
-				->AddTab(FInsightsManagerTabs::SessionInfoTabId, ETabState::ClosedTab)
-				->AddTab(FInsightsManagerTabs::TimingProfilerTabId, ETabState::ClosedTab)
-				->AddTab(FInsightsManagerTabs::LoadingProfilerTabId, ETabState::ClosedTab)
-				->AddTab(FInsightsManagerTabs::NetworkingProfilerTabId, ETabState::ClosedTab)
-				->AddTab(FInsightsManagerTabs::MemoryProfilerTabId, ETabState::ClosedTab)
-				->SetForegroundTab(FTabId(FInsightsManagerTabs::StartPageTabId))
-			)
-		);
-	}
-
+			FTabManager::NewStack()
+			->AddTab(FInsightsManagerTabs::StartPageTabId, ETabState::OpenedTab)
+		)
+	);
+	
 	TSharedRef<SWindow> RootWindow = SNew(SWindow)
 		.AutoCenter(EAutoCenter::PreferredWorkArea)
 		.Title(NSLOCTEXT("TraceInsightsModule", "UnrealInsightsBrowserAppName", "Unreal Insights Session Browser"))
@@ -220,7 +194,7 @@ void FTraceInsightsModule::CreateSessionBrowser(bool bAllowDebugTools, bool bSin
 	FGlobalTabmanager::Get()->SetRootWindow(RootWindow);
 	FSlateNotificationManager::Get().SetRootWindow(RootWindow);
 
-	AddAreaForWidgetReflector(DefaultLayout, bAllowDebugTools);
+	AddAreaForWidgetReflector(DefaultLayout, Params.bAllowDebugTools);
 
 	// Load layout from ini file.
 	PersistentLayout = FLayoutSaveRestore::LoadFromConfig(UnrealInsightsLayoutIni, DefaultLayout);
@@ -235,8 +209,9 @@ void FTraceInsightsModule::CreateSessionBrowser(bool bAllowDebugTools, bool bSin
 	TSharedPtr<class SStartPageWindow> StartPageWnd = FInsightsManager::Get()->GetStartPageWindow();
 	if (StartPageWnd.IsValid())
 	{
-		StartPageWnd->SetEnableAutomaticTesting(bInitializeTesting);
-		StartPageWnd->SetEnableDebugTools(bAllowDebugTools);
+		StartPageWnd->SetEnableAutomaticTesting(Params.bInitializeTesting);
+		StartPageWnd->SetEnableDebugTools(Params.bAllowDebugTools);
+		StartPageWnd->SetStartProcessWithStompMalloc(Params.bStartProcessWithStompMalloc);
 	}
 }
 
diff --git a/Engine/Source/Developer/TraceInsights/Private/Insights/TraceInsightsModule.h b/Engine/Source/Developer/TraceInsights/Private/Insights/TraceInsightsModule.h
--- a/Engine/Source/Developer/TraceInsights/Private/Insights/TraceInsightsModule.h
+++ b/Engine/Source/Developer/TraceInsights/Private/Insights/TraceInsightsModule.h
@@ -43,7 +43,7 @@ public:
 	virtual UE::Trace::FStoreClient* GetStoreClient() override;
 	virtual bool ConnectToStore(const TCHAR* InStoreHost, uint32 InStorePort) override;
 
-	virtual void CreateSessionBrowser(bool bAllowDebugTools = false, bool bSingleProcess = false, bool bInitializeTesting = false) override;
+	virtual void CreateSessionBrowser(const FCreateSessionBrowserParams& Params) override;
 	virtual void CreateSessionViewer(bool bAllowDebugTools = false) override;
 
 	virtual TSharedPtr<const TraceServices::IAnalysisSession> GetAnalysisSession() const override;
diff --git a/Engine/Source/Developer/TraceInsights/Private/Insights/Widgets/SStartPageWindow.cpp b/Engine/Source/Developer/TraceInsights/Private/Insights/Widgets/SStartPageWindow.cpp
--- a/Engine/Source/Developer/TraceInsights/Private/Insights/Widgets/SStartPageWindow.cpp
+++ b/Engine/Source/Developer/TraceInsights/Private/Insights/Widgets/SStartPageWindow.cpp
@@ -1611,92 +1611,76 @@ void SStartPageWindow::LoadTraceSession(TSharedPtr<FTraceViewModel> InTraceSessi
 
 void SStartPageWindow::LoadTraceFile(const FString& InTraceFile)
 {
-	if (FInsightsManager::Get()->ShouldOpenAnalysisInSeparateProcess())
-	{
-		UE_LOG(TraceInsights, Log, TEXT("[StartPage] Start analysis (in separate process) for trace file: \"%s\""), *InTraceFile);
+	UE_LOG(TraceInsights, Log, TEXT("[StartPage] Start analysis (in separate process) for trace file: \"%s\""), *InTraceFile);
 
-		const TCHAR* ExecutablePath = FPlatformProcess::ExecutablePath();
+	const TCHAR* ExecutablePath = FPlatformProcess::ExecutablePath();
 
-		FString CmdLine = TEXT("-OpenTraceFile=\"") + InTraceFile + TEXT("\"");
+	FString CmdLine = TEXT("-OpenTraceFile=\"") + InTraceFile + TEXT("\"");
 
-		FString ExtraCmdParams;
-		GetExtraCommandLineParams(ExtraCmdParams);
-		CmdLine += ExtraCmdParams;
+	FString ExtraCmdParams;
+	GetExtraCommandLineParams(ExtraCmdParams);
+	CmdLine += ExtraCmdParams;
 
-		constexpr bool bLaunchDetached = true;
-		constexpr bool bLaunchHidden = false;
-		constexpr bool bLaunchReallyHidden = false;
+	constexpr bool bLaunchDetached = true;
+	constexpr bool bLaunchHidden = false;
+	constexpr bool bLaunchReallyHidden = false;
 
-		uint32 ProcessID = 0;
-		const int32 PriorityModifier = 0;
-		const TCHAR* OptionalWorkingDirectory = nullptr;
+	uint32 ProcessID = 0;
+	const int32 PriorityModifier = 0;
+	const TCHAR* OptionalWorkingDirectory = nullptr;
 
-		void* PipeWriteChild = nullptr;
-		void* PipeReadChild = nullptr;
+	void* PipeWriteChild = nullptr;
+	void* PipeReadChild = nullptr;
 
-		FProcHandle Handle = FPlatformProcess::CreateProc(ExecutablePath, *CmdLine, bLaunchDetached, bLaunchHidden, bLaunchReallyHidden, &ProcessID, PriorityModifier, OptionalWorkingDirectory, PipeWriteChild, PipeReadChild);
-		if (Handle.IsValid())
-		{
-			FPlatformProcess::CloseProc(Handle);
-		}
-
-		SplashScreenOverlayTraceFile = FPaths::GetBaseFilename(InTraceFile);
-		ShowSplashScreenOverlay();
-	}
-	else
+	FProcHandle Handle = FPlatformProcess::CreateProc(ExecutablePath, *CmdLine, bLaunchDetached, bLaunchHidden, bLaunchReallyHidden, &ProcessID, PriorityModifier, OptionalWorkingDirectory, PipeWriteChild, PipeReadChild);
+	if (Handle.IsValid())
 	{
-		UE_LOG(TraceInsights, Log, TEXT("[StartPage] Start analysis for trace file: \"%s\""), *InTraceFile);
-		FInsightsManager::Get()->LoadTraceFile(InTraceFile);
+		FPlatformProcess::CloseProc(Handle);
 	}
+
+	SplashScreenOverlayTraceFile = FPaths::GetBaseFilename(InTraceFile);
+	ShowSplashScreenOverlay();
 }
 
 ////////////////////////////////////////////////////////////////////////////////////////////////////
 
 void SStartPageWindow::LoadTrace(uint32 InTraceId)
 {
-	if (FInsightsManager::Get()->ShouldOpenAnalysisInSeparateProcess())
-	{
-		UE_LOG(TraceInsights, Log, TEXT("[StartPage] Start analysis (in separate process) for trace id: 0x%08X"), InTraceId);
+	UE_LOG(TraceInsights, Log, TEXT("[StartPage] Start analysis (in separate process) for trace id: 0x%08X"), InTraceId);
 
-		const TCHAR* ExecutablePath = FPlatformProcess::ExecutablePath();
+	const TCHAR* ExecutablePath = FPlatformProcess::ExecutablePath();
 
-		const uint32 StorePort = FInsightsManager::Get()->GetStoreClient()->GetStorePort();
-		FString CmdLine = FString::Printf(TEXT("-OpenTraceId=%d -StorePort=%d"), InTraceId, StorePort);
+	const uint32 StorePort = FInsightsManager::Get()->GetStoreClient()->GetStorePort();
+	FString CmdLine = FString::Printf(TEXT("-OpenTraceId=%d -StorePort=%d"), InTraceId, StorePort);
 		
-		FString ExtraCmdParams;
-		GetExtraCommandLineParams(ExtraCmdParams);
-		CmdLine += ExtraCmdParams;
-
-		constexpr bool bLaunchDetached = true;
-		constexpr bool bLaunchHidden = false;
-		constexpr bool bLaunchReallyHidden = false;
+	FString ExtraCmdParams;
+	GetExtraCommandLineParams(ExtraCmdParams);
+	CmdLine += ExtraCmdParams;
 
-		uint32 ProcessID = 0;
-		const int32 PriorityModifier = 0;
-		const TCHAR* OptionalWorkingDirectory = nullptr;
+	constexpr bool bLaunchDetached = true;
+	constexpr bool bLaunchHidden = false;
+	constexpr bool bLaunchReallyHidden = false;
 
-		void* PipeWriteChild = nullptr;
-		void* PipeReadChild = nullptr;
+	uint32 ProcessID = 0;
+	const int32 PriorityModifier = 0;
+	const TCHAR* OptionalWorkingDirectory = nullptr;
 
-		FProcHandle Handle = FPlatformProcess::CreateProc(ExecutablePath, *CmdLine, bLaunchDetached, bLaunchHidden, bLaunchReallyHidden, &ProcessID, PriorityModifier, OptionalWorkingDirectory, PipeWriteChild, PipeReadChild);
-		if (Handle.IsValid())
-		{
-			FPlatformProcess::CloseProc(Handle);
-		}
+	void* PipeWriteChild = nullptr;
+	void* PipeReadChild = nullptr;
 
-		TSharedPtr<FTraceViewModel>* TraceSessionPtrPtr = TraceViewModelMap.Find(InTraceId);
-		if (TraceSessionPtrPtr)
-		{
-			FTraceViewModel& TraceSession = **TraceSessionPtrPtr;
-			SplashScreenOverlayTraceFile = FPaths::GetBaseFilename(TraceSession.Uri.ToString());
-		}
-		ShowSplashScreenOverlay();
+	FProcHandle Handle = FPlatformProcess::CreateProc(ExecutablePath, *CmdLine, bLaunchDetached, bLaunchHidden, bLaunchReallyHidden, &ProcessID, PriorityModifier, OptionalWorkingDirectory, PipeWriteChild, PipeReadChild);
+	if (Handle.IsValid())
+	{
+		FPlatformProcess::CloseProc(Handle);
 	}
-	else
+
+	TSharedPtr<FTraceViewModel>* TraceSessionPtrPtr = TraceViewModelMap.Find(InTraceId);
+	if (TraceSessionPtrPtr)
 	{
-		UE_LOG(TraceInsights, Log, TEXT("[StartPage] Start analysis for trace id: 0x%08X"), InTraceId);
-		FInsightsManager::Get()->LoadTrace(InTraceId);
+		FTraceViewModel& TraceSession = **TraceSessionPtrPtr;
+		SplashScreenOverlayTraceFile = FPaths::GetBaseFilename(TraceSession.Uri.ToString());
 	}
+	ShowSplashScreenOverlay();
 }
 
 ////////////////////////////////////////////////////////////////////////////////////////////////////
@@ -1923,6 +1907,10 @@ void SStartPageWindow::GetExtraCommandLineParams(FString& OutParams) const
 	{
 		OutParams.Append(TEXT(" -DebugTools"));
 	}
+	if (bStartProcessWithStompMalloc)
+	{
+		OutParams.Append(TEXT(" -stompmalloc"));
+	}
 }
 
 ////////////////////////////////////////////////////////////////////////////////////////////////////
diff --git a/Engine/Source/Developer/TraceInsights/Private/Insights/Widgets/SStartPageWindow.h b/Engine/Source/Developer/TraceInsights/Private/Insights/Widgets/SStartPageWindow.h
--- a/Engine/Source/Developer/TraceInsights/Private/Insights/Widgets/SStartPageWindow.h
+++ b/Engine/Source/Developer/TraceInsights/Private/Insights/Widgets/SStartPageWindow.h
@@ -110,6 +110,9 @@ public:
 	void SetEnableDebugTools(bool InValue) { bEnableDebugTools = InValue; };
 	bool GetEnableDebugTools() const { return bEnableDebugTools; };
 
+	void SetStartProcessWithStompMalloc(bool InValue) { bStartProcessWithStompMalloc = InValue; };
+	bool GetStartProcessWithStompMalloc() const { return bStartProcessWithStompMalloc; };
+
 private:
 	TSharedRef<SWidget> ConstructSessionsPanel();
 	TSharedRef<SWidget> ConstructLoadPanel();
@@ -289,4 +292,5 @@ private:
 
 	bool bEnableAutomaticTesting = false;
 	bool bEnableDebugTools = false;
+	bool bStartProcessWithStompMalloc = false;
 };
diff --git a/Engine/Source/Developer/TraceInsights/Public/Insights/IUnrealInsightsModule.h b/Engine/Source/Developer/TraceInsights/Public/Insights/IUnrealInsightsModule.h
--- a/Engine/Source/Developer/TraceInsights/Public/Insights/IUnrealInsightsModule.h
+++ b/Engine/Source/Developer/TraceInsights/Public/Insights/IUnrealInsightsModule.h
@@ -146,6 +146,16 @@ protected:
 
 ////////////////////////////////////////////////////////////////////////////////////////////////////
 
+/** Contains parameters that are passed to the CreateSessionBrowser function to control specific behaviors. */
+struct TRACEINSIGHTS_API FCreateSessionBrowserParams
+{
+	bool bAllowDebugTools = false;
+	bool bInitializeTesting = false;
+	bool bStartProcessWithStompMalloc = false;
+};
+
+////////////////////////////////////////////////////////////////////////////////////////////////////
+
 /** Called back to register common layout extensions */
 DECLARE_MULTICAST_DELEGATE_OneParam(FOnRegisterMajorTabExtensions, FInsightsMajorTabExtender& /*MajorTabExtender*/);
 
@@ -253,7 +263,7 @@ public:
 	/**
 	 * Called when the application starts in "Browser" mode.
 	 */
-	virtual void CreateSessionBrowser(bool bAllowDebugTools = false, bool bSingleProcess = false, bool bInitializeTesting = false) = 0;
+	virtual void CreateSessionBrowser(const FCreateSessionBrowserParams& Params) = 0;
 
 	/**
 	 * Called when the application starts in "Viewer" mode.
diff --git a/Engine/Source/Programs/UnrealInsights/Private/UserInterfaceCommand.cpp b/Engine/Source/Programs/UnrealInsights/Private/UserInterfaceCommand.cpp
--- a/Engine/Source/Programs/UnrealInsights/Private/UserInterfaceCommand.cpp
+++ b/Engine/Source/Programs/UnrealInsights/Private/UserInterfaceCommand.cpp
@@ -321,8 +321,12 @@ void FUserInterfaceCommand::InitializeSlateApplication(bool bOpenTraceFile, cons
 			{
 				TraceInsightsModule.ConnectToStore(StoreHost, StorePort);
 			}
-			const bool bSingleProcess = FParse::Param(FCommandLine::Get(), TEXT("SingleProcess"));
-			TraceInsightsModule.CreateSessionBrowser(bAllowDebugTools, bSingleProcess, bInitializeTesting);
+
+			FCreateSessionBrowserParams Params;
+			Params.bAllowDebugTools = bAllowDebugTools;
+			Params.bInitializeTesting = bInitializeTesting;
+			Params.bStartProcessWithStompMalloc = FParse::Param(FCommandLine::Get(), TEXT("stompmalloc"));
+			TraceInsightsModule.CreateSessionBrowser(Params);
 		}
 	}
 
