[Insights] Integrate the AsyncOperationStatus Widget with Memory Alloc Tree View

diff --git a/Engine/Source/Developer/TraceInsights/Private/Insights/Common/InsightsAsyncWorkUtils.h b/Engine/Source/Developer/TraceInsights/Private/Insights/Common/InsightsAsyncWorkUtils.h
new file mode 100644
--- /dev/null
+++ b/Engine/Source/Developer/TraceInsights/Private/Insights/Common/InsightsAsyncWorkUtils.h
@@ -0,0 +1,26 @@
+// Copyright Epic Games, Inc. All Rights Reserved.
+
+#pragma once
+
+#include "CoreMinimal.h"
+
+namespace Insights
+{
+
+////////////////////////////////////////////////////////////////////////////////////////////////////
+
+class IAsyncOperationStatusProvider
+{
+public:
+	virtual bool IsRunning() const = 0;
+
+	virtual double GetAllOperationsDuration() = 0;
+	virtual double GetCurrentOperationDuration() = 0;
+	virtual uint32 GetOperationCount() const = 0;
+
+	virtual FText GetCurrentOperationName() const = 0;
+};
+
+////////////////////////////////////////////////////////////////////////////////////////////////////
+
+} // namespace Insights
\ No newline at end of file
diff --git a/Engine/Source/Developer/TraceInsights/Private/Insights/Table/Widgets/STableTreeView.cpp b/Engine/Source/Developer/TraceInsights/Private/Insights/Table/Widgets/STableTreeView.cpp
--- a/Engine/Source/Developer/TraceInsights/Private/Insights/Table/Widgets/STableTreeView.cpp
+++ b/Engine/Source/Developer/TraceInsights/Private/Insights/Table/Widgets/STableTreeView.cpp
@@ -18,7 +18,6 @@
 #include "Widgets/Views/STableViewBase.h"
 
 // Insights
-#include "Insights/Common/Stopwatch.h"
 #include "Insights/InsightsManager.h"
 #include "Insights/Log.h"
 #include "Insights/Table/ViewModels/Table.h"
@@ -30,6 +29,7 @@
 #include "Insights/Table/Widgets/STableTreeViewTooltip.h"
 #include "Insights/Table/Widgets/STableTreeViewRow.h"
 #include "Insights/TimingProfilerCommon.h"
+#include "Insights/Widgets/SAsyncOperationStatus.h"
 
 #define LOCTEXT_NAMESPACE "STableTreeView"
 
@@ -176,25 +176,40 @@ void STableTreeView::ConstructWidget(TSharedPtr<FTable> InTablePtr)
 
 				+ SScrollBox::Slot()
 				[
-					SNew(SBorder)
-					.BorderImage(FEditorStyle::GetBrush("ToolPanel.GroupBorder"))
-					.Padding(0.0f)
+					SNew(SOverlay)
+
+					+ SOverlay::Slot()
+					.HAlign(HAlign_Fill)
+					.VAlign(VAlign_Fill)
+					[
+						SNew(SBorder)
+						.BorderImage(FEditorStyle::GetBrush("ToolPanel.GroupBorder"))
+						.Padding(0.0f)
+						[
+							SAssignNew(TreeView, STreeView<FTableTreeNodePtr>)
+							.ExternalScrollbar(ExternalScrollbar)
+							.SelectionMode(ESelectionMode::Multi)
+							.TreeItemsSource(&FilteredGroupNodes)
+							.OnGetChildren(this, &STableTreeView::TreeView_OnGetChildren)
+							.OnGenerateRow(this, &STableTreeView::TreeView_OnGenerateRow)
+							.OnSelectionChanged(this, &STableTreeView::TreeView_OnSelectionChanged)
+							.OnMouseButtonDoubleClick(this, &STableTreeView::TreeView_OnMouseButtonDoubleClick)
+							.OnContextMenuOpening(FOnContextMenuOpening::CreateSP(this, &STableTreeView::TreeView_GetMenuContent))
+							.ItemHeight(12.0f)
+							.HeaderRow
+							(
+								SAssignNew(TreeViewHeaderRow, SHeaderRow)
+								.Visibility(EVisibility::Visible)
+							)
+						]
+					]
+
+					+ SOverlay::Slot()
+					.HAlign(HAlign_Right)
+					.VAlign(VAlign_Bottom)
+					.Padding(16.0f)
 					[
-						SAssignNew(TreeView, STreeView<FTableTreeNodePtr>)
-						.ExternalScrollbar(ExternalScrollbar)
-						.SelectionMode(ESelectionMode::Multi)
-						.TreeItemsSource(&FilteredGroupNodes)
-						.OnGetChildren(this, &STableTreeView::TreeView_OnGetChildren)
-						.OnGenerateRow(this, &STableTreeView::TreeView_OnGenerateRow)
-						.OnSelectionChanged(this, &STableTreeView::TreeView_OnSelectionChanged)
-						.OnMouseButtonDoubleClick(this, &STableTreeView::TreeView_OnMouseButtonDoubleClick)
-						.OnContextMenuOpening(FOnContextMenuOpening::CreateSP(this, &STableTreeView::TreeView_GetMenuContent))
-						.ItemHeight(12.0f)
-						.HeaderRow
-						(
-							SAssignNew(TreeViewHeaderRow, SHeaderRow)
-							.Visibility(EVisibility::Visible)
-						)
+						SAssignNew(AsyncOperationStatus, Insights::SAsyncOperationStatus, SharedThis(this))
 					]
 				]
 			]
@@ -2034,6 +2049,7 @@ void STableTreeView::OnPreAsyncUpdate()
 {
 	check(!bIsUpdateRunning);
 
+	AsyncUpdateStopwatch.Restart();
 	bIsUpdateRunning = true;
 
 	ExpandedNodes.Empty();
@@ -2068,6 +2084,8 @@ void STableTreeView::OnPostAsyncUpdate()
 
 	ClearPendingAsyncOperations();
 	TreeView_Refresh();
+
+	AsyncUpdateStopwatch.Stop();
 }
 
 ////////////////////////////////////////////////////////////////////////////////////////////////////
@@ -2140,6 +2158,21 @@ void STableTreeView::OnClose()
 
 ////////////////////////////////////////////////////////////////////////////////////////////////////
 
+FText STableTreeView::GetCurrentOperationName() const
+{
+	return LOCTEXT("CurrentOperationName", "Updating Tree");
+}
+
+////////////////////////////////////////////////////////////////////////////////////////////////////
+
+double STableTreeView::GetAllOperationsDuration()
+{ 
+	AsyncUpdateStopwatch.Update();  
+	return AsyncUpdateStopwatch.GetAccumulatedTime(); 
+}
+
+////////////////////////////////////////////////////////////////////////////////////////////////////
+
 } // namespace Insights
 
 #undef LOCTEXT_NAMESPACE
diff --git a/Engine/Source/Developer/TraceInsights/Private/Insights/Table/Widgets/STableTreeView.h b/Engine/Source/Developer/TraceInsights/Private/Insights/Table/Widgets/STableTreeView.h
--- a/Engine/Source/Developer/TraceInsights/Private/Insights/Table/Widgets/STableTreeView.h
+++ b/Engine/Source/Developer/TraceInsights/Private/Insights/Table/Widgets/STableTreeView.h
@@ -17,6 +17,8 @@
 #include "Widgets/Views/STreeView.h"
 
 // Insights
+#include "Insights/Common/InsightsAsyncWorkUtils.h"
+#include "Insights/Common/Stopwatch.h"
 #include "Insights/Table/ViewModels/TableTreeNode.h"
 
 class FMenuBuilder;
@@ -57,7 +59,7 @@ ENUM_CLASS_FLAGS(EAsyncOperationType);
 /**
  * A custom widget used to display the list of tree nodes.
  */
-class STableTreeView : public SCompoundWidget
+class STableTreeView : public SCompoundWidget, public IAsyncOperationStatusProvider
 {
 	friend class FTableTreeViewFilterAsyncTask;
 	friend class FTableTreeViewSortAsyncTask;
@@ -107,6 +109,18 @@ public:
 
 	void OnClose();
 
+	////////////////////////////////////////////////////////////////////////////////////////////////////
+	// IAsyncOperationStatusProvider implementation
+
+	virtual bool IsRunning() const override { return bIsUpdateRunning; }
+
+	virtual double GetAllOperationsDuration() override;
+	virtual double GetCurrentOperationDuration() override { return 0.0; }
+	virtual uint32 GetOperationCount() const override { return 1; }
+	virtual FText GetCurrentOperationName() const override;
+
+	////////////////////////////////////////////////////////////////////////////////////////////////////
+
 protected:
 	void ConstructWidget(TSharedPtr<FTable> InTablePtr);
 	void UpdateTree();
@@ -377,11 +391,15 @@ protected:
 	//////////////////////////////////////////////////
 	// Async
 	bool bRunInAsyncMode = false;
-	FGraphEventRef PendingAsyncOperationEvent;;
 	bool bIsUpdateRunning = false;
-	TArray<FTableTreeNodePtr> DummyGroupNodes;
 	bool bIsCloseScheduled = false;
+
+	TArray<FTableTreeNodePtr> DummyGroupNodes;
+	FGraphEventRef PendingAsyncOperationEvent;;
 	EAsyncOperationType PendingAsyncOperations = static_cast<EAsyncOperationType>(0);
+	TSharedPtr<class SAsyncOperationStatus> AsyncOperationStatus;
+	FStopwatch AsyncUpdateStopwatch;
+
 	//////////////////////////////////////////////////
 
 	double StatsStartTime;
diff --git a/Engine/Source/Developer/TraceInsights/Private/Insights/ViewModels/StatsAggregator.cpp b/Engine/Source/Developer/TraceInsights/Private/Insights/ViewModels/StatsAggregator.cpp
--- a/Engine/Source/Developer/TraceInsights/Private/Insights/ViewModels/StatsAggregator.cpp
+++ b/Engine/Source/Developer/TraceInsights/Private/Insights/ViewModels/StatsAggregator.cpp
@@ -4,6 +4,8 @@
 
 #include "Insights/TimingProfilerCommon.h"
 
+#define LOCTEXT_NAMESPACE "SStatsAggregator"
+
 namespace Insights
 {
 
@@ -163,6 +165,13 @@ void FStatsAggregator::Tick(TSharedPtr<const TraceServices::IAnalysisSession> In
 
 ////////////////////////////////////////////////////////////////////////////////////////////////////
 
+FText FStatsAggregator::GetCurrentOperationName() const
+{
+	return LOCTEXT("OperationName", "Computing aggregated stats");
+}
+
+////////////////////////////////////////////////////////////////////////////////////////////////////
+
 IStatsAggregationWorker* FStatsAggregator::GetWorker() const
 {
 	// It can only be called from OnFinishedCallback.
@@ -174,3 +183,5 @@ IStatsAggregationWorker* FStatsAggregator::GetWorker() const
 ////////////////////////////////////////////////////////////////////////////////////////////////////
 
 } // namespace Insights
+
+#undef LOCTEXT_NAMESPACE
diff --git a/Engine/Source/Developer/TraceInsights/Private/Insights/ViewModels/StatsAggregator.h b/Engine/Source/Developer/TraceInsights/Private/Insights/ViewModels/StatsAggregator.h
--- a/Engine/Source/Developer/TraceInsights/Private/Insights/ViewModels/StatsAggregator.h
+++ b/Engine/Source/Developer/TraceInsights/Private/Insights/ViewModels/StatsAggregator.h
@@ -4,6 +4,7 @@
 
 #include "CoreMinimal.h"
 #include "Async/AsyncWork.h"
+#include "Insights/Common/InsightsAsyncWorkUtils.h"
 #include "Insights/Common/Stopwatch.h"
 
 ////////////////////////////////////////////////////////////////////////////////////////////////////
@@ -18,18 +19,13 @@ namespace Insights
 
 ////////////////////////////////////////////////////////////////////////////////////////////////////
 
-class IStatsAggregator
+class IStatsAggregator : public IAsyncOperationStatusProvider
 {
 public:
 	virtual void Start() = 0;
 	virtual void Cancel() = 0;
 
 	virtual bool IsCancelRequested() const = 0;
-	virtual bool IsRunning() const = 0;
-
-	virtual double GetAllOperationsDuration() = 0;
-	virtual double GetCurrentOperationDuration() = 0;
-	virtual uint32 GetOperationCount() const = 0;
 };
 
 ////////////////////////////////////////////////////////////////////////////////////////////////////
@@ -78,6 +74,8 @@ public:
 	virtual double GetCurrentOperationDuration() override { CurrentOpStopwatch.Update(); return CurrentOpStopwatch.GetAccumulatedTime(); }
 	virtual uint32 GetOperationCount() const override { return OperationCount; }
 
+	virtual FText GetCurrentOperationName() const;
+
 	//////////////////////////////////////////////////
 
 protected:
diff --git a/Engine/Source/Developer/TraceInsights/Private/Insights/Widgets/SAggregatorStatus.cpp b/Engine/Source/Developer/TraceInsights/Private/Insights/Widgets/SAggregatorStatus.cpp
deleted file mode 100644
diff --git a/Engine/Source/Developer/TraceInsights/Private/Insights/Widgets/SAggregatorStatus.h b/Engine/Source/Developer/TraceInsights/Private/Insights/Widgets/SAggregatorStatus.h
deleted file mode 100644
diff --git a/Engine/Source/Developer/TraceInsights/Private/Insights/Widgets/SAsyncOperationStatus.cpp b/Engine/Source/Developer/TraceInsights/Private/Insights/Widgets/SAsyncOperationStatus.cpp
new file mode 100644
--- /dev/null
+++ b/Engine/Source/Developer/TraceInsights/Private/Insights/Widgets/SAsyncOperationStatus.cpp
@@ -0,0 +1,150 @@
+// Copyright Epic Games, Inc. All Rights Reserved.
+
+#include "SAsyncOperationStatus.h"
+
+#include "EditorStyleSet.h"
+#include "Styling/CoreStyle.h"
+#include "Widgets/Layout/SBorder.h"
+#include "Widgets/Layout/SBox.h"
+#include "Widgets/SBoxPanel.h"
+#include "Widgets/Text/STextBlock.h"
+
+#include "Insights/Common/InsightsAsyncWorkUtils.h"
+#include "Insights/Common/TimeUtils.h"
+
+#define LOCTEXT_NAMESPACE "SAsyncOperationStatus"
+
+namespace Insights
+{
+
+////////////////////////////////////////////////////////////////////////////////////////////////////
+
+void SAsyncOperationStatus::Construct(const FArguments& InArgs, TSharedRef<IAsyncOperationStatusProvider> InStatusProvider)
+{
+	StatusProvider = InStatusProvider;
+
+	ChildSlot
+	[
+		SNew(SBox)
+		.Visibility(this, &SAsyncOperationStatus::GetContentVisibility)
+		[
+			SNew(SBorder)
+			.BorderImage(FEditorStyle::GetBrush("NotificationList.ItemBackground"))
+			.BorderBackgroundColor(this, &SAsyncOperationStatus::GetBackgroundColorAndOpacity)
+			.Padding(FMargin(16.0f, 8.0f, 16.0f, 8.0f))
+			.HAlign(HAlign_Fill)
+			.VAlign(VAlign_Fill)
+			[
+				SNew(SHorizontalBox)
+
+				+ SHorizontalBox::Slot()
+				.AutoWidth()
+				.HAlign(HAlign_Center)
+				.VAlign(VAlign_Center)
+				[
+					SNew(STextBlock)
+					.Font(FCoreStyle::GetDefaultFontStyle("Bold", 11))
+					.Text(this, &SAsyncOperationStatus::GetText)
+					.ToolTipText(this, &SAsyncOperationStatus::GetTooltipText)
+					.ColorAndOpacity(this, &SAsyncOperationStatus::GetTextColorAndOpacity)
+				]
+
+				+ SHorizontalBox::Slot()
+				.AutoWidth()
+				.HAlign(HAlign_Left)
+				.VAlign(VAlign_Center)
+				[
+					SNew(STextBlock)
+					.MinDesiredWidth(16.0f)
+					.Font(FCoreStyle::GetDefaultFontStyle("Bold", 11))
+					.Text(this, &SAsyncOperationStatus::GetAnimatedText)
+					.ToolTipText(this, &SAsyncOperationStatus::GetTooltipText)
+					.ColorAndOpacity(this, &SAsyncOperationStatus::GetTextColorAndOpacity)
+				]
+			]
+		]
+	];
+}
+
+////////////////////////////////////////////////////////////////////////////////////////////////////
+
+EVisibility SAsyncOperationStatus::GetContentVisibility() const
+{
+	return StatusProvider->IsRunning() ? EVisibility::Visible : EVisibility::Collapsed;
+}
+
+////////////////////////////////////////////////////////////////////////////////////////////////////
+
+FSlateColor SAsyncOperationStatus::GetBackgroundColorAndOpacity() const
+{
+	const float Opacity = ComputeOpacity();
+	return FSlateColor(FLinearColor(0.2f, 0.2f, 0.2f, Opacity));
+}
+
+////////////////////////////////////////////////////////////////////////////////////////////////////
+
+FSlateColor SAsyncOperationStatus::GetTextColorAndOpacity() const
+{
+	const float Opacity = ComputeOpacity();
+	return FSlateColor(FLinearColor(1.0f, 1.0f, 1.0f, Opacity));
+}
+
+////////////////////////////////////////////////////////////////////////////////////////////////////
+
+float SAsyncOperationStatus::ComputeOpacity() const
+{
+	const double TotalDuration = StatusProvider->GetAllOperationsDuration();
+	constexpr double FadeInStartTime = 0.1; // [second]
+	constexpr double FadeInEndTime = 3.0; // [second]
+	constexpr double FadeInDuration = FadeInEndTime - FadeInStartTime;
+	const float Opacity = static_cast<float>(FMath::Clamp(TotalDuration - FadeInStartTime, 0.0, FadeInDuration) / FadeInDuration);
+	return Opacity;
+}
+
+////////////////////////////////////////////////////////////////////////////////////////////////////
+
+FText SAsyncOperationStatus::GetText() const
+{
+	return LOCTEXT("DefaultText", "Computing");
+}
+
+////////////////////////////////////////////////////////////////////////////////////////////////////
+
+FText SAsyncOperationStatus::GetAnimatedText() const
+{
+	const double TotalDuration = StatusProvider->GetAllOperationsDuration();
+	const TCHAR* Anim[] = { TEXT(""), TEXT("."), TEXT(".."), TEXT("..."), };
+	int32 AnimIndex = static_cast<int32>(TotalDuration / 0.2) % UE_ARRAY_COUNT(Anim);
+	return FText::FromString(FString(Anim[AnimIndex]));
+}
+
+////////////////////////////////////////////////////////////////////////////////////////////////////
+
+FText SAsyncOperationStatus::GetTooltipText() const
+{
+	const double TotalDuration = StatusProvider->GetAllOperationsDuration();
+	const uint32 OperationCount = StatusProvider->GetOperationCount();
+	FText CurrentOpName = StatusProvider->GetCurrentOperationName();
+
+	if (OperationCount == 1)
+	{
+		return FText::Format(LOCTEXT("DefaultTooltip_Fmt2", "{0}...\nElapsed Time: {1}"),
+			CurrentOpName,
+			FText::FromString(TimeUtils::FormatTime(TotalDuration, TimeUtils::Second)));
+	}
+	else
+	{
+		const double Duration = StatusProvider->GetCurrentOperationDuration();
+		return FText::Format(LOCTEXT("DefaultTooltip_Fmt3", "{0}...\nElapsed Time: {1}\nElapsed Time (op {2}): {3}"),
+			CurrentOpName,
+			FText::FromString(TimeUtils::FormatTime(TotalDuration, TimeUtils::Second)),
+			FText::AsNumber(OperationCount),
+			FText::FromString(TimeUtils::FormatTime(Duration, TimeUtils::Second)));
+	}
+}
+
+////////////////////////////////////////////////////////////////////////////////////////////////////
+
+} // namespace Insights
+
+#undef LOCTEXT_NAMESPACE
diff --git a/Engine/Source/Developer/TraceInsights/Private/Insights/Widgets/SAsyncOperationStatus.h b/Engine/Source/Developer/TraceInsights/Private/Insights/Widgets/SAsyncOperationStatus.h
new file mode 100644
--- /dev/null
+++ b/Engine/Source/Developer/TraceInsights/Private/Insights/Widgets/SAsyncOperationStatus.h
@@ -0,0 +1,39 @@
+// Copyright Epic Games, Inc. All Rights Reserved.
+
+#pragma once
+
+#include "CoreMinimal.h"
+#include "Widgets/DeclarativeSyntaxSupport.h"
+#include "Widgets/SCompoundWidget.h"
+
+namespace Insights
+{
+
+class IAsyncOperationStatusProvider;
+
+// A widget representing the status message used in tree views (Timers, Counters, etc.) when computing the aggregated stats in a background operation.
+class SAsyncOperationStatus : public SCompoundWidget
+{
+public:
+	SLATE_BEGIN_ARGS(SAsyncOperationStatus) {}
+
+	SLATE_END_ARGS()
+
+	void Construct(const FArguments& InArgs, TSharedRef<IAsyncOperationStatusProvider> InStatusProvider);
+
+private:
+	EVisibility GetContentVisibility() const;
+
+	FSlateColor GetBackgroundColorAndOpacity() const;
+	FSlateColor GetTextColorAndOpacity() const;
+	float ComputeOpacity() const;
+
+	FText GetText() const;
+	FText GetAnimatedText() const;
+	FText GetTooltipText() const;
+
+private:
+	TSharedPtr<IAsyncOperationStatusProvider> StatusProvider;
+};
+
+} // namespace Insights
diff --git a/Engine/Source/Developer/TraceInsights/Private/Insights/Widgets/SStatsView.cpp b/Engine/Source/Developer/TraceInsights/Private/Insights/Widgets/SStatsView.cpp
--- a/Engine/Source/Developer/TraceInsights/Private/Insights/Widgets/SStatsView.cpp
+++ b/Engine/Source/Developer/TraceInsights/Private/Insights/Widgets/SStatsView.cpp
@@ -29,7 +29,7 @@
 #include "Insights/ViewModels/StatsNodeHelper.h"
 #include "Insights/ViewModels/StatsViewColumnFactory.h"
 #include "Insights/ViewModels/TimingGraphTrack.h"
-#include "Insights/Widgets/SAggregatorStatus.h"
+#include "Insights/Widgets/SAsyncOperationStatus.h"
 #include "Insights/Widgets/SStatsViewTooltip.h"
 #include "Insights/Widgets/SStatsTableRow.h"
 #include "Insights/Widgets/STimingProfilerWindow.h"
@@ -273,7 +273,7 @@ void SStatsView::Construct(const FArguments& InArgs)
 					.VAlign(VAlign_Bottom)
 					.Padding(16.0f)
 					[
-						SAssignNew(AggregatorStatus, Insights::SAggregatorStatus, Aggregator)
+						SAssignNew(AsyncOperationStatus, Insights::SAsyncOperationStatus, Aggregator)
 					]
 				]
 			]
diff --git a/Engine/Source/Developer/TraceInsights/Private/Insights/Widgets/SStatsView.h b/Engine/Source/Developer/TraceInsights/Private/Insights/Widgets/SStatsView.h
--- a/Engine/Source/Developer/TraceInsights/Private/Insights/Widgets/SStatsView.h
+++ b/Engine/Source/Developer/TraceInsights/Private/Insights/Widgets/SStatsView.h
@@ -37,7 +37,7 @@ namespace Insights
 	class ITableCellValueSorter;
 
 	class FCounterAggregator;
-	class SAggregatorStatus;
+	class SAsyncOperationStatus;
 }
 
 ////////////////////////////////////////////////////////////////////////////////////////////////////
@@ -372,7 +372,7 @@ private:
 	//////////////////////////////////////////////////
 
 	TSharedRef<Insights::FCounterAggregator> Aggregator;
-	TSharedPtr<Insights::SAggregatorStatus> AggregatorStatus;
+	TSharedPtr<Insights::SAsyncOperationStatus> AsyncOperationStatus;
 };
 
 ////////////////////////////////////////////////////////////////////////////////////////////////////
diff --git a/Engine/Source/Developer/TraceInsights/Private/Insights/Widgets/STimerTreeView.cpp b/Engine/Source/Developer/TraceInsights/Private/Insights/Widgets/STimerTreeView.cpp
--- a/Engine/Source/Developer/TraceInsights/Private/Insights/Widgets/STimerTreeView.cpp
+++ b/Engine/Source/Developer/TraceInsights/Private/Insights/Widgets/STimerTreeView.cpp
@@ -19,7 +19,7 @@
 #include "Insights/TimingProfilerManager.h"
 #include "Insights/ViewModels/TimerButterflyAggregation.h"
 #include "Insights/ViewModels/TimersViewColumnFactory.h"
-#include "Insights/Widgets/SAggregatorStatus.h"
+#include "Insights/Widgets/SAsyncOperationStatus.h"
 #include "Insights/Widgets/STimersViewTooltip.h"
 #include "Insights/Widgets/STimerTableRow.h"
 
@@ -146,7 +146,7 @@ void STimerTreeView::Construct(const FArguments& InArgs, const FText& InViewName
 				.VAlign(VAlign_Bottom)
 				.Padding(16.0f)
 				[
-					SAssignNew(AggregatorStatus, Insights::SAggregatorStatus, TimerButterflyAggregator)
+					SAssignNew(AsyncOperationStatus, Insights::SAsyncOperationStatus, TimerButterflyAggregator)
 				]
 			]
 		]
diff --git a/Engine/Source/Developer/TraceInsights/Private/Insights/Widgets/STimerTreeView.h b/Engine/Source/Developer/TraceInsights/Private/Insights/Widgets/STimerTreeView.h
--- a/Engine/Source/Developer/TraceInsights/Private/Insights/Widgets/STimerTreeView.h
+++ b/Engine/Source/Developer/TraceInsights/Private/Insights/Widgets/STimerTreeView.h
@@ -27,7 +27,7 @@ namespace Insights
 	class FTable;
 	class FTableColumn;
 	class ITableCellValueSorter;
-	class SAggregatorStatus;
+	class SAsyncOperationStatus;
 }
 
 ////////////////////////////////////////////////////////////////////////////////////////////////////
@@ -228,7 +228,7 @@ private:
 	/** How we sort the nodes? Ascending or Descending. */
 	EColumnSortMode::Type ColumnSortMode;
 
-	TSharedPtr<Insights::SAggregatorStatus> AggregatorStatus;
+	TSharedPtr<Insights::SAsyncOperationStatus> AsyncOperationStatus;
 };
 
 ////////////////////////////////////////////////////////////////////////////////////////////////////
diff --git a/Engine/Source/Developer/TraceInsights/Private/Insights/Widgets/STimersView.cpp b/Engine/Source/Developer/TraceInsights/Private/Insights/Widgets/STimersView.cpp
--- a/Engine/Source/Developer/TraceInsights/Private/Insights/Widgets/STimersView.cpp
+++ b/Engine/Source/Developer/TraceInsights/Private/Insights/Widgets/STimersView.cpp
@@ -30,7 +30,7 @@
 #include "Insights/ViewModels/TimerNodeHelper.h"
 #include "Insights/ViewModels/TimersViewColumnFactory.h"
 #include "Insights/ViewModels/TimingGraphTrack.h"
-#include "Insights/Widgets/SAggregatorStatus.h"
+#include "Insights/Widgets/SAsyncOperationStatus.h"
 #include "Insights/Widgets/STimersViewTooltip.h"
 #include "Insights/Widgets/STimerTableRow.h"
 #include "Insights/Widgets/STimingProfilerWindow.h"
@@ -282,7 +282,7 @@ void STimersView::Construct(const FArguments& InArgs)
 					.VAlign(VAlign_Bottom)
 					.Padding(16.0f)
 					[
-						SAssignNew(AggregatorStatus, Insights::SAggregatorStatus, Aggregator)
+						SAssignNew(AsyncOperationStatus, Insights::SAsyncOperationStatus, Aggregator)
 					]
 				]
 			]
diff --git a/Engine/Source/Developer/TraceInsights/Private/Insights/Widgets/STimersView.h b/Engine/Source/Developer/TraceInsights/Private/Insights/Widgets/STimersView.h
--- a/Engine/Source/Developer/TraceInsights/Private/Insights/Widgets/STimersView.h
+++ b/Engine/Source/Developer/TraceInsights/Private/Insights/Widgets/STimersView.h
@@ -37,7 +37,7 @@ namespace Insights
 	class ITableCellValueSorter;
 
 	class FTimerAggregator;
-	class SAggregatorStatus;
+	class SAsyncOperationStatus;
 }
 
 ////////////////////////////////////////////////////////////////////////////////////////////////////
@@ -374,7 +374,7 @@ private:
 	//////////////////////////////////////////////////
 
 	TSharedRef<Insights::FTimerAggregator> Aggregator;
-	TSharedPtr<Insights::SAggregatorStatus> AggregatorStatus;
+	TSharedPtr<Insights::SAsyncOperationStatus> AsyncOperationStatus;
 };
 
 ////////////////////////////////////////////////////////////////////////////////////////////////////
