[Insights] Fixed the Main Memory Graph track not showing correctly the vertical axis grid when first series is hidden.

diff --git a/Engine/Source/Developer/TraceInsights/Private/Insights/MemoryProfiler/ViewModels/MemoryGraphTrack.cpp b/Engine/Source/Developer/TraceInsights/Private/Insights/MemoryProfiler/ViewModels/MemoryGraphTrack.cpp
--- a/Engine/Source/Developer/TraceInsights/Private/Insights/MemoryProfiler/ViewModels/MemoryGraphTrack.cpp
+++ b/Engine/Source/Developer/TraceInsights/Private/Insights/MemoryProfiler/ViewModels/MemoryGraphTrack.cpp
@@ -417,10 +417,18 @@ void FMemoryGraphTrack::DrawVerticalAxisGrid(const ITimingTrackDrawContext& Cont
 {
 	TSharedPtr<FMemoryGraphSeries> Series = MainSeries;
 
-	if (!Series.IsValid() && AllSeries.Num() > 0)
+	if (!Series.IsValid())
 	{
-		//TODO: if (Series->Is<FMemoryGraphSeries>())
-		Series = StaticCastSharedPtr<FMemoryGraphSeries>(AllSeries[0]);
+		// Use the first visible series.
+		for (const TSharedPtr<FGraphSeries>& GraphSeries : AllSeries)
+		{
+			if (GraphSeries->IsVisible())
+			{
+				//TODO: if (GraphSeries->Is<FMemoryGraphSeries>())
+				Series = StaticCastSharedPtr<FMemoryGraphSeries>(GraphSeries);
+				break;
+			}
+		}
 	}
 
 	if (!Series.IsValid())
