[Insights] Permit modifications to groupings during an async operation in Alloc Tree View.

diff --git a/Engine/Source/Developer/TraceInsights/Private/Insights/Table/Widgets/STableTreeView.cpp b/Engine/Source/Developer/TraceInsights/Private/Insights/Table/Widgets/STableTreeView.cpp
--- a/Engine/Source/Developer/TraceInsights/Private/Insights/Table/Widgets/STableTreeView.cpp
+++ b/Engine/Source/Developer/TraceInsights/Private/Insights/Table/Widgets/STableTreeView.cpp
@@ -598,10 +598,11 @@ void STableTreeView::Tick(const FGeometry& AllottedGeometry, const double InCurr
 {
 	if (bRunInAsyncMode && bIsUpdateRunning && !bIsCloseScheduled)
 	{
-		check(PendingAsyncOperationEvent.IsValid());
-		if (PendingAsyncOperationEvent->IsComplete())
+		check(InProgressAsyncOperationEvent.IsValid());
+		if (InProgressAsyncOperationEvent->IsComplete())
 		{
 			OnPostAsyncUpdate();
+			StartPendingAsyncOperations();
 		}
 	}
 }
@@ -615,15 +616,18 @@ void STableTreeView::UpdateTree()
 
 	if (bRunInAsyncMode)
 	{
-		OnPreAsyncUpdate();
+		if (!bIsUpdateRunning)
+		{
+			OnPreAsyncUpdate();
 
-		FGraphEventRef CompletedEvent = StartCreateGroupsTask();
-		CompletedEvent = StartSortTreeNodesTask(CompletedEvent);
-		PendingAsyncOperationEvent = StartApplyFiltersTask(CompletedEvent);
+			FGraphEventRef CompletedEvent = StartCreateGroupsTask();
+			CompletedEvent = StartSortTreeNodesTask(CompletedEvent);
+			InProgressAsyncOperationEvent = StartApplyFiltersTask(CompletedEvent);
+		}
 	}
 	else
 	{
-		CreateGroups();
+		CreateGroups(CurrentGroupings);
 		SortTreeNodes();
 		ApplyFiltering();
 	}
@@ -882,11 +886,14 @@ void STableTreeView::SearchBox_OnTextChanged(const FText& InFilterText)
 	TextFilter->SetRawFilterText(InFilterText);
 	SearchBox->SetError(TextFilter->GetFilterErrorText());
 
-	if(bRunInAsyncMode)
+	if (bRunInAsyncMode)
 	{
-		OnPreAsyncUpdate();
+		if (!bIsUpdateRunning)
+		{
+			OnPreAsyncUpdate();
 
-		PendingAsyncOperationEvent = StartApplyFiltersTask();
+			InProgressAsyncOperationEvent = StartApplyFiltersTask();
+		}
 	}
 	else
 	{
@@ -910,12 +917,12 @@ bool STableTreeView::SearchBox_IsEnabled() const
 // Grouping
 ////////////////////////////////////////////////////////////////////////////////////////////////////
 
-void STableTreeView::CreateGroups()
+void STableTreeView::CreateGroups(const TArray<TSharedPtr<FTreeNodeGrouping>>& Groupings)
 {
 	FStopwatch Stopwatch;
 	Stopwatch.Start();
 
-	GroupNodesRec(TableTreeNodes, *Root, 0);
+	GroupNodesRec(TableTreeNodes, *Root, 0, Groupings);
 
 	Stopwatch.Update();
 	const double Time1 = Stopwatch.GetAccumulatedTime();
@@ -953,11 +960,11 @@ void STableTreeView::CreateGroups()
 
 ////////////////////////////////////////////////////////////////////////////////////////////////////
 
-void STableTreeView::GroupNodesRec(const TArray<FTableTreeNodePtr>& Nodes, FTableTreeNode& ParentGroup, int32 GroupingDepth)
+void STableTreeView::GroupNodesRec(const TArray<FTableTreeNodePtr>& Nodes, FTableTreeNode& ParentGroup, int32 GroupingDepth, const TArray<TSharedPtr<FTreeNodeGrouping>>& Groupings)
 {
-	ensure(CurrentGroupings.Num() > 0);
+	ensure(Groupings.Num() > 0);
 
-	FTreeNodeGrouping& Grouping = *CurrentGroupings[GroupingDepth];
+	FTreeNodeGrouping& Grouping = *Groupings[GroupingDepth];
 
 	Grouping.GroupNodes(Nodes, ParentGroup, Table);
 
@@ -971,7 +978,7 @@ void STableTreeView::GroupNodesRec(const TArray<FTableTreeNodePtr>& Nodes, FTabl
 		}
 	}
 
-	if (GroupingDepth < CurrentGroupings.Num() - 1)
+	if (GroupingDepth < Groupings.Num() - 1)
 	{
 		TArray<FTableTreeNodePtr> ChildNodes;
 
@@ -987,7 +994,7 @@ void STableTreeView::GroupNodesRec(const TArray<FTableTreeNodePtr>& Nodes, FTabl
 				ChildNodes.Add(StaticCastSharedPtr<FTableTreeNode>(ChildPtr));
 			}
 
-			GroupNodesRec(ChildNodes, Group, GroupingDepth + 1);
+			GroupNodesRec(ChildNodes, Group, GroupingDepth + 1, Groupings);
 		}
 	}
 }
@@ -1112,7 +1119,7 @@ void STableTreeView::CreateGroupings()
 
 void STableTreeView::PreChangeGroupings()
 {
-	for (TSharedPtr<FTreeNodeGrouping> GroupingPtr : CurrentGroupings)
+	for (TSharedPtr<FTreeNodeGrouping>& GroupingPtr : CurrentGroupings)
 	{
 		const FName& ColumnId = GroupingPtr->GetColumnId();
 		if (ColumnId != NAME_None)
@@ -1135,7 +1142,7 @@ void STableTreeView::PostChangeGroupings()
 	//FString GroupingStr;
 
 	int32 GroupingDepth = 0;
-	for (TSharedPtr<FTreeNodeGrouping> GroupingPtr : CurrentGroupings)
+	for (TSharedPtr<FTreeNodeGrouping>& GroupingPtr : CurrentGroupings)
 	{
 		const FName& ColumnId = GroupingPtr->GetColumnId();
 
@@ -1199,20 +1206,23 @@ void STableTreeView::PostChangeGroupings()
 
 	if (bRunInAsyncMode)
 	{
-		OnPreAsyncUpdate();
+		if (!bIsUpdateRunning)
+		{
+			OnPreAsyncUpdate();
 
-		FGraphEventRef CompletedEvent = StartCreateGroupsTask();
-		CompletedEvent = StartSortTreeNodesTask(CompletedEvent);
-		PendingAsyncOperationEvent = StartApplyFiltersTask(CompletedEvent);
+			FGraphEventRef CompletedEvent = StartCreateGroupsTask();
+			CompletedEvent = StartSortTreeNodesTask(CompletedEvent);
+			InProgressAsyncOperationEvent = StartApplyFiltersTask(CompletedEvent);
+		}
 	}
 	else
 	{
-		CreateGroups();
+		CreateGroups(CurrentGroupings);
 		SortTreeNodes();
 		ApplyFiltering();
-
-		RebuildGroupingCrumbs();
 	}
+
+	RebuildGroupingCrumbs();
 }
 
 ////////////////////////////////////////////////////////////////////////////////////////////////////
@@ -1329,9 +1339,9 @@ TSharedRef<SWidget> STableTreeView::GetGroupingCrumbMenuContent(const TSharedPtr
 	}
 	MenuBuilder.EndSection();
 
-	auto CanExecute = [this]()
+	auto CanExecute = []()
 	{
-		return !this->IsRunningAsyncUpdate();
+		return true;
 	};
 
 	if (CrumbGroupingDepth >= 0)
@@ -1513,11 +1523,6 @@ void STableTreeView::GroupingCrumbMenu_Change_Execute(const TSharedPtr<FTreeNode
 
 bool STableTreeView::GroupingCrumbMenu_Change_CanExecute(const TSharedPtr<FTreeNodeGrouping> OldGrouping, const TSharedPtr<FTreeNodeGrouping> NewGrouping) const
 {
-	if (bIsUpdateRunning)
-	{
-		return false;
-	}
-
 	return NewGrouping != OldGrouping;
 }
 
@@ -1565,11 +1570,6 @@ void STableTreeView::GroupingCrumbMenu_Add_Execute(const TSharedPtr<FTreeNodeGro
 
 bool STableTreeView::GroupingCrumbMenu_Add_CanExecute(const TSharedPtr<FTreeNodeGrouping> Grouping, const TSharedPtr<FTreeNodeGrouping> AfterGrouping) const
 {
-	if (bIsUpdateRunning)
-	{
-		return false;
-	}
-
 	if (AfterGrouping.IsValid())
 	{
 		const int32 AfterGroupingDepth = GetGroupingDepth(AfterGrouping);
@@ -1697,10 +1697,13 @@ void STableTreeView::SetSortModeForColumn(const FName& ColumnId, const EColumnSo
 
 	if (bRunInAsyncMode)
 	{
-		OnPreAsyncUpdate();
+		if (!bIsUpdateRunning)
+		{
+			OnPreAsyncUpdate();
 
-		FGraphEventRef CompletedEvent = StartSortTreeNodesTask();
-		PendingAsyncOperationEvent = StartApplyFiltersTask(CompletedEvent);
+			FGraphEventRef CompletedEvent = StartSortTreeNodesTask();
+			InProgressAsyncOperationEvent = StartApplyFiltersTask(CompletedEvent);
+		}
 	}
 	else
 	{
@@ -2067,14 +2070,12 @@ void STableTreeView::OnPostAsyncUpdate()
 
 	bIsUpdateRunning = false;
 
-	RebuildGroupingCrumbs();
-
 	TreeView->SetTreeItemsSource(&FilteredGroupNodes);
 
 	TreeView->ClearExpandedItems();
 
 	// Grouping can result in old group nodes no longer existing in the tree, so we don't keep the expanded list.
-	if (!HasPendingAsyncOperation(EAsyncOperationType::GroupingOp))
+	if (!HasInProgressAsyncOperation(EAsyncOperationType::GroupingOp))
 	{
 		for (auto It = ExpandedNodes.CreateConstIterator(); It; ++It)
 		{
@@ -2082,7 +2083,7 @@ void STableTreeView::OnPostAsyncUpdate()
 		}
 	}
 
-	ClearPendingAsyncOperations();
+	ClearInProgressAsyncOperations();
 	TreeView_Refresh();
 
 	AsyncUpdateStopwatch.Stop();
@@ -2092,7 +2093,7 @@ void STableTreeView::OnPostAsyncUpdate()
 
 FGraphEventRef STableTreeView::StartSortTreeNodesTask(FGraphEventRef Prerequisite)
 {
-	AddPendingAsyncOperation(EAsyncOperationType::SortingOp);
+	AddInProgressAsyncOperation(EAsyncOperationType::SortingOp);
 
 	if (Prerequisite.IsValid())
 	{
@@ -2110,17 +2111,20 @@ FGraphEventRef STableTreeView::StartSortTreeNodesTask(FGraphEventRef Prerequisit
 
 FGraphEventRef STableTreeView::StartCreateGroupsTask(FGraphEventRef Prerequisite)
 {
-	AddPendingAsyncOperation(EAsyncOperationType::GroupingOp);
+	AddInProgressAsyncOperation(EAsyncOperationType::GroupingOp);
+
+	CurrentAsyncOpGroupings.Empty();
+	CurrentAsyncOpGroupings.Insert(CurrentGroupings, 0);
 
 	if (Prerequisite.IsValid())
 	{
 		FGraphEventArray Prerequisites;
 		Prerequisites.Add(Prerequisite);
-		return TGraphTask<FTableTreeViewGroupAsyncTask>::CreateTask(&Prerequisites).ConstructAndDispatchWhenReady(this);
+		return TGraphTask<FTableTreeViewGroupAsyncTask>::CreateTask(&Prerequisites).ConstructAndDispatchWhenReady(this, &CurrentAsyncOpGroupings);
 	}
 	else
 	{
-		return TGraphTask<FTableTreeViewGroupAsyncTask>::CreateTask().ConstructAndDispatchWhenReady(this);
+		return TGraphTask<FTableTreeViewGroupAsyncTask>::CreateTask().ConstructAndDispatchWhenReady(this, &CurrentAsyncOpGroupings);
 	}
 }
 
@@ -2128,7 +2132,7 @@ FGraphEventRef STableTreeView::StartCreateGroupsTask(FGraphEventRef Prerequisite
 
 FGraphEventRef STableTreeView::StartApplyFiltersTask(FGraphEventRef Prerequisite)
 {
-	AddPendingAsyncOperation(EAsyncOperationType::FilteringOp);
+	AddInProgressAsyncOperation(EAsyncOperationType::FilteringOp);
 
 	if (Prerequisite.IsValid())
 	{
@@ -2146,12 +2150,12 @@ FGraphEventRef STableTreeView::StartApplyFiltersTask(FGraphEventRef Prerequisite
 
 void STableTreeView::OnClose()
 {
-	if (bIsUpdateRunning && PendingAsyncOperationEvent.IsValid() && !PendingAsyncOperationEvent->IsComplete())
+	if (bIsUpdateRunning && InProgressAsyncOperationEvent.IsValid() && !InProgressAsyncOperationEvent->IsComplete())
 	{
 		bIsCloseScheduled = true;
 
 		FGraphEventArray Prerequisites;
-		Prerequisites.Add(PendingAsyncOperationEvent);
+		Prerequisites.Add(InProgressAsyncOperationEvent);
 		TGraphTask<FTableTreeViewAsyncCompleteTask>::CreateTask(&Prerequisites).ConstructAndDispatchWhenReady(SharedThis(this));
 	}
 }
@@ -2173,6 +2177,34 @@ double STableTreeView::GetAllOperationsDuration()
 
 ////////////////////////////////////////////////////////////////////////////////////////////////////
 
+void STableTreeView::StartPendingAsyncOperations()
+{
+	bool bGroupingsHaveChanged = CurrentGroupings.Num() != CurrentAsyncOpGroupings.Num();
+
+	if (!bGroupingsHaveChanged)
+	{
+		for (int Index = 0; Index < CurrentGroupings.Num(); ++Index)
+		{
+			if (CurrentGroupings[Index] != CurrentAsyncOpGroupings[Index])
+			{
+				bGroupingsHaveChanged = true;
+				break;
+			}
+		}
+	}
+
+	if (bGroupingsHaveChanged)
+	{
+		OnPreAsyncUpdate();
+
+		FGraphEventRef CompletedEvent = StartCreateGroupsTask();
+		CompletedEvent = StartSortTreeNodesTask(CompletedEvent);
+		InProgressAsyncOperationEvent = StartApplyFiltersTask(CompletedEvent);
+	}
+}
+
+////////////////////////////////////////////////////////////////////////////////////////////////////
+
 } // namespace Insights
 
 #undef LOCTEXT_NAMESPACE
diff --git a/Engine/Source/Developer/TraceInsights/Private/Insights/Table/Widgets/STableTreeView.h b/Engine/Source/Developer/TraceInsights/Private/Insights/Table/Widgets/STableTreeView.h
--- a/Engine/Source/Developer/TraceInsights/Private/Insights/Table/Widgets/STableTreeView.h
+++ b/Engine/Source/Developer/TraceInsights/Private/Insights/Table/Widgets/STableTreeView.h
@@ -201,8 +201,8 @@ protected:
 
 	void CreateGroupings();
 
-	void CreateGroups();
-	void GroupNodesRec(const TArray<FTableTreeNodePtr>& Nodes, FTableTreeNode& ParentGroup, int32 GroupingDepth);
+	void CreateGroups(const TArray<TSharedPtr<FTreeNodeGrouping>>& Groupings);
+	void GroupNodesRec(const TArray<FTableTreeNodePtr>& Nodes, FTableTreeNode& ParentGroup, int32 GroupingDepth, const TArray<TSharedPtr<FTreeNodeGrouping>>& Groupings);
 
 	void ResetAggregatedValuesRec(FTableTreeNode& GroupNode);
 	void UpdateInt64SumAggregationRec(FTableColumn& Column, FTableTreeNode& GroupNode);
@@ -294,9 +294,11 @@ protected:
 	FGraphEventRef StartCreateGroupsTask(FGraphEventRef Prerequisite = nullptr);
 	FGraphEventRef StartApplyFiltersTask(FGraphEventRef Prerequisite = nullptr);
 
-	void AddPendingAsyncOperation(EAsyncOperationType InType) { EnumAddFlags(PendingAsyncOperations, InType); }
-	bool HasPendingAsyncOperation(EAsyncOperationType InType) const { return EnumHasAnyFlags(PendingAsyncOperations, InType); }
-	void ClearPendingAsyncOperations() { PendingAsyncOperations = static_cast<EAsyncOperationType>(0); }
+	void AddInProgressAsyncOperation(EAsyncOperationType InType) { EnumAddFlags(InProgressAsyncOperations, InType); }
+	bool HasInProgressAsyncOperation(EAsyncOperationType InType) const { return EnumHasAnyFlags(InProgressAsyncOperations, InType); }
+	void ClearInProgressAsyncOperations() { InProgressAsyncOperations = static_cast<EAsyncOperationType>(0); }
+
+	void StartPendingAsyncOperations();
 
 	////////////////////////////////////////////////////////////////////////////////////////////////////
 
@@ -395,11 +397,13 @@ protected:
 	bool bIsCloseScheduled = false;
 
 	TArray<FTableTreeNodePtr> DummyGroupNodes;
-	FGraphEventRef PendingAsyncOperationEvent;;
-	EAsyncOperationType PendingAsyncOperations = static_cast<EAsyncOperationType>(0);
+	FGraphEventRef InProgressAsyncOperationEvent;;
+	EAsyncOperationType InProgressAsyncOperations = static_cast<EAsyncOperationType>(0);
 	TSharedPtr<class SAsyncOperationStatus> AsyncOperationStatus;
 	FStopwatch AsyncUpdateStopwatch;
 
+	TArray<TSharedPtr<FTreeNodeGrouping>> CurrentAsyncOpGroupings;
+
 	//////////////////////////////////////////////////
 
 	double StatsStartTime;
@@ -463,9 +467,10 @@ private:
 class FTableTreeViewGroupAsyncTask
 {
 public:
-	FTableTreeViewGroupAsyncTask(STableTreeView* InPtr)
+	FTableTreeViewGroupAsyncTask(STableTreeView* InPtr, TArray<TSharedPtr<FTreeNodeGrouping>>* InGroupings)
 	{
 		TableTreeViewPtr = InPtr;
+		Groupings = InGroupings;
 	}
 
 	FORCEINLINE TStatId GetStatId() const { RETURN_QUICK_DECLARE_CYCLE_STAT(FTableTreeViewGroupAsyncTask, STATGROUP_TaskGraphTasks); }
@@ -476,12 +481,13 @@ public:
 	{
 		if (TableTreeViewPtr)
 		{
-			TableTreeViewPtr->CreateGroups();
+			TableTreeViewPtr->CreateGroups(*Groupings);
 		}
 	}
 
 private:
 	STableTreeView* TableTreeViewPtr = nullptr;
+	TArray<TSharedPtr<FTreeNodeGrouping>>* Groupings;
 };
 
 ////////////////////////////////////////////////////////////////////////////////////////////////////
