[Insights]
  - Session Info: Changed "Command Line" to be automatically displayed on multiple lines.
  - Session Info: Increased precision for estimated analysis speed.

diff --git a/Engine/Source/Developer/TraceInsights/Private/Insights/Widgets/SSessionInfoWindow.cpp b/Engine/Source/Developer/TraceInsights/Private/Insights/Widgets/SSessionInfoWindow.cpp
--- a/Engine/Source/Developer/TraceInsights/Private/Insights/Widgets/SSessionInfoWindow.cpp
+++ b/Engine/Source/Developer/TraceInsights/Private/Insights/Widgets/SSessionInfoWindow.cpp
@@ -9,6 +9,7 @@
 #include "Misc/Paths.h"
 #include "Widgets/Images/SImage.h"
 #include "Widgets/Input/SEditableTextBox.h"
+#include "Widgets/Input/SMultiLineEditableTextBox.h"
 #include "Widgets/Layout/SScrollBox.h"
 #include "Widgets/SOverlay.h"
 #include "Widgets/Text/STextBlock.h"
@@ -136,7 +137,7 @@ void SSessionInfoWindow::Construct(const FArguments& InArgs, const TSharedRef<SD
 
 ////////////////////////////////////////////////////////////////////////////////////////////////////
 
-void SSessionInfoWindow::AddInfoLine(TSharedPtr<SVerticalBox> InVerticalBox, const FText& InHeader, const TAttribute<FText>& InValue) const
+void SSessionInfoWindow::AddInfoLine(TSharedPtr<SVerticalBox> InVerticalBox, const FText& InHeader, const TAttribute<FText>& InValue, bool bMultiLine) const
 {
 	InVerticalBox->AddSlot()
 		.AutoHeight()
@@ -147,16 +148,30 @@ void SSessionInfoWindow::AddInfoLine(TSharedPtr<SVerticalBox> InVerticalBox, con
 			.Font(FCoreStyle::GetDefaultFontStyle("Bold", 11))
 		];
 
+	TSharedPtr<SWidget> TextBox;
+	if (bMultiLine)
+	{
+		TextBox = SNew(SMultiLineEditableTextBox)
+			.Text(InValue)
+			.HintText(LOCTEXT("HintText", "N/A"))
+			.AutoWrapText(true)
+			.BackgroundColor(FLinearColor(0.1f, 0.1f, 0.1f, 1.0f))
+			.IsReadOnly(true);
+	}
+	else
+	{
+		TextBox = SNew(SEditableTextBox)
+			.Text(InValue)
+			.HintText(LOCTEXT("HintText", "N/A"))
+			.BackgroundColor(FLinearColor(0.1f, 0.1f, 0.1f, 1.0f))
+			.IsReadOnly(true);
+	}
+
 	InVerticalBox->AddSlot()
 		.AutoHeight()
 		.Padding(8.0f, 0.0f, 8.0f, 4.0f)
 		[
-			SNew(SEditableTextBox)
-			.Text(InValue)
-			.HintText(LOCTEXT("HintText", "N/A"))
-			//.BackgroundColor(FLinearColor(0.243f, 0.243f, 0.243f, 1.0f))
-			.BackgroundColor(FLinearColor(0.1f, 0.1f, 0.1f, 1.0f))
-			.IsReadOnly(true)
+			TextBox.ToSharedRef()
 		];
 }
 
@@ -175,7 +190,7 @@ TSharedRef<SDockTab> SSessionInfoWindow::SpawnTab_SessionInfo(const FSpawnTabArg
 		[
 			SNew(SOverlay)
 
-			//Overlay slot for the Background image
+			// Overlay slot for the Background image
 			+ SOverlay::Slot()
 			.HAlign(HAlign_Fill)
 			.VAlign(VAlign_Fill)
@@ -183,7 +198,7 @@ TSharedRef<SDockTab> SSessionInfoWindow::SpawnTab_SessionInfo(const FSpawnTabArg
 				SAssignNew(Image, SImage)
 			]
 
-			//Overlay slot for the ScrollBox containing the data
+			// Overlay slot for the ScrollBox containing the data
 			+ SOverlay::Slot()
 			.HAlign(HAlign_Fill)
 			.VAlign(VAlign_Top)
@@ -219,10 +234,10 @@ TSharedRef<SDockTab> SSessionInfoWindow::SpawnTab_SessionInfo(const FSpawnTabArg
 	AddInfoLine(VerticalBox, LOCTEXT("AppName_HeaderText",		"Application Name:"),	TAttribute<FText>(this, &SSessionInfoWindow::GetAppNameText));
 	AddInfoLine(VerticalBox, LOCTEXT("BuildConfig_HeaderText",	"Build Config:"),		TAttribute<FText>(this, &SSessionInfoWindow::GetBuildConfigText));
 	AddInfoLine(VerticalBox, LOCTEXT("BuildTarget_HeaderText",	"Build Target:"),		TAttribute<FText>(this, &SSessionInfoWindow::GetBuildTargetText));
-	AddInfoLine(VerticalBox, LOCTEXT("CommandLine_HeaderText",	"Command Line:"),		TAttribute<FText>(this, &SSessionInfoWindow::GetCommandLineText));
+	AddInfoLine(VerticalBox, LOCTEXT("CommandLine_HeaderText",	"Command Line:"),		TAttribute<FText>(this, &SSessionInfoWindow::GetCommandLineText), true);
 	//AddInfoLine(VerticalBox, LOCTEXT("FileSize_HeaderText",		"File Size:"),			TAttribute<FText>(this, &SSessionInfoWindow::GetFileSizeText));
-	AddInfoLine(VerticalBox, LOCTEXT("Status_HeaderText",		"Status:"),				TAttribute<FText>(this, &SSessionInfoWindow::GetStatusText));
-	AddInfoLine(VerticalBox, LOCTEXT("Modules_HeaderText",		"Modules:"),			TAttribute<FText>(this, &SSessionInfoWindow::GetModulesText));
+	AddInfoLine(VerticalBox, LOCTEXT("Status_HeaderText",		"Analysis Status:"),	TAttribute<FText>(this, &SSessionInfoWindow::GetStatusText));
+	AddInfoLine(VerticalBox, LOCTEXT("Modules_HeaderText",		"Analysis Modules:"),	TAttribute<FText>(this, &SSessionInfoWindow::GetModulesText));
 
 	DockTab->SetOnTabClosed(SDockTab::FOnTabClosedCallback::CreateRaw(this, &SSessionInfoWindow::OnSessionInfoTabClosed));
 
@@ -419,11 +434,14 @@ FText SSessionInfoWindow::GetStatusText() const
 	TSharedPtr<FInsightsManager> InsightsManager = FInsightsManager::Get();
 	InsightsManager->UpdateSessionDuration();
 
+	FNumberFormattingOptions FormattingOptions;
+	FormattingOptions.MaximumFractionalDigits = 1;
+
 	FText Status = FText::Format(LOCTEXT("StatusFmt", "{0}\nSession Duration: {1}\nAnalyzed in {2} at {3}X speed."),
 		InsightsManager->IsAnalysisComplete() ? FText::FromString(FString(TEXT("ANALYSIS COMPLETED."))) : FText::FromString(FString(TEXT("ANALYZING..."))),
 		FText::FromString(TimeUtils::FormatTimeAuto(InsightsManager->GetSessionDuration(), 2)),
 		FText::FromString(TimeUtils::FormatTimeAuto(InsightsManager->GetAnalysisDuration(), 2)),
-		FMath::RoundToInt(static_cast<float>(InsightsManager->GetAnalysisSpeedFactor())));
+		FText::AsNumber(InsightsManager->GetAnalysisSpeedFactor(), &FormattingOptions));
 
 	return Status;
 }
diff --git a/Engine/Source/Developer/TraceInsights/Private/Insights/Widgets/SSessionInfoWindow.h b/Engine/Source/Developer/TraceInsights/Private/Insights/Widgets/SSessionInfoWindow.h
--- a/Engine/Source/Developer/TraceInsights/Private/Insights/Widgets/SSessionInfoWindow.h
+++ b/Engine/Source/Developer/TraceInsights/Private/Insights/Widgets/SSessionInfoWindow.h
@@ -120,7 +120,7 @@ private:
 
 
 private:
-	void AddInfoLine(TSharedPtr<SVerticalBox> InVerticalBox, const FText& InHeader, const TAttribute<FText>& InValue) const;
+	void AddInfoLine(TSharedPtr<SVerticalBox> InVerticalBox, const FText& InHeader, const TAttribute<FText>& InValue, bool bMultiLine = false) const;
 
 	TSharedRef<SDockTab> SpawnTab_SessionInfo(const FSpawnTabArgs& Args);
 	void OnSessionInfoTabClosed(TSharedRef<SDockTab> TabBeingClosed);
