[Insights] Build Allocation Tree only when a query is complete

diff --git a/Engine/Source/Developer/TraceInsights/Private/Insights/MemoryProfiler/Widgets/SMemAllocTableTreeView.cpp b/Engine/Source/Developer/TraceInsights/Private/Insights/MemoryProfiler/Widgets/SMemAllocTableTreeView.cpp
--- a/Engine/Source/Developer/TraceInsights/Private/Insights/MemoryProfiler/Widgets/SMemAllocTableTreeView.cpp
+++ b/Engine/Source/Developer/TraceInsights/Private/Insights/MemoryProfiler/Widgets/SMemAllocTableTreeView.cpp
@@ -6,7 +6,6 @@
 #include "TraceServices/Model/AllocationsProvider.h"
 
 // Insights
-#include "Insights/Common/Stopwatch.h"
 #include "Insights/MemoryProfiler/MemoryProfilerManager.h"
 #include "Insights/MemoryProfiler/ViewModels/MemAllocNode.h"
 #include "Insights/MemoryProfiler/ViewModels/MemAllocTable.h"
@@ -110,27 +109,31 @@ void SMemAllocTableTreeView::RebuildTree(bool bResync)
 	SyncStopwatch.Start();
 	if (Session.IsValid() && MemAllocTable.IsValid())
 	{
-		UpdateQuery();
+		TraceServices::IAllocationsProvider::EQueryStatus QueryStatus;
+		UpdateQuery(QueryStatus);
 
-		TArray<FMemoryAlloc>&  Allocs = MemAllocTable->GetAllocs();
-
-		const int32 TotalAllocCount = Allocs.Num();
-		if (TotalAllocCount != TableTreeNodes.Num())
+		if (QueryStatus == TraceServices::IAllocationsProvider::EQueryStatus::Done)
 		{
-			if (TableTreeNodes.Num() > TotalAllocCount)
-			{
-				TableTreeNodes.Empty();
-			}
-			TableTreeNodes.Reserve(TotalAllocCount);
+			TArray<FMemoryAlloc>& Allocs = MemAllocTable->GetAllocs();
 
-			FName BaseNodeName(TEXT("alloc"));
-			for (int32 AllocIndex = TableTreeNodes.Num(); AllocIndex < TotalAllocCount; ++AllocIndex)
+			const int32 TotalAllocCount = Allocs.Num();
+			if (TotalAllocCount != TableTreeNodes.Num())
 			{
-				FName NodeName(BaseNodeName, AllocIndex + 1);
-				FMemAllocNodePtr NodePtr = MakeShared<FMemAllocNode>(NodeName, MemAllocTable, AllocIndex);
-				TableTreeNodes.Add(NodePtr);
+				if (TableTreeNodes.Num() > TotalAllocCount)
+				{
+					TableTreeNodes.Empty();
+				}
+				TableTreeNodes.Reserve(TotalAllocCount);
+
+				FName BaseNodeName(TEXT("alloc"));
+				for (int32 AllocIndex = TableTreeNodes.Num(); AllocIndex < TotalAllocCount; ++AllocIndex)
+				{
+					FName NodeName(BaseNodeName, AllocIndex + 1);
+					FMemAllocNodePtr NodePtr = MakeShared<FMemAllocNode>(NodeName, MemAllocTable, AllocIndex);
+					TableTreeNodes.Add(NodePtr);
+				}
+				ensure(TableTreeNodes.Num() == TotalAllocCount);
 			}
-			ensure(TableTreeNodes.Num() == TotalAllocCount);
 		}
 	}
 	SyncStopwatch.Stop();
@@ -219,102 +222,113 @@ void SMemAllocTableTreeView::StartQuery()
 	{
 		UE_LOG(MemoryProfiler, Error, TEXT("[MemAlloc] Unsupported query rule (%s)!"), *Rule->GetShortName().ToString());
 	}
+	else
+	{
+		QueryStopwatch.Reset();
+		QueryStopwatch.Start();
+	}
 }
 
 ////////////////////////////////////////////////////////////////////////////////////////////////////
 
-void SMemAllocTableTreeView::UpdateQuery()
+void SMemAllocTableTreeView::UpdateQuery(TraceServices::IAllocationsProvider::EQueryStatus& OutStatus)
 {
-	if (Query != 0)
+	if (Query == 0)
 	{
-		const TraceServices::IAllocationsProvider* AllocationsProvider = TraceServices::ReadAllocationsProvider(*Session.Get());
-		if (!AllocationsProvider)
-		{
-			UE_LOG(MemoryProfiler, Warning, TEXT("[MemAlloc] Invalid allocations provider!"));
-			return;
-		}
-		const TraceServices::IAllocationsProvider& Provider = *AllocationsProvider;
+		OutStatus = TraceServices::IAllocationsProvider::EQueryStatus::Unknown;
+		return;
+	}
 
-		TraceServices::IAllocationsProvider::FQueryStatus Status = Provider.PollQuery(Query);
-		if (Status.Status <= TraceServices::IAllocationsProvider::EQueryStatus::Done)
-		{
-			UE_LOG(MemoryProfiler, Log, TEXT("[MemAlloc] Query completed."));
-			Query = 0;
-			return;
-		}
+	const TraceServices::IAllocationsProvider* AllocationsProvider = TraceServices::ReadAllocationsProvider(*Session.Get());
+	if (!AllocationsProvider)
+	{
+		UE_LOG(MemoryProfiler, Warning, TEXT("[MemAlloc] Invalid allocations provider!"));
+		return;
+	}
+	const TraceServices::IAllocationsProvider& Provider = *AllocationsProvider;
 
-		TSharedPtr<Insights::FMemAllocTable> MemAllocTable = GetMemAllocTable();
-		if (MemAllocTable)
-		{
-			TArray<FMemoryAlloc>& Allocs = MemAllocTable->GetAllocs();
+	TraceServices::IAllocationsProvider::FQueryStatus Status = Provider.PollQuery(Query);
 
-			FStopwatch TotalStopwatch;
-			FStopwatch PageStopwatch;
-			TotalStopwatch.Start();
-			uint32 PageCount = 0;
-			uint32 TotatAllocCount = 0;
+	OutStatus = Status.Status;
+	if (Status.Status <= TraceServices::IAllocationsProvider::EQueryStatus::Done)
+	{
+		UE_LOG(MemoryProfiler, Log, TEXT("[MemAlloc] Query completed."));
+		Query = 0;
+		QueryStopwatch.Stop();
+		return;
+	}
 
-			// Multiple 'pages' of results will be returned. No guarantees are made
-			// about the order of pages or the allocations they report.
-			TraceServices::IAllocationsProvider::FQueryResult Result = Status.NextResult();
-			while (Result.IsValid())
-			{
-				++PageCount;
-				PageStopwatch.Restart();
+	TSharedPtr<Insights::FMemAllocTable> MemAllocTable = GetMemAllocTable();
+	if (MemAllocTable)
+	{
+		TArray<FMemoryAlloc>& Allocs = MemAllocTable->GetAllocs();
 
-				const uint32 AllocCount = Result->Num();
-				TotatAllocCount += AllocCount;
-#if 0
-				Allocs.Reserve(Allocs.Num() + AllocCount);
+		FStopwatch TotalStopwatch;
+		FStopwatch PageStopwatch;
+		TotalStopwatch.Start();
+		uint32 PageCount = 0;
+		uint32 TotatAllocCount = 0;
 
-				for (uint32 AllocIndex = 0; AllocIndex < AllocCount; ++AllocIndex)
-				{
-					const IAllocationsProvider::FAllocation* Allocation = Result->Get(AllocIndex);
+		// Multiple 'pages' of results will be returned. No guarantees are made
+		// about the order of pages or the allocations they report.
+		TraceServices::IAllocationsProvider::FQueryResult Result = Status.NextResult();
+		while (Result.IsValid())
+		{
+			++PageCount;
+			PageStopwatch.Restart();
 
-					const double AllocStartTime = Allocation->GetStartTime();
-					const double AllocEndTime = Allocation->GetEndTime();
-					const uint64 AllocAddress = Allocation->GetAddress();
-					const uint64 AllocationSize = Allocation->GetSize();
-					const FMemoryTagId AllocTag = static_cast<FMemoryTagId>(Allocation->GetTag());
-					const uint64 BacktraceId = Allocation->GetBacktraceId();
+			const uint32 AllocCount = Result->Num();
+			TotatAllocCount += AllocCount;
+#if 0
+			Allocs.Reserve(Allocs.Num() + AllocCount);
 
-					//Allocs.Add(FMemoryAlloc(AllocStartTime, AllocEndTime, AllocAddress, AllocationSize, AllocTag, BacktraceId));
-					Allocs.Emplace(AllocStartTime, AllocEndTime, AllocAddress, AllocationSize, AllocTag, BacktraceId);
-				}
-#else // 14% faster, but uses more mem
-				uint64 AllocsDestIndex = Allocs.Num();
-				Allocs.AddUninitialized(AllocCount);
-				for (uint32 AllocIndex = 0; AllocIndex < AllocCount; ++AllocIndex, ++AllocsDestIndex)
-				{
-					const TraceServices::IAllocationsProvider::FAllocation* Allocation = Result->Get(AllocIndex);
-					FMemoryAlloc& Alloc = Allocs[AllocsDestIndex];
-					Alloc.StartTime = Allocation->GetStartTime();
-					Alloc.EndTime = Allocation->GetEndTime();
-					Alloc.Address = Allocation->GetAddress();
-					Alloc.Size = Allocation->GetSize();
-					Alloc.MemTag = static_cast<FMemoryTagId>(Allocation->GetTag());
-					Alloc.BacktraceId = Allocation->GetBacktraceId();
-				}
-#endif
+			for (uint32 AllocIndex = 0; AllocIndex < AllocCount; ++AllocIndex)
+			{
+				const IAllocationsProvider::FAllocation* Allocation = Result->Get(AllocIndex);
 
-				PageStopwatch.Stop();
-				const double PageTime = PageStopwatch.GetAccumulatedTime();
-				if (PageTime > 0.01)
-				{
-					const double Speed = (PageTime * 1000000.0) / AllocCount;
-					UE_LOG(MemoryProfiler, Log, TEXT("[MemAlloc] Query result for page %u (%u allocs, slack=%u) retrieved in %.3fs (speed: %.3f seconds per 1M allocs)."), PageCount, AllocCount, Allocs.GetSlack(), PageTime, Speed);
-				}
+				const double AllocStartTime = Allocation->GetStartTime();
+				const double AllocEndTime = Allocation->GetEndTime();
+				const uint64 AllocAddress = Allocation->GetAddress();
+				const uint64 AllocationSize = Allocation->GetSize();
+				const FMemoryTagId AllocTag = static_cast<FMemoryTagId>(Allocation->GetTag());
+				const uint64 BacktraceId = Allocation->GetBacktraceId();
 
-				Result = Status.NextResult();
+				//Allocs.Add(FMemoryAlloc(AllocStartTime, AllocEndTime, AllocAddress, AllocationSize, AllocTag, BacktraceId));
+				Allocs.Emplace(AllocStartTime, AllocEndTime, AllocAddress, AllocationSize, AllocTag, BacktraceId);
 			}
+#else // 14% faster, but uses more mem
+			uint64 AllocsDestIndex = Allocs.Num();
+			Allocs.AddUninitialized(AllocCount);
+			for (uint32 AllocIndex = 0; AllocIndex < AllocCount; ++AllocIndex, ++AllocsDestIndex)
+			{
+				const TraceServices::IAllocationsProvider::FAllocation* Allocation = Result->Get(AllocIndex);
+				FMemoryAlloc& Alloc = Allocs[AllocsDestIndex];
+				Alloc.StartTime = Allocation->GetStartTime();
+				Alloc.EndTime = Allocation->GetEndTime();
+				Alloc.Address = Allocation->GetAddress();
+				Alloc.Size = Allocation->GetSize();
+				Alloc.MemTag = static_cast<FMemoryTagId>(Allocation->GetTag());
+				Alloc.BacktraceId = Allocation->GetBacktraceId();
+			}
+#endif
 
-			TotalStopwatch.Stop();
-			const double TotalTime = TotalStopwatch.GetAccumulatedTime();
-			if (TotalTime > 0.01)
+			PageStopwatch.Stop();
+			const double PageTime = PageStopwatch.GetAccumulatedTime();
+			if (PageTime > 0.01)
 			{
-				const double Speed = (TotalTime * 1000000.0) / TotatAllocCount;
-				UE_LOG(MemoryProfiler, Log, TEXT("[MemAlloc] Query results (%u pages, %u allocs, slack=%u) retrieved in %.3fs (speed: %.3f seconds per 1M allocs)."), PageCount, TotatAllocCount, Allocs.GetSlack(), TotalTime, Speed);
+				const double Speed = (PageTime * 1000000.0) / AllocCount;
+				UE_LOG(MemoryProfiler, Log, TEXT("[MemAlloc] Query result for page %u (%u allocs, slack=%u) retrieved in %.3fs (speed: %.3f seconds per 1M allocs)."), PageCount, AllocCount, Allocs.GetSlack(), PageTime, Speed);
 			}
+
+			Result = Status.NextResult();
+		}
+
+		TotalStopwatch.Stop();
+		const double TotalTime = TotalStopwatch.GetAccumulatedTime();
+		if (TotalTime > 0.01)
+		{
+			const double Speed = (TotalTime * 1000000.0) / TotatAllocCount;
+			UE_LOG(MemoryProfiler, Log, TEXT("[MemAlloc] Query results (%u pages, %u allocs, slack=%u) retrieved in %.3fs (speed: %.3f seconds per 1M allocs)."), PageCount, TotatAllocCount, Allocs.GetSlack(), TotalTime, Speed);
 		}
 	}
 }
@@ -331,7 +345,9 @@ void SMemAllocTableTreeView::CancelQuery()
 			AllocationsProvider->CancelQuery(Query);
 			UE_LOG(MemoryProfiler, Log, TEXT("[MemAlloc] Query canceled."));
 		}
+
 		Query = 0;
+		QueryStopwatch.Stop();
 	}
 }
 
@@ -426,6 +442,38 @@ TSharedRef<ITableCellValueFormatter> SMemAllocTableTreeView::CreateCachedLlmTagV
 
 ////////////////////////////////////////////////////////////////////////////////////////////////////
 
+bool SMemAllocTableTreeView::IsRunning() const
+{
+	return Query != 0 || STableTreeView::IsRunning();
+}
+
+////////////////////////////////////////////////////////////////////////////////////////////////////
+
+double SMemAllocTableTreeView::GetAllOperationsDuration()
+{
+	if (Query != 0)
+	{
+		QueryStopwatch.Update();
+		return QueryStopwatch.GetAccumulatedTime();
+	}
+
+	return STableTreeView::GetAllOperationsDuration();
+}
+
+////////////////////////////////////////////////////////////////////////////////////////////////////
+
+FText SMemAllocTableTreeView::GetCurrentOperationName() const
+{
+	if (Query != 0)
+	{
+		return LOCTEXT("CurrentOperationName", "Running Query");
+	}
+
+	return STableTreeView::GetCurrentOperationName();
+}
+
+////////////////////////////////////////////////////////////////////////////////////////////////////
+
 } // namespace Insights
 
 #undef LOCTEXT_NAMESPACE
diff --git a/Engine/Source/Developer/TraceInsights/Private/Insights/MemoryProfiler/Widgets/SMemAllocTableTreeView.h b/Engine/Source/Developer/TraceInsights/Private/Insights/MemoryProfiler/Widgets/SMemAllocTableTreeView.h
--- a/Engine/Source/Developer/TraceInsights/Private/Insights/MemoryProfiler/Widgets/SMemAllocTableTreeView.h
+++ b/Engine/Source/Developer/TraceInsights/Private/Insights/MemoryProfiler/Widgets/SMemAllocTableTreeView.h
@@ -7,6 +7,7 @@
 #include "TraceServices/Model/AllocationsProvider.h"
 
 // Insights
+#include "Insights/Common/Stopwatch.h"
 #include "Insights/MemoryProfiler/ViewModels/MemAllocTable.h"
 #include "Insights/Table/Widgets/STableTreeView.h"
 
@@ -70,6 +71,16 @@ public:
 	int32 GetTabIndex() const { return TabIndex; }
 	void SetTabIndex(int32 InTabIndex) { TabIndex = InTabIndex; }
 
+
+	////////////////////////////////////////////////////////////////////////////////////////////////////
+	// IAsyncOperationStatusProvider implementation
+
+	virtual bool IsRunning() const override;
+	virtual double GetAllOperationsDuration() override;
+	virtual FText GetCurrentOperationName() const override;
+
+	////////////////////////////////////////////////////////////////////////////////////////////////////
+
 protected:
 	virtual void OnPreAsyncUpdate() override;
 	virtual void OnPostAsyncUpdate() override;
@@ -77,7 +88,7 @@ protected:
 private:
 	void OnQueryInvalidated();
 	void StartQuery();
-	void UpdateQuery();
+	void UpdateQuery(TraceServices::IAllocationsProvider::EQueryStatus& OutStatus);
 	void CancelQuery();
 
 	TSharedRef<class ITableCellValueFormatter> CreateCachedLlmTagValueFormatter();
@@ -89,6 +100,8 @@ private:
 	TraceServices::IAllocationsProvider::FQueryHandle Query = 0;
 
 	TSharedPtr<class ITableCellValueFormatter> OriginalLlmTagValueFormatter;
+
+	FStopwatch QueryStopwatch;
 };
 
 ////////////////////////////////////////////////////////////////////////////////////////////////////
