[Insights]
  - Fixed the Allocations Analysis to process the updated alloc trace events.
  - Added timelines in the Allocations Provider with variation for several stats based on 1ms sample:
      * MinTotalAllocatedMemory
      * MaxTotalAllocatedMemory
      * MinLiveAllocationCount
      * MaxLiveAllocationCount
      * AllocEventCount
      * FreeEventCount
  - Added graph tracks and graph series for all above timelines in Memory Insights.

diff --git a/Engine/Source/Developer/TraceInsights/Private/Insights/MemoryProfiler/MemoryProfilerManager.cpp b/Engine/Source/Developer/TraceInsights/Private/Insights/MemoryProfiler/MemoryProfilerManager.cpp
--- a/Engine/Source/Developer/TraceInsights/Private/Insights/MemoryProfiler/MemoryProfilerManager.cpp
+++ b/Engine/Source/Developer/TraceInsights/Private/Insights/MemoryProfiler/MemoryProfilerManager.cpp
@@ -256,6 +256,7 @@ bool FMemoryProfilerManager::Tick(float DeltaTime)
 			const TraceServices::IAllocationsProvider* AllocationsProvider = TraceServices::ReadAllocationsProvider(*Session.Get());
 			if (AllocationsProvider)
 			{
+				TraceServices::IAllocationsProvider::FReadScopeLock _(*AllocationsProvider);
 				if (AllocationsProvider->IsInitialized())
 				{
 					bShouldBeAvailable = true;
diff --git a/Engine/Source/Developer/TraceInsights/Private/Insights/MemoryProfiler/ViewModels/MemoryGraphTrack.cpp b/Engine/Source/Developer/TraceInsights/Private/Insights/MemoryProfiler/ViewModels/MemoryGraphTrack.cpp
--- a/Engine/Source/Developer/TraceInsights/Private/Insights/MemoryProfiler/ViewModels/MemoryGraphTrack.cpp
+++ b/Engine/Source/Developer/TraceInsights/Private/Insights/MemoryProfiler/ViewModels/MemoryGraphTrack.cpp
@@ -6,6 +6,7 @@
 #include "Fonts/SlateFontInfo.h"
 #include "Styling/SlateBrush.h"
 #include "TraceServices/AnalysisService.h"
+#include "TraceServices/Model/AllocationsProvider.h"
 #include "TraceServices/Model/TimingProfiler.h"
 #include "TraceServices/Model/Counters.h"
 #include "TraceServices/Model/Memory.h"
@@ -31,18 +32,32 @@
 
 FString FMemoryGraphSeries::FormatValue(double Value) const
 {
-	const int64 MemValue = static_cast<int64>(Value);
-	if (MemValue == 0)
+	const int64 ValueInt64 = static_cast<int64>(Value);
+	if (ValueInt64 == 0)
 	{
 		return TEXT("0");
 	}
-	else if (MemValue < 0)
+
+	if (TimelineType <= ETimelineType::MaxTotalMem)
 	{
-		return FString::Printf(TEXT("-%s (%s bytes)"), *FText::AsMemory(-MemValue).ToString(), *FText::AsNumber(MemValue).ToString());
+		double UnitValue;
+		const TCHAR* UnitText;
+		FMemoryGraphTrack::GetUnit(EGraphTrackLabelUnit::Auto, FMath::Abs(Value), UnitValue, UnitText);
+
+		if (ValueInt64 < 0)
+		{
+			FString Auto = FMemoryGraphTrack::FormatValue(-Value, UnitValue, UnitText, 0);
+			return FString::Printf(TEXT("-%s (%s bytes)"), *Auto, *FText::AsNumber(ValueInt64).ToString());
+		}
+		else
+		{
+			FString Auto = FMemoryGraphTrack::FormatValue(Value, UnitValue, UnitText, 0);
+			return FString::Printf(TEXT("%s (%s bytes)"), *Auto, *FText::AsNumber(ValueInt64).ToString());
+		}
 	}
 	else
 	{
-		return FString::Printf(TEXT("%s (%s bytes)"), *FText::AsMemory(MemValue).ToString(), *FText::AsNumber(MemValue).ToString());
+		return FText::AsNumber(ValueInt64).ToString();
 	}
 }
 
@@ -141,7 +156,14 @@ void FMemoryGraphTrack::Update(const ITimingTrackUpdateContext& Context)
 
 				if (bIsEntireGraphTrackDirty || Series->IsDirty())
 				{
-					PreUpdateMemTagSeries(*MemorySeries, Viewport);
+					if (MemorySeries->GetTimelineType() == FMemoryGraphSeries::ETimelineType::MemTag)
+					{
+						PreUpdateMemTagSeries(*MemorySeries, Viewport);
+					}
+					else
+					{
+						PreUpdateAllocationsTimelineSeries(*MemorySeries, Viewport);
+					}
 				}
 
 				const double CurrentSeriesMinValue = MemorySeries->GetMinValue();
@@ -185,7 +207,15 @@ void FMemoryGraphTrack::Update(const ITimingTrackUpdateContext& Context)
 
 				//TODO: if (Series->Is<FMemoryGraphSeries>())
 				TSharedPtr<FMemoryGraphSeries> MemorySeries = StaticCastSharedPtr<FMemoryGraphSeries>(Series);
-				UpdateMemTagSeries(*MemorySeries, Viewport);
+
+				if (MemorySeries->GetTimelineType() == FMemoryGraphSeries::ETimelineType::MemTag)
+				{
+					UpdateMemTagSeries(*MemorySeries, Viewport);
+				}
+				else
+				{
+					UpdateAllocationsTimelineSeries(*MemorySeries, Viewport);
+				}
 
 				if (Series->IsAutoZoomDirty())
 				{
@@ -335,6 +365,252 @@ void FMemoryGraphTrack::UpdateMemTagSeries(FMemoryGraphSeries& Series, const FTi
 
 ////////////////////////////////////////////////////////////////////////////////////////////////////
 
+TSharedPtr<FMemoryGraphSeries> FMemoryGraphTrack::GetTimelineSeries(FMemoryGraphSeries::ETimelineType InTimelineType)
+{
+	return nullptr; //TODO
+}
+
+////////////////////////////////////////////////////////////////////////////////////////////////////
+
+TSharedPtr<FMemoryGraphSeries> FMemoryGraphTrack::AddTimelineSeries(FMemoryGraphSeries::ETimelineType InTimelineType)
+{
+	TSharedPtr<FMemoryGraphSeries> Series = GetTimelineSeries(InTimelineType);
+
+	if (!Series.IsValid())
+	{
+		Series = MakeShared<FMemoryGraphSeries>();
+
+		switch (InTimelineType)
+		{
+			case FMemoryGraphSeries::ETimelineType::MinTotalMem:
+			{
+				Series->SetName(TEXT("Total Allocated Memory (Min)"));
+				Series->SetDescription(TEXT("Minimum value per sample for the Total Allocated Memory"));
+				const FLinearColor Color = FLinearColor(0.0f, 0.5f, 1.0f, 1.0f);
+				const FLinearColor BorderColor(FMath::Min(Color.R + 0.4f, 1.0f), FMath::Min(Color.G + 0.4f, 1.0f), FMath::Min(Color.B + 0.4f, 1.0f), 1.0f);
+				Series->SetColor(Color, BorderColor, Color.CopyWithNewOpacity(0.1f));
+				break;
+			}
+
+			case FMemoryGraphSeries::ETimelineType::MaxTotalMem:
+			{
+				Series->SetName(TEXT("Total Allocated Memory (Max)"));
+				Series->SetDescription(TEXT("Maximum value per sample for the Total Allocated Memory"));
+				const FLinearColor Color = FLinearColor(1.0f, 0.25f, 1.0f, 1.0f);
+				const FLinearColor BorderColor(FMath::Min(Color.R + 0.4f, 1.0f), FMath::Min(Color.G + 0.4f, 1.0f), FMath::Min(Color.B + 0.4f, 1.0f), 1.0f);
+				Series->SetColor(Color, BorderColor, Color.CopyWithNewOpacity(0.1f));
+				break;
+			}
+
+			case FMemoryGraphSeries::ETimelineType::MinLiveAllocs:
+			{
+				Series->SetName(TEXT("Live Allocation Count (Min)"));
+				Series->SetDescription(TEXT("Minimum value per sample for the Live Allocation Count"));
+				const FLinearColor Color = FLinearColor(1.0f, 1.0f, 0.25f, 1.0f);
+				const FLinearColor BorderColor(FMath::Min(Color.R + 0.4f, 1.0f), FMath::Min(Color.G + 0.4f, 1.0f), FMath::Min(Color.B + 0.4f, 1.0f), 1.0f);
+				Series->SetColor(Color, BorderColor, Color.CopyWithNewOpacity(0.1f));
+				break;
+			}
+
+			case FMemoryGraphSeries::ETimelineType::MaxLiveAllocs:
+			{
+				Series->SetName(TEXT("Live Allocation Count (Max)"));
+				Series->SetDescription(TEXT("Maximum value per sample for the Live Allocation Count"));
+				const FLinearColor Color = FLinearColor(1.0f, 0.25f, 1.0f, 1.0f);
+				const FLinearColor BorderColor(FMath::Min(Color.R + 0.4f, 1.0f), FMath::Min(Color.G + 0.4f, 1.0f), FMath::Min(Color.B + 0.4f, 1.0f), 1.0f);
+				Series->SetColor(Color, BorderColor, Color.CopyWithNewOpacity(0.1f));
+				break;
+			}
+
+			case FMemoryGraphSeries::ETimelineType::AllocEvents:
+			{
+				Series->SetName(TEXT("Alloc Event Count"));
+				Series->SetDescription(TEXT("Number of alloc events per sample"));
+				const FLinearColor Color = FLinearColor(0.0f, 1.0f, 0.5f, 1.0f);
+				const FLinearColor BorderColor(FMath::Min(Color.R + 0.4f, 1.0f), FMath::Min(Color.G + 0.4f, 1.0f), FMath::Min(Color.B + 0.4f, 1.0f), 1.0f);
+				Series->SetColor(Color, BorderColor, Color.CopyWithNewOpacity(0.1f));
+				break;
+			}
+
+			case FMemoryGraphSeries::ETimelineType::FreeEvents:
+			{
+				Series->SetName(TEXT("Free Event Count"));
+				Series->SetDescription(TEXT("Number of free events per sample"));
+				const FLinearColor Color = FLinearColor(1.0f, 0.5f, 0.25f, 1.0f);
+				const FLinearColor BorderColor(FMath::Min(Color.R + 0.4f, 1.0f), FMath::Min(Color.G + 0.4f, 1.0f), FMath::Min(Color.B + 0.4f, 1.0f), 1.0f);
+				Series->SetColor(Color, BorderColor, Color.CopyWithNewOpacity(0.1f));
+				break;
+			}
+		}
+
+		Series->SetTimelineType(InTimelineType);
+		Series->SetValueRange(0.0f, 0.0f);
+
+		Series->SetBaselineY(GetHeight() - 1.0f);
+		Series->SetScaleY(1.0);
+
+		AllSeries.Add(Series);
+		SetDirtyFlag();
+	}
+
+	return Series;
+}
+
+////////////////////////////////////////////////////////////////////////////////////////////////////
+
+void FMemoryGraphTrack::PreUpdateAllocationsTimelineSeries(FMemoryGraphSeries& Series, const FTimingTrackViewport& Viewport)
+{
+	TSharedPtr<const TraceServices::IAnalysisSession> Session = FInsightsManager::Get()->GetSession();
+	if (Session.IsValid())
+	{
+		const TraceServices::IAllocationsProvider* AllocationsProvider = TraceServices::ReadAllocationsProvider(*Session.Get());
+		if (AllocationsProvider)
+		{
+			TraceServices::IAllocationsProvider::FReadScopeLock ProviderReadScope(*AllocationsProvider);
+
+			int32 StartIndex = -1;
+			int32 EndIndex = -1;
+			AllocationsProvider->GetTimelineIndexRange(Viewport.GetStartTime(), Viewport.GetEndTime(), StartIndex, EndIndex);
+			if (EndIndex >= 0)
+			{
+				--StartIndex; // include one more point on the left side
+				++EndIndex; // include one more point on the right side
+			}
+
+			uint64 MinValue = std::numeric_limits<uint64>::max();
+			uint64 MaxValue = 0;
+
+			// Compute Min/Max values.
+			auto Callback64 = [&MinValue, &MaxValue](double Time, double Duration, uint64 Value)
+			{
+				if (Value < MinValue)
+				{
+					MinValue = Value;
+				}
+				if (Value > MaxValue)
+				{
+					MaxValue = Value;
+				}
+			};
+			auto Callback32 = [&MinValue, &MaxValue](double Time, double Duration, uint32 Value)
+			{
+				if (Value < MinValue)
+				{
+					MinValue = Value;
+				}
+				if (Value > MaxValue)
+				{
+					MaxValue = Value;
+				}
+			};
+
+			switch (Series.GetTimelineType())
+			{
+			case FMemoryGraphSeries::ETimelineType::MinTotalMem:
+				AllocationsProvider->EnumerateMinTotalAllocatedMemoryTimeline(StartIndex, EndIndex, Callback64);
+				break;
+			case FMemoryGraphSeries::ETimelineType::MaxTotalMem:
+				AllocationsProvider->EnumerateMaxTotalAllocatedMemoryTimeline(StartIndex, EndIndex, Callback64);
+				break;
+			case FMemoryGraphSeries::ETimelineType::MinLiveAllocs:
+				AllocationsProvider->EnumerateMinLiveAllocationsTimeline(StartIndex, EndIndex, Callback32);
+				break;
+			case FMemoryGraphSeries::ETimelineType::MaxLiveAllocs:
+				AllocationsProvider->EnumerateMaxLiveAllocationsTimeline(StartIndex, EndIndex, Callback32);
+				break;
+			case FMemoryGraphSeries::ETimelineType::AllocEvents:
+				AllocationsProvider->EnumerateAllocEventsTimeline(StartIndex, EndIndex, Callback32);
+				break;
+			case FMemoryGraphSeries::ETimelineType::FreeEvents:
+				AllocationsProvider->EnumerateFreeEventsTimeline(StartIndex, EndIndex, Callback32);
+				break;
+			}
+
+			if (Series.GetTimelineType() == FMemoryGraphSeries::ETimelineType::FreeEvents)
+			{
+				// Shows FreeEvents as negative values in order to be displayed on same graph as AllocEvents.
+				Series.SetValueRange(-static_cast<double>(MaxValue), -static_cast<double>(MinValue));
+			}
+			else
+			{
+				Series.SetValueRange(static_cast<double>(MinValue), static_cast<double>(MaxValue));
+			}
+		}
+	}
+}
+
+////////////////////////////////////////////////////////////////////////////////////////////////////
+
+void FMemoryGraphTrack::UpdateAllocationsTimelineSeries(FMemoryGraphSeries& Series, const FTimingTrackViewport& Viewport)
+{
+	FGraphTrackBuilder Builder(*this, Series, Viewport);
+
+	TSharedPtr<const TraceServices::IAnalysisSession> Session = FInsightsManager::Get()->GetSession();
+	if (Session.IsValid())
+	{
+		const TraceServices::IAllocationsProvider* AllocationsProvider = TraceServices::ReadAllocationsProvider(*Session.Get());
+		if (AllocationsProvider)
+		{
+			TraceServices::IAllocationsProvider::FReadScopeLock ProviderReadScope(*AllocationsProvider);
+
+			int32 StartIndex = -1;
+			int32 EndIndex = -1;
+			AllocationsProvider->GetTimelineIndexRange(Viewport.GetStartTime(), Viewport.GetEndTime(), StartIndex, EndIndex);
+			if (EndIndex >= 0)
+			{
+				--StartIndex; // include one more point on the left side
+				++EndIndex; // include one more point on the right side
+			}
+
+			const float TopY = 4.0f;
+			const float BottomY = GetHeight() - 4.0f;
+
+			if (Series.IsAutoZoomEnabled() && TopY < BottomY)
+			{
+				Series.UpdateAutoZoom(TopY, BottomY, Series.GetMinValue(), Series.GetMaxValue());
+			}
+
+			auto Callback64 = [&Builder](double Time, double Duration, uint64 Value)
+			{
+				Builder.AddEvent(Time, Duration, static_cast<double>(Value));
+			};
+			auto Callback32 = [&Builder](double Time, double Duration, uint32 Value)
+			{
+				Builder.AddEvent(Time, Duration, static_cast<double>(Value));
+			};
+			auto Callback32Negative = [&Builder](double Time, double Duration, uint32 Value)
+			{
+				// Shows FreeEvents as negative values in order to be displayed on same graph as AllocEvents.
+				Builder.AddEvent(Time, Duration, -static_cast<double>(Value));
+			};
+
+			switch (Series.GetTimelineType())
+			{
+			case FMemoryGraphSeries::ETimelineType::MinTotalMem:
+				AllocationsProvider->EnumerateMinTotalAllocatedMemoryTimeline(StartIndex, EndIndex, Callback64);
+				break;
+			case FMemoryGraphSeries::ETimelineType::MaxTotalMem:
+				AllocationsProvider->EnumerateMaxTotalAllocatedMemoryTimeline(StartIndex, EndIndex, Callback64);
+				break;
+			case FMemoryGraphSeries::ETimelineType::MinLiveAllocs:
+				AllocationsProvider->EnumerateMinLiveAllocationsTimeline(StartIndex, EndIndex, Callback32);
+				break;
+			case FMemoryGraphSeries::ETimelineType::MaxLiveAllocs:
+				AllocationsProvider->EnumerateMaxLiveAllocationsTimeline(StartIndex, EndIndex, Callback32);
+				break;
+			case FMemoryGraphSeries::ETimelineType::AllocEvents:
+				AllocationsProvider->EnumerateAllocEventsTimeline(StartIndex, EndIndex, Callback32);
+				break;
+			case FMemoryGraphSeries::ETimelineType::FreeEvents:
+				AllocationsProvider->EnumerateFreeEventsTimeline(StartIndex, EndIndex, Callback32Negative);
+				break;
+			}
+		}
+	}
+}
+
+////////////////////////////////////////////////////////////////////////////////////////////////////
+
 void FMemoryGraphTrack::Draw(const ITimingTrackDrawContext& Context) const
 {
 	FGraphTrack::Draw(Context);
@@ -357,49 +633,52 @@ void FMemoryGraphTrack::Draw(const ITimingTrackDrawContext& Context) const
 				//TODO: if (Series->Is<FMemoryGraphSeries>())
 				TSharedPtr<FMemoryGraphSeries> MemorySeries = StaticCastSharedPtr<FMemoryGraphSeries>(Series);
 
-				const Insights::FMemoryTag* Tag = SharedState.GetTagList().GetTagById(MemorySeries->GetTagId());
-				const uint64 TrackerFlag = 1ULL << MemorySeries->GetTrackerId();
-
-				if ((Tag->GetTrackers() & TrackerFlag) != TrackerFlag)
+				if (MemorySeries->GetTimelineType() == FMemoryGraphSeries::ETimelineType::MemTag)
 				{
-					FDrawContext& DrawContext = Context.GetDrawContext();
-					const FSlateBrush* Brush = Context.GetHelper().GetWhiteBrush();
+					const Insights::FMemoryTag* Tag = SharedState.GetTagList().GetTagById(MemorySeries->GetTagId());
+					const uint64 TrackerFlag = 1ULL << MemorySeries->GetTrackerId();
 
-					if (WarningCount == 0)
+					if ((Tag->GetTrackers() & TrackerFlag) != TrackerFlag)
 					{
-						const FText WarningText1 = FText::Format(LOCTEXT("WarningFmt1", "Warning: One or more LLM tags in this graph are not used by the current tracker ({0})!"),
-							SharedState.GetCurrentTracker() ? FText::FromString(SharedState.GetCurrentTracker()->GetName()) : FText::GetEmpty());
-						const FLinearColor WarningTextColor1(1.0f, 0.2f, 0.2f, 1.0f);
-						DrawContext.DrawText(DrawContext.LayerId, X, Y, WarningText1.ToString(), Font, WarningTextColor1);
-						Y += 12.0f;
-
-						const FText WarningText2 = FText::Format(LOCTEXT("WarningLine2", "In order to see graph series for these LLM tags, change the current tracker to one that uses the LLM tag."),
-							SharedState.GetCurrentTracker() ? FText::FromString(SharedState.GetCurrentTracker()->GetName()) : FText::GetEmpty());
-						const FLinearColor WarningTextColor2(0.75f, 0.5f, 0.25f, 1.0f);
-						DrawContext.DrawText(DrawContext.LayerId, X, Y, WarningText2.ToString(), Font, WarningTextColor2);
-						Y += 12.0f;
-					}
+						FDrawContext& DrawContext = Context.GetDrawContext();
+						const FSlateBrush* Brush = Context.GetHelper().GetWhiteBrush();
 
-					if (Tag->GetTrackers() != 0)
-					{
-						const FText Conjunction = LOCTEXT("WarningTrackersConjunction", " and the ");
-						const FText Text = FText::Format(LOCTEXT("WarningFmt2", "The {0} LLM tag is only used by the {1} tracker(s)."),
-							FText::FromString(Tag->GetStatName()),
-							FText::FromString(SharedState.TrackersToString(Tag->GetTrackers(), *Conjunction.ToString())));
-						const FLinearColor TextColor(1.0f, 0.75f, 0.5f, 1.0f);
-						DrawContext.DrawText(DrawContext.LayerId, X, Y, Text.ToString(), Font, TextColor);
-						Y += 12.0f;
-					}
-					else
-					{
-						const FText Text = FText::Format(LOCTEXT("WarningFmt3", "The {0} LLM tag is not used by any tracker."),
-							FText::FromString(Tag->GetStatName()));
-						const FLinearColor TextColor(1.0f, 0.75f, 0.5f, 1.0f);
-						DrawContext.DrawText(DrawContext.LayerId, X, Y, Text.ToString(), Font, TextColor);
-						Y += 12.0f;
-					}
+						if (WarningCount == 0)
+						{
+							const FText WarningText1 = FText::Format(LOCTEXT("WarningFmt1", "Warning: One or more LLM tags in this graph are not used by the current tracker ({0})!"),
+								SharedState.GetCurrentTracker() ? FText::FromString(SharedState.GetCurrentTracker()->GetName()) : FText::GetEmpty());
+							const FLinearColor WarningTextColor1(1.0f, 0.2f, 0.2f, 1.0f);
+							DrawContext.DrawText(DrawContext.LayerId, X, Y, WarningText1.ToString(), Font, WarningTextColor1);
+							Y += 12.0f;
+
+							const FText WarningText2 = FText::Format(LOCTEXT("WarningLine2", "In order to see graph series for these LLM tags, change the current tracker to one that uses the LLM tag."),
+								SharedState.GetCurrentTracker() ? FText::FromString(SharedState.GetCurrentTracker()->GetName()) : FText::GetEmpty());
+							const FLinearColor WarningTextColor2(0.75f, 0.5f, 0.25f, 1.0f);
+							DrawContext.DrawText(DrawContext.LayerId, X, Y, WarningText2.ToString(), Font, WarningTextColor2);
+							Y += 12.0f;
+						}
+
+						if (Tag->GetTrackers() != 0)
+						{
+							const FText Conjunction = LOCTEXT("WarningTrackersConjunction", " and the ");
+							const FText Text = FText::Format(LOCTEXT("WarningFmt2", "The {0} LLM tag is only used by the {1} tracker(s)."),
+								FText::FromString(Tag->GetStatName()),
+								FText::FromString(SharedState.TrackersToString(Tag->GetTrackers(), *Conjunction.ToString())));
+							const FLinearColor TextColor(1.0f, 0.75f, 0.5f, 1.0f);
+							DrawContext.DrawText(DrawContext.LayerId, X, Y, Text.ToString(), Font, TextColor);
+							Y += 12.0f;
+						}
+						else
+						{
+							const FText Text = FText::Format(LOCTEXT("WarningFmt3", "The {0} LLM tag is not used by any tracker."),
+								FText::FromString(Tag->GetStatName()));
+							const FLinearColor TextColor(1.0f, 0.75f, 0.5f, 1.0f);
+							DrawContext.DrawText(DrawContext.LayerId, X, Y, Text.ToString(), Font, TextColor);
+							Y += 12.0f;
+						}
 
-					++WarningCount;
+						++WarningCount;
+					}
 				}
 			}
 		}
@@ -491,33 +770,44 @@ void FMemoryGraphTrack::DrawVerticalAxisGrid(const ITimingTrackDrawContext& Cont
 
 	if (Delta > 0.0)
 	{
-		int64 DeltaBytes = static_cast<int64>(Delta);
-		if (DeltaBytes <= 0)
-		{
-			DeltaBytes = 1;
-		}
+		double Grid;
 
-		// Compute rounding based on magnitude of visible range of values (Delta).
-		int64 Delta2 = DeltaBytes;
-		int64 Power2 = 1;
-		while (Delta2 > 0)
+		if (Series->GetTimelineType() <= FMemoryGraphSeries::ETimelineType::MaxTotalMem)
 		{
-			Delta2 >>= 1;
-			Power2 <<= 1;
+			const uint64 DeltaBytes = FMath::Max(1ULL, static_cast<uint64>(Delta));
+			Grid = static_cast<double>(FMath::RoundUpToPowerOfTwo64(DeltaBytes));
 		}
-		if (Power2 >> 1 == DeltaBytes)
+		else
 		{
-			Power2 >>= 1;
+			const uint64 DeltaCount = FMath::Max(1ULL, static_cast<uint64>(Delta));
+
+			// Compute rounding based on magnitude of visible range of values (Delta).
+			uint64 Delta10 = DeltaCount;
+			uint64 Power10 = 1;
+			while (Delta10 > 0)
+			{
+				Delta10 /= 10;
+				Power10 *= 10;
+			}
+			if (Power10 >= 100)
+			{
+				Power10 /= 100;
+			}
+			else
+			{
+				Power10 = 1;
+			}
+
+			// Compute Grid as the next value divisible with a multiple of 10.
+			Grid = static_cast<double>(((DeltaCount + Power10 - 1) / Power10) * Power10);
 		}
-		const double Grid = static_cast<double>(Power2);
+
 		const double StartValue = FMath::GridSnap(BottomValue, Grid);
 
 		TDrawHorizontalAxisLabelParams Params(DrawContext, Brush, FontMeasureService);
 		Params.TextBgColor = FLinearColor(0.05f, 0.05f, 0.05f, 1.0f);
 		Params.TextColor = FLinearColor(1.0f, 1.0f, 1.0f, 1.0f);
 		Params.X = X0;
-		//Params.Y;
-		//Params.Value;
 		Params.Precision = Precision;
 
 		const FLinearColor GridColor(0.0f, 0.0f, 0.0f, 0.1f);
@@ -560,8 +850,6 @@ void FMemoryGraphTrack::DrawVerticalAxisGrid(const ITimingTrackDrawContext& Cont
 		}
 
 		Params.X = X0;
-		//Params.Y;
-		//Params.Value;
 		Params.Precision = -Precision; // format with detailed text
 
 		int32 MinMaxAxis = 0;
@@ -569,8 +857,8 @@ void FMemoryGraphTrack::DrawVerticalAxisGrid(const ITimingTrackDrawContext& Cont
 		// Draw horizontal axis at the max value.
 		if (MaxValueY >= Y0 && MaxValueY <= Y0 + RoundedViewHeight)
 		{
-			Params.Value = AllSeriesMaxValue;
 			Params.Y = MaxValueY;
+			Params.Value = AllSeriesMaxValue;
 			DrawHorizontalAxisLabel(Params);
 			++MinMaxAxis;
 		}
@@ -594,10 +882,6 @@ void FMemoryGraphTrack::DrawVerticalAxisGrid(const ITimingTrackDrawContext& Cont
 
 			if (MX > ViewWidth - MX2 && MY >= MaxValueY && MY < MinValueY + TextH)
 			{
-				//const float Alpha = FMath::Min(1.0f, (MX - ViewWidth + MX2) / (MX2 - MX1));
-				//Params.TextBgColor.A = Alpha;
-				//Params.TextColor.A = Alpha;
-
 				const float LineX = MX - 16.0f;
 				DrawContext.DrawBox(DrawContext.LayerId + 1, LineX, ActualMaxValueY, X0 - LineX, 1.0f, Params.Brush, Params.TextBgColor);
 				DrawContext.DrawBox(DrawContext.LayerId + 1, LineX, ActualMaxValueY, 1.0f, ActualMinValueY - ActualMaxValueY, Params.Brush, Params.TextBgColor);
@@ -672,6 +956,13 @@ void FMemoryGraphTrack::GetUnit(const EGraphTrackLabelUnit InLabelUnit, const do
 	constexpr double PiB = (double)(1LL << 50); // 2^50 bytes
 	constexpr double EiB = (double)(1LL << 60); // 2^60 bytes
 
+	constexpr double K10 = 1000.0;    // 10^3
+	constexpr double M10 = K10 * K10; // 10^6
+	constexpr double G10 = M10 * K10; // 10^9
+	constexpr double T10 = G10 * K10; // 10^12
+	constexpr double P10 = T10 * K10; // 10^15
+	constexpr double E10 = P10 * K10; // 10^18
+
 	switch (InLabelUnit)
 	{
 		case EGraphTrackLabelUnit::Auto:
@@ -745,10 +1036,55 @@ void FMemoryGraphTrack::GetUnit(const EGraphTrackLabelUnit InLabelUnit, const do
 			break;
 
 		case EGraphTrackLabelUnit::Byte:
-		default:
 			OutUnitValue = 1.0;
 			OutUnitText = TEXT("B");
 			break;
+
+		case EGraphTrackLabelUnit::AutoCount:
+		{
+			if (InPrecision >= E10)
+			{
+				OutUnitValue = E10;
+				OutUnitText = TEXT("E");
+			}
+			else if (InPrecision >= P10)
+			{
+				OutUnitValue = P10;
+				OutUnitText = TEXT("P");
+			}
+			else if (InPrecision >= T10)
+			{
+				OutUnitValue = T10;
+				OutUnitText = TEXT("T");
+			}
+			else if (InPrecision >= G10)
+			{
+				OutUnitValue = G10;
+				OutUnitText = TEXT("G");
+			}
+			else if (InPrecision >= M10)
+			{
+				OutUnitValue = M10;
+				OutUnitText = TEXT("M");
+			}
+			else if (InPrecision >= K10)
+			{
+				OutUnitValue = K10;
+				OutUnitText = TEXT("K");
+			}
+			else
+			{
+				OutUnitValue = 1.0;
+				OutUnitText = TEXT("");
+			}
+		}
+		break;
+
+		case EGraphTrackLabelUnit::Count:
+		default:
+			OutUnitValue = 1.0;
+			OutUnitText = TEXT("");
+			break;
 	}
 }
 
@@ -756,6 +1092,11 @@ void FMemoryGraphTrack::GetUnit(const EGraphTrackLabelUnit InLabelUnit, const do
 
 FString FMemoryGraphTrack::FormatValue(const double InValue, const double InUnitValue, const TCHAR* InUnitText, const int32 InDecimalDigitCount)
 {
+	if (InUnitText[0] == TEXT('\0') && InDecimalDigitCount == 0)
+	{
+		return FText::AsNumber(static_cast<int64>(InValue)).ToString();
+	}
+
 	FString OutText;
 
 	TCHAR FormatString[32];
@@ -776,8 +1117,11 @@ FString FMemoryGraphTrack::FormatValue(const double InValue, const double InUnit
 		}
 	}
 
-	OutText += TEXT(' ');
-	OutText += InUnitText;
+	if (InUnitText[0] != TEXT('\0'))
+	{
+		OutText += TEXT(' ');
+		OutText += InUnitText;
+	}
 
 	return OutText;
 }
@@ -818,8 +1162,12 @@ void FMemoryGraphTrack::InitTooltip(FTooltipDrawState& InOutTooltip, const ITimi
 
 		InOutTooltip.ResetContent();
 		InOutTooltip.AddTitle(Series->GetName().ToString(), Series->GetColor());
-		FString SubTitle = FString::Printf(TEXT("(tag id %lli, tracker id %lli)"), (int64)Series->GetTagId(), (int64)Series->GetTrackerId());
-		InOutTooltip.AddTitle(SubTitle, Series->GetColor());
+
+		if (Series->GetTimelineType() == FMemoryGraphSeries::ETimelineType::MemTag)
+		{
+			FString SubTitle = FString::Printf(TEXT("(tag id %lli, tracker id %lli)"), (int64)Series->GetTagId(), (int64)Series->GetTrackerId());
+			InOutTooltip.AddTitle(SubTitle, Series->GetColor());
+		}
 
 		InOutTooltip.AddNameValueTextLine(TEXT("Time:"), TimeUtils::FormatTimeAuto(TooltipEvent.GetStartTime()));
 		if (Series->HasEventDuration())
diff --git a/Engine/Source/Developer/TraceInsights/Private/Insights/MemoryProfiler/ViewModels/MemoryGraphTrack.h b/Engine/Source/Developer/TraceInsights/Private/Insights/MemoryProfiler/ViewModels/MemoryGraphTrack.h
--- a/Engine/Source/Developer/TraceInsights/Private/Insights/MemoryProfiler/ViewModels/MemoryGraphTrack.h
+++ b/Engine/Source/Developer/TraceInsights/Private/Insights/MemoryProfiler/ViewModels/MemoryGraphTrack.h
@@ -27,13 +27,27 @@ enum class EGraphTrackLabelUnit
 	GiB, // 2^30 bytes (gibibyte)
 	TiB, // 2^40 bytes (tebibyte)
 	PiB, // 2^50 bytes (pebibyte)
-	EiB  // 2^60 bytes (exbibyte)
+	EiB, // 2^60 bytes (exbibyte)
+	AutoCount,
+	Count,
 };
 
 ////////////////////////////////////////////////////////////////////////////////////////////////////
 
 class FMemoryGraphSeries : public FGraphSeries
 {
+public:
+	enum class ETimelineType
+	{
+		MemTag,
+		MinTotalMem,
+		MaxTotalMem,
+		MinLiveAllocs,
+		MaxLiveAllocs,
+		AllocEvents,
+		FreeEvents
+	};
+
 public:
 	virtual FString FormatValue(double Value) const override;
 
@@ -43,15 +57,19 @@ public:
 	Insights::FMemoryTagId GetTagId() const { return TagId; }
 	void SetTagId(Insights::FMemoryTagId InTagId) { TagId = InTagId; }
 
+	ETimelineType GetTimelineType() const { return TimelineType; }
+	void SetTimelineType(ETimelineType InTimelineType) { TimelineType = InTimelineType; }
+
 	double GetMinValue() const { return MinValue; }
 	double GetMaxValue() const { return MaxValue; }
 	void SetValueRange(double Min, double Max) { MinValue = Min; MaxValue = Max; }
 
 private:
-	Insights::FMemoryTrackerId TrackerId; // LLM tracker id
-	Insights::FMemoryTagId TagId; // LLM tag id
-	double MinValue;
-	double MaxValue;
+	Insights::FMemoryTrackerId TrackerId = Insights::FMemoryTracker::InvalidTrackerId; // LLM tracker id
+	Insights::FMemoryTagId TagId = Insights::FMemoryTag::InvalidTagId; // LLM tag id
+	double MinValue = 0.0;
+	double MaxValue = 0.0;
+	ETimelineType TimelineType = ETimelineType::MemTag;
 };
 
 ////////////////////////////////////////////////////////////////////////////////////////////////////
@@ -99,13 +117,22 @@ public:
 	int32 RemoveMemTagSeries(Insights::FMemoryTagId MemTagId);
 	int32 RemoveAllMemTagSeries();
 
+	TSharedPtr<FMemoryGraphSeries> GetTimelineSeries(FMemoryGraphSeries::ETimelineType InTimelineType);
+	TSharedPtr<FMemoryGraphSeries> AddTimelineSeries(FMemoryGraphSeries::ETimelineType InTimelineType);
+
 	void SetAvailableTrackHeight(EMemoryTrackHeightMode InMode, float InTrackHeight);
 	void SetCurrentTrackHeight(EMemoryTrackHeightMode InMode);
 
+	static void GetUnit(const EGraphTrackLabelUnit InLabelUnit, const double InPrecision, double& OutUnitValue, const TCHAR*& OutUnitText);
+	static FString FormatValue(const double InValue, const double InUnitValue, const TCHAR* InUnitText, const int32 InDecimalDigitCount);
+
 protected:
 	void PreUpdateMemTagSeries(FMemoryGraphSeries& Series, const FTimingTrackViewport& Viewport);
 	void UpdateMemTagSeries(FMemoryGraphSeries& Series, const FTimingTrackViewport& Viewport);
 
+	void PreUpdateAllocationsTimelineSeries(FMemoryGraphSeries& Series, const FTimingTrackViewport& Viewport);
+	void UpdateAllocationsTimelineSeries(FMemoryGraphSeries& Series, const FTimingTrackViewport& Viewport);
+
 	virtual void DrawVerticalAxisGrid(const ITimingTrackDrawContext& Context) const override;
 
 	struct TDrawHorizontalAxisLabelParams
@@ -129,9 +156,6 @@ protected:
 	};
 	void DrawHorizontalAxisLabel(const TDrawHorizontalAxisLabelParams& Params) const;
 
-	static void GetUnit(const EGraphTrackLabelUnit InLabelUnit, const double InPrecision, double& OutUnitValue, const TCHAR*& OutUnitText);
-	static FString FormatValue(const double InValue, const double InUnitValue, const TCHAR* InUnitText, const int32 InDecimalDigitCount);
-
 protected:
 	FMemorySharedState& SharedState;
 
diff --git a/Engine/Source/Developer/TraceInsights/Private/Insights/MemoryProfiler/ViewModels/MemorySharedState.cpp b/Engine/Source/Developer/TraceInsights/Private/Insights/MemoryProfiler/ViewModels/MemorySharedState.cpp
--- a/Engine/Source/Developer/TraceInsights/Private/Insights/MemoryProfiler/ViewModels/MemorySharedState.cpp
+++ b/Engine/Source/Developer/TraceInsights/Private/Insights/MemoryProfiler/ViewModels/MemorySharedState.cpp
@@ -36,6 +36,8 @@ FMemorySharedState::FMemorySharedState()
 	, DefaultTracker(nullptr)
 	, CurrentTracker(nullptr)
 	, MainGraphTrack(nullptr)
+	, LiveAllocsGraphTrack(nullptr)
+	, AllocFreeGraphTrack(nullptr)
 	, TrackHeightMode(EMemoryTrackHeightMode::Medium)
 	, bShowHideAllMemoryTracks(false)
 	, CreatedDefaultTracks()
@@ -66,6 +68,8 @@ void FMemorySharedState::OnBeginSession(Insights::ITimingViewSession& InSession)
 	CurrentTracker = nullptr;
 
 	MainGraphTrack = nullptr;
+	LiveAllocsGraphTrack = nullptr;
+	AllocFreeGraphTrack = nullptr;
 	AllTracks.Reset();
 
 	bShowHideAllMemoryTracks = true;
@@ -91,6 +95,8 @@ void FMemorySharedState::OnEndSession(Insights::ITimingViewSession& InSession)
 	CurrentTracker = nullptr;
 
 	MainGraphTrack = nullptr;
+	LiveAllocsGraphTrack = nullptr;
+	AllocFreeGraphTrack = nullptr;
 	AllTracks.Reset();
 
 	bShowHideAllMemoryTracks = false;
@@ -112,9 +118,14 @@ void FMemorySharedState::Tick(Insights::ITimingViewSession& InSession, const Tra
 	if (!MainGraphTrack.IsValid())
 	{
 		MainGraphTrack = CreateMemoryGraphTrack();
+		check(MainGraphTrack);
 
 		MainGraphTrack->SetOrder(FTimingTrackOrder::First);
 		MainGraphTrack->SetName(TEXT("Main Memory Graph"));
+
+		MainGraphTrack->AddTimelineSeries(FMemoryGraphSeries::ETimelineType::MaxTotalMem);
+		MainGraphTrack->AddTimelineSeries(FMemoryGraphSeries::ETimelineType::MinTotalMem);
+
 		MainGraphTrack->SetVisibilityFlag(bShowHideAllMemoryTracks);
 
 		MainGraphTrack->SetAvailableTrackHeight(EMemoryTrackHeightMode::Small, 100.0f);
@@ -125,6 +136,50 @@ void FMemorySharedState::Tick(Insights::ITimingViewSession& InSession, const Tra
 		TimingView->InvalidateScrollableTracksOrder();
 	}
 
+	if (!LiveAllocsGraphTrack.IsValid())
+	{
+		LiveAllocsGraphTrack = CreateMemoryGraphTrack();
+		check(LiveAllocsGraphTrack);
+
+		LiveAllocsGraphTrack->SetOrder(FTimingTrackOrder::First + 1);
+		LiveAllocsGraphTrack->SetName(TEXT("Live Allocation Count"));
+		LiveAllocsGraphTrack->SetLabelUnit(EGraphTrackLabelUnit::Count, 0);
+
+		LiveAllocsGraphTrack->AddTimelineSeries(FMemoryGraphSeries::ETimelineType::MaxLiveAllocs);
+		LiveAllocsGraphTrack->AddTimelineSeries(FMemoryGraphSeries::ETimelineType::MinLiveAllocs);
+
+		LiveAllocsGraphTrack->SetVisibilityFlag(bShowHideAllMemoryTracks);
+
+		LiveAllocsGraphTrack->SetAvailableTrackHeight(EMemoryTrackHeightMode::Small, 50.0f);
+		LiveAllocsGraphTrack->SetAvailableTrackHeight(EMemoryTrackHeightMode::Medium, 100.0f);
+		LiveAllocsGraphTrack->SetAvailableTrackHeight(EMemoryTrackHeightMode::Large, 200.0f);
+		LiveAllocsGraphTrack->SetCurrentTrackHeight(TrackHeightMode);
+
+		TimingView->InvalidateScrollableTracksOrder();
+	}
+
+	if (!AllocFreeGraphTrack.IsValid())
+	{
+		AllocFreeGraphTrack = CreateMemoryGraphTrack();
+		check(AllocFreeGraphTrack);
+
+		AllocFreeGraphTrack->SetOrder(FTimingTrackOrder::First + 2);
+		AllocFreeGraphTrack->SetName(TEXT("Alloc/Free Event Count"));
+		AllocFreeGraphTrack->SetLabelUnit(EGraphTrackLabelUnit::Count, 0);
+
+		AllocFreeGraphTrack->AddTimelineSeries(FMemoryGraphSeries::ETimelineType::AllocEvents);
+		AllocFreeGraphTrack->AddTimelineSeries(FMemoryGraphSeries::ETimelineType::FreeEvents);
+
+		AllocFreeGraphTrack->SetVisibilityFlag(bShowHideAllMemoryTracks);
+
+		AllocFreeGraphTrack->SetAvailableTrackHeight(EMemoryTrackHeightMode::Small, 50.0f);
+		AllocFreeGraphTrack->SetAvailableTrackHeight(EMemoryTrackHeightMode::Medium, 100.0f);
+		AllocFreeGraphTrack->SetAvailableTrackHeight(EMemoryTrackHeightMode::Large, 200.0f);
+		AllocFreeGraphTrack->SetCurrentTrackHeight(TrackHeightMode);
+
+		TimingView->InvalidateScrollableTracksOrder();
+	}
+
 	const int32 PrevTagCount = TagList.GetTags().Num();
 
 	TagList.Update();
diff --git a/Engine/Source/Developer/TraceInsights/Private/Insights/MemoryProfiler/ViewModels/MemorySharedState.h b/Engine/Source/Developer/TraceInsights/Private/Insights/MemoryProfiler/ViewModels/MemorySharedState.h
--- a/Engine/Source/Developer/TraceInsights/Private/Insights/MemoryProfiler/ViewModels/MemorySharedState.h
+++ b/Engine/Source/Developer/TraceInsights/Private/Insights/MemoryProfiler/ViewModels/MemorySharedState.h
@@ -129,7 +129,9 @@ private:
 	TSharedPtr<Insights::FMemoryTracker> DefaultTracker;
 	TSharedPtr<Insights::FMemoryTracker> CurrentTracker;
 
-	TSharedPtr<FMemoryGraphTrack> MainGraphTrack; // the Main Memory Graph track
+	TSharedPtr<FMemoryGraphTrack> MainGraphTrack; // the Main Memory Graph track; also hosts the Total Allocated Memory series
+	TSharedPtr<FMemoryGraphTrack> LiveAllocsGraphTrack; // the graph track for the Live Allocation Count series
+	TSharedPtr<FMemoryGraphTrack> AllocFreeGraphTrack; // the graph track for the Alloc Event Count and the Free Event Count series
 	TSet<TSharedPtr<FMemoryGraphTrack>> AllTracks;
 
 	EMemoryTrackHeightMode TrackHeightMode;
diff --git a/Engine/Source/Developer/TraceInsights/Private/Insights/MemoryProfiler/Widgets/SMemAllocTableTreeView.cpp b/Engine/Source/Developer/TraceInsights/Private/Insights/MemoryProfiler/Widgets/SMemAllocTableTreeView.cpp
--- a/Engine/Source/Developer/TraceInsights/Private/Insights/MemoryProfiler/Widgets/SMemAllocTableTreeView.cpp
+++ b/Engine/Source/Developer/TraceInsights/Private/Insights/MemoryProfiler/Widgets/SMemAllocTableTreeView.cpp
@@ -213,10 +213,13 @@ void SMemAllocTableTreeView::StartQuery()
 		UE_LOG(MemoryProfiler, Warning, TEXT("[MemAlloc] Invalid allocations provider!"));
 		return;
 	}
-	const TraceServices::IAllocationsProvider& Provider = *AllocationsProvider;
 
-	TraceServices::IAllocationsProvider::FQueryParams Params = { Rule->GetValue(), TimeMarkers[0], TimeMarkers[1], TimeMarkers[2], TimeMarkers[3] };
-	Query = Provider.StartQuery(Params);
+	{
+		const TraceServices::IAllocationsProvider& Provider = *AllocationsProvider;
+		TraceServices::IAllocationsProvider::FReadScopeLock _(Provider);
+		TraceServices::IAllocationsProvider::FQueryParams Params = { Rule->GetValue(), TimeMarkers[0], TimeMarkers[1], TimeMarkers[2], TimeMarkers[3] };
+		Query = Provider.StartQuery(Params);
+	}
 
 	if (Query == 0)
 	{
@@ -246,6 +249,7 @@ void SMemAllocTableTreeView::UpdateQuery(TraceServices::IAllocationsProvider::EQ
 		return;
 	}
 	const TraceServices::IAllocationsProvider& Provider = *AllocationsProvider;
+	TraceServices::IAllocationsProvider::FReadScopeLock _(Provider);
 
 	TraceServices::IAllocationsProvider::FQueryStatus Status = Provider.PollQuery(Query);
 
diff --git a/Engine/Source/Developer/TraceInsights/Private/Insights/MemoryProfiler/Widgets/SMemTagTreeView.cpp b/Engine/Source/Developer/TraceInsights/Private/Insights/MemoryProfiler/Widgets/SMemTagTreeView.cpp
--- a/Engine/Source/Developer/TraceInsights/Private/Insights/MemoryProfiler/Widgets/SMemTagTreeView.cpp
+++ b/Engine/Source/Developer/TraceInsights/Private/Insights/MemoryProfiler/Widgets/SMemTagTreeView.cpp
@@ -21,6 +21,7 @@
 #include "Widgets/Views/STableViewBase.h"
 
 // Insights
+#include "Insights/Common/Stopwatch.h"
 #include "Insights/Common/TimeUtils.h"
 #include "Insights/InsightsStyle.h"
 #include "Insights/MemoryProfiler/MemoryProfilerManager.h"
diff --git a/Engine/Source/Developer/TraceServices/Private/Analyzers/AllocationsAnalysis.cpp b/Engine/Source/Developer/TraceServices/Private/Analyzers/AllocationsAnalysis.cpp
--- a/Engine/Source/Developer/TraceServices/Private/Analyzers/AllocationsAnalysis.cpp
+++ b/Engine/Source/Developer/TraceServices/Private/Analyzers/AllocationsAnalysis.cpp
@@ -27,20 +27,21 @@ void FAllocationsAnalyzer::OnAnalysisBegin(const FOnAnalysisContext& Context)
 {
 	FInterfaceBuilder& Builder = Context.InterfaceBuilder;
 
-	Builder.RouteEvent(RouteId_Init,       "Memory", "Init");
-	Builder.RouteEvent(RouteId_CoreAdd,    "Memory", "CoreAdd");
-	Builder.RouteEvent(RouteId_CoreRemove, "Memory", "CoreRemove");
-	Builder.RouteEvent(RouteId_Alloc,      "Memory", "Alloc");
-	Builder.RouteEvent(RouteId_Realloc,    "Memory", "Realloc");
-	Builder.RouteEvent(RouteId_Free,       "Memory", "Free");
-	Builder.RouteEvent(RouteId_Marker,	   "Memory", "Marker");
+	Builder.RouteEvent(RouteId_Init,         "Memory", "Init");
+	Builder.RouteEvent(RouteId_CoreAdd,      "Memory", "CoreAdd");
+	Builder.RouteEvent(RouteId_CoreRemove,   "Memory", "CoreRemove");
+	Builder.RouteEvent(RouteId_Alloc,        "Memory", "Alloc");
+	Builder.RouteEvent(RouteId_Free,         "Memory", "Free");
+	Builder.RouteEvent(RouteId_ReallocAlloc, "Memory", "ReallocAlloc");
+	Builder.RouteEvent(RouteId_Marker,	     "Memory", "Marker");
 }
 
 ////////////////////////////////////////////////////////////////////////////////////////////////////
 
 void FAllocationsAnalyzer::OnAnalysisEnd()
 {
-	//AllocationsProvider.DebugPrint();
+	FAllocationsProvider::FEditScopeLock _(AllocationsProvider);
+	AllocationsProvider.EditOnAnalysisCompleted();
 }
 
 ////////////////////////////////////////////////////////////////////////////////////////////////////
@@ -56,13 +57,12 @@ bool FAllocationsAnalyzer::OnEvent(uint16 RouteId, EStyle Style, const FOnEventC
 			const uint8 MinAlignment = EventData.GetValue<uint8>("MinAlignment");
 			const uint8 SizeShift = EventData.GetValue<uint8>("SizeShift");
 			const uint8 SummarySizeShift = EventData.GetValue<uint8>("SummarySizeShift");
-			const uint8 Mode = EventData.GetValue<uint8>("Mode");
 
 			BaseCycle = EventData.GetValue<uint64>("BaseCycle");
 			MarkerPeriod = EventData.GetValue<uint32>("MarkerPeriod");
 
 			FAllocationsProvider::FEditScopeLock _(AllocationsProvider);
-			AllocationsProvider.EditInit(Time, MinAlignment, SizeShift, SummarySizeShift, Mode);
+			AllocationsProvider.EditInit(Time, MinAlignment, SizeShift, SummarySizeShift);
 			break;
 		}
 		case RouteId_CoreAdd:
@@ -88,43 +88,29 @@ bool FAllocationsAnalyzer::OnEvent(uint16 RouteId, EStyle Style, const FOnEventC
 			break;
 		}
 		case RouteId_Alloc:
-		{
-			const double Time = GetCurrentTime();
-			const uint64 Owner = EventData.GetValue<uint64>("Owner");
-			const uint64 Address = EventData.GetValue<uint64>("Address");
-			const uint32 Size = EventData.GetValue<uint32>("Size");
-			const uint8 Alignment = EventData.GetValue<uint8>("Alignment");
-			const uint8 Waste = EventData.GetValue<uint8>("Waste");
-			uint32 Tag = 0;//TODO
-
-			FAllocationsProvider::FEditScopeLock _(AllocationsProvider);
-			AllocationsProvider.EditAlloc(Time, Owner, Address, Size, Alignment, Waste, Tag);
-			break;
-		}
-		case RouteId_Realloc:
+		case RouteId_ReallocAlloc:
 		{
 			// TODO: Can we have a struct mapping over the EventData?
-			//       something like: const auto& Ev = (const ReallocEvent&)EventData.Get();
+			//       something like: const auto& Ev = (const AllocEvent&)EventData.Get();
 			const double Time = GetCurrentTime();
 			const uint64 Owner = EventData.GetValue<uint64>("Owner");
-			const uint64 FreeAddress = EventData.GetValue<uint64>("FreeAddress");
 			const uint64 Address = EventData.GetValue<uint64>("Address");
 			const uint32 Size = EventData.GetValue<uint32>("Size");
-			const uint8 Alignment = EventData.GetValue<uint8>("Alignment");
-			const uint8 Waste = EventData.GetValue<uint8>("Waste");
+			const uint8 Alignment_SizeLower = EventData.GetValue<uint8>("Alignment_SizeLower");
+			const uint32 Tag = 0; //TODO
 
 			FAllocationsProvider::FEditScopeLock _(AllocationsProvider);
-			AllocationsProvider.EditRealloc(Time, Owner, FreeAddress, Address, Size, Alignment, Waste);
+			AllocationsProvider.EditAlloc(Time, Owner, Address, Size, Alignment_SizeLower, Tag);
 			break;
 		}
 		case RouteId_Free:
+		case RouteId_ReallocFree:
 		{
 			const double Time = GetCurrentTime();
-			const uint64 Owner = EventData.GetValue<uint64>("Owner");
-			const uint64 FreeAddress = EventData.GetValue<uint64>("FreeAddress");
+			const uint64 Address = EventData.GetValue<uint64>("Address");
 
 			FAllocationsProvider::FEditScopeLock _(AllocationsProvider);
-			AllocationsProvider.EditFree(Time, Owner, FreeAddress);
+			AllocationsProvider.EditFree(Time, Address);
 			break;
 		}
 		case RouteId_Marker:
diff --git a/Engine/Source/Developer/TraceServices/Private/Analyzers/AllocationsAnalysis.h b/Engine/Source/Developer/TraceServices/Private/Analyzers/AllocationsAnalysis.h
--- a/Engine/Source/Developer/TraceServices/Private/Analyzers/AllocationsAnalysis.h
+++ b/Engine/Source/Developer/TraceServices/Private/Analyzers/AllocationsAnalysis.h
@@ -20,8 +20,9 @@ private:
 		RouteId_CoreAdd,
 		RouteId_CoreRemove,
 		RouteId_Alloc,
-		RouteId_Realloc,
 		RouteId_Free,
+		RouteId_ReallocAlloc,
+		RouteId_ReallocFree,
 		RouteId_Marker,
 	};
 
diff --git a/Engine/Source/Developer/TraceServices/Private/Model/AllocationsProvider.cpp b/Engine/Source/Developer/TraceServices/Private/Model/AllocationsProvider.cpp
--- a/Engine/Source/Developer/TraceServices/Private/Model/AllocationsProvider.cpp
+++ b/Engine/Source/Developer/TraceServices/Private/Model/AllocationsProvider.cpp
@@ -101,26 +101,18 @@ uint64 IAllocationsProvider::FAllocation::GetAddress() const
 
 ////////////////////////////////////////////////////////////////////////////////////////////////////
 
-uint32 IAllocationsProvider::FAllocation::GetSize() const
+uint64 IAllocationsProvider::FAllocation::GetSize() const
 {
 	const auto* Inner = (const FAllocationItem*)this;
-	return Inner->Size;
+	return Inner->GetSize();
 }
 
 ////////////////////////////////////////////////////////////////////////////////////////////////////
 
-uint8 IAllocationsProvider::FAllocation::GetAlignment() const
+uint32 IAllocationsProvider::FAllocation::GetAlignment() const
 {
 	const auto* Inner = (const FAllocationItem*)this;
-	return Inner->Alignment;
-}
-
-////////////////////////////////////////////////////////////////////////////////////////////////////
-
-uint8 IAllocationsProvider::FAllocation::GetWaste() const
-{
-	const auto* Inner = (const FAllocationItem*)this;
-	return Inner->Waste;
+	return Inner->GetAlignment();
 }
 
 ////////////////////////////////////////////////////////////////////////////////////////////////////
@@ -192,6 +184,13 @@ IAllocationsProvider::FQueryResult IAllocationsProvider::FQueryStatus::NextResul
 FAllocationsProvider::FAllocationsProvider(ILinearAllocator& InAllocator)
 	: Allocator(InAllocator)
 	, SbTree(nullptr)
+	, Timeline(InAllocator, 1024)
+	, MinTotalAllocatedMemoryTimeline(InAllocator, 1024)
+	, MaxTotalAllocatedMemoryTimeline(InAllocator, 1024)
+	, MinLiveAllocationsTimeline(InAllocator, 1024)
+	, MaxLiveAllocationsTimeline(InAllocator, 1024)
+	, AllocEventsTimeline(InAllocator, 1024)
+	, FreeEventsTimeline(InAllocator, 1024)
 {
 	const uint32 ColumnShift = 17; // 1<<17 = 128K
 	SbTree = new FSbTree(InAllocator, ColumnShift);
@@ -218,8 +217,10 @@ FName FAllocationsProvider::GetName()
 
 ////////////////////////////////////////////////////////////////////////////////////////////////////
 
-void FAllocationsProvider::EditInit(double InTime, uint8 InMinAlignment, uint8 InSizeShift, uint8 InSummarySizeShift, uint8 InMode)
+void FAllocationsProvider::EditInit(double InTime, uint8 InMinAlignment, uint8 InSizeShift, uint8 InSummarySizeShift)
 {
+	Lock.WriteAccessCheck();
+
 	if (bInitialized)
 	{
 		// error: already initialized
@@ -230,9 +231,10 @@ void FAllocationsProvider::EditInit(double InTime, uint8 InMinAlignment, uint8 I
 	MinAlignment = InMinAlignment;
 	SizeShift = InSizeShift;
 	SummarySizeShift = InSummarySizeShift;
-	Mode = InMode;
 
 	bInitialized = true;
+
+	AdvanceTimelines(InTime);
 }
 
 ////////////////////////////////////////////////////////////////////////////////////////////////////
@@ -265,7 +267,7 @@ void FAllocationsProvider::EditRemoveCore(double Time, uint64 Owner, uint64 Base
 
 ////////////////////////////////////////////////////////////////////////////////////////////////////
 
-void FAllocationsProvider::EditAlloc(double Time, uint64 Owner, uint64 Address, uint32 Size, uint8 Alignment, uint8 Waste, uint32 Tag)
+void FAllocationsProvider::EditAlloc(double Time, uint64 Owner, uint64 Address, uint32 InSize, uint8 InAlignmentAndSizeLower, uint32 Tag)
 {
 	Lock.WriteAccessCheck();
 
@@ -274,6 +276,13 @@ void FAllocationsProvider::EditAlloc(double Time, uint64 Owner, uint64 Address,
 		return;
 	}
 
+	if (Address == 0)
+	{
+		return;
+	}
+
+	SbTree->SetTimeForEvent(EventIndex, Time);
+
 	FAllocationItem* AllocationPtr = LiveAllocs.Find(Address);
 	if (AllocationPtr)
 	{
@@ -281,7 +290,12 @@ void FAllocationsProvider::EditAlloc(double Time, uint64 Owner, uint64 Address,
 	}
 	else
 	{
-		SbTree->SetTimeForEvent(EventIndex, Time);
+		AdvanceTimelines(Time);
+
+		const uint8 SizeLowerMask = ((1 << SizeShift) - 1);
+		const uint8 AlignmentMask = ~SizeLowerMask;
+
+		const uint64 Size = (static_cast<uint64>(InSize) << SizeShift) | static_cast<uint64>(InAlignmentAndSizeLower & SizeLowerMask);
 
 		FAllocationItem Allocation =
 		{
@@ -291,10 +305,8 @@ void FAllocationsProvider::EditAlloc(double Time, uint64 Owner, uint64 Address,
 			std::numeric_limits<double>::infinity(),
 			Owner,
 			Address,
-			Size,
-			Alignment,
-			Waste,
-			static_cast<uint16>(Tag),
+			FAllocationItem::PackSizeAndAlignment(Size, InAlignmentAndSizeLower & AlignmentMask),
+			Tag,
 			0, // Reserverd1
 			0, // Reserverd2
 		};
@@ -302,12 +314,19 @@ void FAllocationsProvider::EditAlloc(double Time, uint64 Owner, uint64 Address,
 
 		//TODO: Cache the last added alloc, as there is ~10% chance to be freed in the next event.
 
-		if (LiveAllocs.Num() > MaxLiveAllocCount)
+		const uint32 CurrentLiveAllocCount = static_cast<uint32>(LiveAllocs.Num());
+		if (CurrentLiveAllocCount > MaxLiveAllocCount)
 		{
-			MaxLiveAllocCount = LiveAllocs.Num();
+			MaxLiveAllocCount = CurrentLiveAllocCount;
 		}
 
-		UpdateHistogramByAllocSize(Size, Alignment, Waste);
+		UpdateHistogramByAllocSize(Size);
+
+		// Update stats for current timeline sample.
+		TotalAllocatedMemory += Size;
+		SampleMaxTotalAllocatedMemory = FMath::Max(SampleMaxTotalAllocatedMemory, TotalAllocatedMemory);
+		SampleMaxLiveAllocations = FMath::Max(SampleMaxLiveAllocations, (uint32)LiveAllocs.Num());
+		++SampleAllocEvents;
 	}
 
 	++AllocCount;
@@ -316,7 +335,7 @@ void FAllocationsProvider::EditAlloc(double Time, uint64 Owner, uint64 Address,
 
 ////////////////////////////////////////////////////////////////////////////////////////////////////
 
-void FAllocationsProvider::EditRealloc(double Time, uint64 Owner, uint64 FreeAddress, uint64 Address, uint32 Size, uint8 Alignment, uint8 Waste)
+void FAllocationsProvider::EditFree(double Time, uint64 Address)
 {
 	Lock.WriteAccessCheck();
 
@@ -325,96 +344,40 @@ void FAllocationsProvider::EditRealloc(double Time, uint64 Owner, uint64 FreeAdd
 		return;
 	}
 
-	SbTree->SetTimeForEvent(EventIndex, Time);
-
-	FAllocationItem* AllocationPtr = LiveAllocs.Find(FreeAddress);
-	if (!AllocationPtr)
-	{
-		++ReallocErrors;
-	}
-	else
-	{
-		check(EventIndex > AllocationPtr->StartEventIndex);
-		AllocationPtr->EndEventIndex = EventIndex;
-		AllocationPtr->EndTime = Time;
-
-		SbTree->AddAlloc(AllocationPtr);
-
-		uint32 EventDistance = AllocationPtr->EndEventIndex - AllocationPtr->StartEventIndex;
-		UpdateHistogramByEventDistance(EventDistance);
-
-		if (FreeAddress == Address)
-		{
-			AllocationPtr->StartEventIndex = EventIndex;
-			AllocationPtr->EndEventIndex = -1;
-			AllocationPtr->StartTime = Time;
-			AllocationPtr->EndTime = std::numeric_limits<double>::infinity();
-			AllocationPtr->Owner = Owner;
-			AllocationPtr->Address = Address;
-			AllocationPtr->Size = Size;
-			AllocationPtr->Alignment = Alignment;
-			AllocationPtr->Waste = Waste;
-		}
-		else
-		{
-			uint16 Tag = AllocationPtr->Tag;
-			LiveAllocs.Remove(FreeAddress);
-
-			FAllocationItem Allocation =
-			{
-				EventIndex,
-				(uint32)-1,
-				Time,
-				std::numeric_limits<double>::infinity(),
-				Owner,
-				Address,
-				Size,
-				Alignment,
-				Waste,
-				Tag,
-				0, // Reserverd1
-				0, // Reserverd2
-			};
-			LiveAllocs.Add(Address, Allocation);
-		}
-
-		UpdateHistogramByAllocSize(Size, Alignment, Waste);
-	}
-
-	++AllocCount;
-	++EventIndex;
-}
-
-////////////////////////////////////////////////////////////////////////////////////////////////////
-
-void FAllocationsProvider::EditFree(double Time, uint64 Owner, uint64 FreeAddress)
-{
-	Lock.WriteAccessCheck();
-
-	if (!bInitialized)
+	if (Address == 0)
 	{
 		return;
 	}
 
 	SbTree->SetTimeForEvent(EventIndex, Time);
 
-	FAllocationItem* AllocationPtr = LiveAllocs.Find(FreeAddress);
+	FAllocationItem* AllocationPtr = LiveAllocs.Find(Address);
 	if (!AllocationPtr)
 	{
 		++FreeErrors;
 	}
 	else
 	{
+		AdvanceTimelines(Time);
+
 		check(EventIndex > AllocationPtr->StartEventIndex);
 		AllocationPtr->EndEventIndex = EventIndex;
 		AllocationPtr->EndTime = Time;
 
+		const uint64 OldSize = AllocationPtr->GetSize();
+
 		SbTree->AddAlloc(AllocationPtr);
 
 		uint32 EventDistance = AllocationPtr->EndEventIndex - AllocationPtr->StartEventIndex;
 		UpdateHistogramByEventDistance(EventDistance);
 
-		LiveAllocs.Remove(FreeAddress);
+		LiveAllocs.Remove(Address);
+
+		// Update stats for current timeline sample.
+		TotalAllocatedMemory -= OldSize;
+		SampleMinTotalAllocatedMemory = FMath::Min(SampleMinTotalAllocatedMemory, TotalAllocatedMemory);
+		SampleMinLiveAllocations = FMath::Min(SampleMinLiveAllocations, (uint32)LiveAllocs.Num());
+		++SampleFreeEvents;
 	}
 
 	++EventIndex;
@@ -422,7 +385,7 @@ void FAllocationsProvider::EditFree(double Time, uint64 Owner, uint64 FreeAddres
 
 ////////////////////////////////////////////////////////////////////////////////////////////////////
 
-void FAllocationsProvider::UpdateHistogramByAllocSize(uint32 Size, uint8 Alignment, uint8 Waste)
+void FAllocationsProvider::UpdateHistogramByAllocSize(uint64 Size)
 {
 	if (Size > MaxAllocSize)
 	{
@@ -437,8 +400,8 @@ void FAllocationsProvider::UpdateHistogramByAllocSize(uint32 Size, uint8 Alignme
 	// ...
 	// i : [2^(i-1) .. 2^i-1], i > 0
 	// ...
-	// 32 : [2^31 .. 2^32 - 1]
-	uint32 HistogramIndexPow2 = 32 - FMath::CountLeadingZeros(Size);
+	// 64 : [2^63 .. 2^64-1]
+	uint32 HistogramIndexPow2 = 64 - FMath::CountLeadingZeros64(Size);
 	++AllocSizeHistogramPow2[HistogramIndexPow2];
 }
 
@@ -452,20 +415,336 @@ void FAllocationsProvider::UpdateHistogramByEventDistance(uint32 EventDistance)
 	}
 
 	// HistogramIndex : Value Range
-	// 0 : [0 .. 1]
-	// 1 : [2 .. 3]
-	// 2 : [4 .. 7]
-	// 3 : [8 .. 15]
+	// 0 : [0]
+	// 1 : [1]
+	// 2 : [2 .. 3]
+	// 3 : [4 .. 7]
 	// ...
-	// i : [2^i .. 2^(i+1)-1]
+	// i : [2^(i-1) .. 2^i-1], i > 0
 	// ...
-	// 63 : [2^63 .. 2^64-1]
-	uint32 HistogramIndexPow2 = (EventDistance > 0) ? 63 - FMath::CountLeadingZeros64(EventDistance) : 0;
+	// 32 : [2^31 .. 2^32-1]
+	uint32 HistogramIndexPow2 = 32 - FMath::CountLeadingZeros(EventDistance);
 	++EventDistanceHistogramPow2[HistogramIndexPow2];
 }
 
 ////////////////////////////////////////////////////////////////////////////////////////////////////
 
+void FAllocationsProvider::AdvanceTimelines(double Time)
+{
+	constexpr double SampleGranularity = 0.001; // 1ms
+
+	// If enough time has passed (since the current sample is started)...
+	if (Time - SampleStartTimestamp > SampleGranularity)
+	{
+		// Add the current sample to the timelines.
+		Timeline.EmplaceBack(SampleStartTimestamp);
+		MinTotalAllocatedMemoryTimeline.EmplaceBack(SampleMinTotalAllocatedMemory);
+		MaxTotalAllocatedMemoryTimeline.EmplaceBack(SampleMaxTotalAllocatedMemory);
+		MinLiveAllocationsTimeline.EmplaceBack(SampleMinLiveAllocations);
+		MaxLiveAllocationsTimeline.EmplaceBack(SampleMaxLiveAllocations);
+		AllocEventsTimeline.EmplaceBack(SampleAllocEvents);
+		FreeEventsTimeline.EmplaceBack(SampleFreeEvents);
+
+		// Start a new sample.
+		SampleStartTimestamp = Time;
+		const uint32 NumLiveAllocs = (uint32)LiveAllocs.Num();
+		SampleMinTotalAllocatedMemory = TotalAllocatedMemory;
+		SampleMaxTotalAllocatedMemory = TotalAllocatedMemory;
+		SampleMinLiveAllocations = NumLiveAllocs;
+		SampleMaxLiveAllocations = NumLiveAllocs;
+		SampleAllocEvents = 0;
+		SampleFreeEvents = 0;
+
+		// If the previous sample is well distanced in time...
+		if (Time - SampleEndTimestamp > SampleGranularity) 
+		{
+			// Add an intermediate "flat region" sample.
+			Timeline.EmplaceBack(SampleEndTimestamp);
+			MinTotalAllocatedMemoryTimeline.EmplaceBack(TotalAllocatedMemory);
+			MaxTotalAllocatedMemoryTimeline.EmplaceBack(TotalAllocatedMemory);
+			MinLiveAllocationsTimeline.EmplaceBack(NumLiveAllocs);
+			MaxLiveAllocationsTimeline.EmplaceBack(NumLiveAllocs);
+			AllocEventsTimeline.EmplaceBack(0);
+			FreeEventsTimeline.EmplaceBack(0);
+		}
+	}
+
+	SampleEndTimestamp = Time;
+}
+
+////////////////////////////////////////////////////////////////////////////////////////////////////
+
+void FAllocationsProvider::EditOnAnalysisCompleted()
+{
+	Lock.WriteAccessCheck();
+
+	if (!bInitialized)
+	{
+		return;
+	}
+
+	AdvanceTimelines(std::numeric_limits<double>::infinity());
+
+	//DebugPrint();
+}
+
+////////////////////////////////////////////////////////////////////////////////////////////////////
+
+void FAllocationsProvider::GetTimelineIndexRange(double StartTime, double EndTime, int32& StartIndex, int32& EndIndex) const
+{
+	using PageType = const TPagedArrayPage<double>;
+
+	PageType* PageData = Timeline.GetPages();
+	if (PageData)
+	{
+		const int32 NumPoints = static_cast<int32>(Timeline.Num());
+
+		const int32 NumPages = static_cast<int32>(Timeline.NumPages());
+		TArrayView<PageType, int32> Pages(PageData, NumPages);
+
+		const int32 StartPageIndex = Algo::UpperBoundBy(Pages, StartTime, [](PageType& Page) { return Page.Items[0]; }) - 1;
+		if (StartPageIndex < 0)
+		{
+			StartIndex = -1;
+		}
+		else
+		{
+			PageType& Page = PageData[StartPageIndex];
+			TArrayView<double> PageValues(Page.Items, static_cast<int32>(Page.Count));
+			const int32 Index = Algo::UpperBound(PageValues, StartTime) - 1;
+			check(Index >= 0);
+			StartIndex = StartPageIndex * static_cast<int32>(Timeline.GetPageSize()) + Index;
+			check(Index < NumPoints);
+		}
+
+		const int32 EndPageIndex = Algo::UpperBoundBy(Pages, EndTime, [](PageType& Page) { return Page.Items[0]; }) - 1;
+		if (EndPageIndex < 0)
+		{
+			EndIndex = -1;
+		}
+		else
+		{
+			PageType& Page = PageData[EndPageIndex];
+			TArrayView<double> PageValues(Page.Items, static_cast<int32>(Page.Count));
+			const int32 Index = Algo::UpperBound(PageValues, EndTime) - 1;
+			check(Index >= 0);
+			EndIndex = EndPageIndex * static_cast<int32>(Timeline.GetPageSize()) + Index;
+			check(Index < NumPoints);
+		}
+	}
+	else
+	{
+		StartIndex = -1;
+		EndIndex = -1;
+	}
+}
+
+////////////////////////////////////////////////////////////////////////////////////////////////////
+
+void FAllocationsProvider::EnumerateMinTotalAllocatedMemoryTimeline(int32 StartIndex, int32 EndIndex, TFunctionRef<void(double Time, double Duration, uint64 Value)> Callback) const
+{
+	const int32 NumPoints = Timeline.Num();
+	StartIndex = FMath::Max(StartIndex, 0);
+	EndIndex = FMath::Min(EndIndex + 1, NumPoints); // make it exclusive
+	if (StartIndex < EndIndex)
+	{
+		auto TimeIt = Timeline.GetIteratorFromItem(StartIndex);
+		auto ValueIt = MinTotalAllocatedMemoryTimeline.GetIteratorFromItem(StartIndex);
+		double PrevTime = *TimeIt;
+		uint64 PrevValue = *ValueIt;
+		++TimeIt;
+		++ValueIt;
+		for (int32 Index = StartIndex + 1; Index < EndIndex; ++Index, ++TimeIt, ++ValueIt)
+		{
+			const double Time = *TimeIt;
+			Callback(PrevTime, Time - PrevTime, PrevValue);
+			PrevTime = *TimeIt;
+			PrevValue = *ValueIt;
+		}
+		if (EndIndex < NumPoints)
+		{
+			const double Time = *TimeIt;
+			Callback(PrevTime, Time - PrevTime, PrevValue);
+		}
+		else
+		{
+			Callback(PrevTime, std::numeric_limits<double>::infinity(), PrevValue);
+		}
+	}
+}
+
+////////////////////////////////////////////////////////////////////////////////////////////////////
+
+void FAllocationsProvider::EnumerateMaxTotalAllocatedMemoryTimeline(int32 StartIndex, int32 EndIndex, TFunctionRef<void(double Time, double Duration, uint64 Value)> Callback) const
+{
+	const int32 NumPoints = Timeline.Num();
+	StartIndex = FMath::Max(StartIndex, 0);
+	EndIndex = FMath::Min(EndIndex + 1, NumPoints); // make it exclusive
+	if (StartIndex < EndIndex)
+	{
+		auto TimeIt = Timeline.GetIteratorFromItem(StartIndex);
+		auto ValueIt = MaxTotalAllocatedMemoryTimeline.GetIteratorFromItem(StartIndex);
+		double PrevTime = *TimeIt;
+		uint64 PrevValue = *ValueIt;
+		++TimeIt;
+		++ValueIt;
+		for (int32 Index = StartIndex + 1; Index < EndIndex; ++Index, ++TimeIt, ++ValueIt)
+		{
+			const double Time = *TimeIt;
+			Callback(PrevTime, Time - PrevTime, PrevValue);
+			PrevTime = *TimeIt;
+			PrevValue = *ValueIt;
+		}
+		if (EndIndex < NumPoints)
+		{
+			const double Time = *TimeIt;
+			Callback(PrevTime, Time - PrevTime, PrevValue);
+		}
+		else
+		{
+			Callback(PrevTime, std::numeric_limits<double>::infinity(), PrevValue);
+		}
+	}
+}
+
+////////////////////////////////////////////////////////////////////////////////////////////////////
+
+void FAllocationsProvider::EnumerateMinLiveAllocationsTimeline(int32 StartIndex, int32 EndIndex, TFunctionRef<void(double Time, double Duration, uint32 Value)> Callback) const
+{
+	const int32 NumPoints = Timeline.Num();
+	StartIndex = FMath::Max(StartIndex, 0);
+	EndIndex = FMath::Min(EndIndex + 1, NumPoints); // make it exclusive
+	if (StartIndex < EndIndex)
+	{
+		auto TimeIt = Timeline.GetIteratorFromItem(StartIndex);
+		auto ValueIt = MinLiveAllocationsTimeline.GetIteratorFromItem(StartIndex);
+		double PrevTime = *TimeIt;
+		uint32 PrevValue = *ValueIt;
+		++TimeIt;
+		++ValueIt;
+		for (int32 Index = StartIndex + 1; Index < EndIndex; ++Index, ++TimeIt, ++ValueIt)
+		{
+			const double Time = *TimeIt;
+			Callback(PrevTime, Time - PrevTime, PrevValue);
+			PrevTime = *TimeIt;
+			PrevValue = *ValueIt;
+		}
+		if (EndIndex < NumPoints)
+		{
+			const double Time = *TimeIt;
+			Callback(PrevTime, Time - PrevTime, PrevValue);
+		}
+		else
+		{
+			Callback(PrevTime, std::numeric_limits<double>::infinity(), PrevValue);
+		}
+	}
+}
+
+////////////////////////////////////////////////////////////////////////////////////////////////////
+
+void FAllocationsProvider::EnumerateMaxLiveAllocationsTimeline(int32 StartIndex, int32 EndIndex, TFunctionRef<void(double Time, double Duration, uint32 Value)> Callback) const
+{
+	const int32 NumPoints = Timeline.Num();
+	StartIndex = FMath::Max(StartIndex, 0);
+	EndIndex = FMath::Min(EndIndex + 1, NumPoints); // make it exclusive
+	if (StartIndex < EndIndex)
+	{
+		auto TimeIt = Timeline.GetIteratorFromItem(StartIndex);
+		auto ValueIt = MaxLiveAllocationsTimeline.GetIteratorFromItem(StartIndex);
+		double PrevTime = *TimeIt;
+		uint32 PrevValue = *ValueIt;
+		++TimeIt;
+		++ValueIt;
+		for (int32 Index = StartIndex + 1; Index < EndIndex; ++Index, ++TimeIt, ++ValueIt)
+		{
+			const double Time = *TimeIt;
+			Callback(PrevTime, Time - PrevTime, PrevValue);
+			PrevTime = *TimeIt;
+			PrevValue = *ValueIt;
+		}
+		if (EndIndex < NumPoints)
+		{
+			const double Time = *TimeIt;
+			Callback(PrevTime, Time - PrevTime, PrevValue);
+		}
+		else
+		{
+			Callback(PrevTime, std::numeric_limits<double>::infinity(), PrevValue);
+		}
+	}
+}
+
+////////////////////////////////////////////////////////////////////////////////////////////////////
+
+void FAllocationsProvider::EnumerateAllocEventsTimeline(int32 StartIndex, int32 EndIndex, TFunctionRef<void(double Time, double Duration, uint32 Value)> Callback) const
+{
+	const int32 NumPoints = Timeline.Num();
+	StartIndex = FMath::Max(StartIndex, 0);
+	EndIndex = FMath::Min(EndIndex + 1, NumPoints); // make it exclusive
+	if (StartIndex < EndIndex)
+	{
+		auto TimeIt = Timeline.GetIteratorFromItem(StartIndex);
+		auto ValueIt = AllocEventsTimeline.GetIteratorFromItem(StartIndex);
+		double PrevTime = *TimeIt;
+		uint32 PrevValue = *ValueIt;
+		++TimeIt;
+		++ValueIt;
+		for (int32 Index = StartIndex + 1; Index < EndIndex; ++Index, ++TimeIt, ++ValueIt)
+		{
+			const double Time = *TimeIt;
+			Callback(PrevTime, Time - PrevTime, PrevValue);
+			PrevTime = *TimeIt;
+			PrevValue = *ValueIt;
+		}
+		if (EndIndex < NumPoints)
+		{
+			const double Time = *TimeIt;
+			Callback(PrevTime, Time - PrevTime, PrevValue);
+		}
+		else
+		{
+			Callback(PrevTime, std::numeric_limits<double>::infinity(), PrevValue);
+		}
+	}
+}
+
+////////////////////////////////////////////////////////////////////////////////////////////////////
+
+void FAllocationsProvider::EnumerateFreeEventsTimeline(int32 StartIndex, int32 EndIndex, TFunctionRef<void(double Time, double Duration, uint32 Value)> Callback) const
+{
+	const int32 NumPoints = Timeline.Num();
+	StartIndex = FMath::Max(StartIndex, 0);
+	EndIndex = FMath::Min(EndIndex + 1, NumPoints); // make it exclusive
+	if (StartIndex < EndIndex)
+	{
+		auto TimeIt = Timeline.GetIteratorFromItem(StartIndex);
+		auto ValueIt = FreeEventsTimeline.GetIteratorFromItem(StartIndex);
+		double PrevTime = *TimeIt;
+		uint32 PrevValue = *ValueIt;
+		++TimeIt;
+		++ValueIt;
+		for (int32 Index = StartIndex + 1; Index < EndIndex; ++Index, ++TimeIt, ++ValueIt)
+		{
+			const double Time = *TimeIt;
+			Callback(PrevTime, Time - PrevTime, PrevValue);
+			PrevTime = *TimeIt;
+			PrevValue = *ValueIt;
+		}
+		if (EndIndex < NumPoints)
+		{
+			const double Time = *TimeIt;
+			Callback(PrevTime, Time - PrevTime, PrevValue);
+		}
+		else
+		{
+			Callback(PrevTime, std::numeric_limits<double>::infinity(), PrevValue);
+		}
+	}
+}
+
+////////////////////////////////////////////////////////////////////////////////////////////////////
+
 void FAllocationsProvider::DebugPrint() const
 {
 	SbTree->DebugPrint();
diff --git a/Engine/Source/Developer/TraceServices/Private/Model/AllocationsProvider.h b/Engine/Source/Developer/TraceServices/Private/Model/AllocationsProvider.h
--- a/Engine/Source/Developer/TraceServices/Private/Model/AllocationsProvider.h
+++ b/Engine/Source/Developer/TraceServices/Private/Model/AllocationsProvider.h
@@ -2,9 +2,10 @@
 
 #pragma once
 
-#include "TraceServices/Model/AllocationsProvider.h"
+#include "Common/PagedArray.h"
 #include "Containers/Map.h"
 #include "HAL/CriticalSection.h"
+#include "TraceServices/Model/AllocationsProvider.h"
 
 namespace TraceServices
 {
@@ -34,18 +35,26 @@ private:
 
 struct FAllocationItem
 {
+	static constexpr uint32 AlignmentBits = 8;
+	static constexpr uint32 AlignmentShift = 56;
+	static constexpr uint64 SizeMask = (1ULL << AlignmentShift) - 1;
+
+	static uint64 UnpackSize(uint64 SizeAndAlignment) { return (SizeAndAlignment & SizeMask); }
+	static uint32 UnpackAlignment(uint64 SizeAndAlignment) { return static_cast<uint32>(SizeAndAlignment >> AlignmentShift); }
+	static uint64 PackSizeAndAlignment(uint64 Size, uint8 Alignment) { return (Size | (static_cast<uint64>(Alignment) << AlignmentShift)); }
+
+	uint64 GetSize() const { return UnpackSize(SizeAndAlignment); }
+	uint32 GetAlignment() const { return UnpackAlignment(SizeAndAlignment); }
+
 	uint32 StartEventIndex;
 	uint32 EndEventIndex;
 	double StartTime;
 	double EndTime;
 	uint64 Owner;
 	uint64 Address;
-	//uint64 SizeAlignmentWaste; // 48 bits (Size) + 8 bits (Alignment) + 8 bits (Waste)
-	uint32 Size;
-	uint8 Alignment;
-	uint8 Waste;
-	uint16 Tag;
-	uint64 Reserved1;
+	uint64 SizeAndAlignment; // (Alignment << AlignmentShift) | Size
+	uint32 Tag;
+	uint32 Reserved1;
 	uint64 Reserved2;
 };
 
@@ -68,73 +77,59 @@ struct FAllocationCoreItem
 
 class FAllocationsProvider : public IAllocationsProvider
 {
-public:
-	struct FEditScopeLock
-	{
-		FEditScopeLock(FAllocationsProvider& InAllocationsProvider)
-			: AllocationsProvider(InAllocationsProvider)
-		{
-			AllocationsProvider.BeginEdit();
-		}
-
-		~FEditScopeLock()
-		{
-			AllocationsProvider.EndEdit();
-		}
-
-		FAllocationsProvider& AllocationsProvider;
-	};
-
-	struct FReadScopeLock
-	{
-		FReadScopeLock(FAllocationsProvider& InAllocationsProvider)
-			: AllocationsProvider(InAllocationsProvider)
-		{
-			AllocationsProvider.BeginRead();
-		}
-
-		~FReadScopeLock()
-		{
-			AllocationsProvider.EndRead();
-		}
-
-		FAllocationsProvider& AllocationsProvider;
-	};
-
 public:
 	FAllocationsProvider(ILinearAllocator& InAllocator);
 	virtual ~FAllocationsProvider();
 
 	static FName GetName();
 
-	virtual bool IsInitialized() const override { /*Lock.ReadAccessCheck();*/ return bInitialized; }
+	virtual void BeginEdit() const override { Lock.BeginWrite(); }
+	virtual void EndEdit() const override { Lock.EndWrite(); }
+	void EditAccessCheck() const { return Lock.WriteAccessCheck(); }
+
+	virtual void BeginRead() const override { Lock.BeginRead(); }
+	virtual void EndRead() const override { Lock.EndRead(); }
+	void ReadAccessCheck() const { return Lock.ReadAccessCheck(); }
+
+	//////////////////////////////////////////////////
+	// Read operations
+
+	virtual bool IsInitialized() const override { ReadAccessCheck(); return bInitialized; }
+
+	virtual uint32 GetTimelineNumPoints() const override { ReadAccessCheck(); return Timeline.Num(); }
+	virtual void GetTimelineIndexRange(double StartTime, double EndTime, int32& StartIndex, int32& EndIndex) const override;
+	virtual void EnumerateMinTotalAllocatedMemoryTimeline(int32 StartIndex, int32 EndIndex, TFunctionRef<void(double Time, double Duration, uint64 Value)> Callback) const override;
+	virtual void EnumerateMaxTotalAllocatedMemoryTimeline(int32 StartIndex, int32 EndIndex, TFunctionRef<void(double Time, double Duration, uint64 Value)> Callback) const override;
+	virtual void EnumerateMinLiveAllocationsTimeline(int32 StartIndex, int32 EndIndex, TFunctionRef<void(double Time, double Duration, uint32 Value)> Callback) const override;
+	virtual void EnumerateMaxLiveAllocationsTimeline(int32 StartIndex, int32 EndIndex, TFunctionRef<void(double Time, double Duration, uint32 Value)> Callback) const override;
+	virtual void EnumerateAllocEventsTimeline(int32 StartIndex, int32 EndIndex, TFunctionRef<void(double Time, double Duration, uint32 Value)> Callback) const override;
+	virtual void EnumerateFreeEventsTimeline(int32 StartIndex, int32 EndIndex, TFunctionRef<void(double Time, double Duration, uint32 Value)> Callback) const override;
 
 	virtual FQueryHandle StartQuery(const FQueryParams& Params) const override;
 	virtual void CancelQuery(FQueryHandle Query) const override;
 	virtual const FQueryStatus PollQuery(FQueryHandle Query) const override;
 
-	void BeginEdit() { Lock.BeginWrite(); }
-	void EndEdit() { Lock.EndWrite(); }
-	void EditAccessCheck() const { return Lock.WriteAccessCheck(); }
+	const FSbTree* GetSbTree() const { ReadAccessCheck(); return SbTree; }
 
-	void BeginRead() const { Lock.BeginRead(); }
-	void EndRead() const { Lock.EndRead(); }
-	void ReadAccessCheck() const { return Lock.ReadAccessCheck(); }
+	void DebugPrint() const;
+
+	//////////////////////////////////////////////////
+	// Edit operations
 
-	void EditInit(double Time, uint8 MinAlignment, uint8 SizeShift, uint8 SummarySizeShift, uint8 Mode);
+	void EditInit(double Time, uint8 MinAlignment, uint8 SizeShift, uint8 SummarySizeShift);
 	void EditAddCore(double Time, uint64 Owner, uint64 Base, uint32 Size);
 	void EditRemoveCore(double Time, uint64 Owner, uint64 Base, uint32 Size);
-	void EditAlloc(double Time, uint64 Owner, uint64 Address, uint32 Size, uint8 Alignment, uint8 Waste, uint32 Tag);
-	void EditRealloc(double Time, uint64 Owner, uint64 FreeAddress, uint64 Address, uint32 Size, uint8 Alignment, uint8 Waste);
-	void EditFree(double Time, uint64 Owner, uint64 FreeAddress);
+	void EditAlloc(double Time, uint64 Owner, uint64 Address, uint32 Size, uint8 AlignmentAndSizeLower, uint32 Tag);
+	void EditFree(double Time, uint64 Address);
 
-	void DebugPrint() const;
+	void EditOnAnalysisCompleted();
 
-	const FSbTree* GetSbTree() const { return SbTree; }
+	//////////////////////////////////////////////////
 
 private:
-	void UpdateHistogramByAllocSize(uint32 Size, uint8 Alignment, uint8 Waste);
+	void UpdateHistogramByAllocSize(uint64 Size);
 	void UpdateHistogramByEventDistance(uint32 EventDistance);
+	void AdvanceTimelines(double Time);
 
 private:
 	ILinearAllocator& Allocator;
@@ -145,27 +140,45 @@ private:
 	uint8 MinAlignment = 0;
 	uint8 SizeShift = 0;
 	uint8 SummarySizeShift = 0;
-	uint8 Mode = 0;
 	bool bInitialized = false;
 
 	uint32 EventIndex = 0;
 
 	uint64 AllocCount = 0;
+	uint64 FreeCount = 0;
 
-	uint64 MaxLiveAllocCount = 0;
+	uint32 MaxLiveAllocCount = 0;
 	TMap<uint64, FAllocationItem> LiveAllocs;
 
 	uint64 AllocErrors = 0;
-	uint64 ReallocErrors = 0;
 	uint64 FreeErrors = 0;
 
-	uint32 MaxAllocSize = 0;
-	uint32 AllocSizeHistogramPow2[33] = { 0 };
+	uint64 MaxAllocSize = 0;
+	uint64 AllocSizeHistogramPow2[65] = { 0 };
 
 	uint32 MaxEventDistance = 0;
-	uint32 EventDistanceHistogramPow2[65] = { 0 };
+	uint32 EventDistanceHistogramPow2[33] = { 0 };
 
 	FSbTree* SbTree;
+
+	uint64 TotalAllocatedMemory = 0;
+
+	double SampleStartTimestamp = 0.0;
+	double SampleEndTimestamp = 0.0;
+	uint64 SampleMinTotalAllocatedMemory = 0;
+	uint64 SampleMaxTotalAllocatedMemory = 0;
+	uint32 SampleMinLiveAllocations = 0;
+	uint32 SampleMaxLiveAllocations = 0;
+	uint32 SampleAllocEvents = 0;
+	uint32 SampleFreeEvents = 0;
+
+	TPagedArray<double> Timeline;
+	TPagedArray<uint64> MinTotalAllocatedMemoryTimeline;
+	TPagedArray<uint64> MaxTotalAllocatedMemoryTimeline;
+	TPagedArray<uint32> MinLiveAllocationsTimeline;
+	TPagedArray<uint32> MaxLiveAllocationsTimeline;
+	TPagedArray<uint32> AllocEventsTimeline;
+	TPagedArray<uint32> FreeEventsTimeline;
 };
 
 ////////////////////////////////////////////////////////////////////////////////////////////////////
diff --git a/Engine/Source/Developer/TraceServices/Private/Modules/MemoryModule.cpp b/Engine/Source/Developer/TraceServices/Private/Modules/MemoryModule.cpp
--- a/Engine/Source/Developer/TraceServices/Private/Modules/MemoryModule.cpp
+++ b/Engine/Source/Developer/TraceServices/Private/Modules/MemoryModule.cpp
@@ -11,75 +11,6 @@
 namespace TraceServices
 {
 
-static void TEMP_TestAllocationsProviderImpl(ILinearAllocator& Allocator, double TimeA, double TimeB)
-{
-	FAllocationsProvider Provider(Allocator);
-
-	// 1. Query is started in response to some form of user interaction.
-	IAllocationsProvider::FQueryHandle Query;
-	{
-		const IAllocationsProvider::FQueryParams Params = { IAllocationsProvider::EQueryRule::aABf, TimeA, TimeB, 0.0, 0.0 };
-		Query = Provider.StartQuery(Params);
-	}
-
-	// 2. At regular intervals the query is polled for results.
-	while (true)
-	{
-		const IAllocationsProvider::FQueryStatus Status = Provider.PollQuery(Query);
-		if (Status.Status <= IAllocationsProvider::EQueryStatus::Done)
-		{
-			break;
-		}
-
-		// Multiple 'pages' of results will be returned. No guarantees are made
-		// about the order of pages or the allocations they report.
-		IAllocationsProvider::FQueryResult Result = Status.NextResult();
-		while (Result.IsValid())
-		{
-			for (uint32 i = 0, n = Result->Num(); i < n; ++i)
-			{
-				const IAllocationsProvider::FAllocation* Allocation = Result->Get(i);
-
-				auto Addr = Allocation->GetAddress();
-
-				auto Size = Allocation->GetSize();
-				auto Alignment = Allocation->GetAlignment();
-				auto Waste = Allocation->GetWaste();
-
-				uint64 BacktraceId = Allocation->GetBacktraceId();
-				/*
-				auto Backtrace = BacktraceProvider.Get(BacktraceId);
-				for (auto& RetAddress : Backtrace)
-				{
-				}
-				*/
-
-				auto Tag = Allocation->GetTag();
-
-				//... report to user (e.g.build treeview)
-			}
-
-			Result = Status.NextResult();
-		}
-	}
-
-	// 3a. Queries can be cancelled.
-	if (/*HasUserPressedCance()*/ false)
-	{
-		Provider.CancelQuery(Query);
-	}
-
-	// 3b. Queries that are polled to completion... */
-	/* Queries will automatically clean themselves up */
-}
-
-static void TEMP_TestAllocationsProvider(ILinearAllocator& Allocator)
-{
-	double TimeA = 3.0;
-	double TimeB = 11.0;
-	TEMP_TestAllocationsProviderImpl(Allocator, TimeA, TimeB);
-}
-
 static const FName MemoryModuleName("TraceModule_Memory");
 
 void FMemoryModule::GetModuleInfo(FModuleInfo& OutModuleInfo)
@@ -94,8 +25,6 @@ void FMemoryModule::OnAnalysisBegin(IAnalysisSession& Session)
 	Session.AddProvider(AllocationsProvider->GetName(), AllocationsProvider);
 	Session.AddAnalyzer(new FAllocationsAnalyzer(Session, *AllocationsProvider));
 
-	//TEMP_TestAllocationsProvider(Session.GetLinearAllocator());
-
 	FCallstacksProvider* CallstacksProvider = new FCallstacksProvider(Session);
 	Session.AddProvider(CallstacksProvider->GetName(), CallstacksProvider);
 	Session.AddAnalyzer(new FCallstacksAnalyzer(Session, CallstacksProvider));
diff --git a/Engine/Source/Developer/TraceServices/Public/TraceServices/Model/AllocationsProvider.h b/Engine/Source/Developer/TraceServices/Public/TraceServices/Model/AllocationsProvider.h
--- a/Engine/Source/Developer/TraceServices/Public/TraceServices/Model/AllocationsProvider.h
+++ b/Engine/Source/Developer/TraceServices/Public/TraceServices/Model/AllocationsProvider.h
@@ -16,6 +16,38 @@ namespace TraceServices
 class IAllocationsProvider : public IProvider
 {
 public:
+	struct TRACESERVICES_API FEditScopeLock
+	{
+		FEditScopeLock(const IAllocationsProvider& InAllocationsProvider)
+			: AllocationsProvider(InAllocationsProvider)
+		{
+			AllocationsProvider.BeginEdit();
+		}
+
+		~FEditScopeLock()
+		{
+			AllocationsProvider.EndEdit();
+		}
+
+		const IAllocationsProvider& AllocationsProvider;
+	};
+
+	struct TRACESERVICES_API FReadScopeLock
+	{
+		FReadScopeLock(const IAllocationsProvider& InAllocationsProvider)
+			: AllocationsProvider(InAllocationsProvider)
+		{
+			AllocationsProvider.BeginRead();
+		}
+
+		~FReadScopeLock()
+		{
+			AllocationsProvider.EndRead();
+		}
+
+		const IAllocationsProvider& AllocationsProvider;
+	};
+
 	// Allocation query rules.
 	// The enum uses the following naming convention:
 	//     A, B, C, D = time markers
@@ -54,9 +86,8 @@ public:
 		double GetStartTime() const;
 		double GetEndTime() const;
 		uint64 GetAddress() const;
-		uint32 GetSize() const;
-		uint8 GetAlignment() const;
-		uint8 GetWaste() const;
+		uint64 GetSize() const;
+		uint32 GetAlignment() const;
 		uint64 GetBacktraceId() const;
 		uint32 GetTag() const;
 	};
@@ -90,8 +121,38 @@ public:
 	typedef UPTRINT FQueryHandle;
 
 public:
+	virtual void BeginEdit() const = 0;
+	virtual void EndEdit() const = 0;
+	virtual void BeginRead() const = 0;
+	virtual void EndRead() const = 0;
+
 	virtual bool IsInitialized() const = 0;
 
+	// Returns the number of points in each timeline (Min/Max Total Allocated Memory, Min/Max Live Allocations, Total Alloc Events, Total Free Events).
+	virtual uint32 GetTimelineNumPoints() const = 0;
+
+	// Returns the inclusive index range [StartIndex, EndIndex] for a time range [StartTime, EndTime].
+	// Index values are in range { -1, 0, .. , N-1, N }, where N = GetTimelineNumPoints().
+	virtual void GetTimelineIndexRange(double StartTime, double EndTime, int32& StartIndex, int32& EndIndex) const = 0;
+
+	// Enumerates the Min Total Allocated Memory timeline points in the inclusive index interval [StartIndex, EndIndex].
+	virtual void EnumerateMinTotalAllocatedMemoryTimeline(int32 StartIndex, int32 EndIndex, TFunctionRef<void(double Time, double Duration, uint64 Value)> Callback) const = 0;
+
+	// Enumerates the Max Total Allocated Memory timeline points in the inclusive index interval [StartIndex, EndIndex].
+	virtual void EnumerateMaxTotalAllocatedMemoryTimeline(int32 StartIndex, int32 EndIndex, TFunctionRef<void(double Time, double Duration, uint64 Value)> Callback) const = 0;
+
+	// Enumerates the Min Live Allocations timeline points in the inclusive index interval [StartIndex, EndIndex].
+	virtual void EnumerateMinLiveAllocationsTimeline(int32 StartIndex, int32 EndIndex, TFunctionRef<void(double Time, double Duration, uint32 Value)> Callback) const = 0;
+
+	// Enumerates the Max Live Allocations timeline points in the inclusive index interval [StartIndex, EndIndex].
+	virtual void EnumerateMaxLiveAllocationsTimeline(int32 StartIndex, int32 EndIndex, TFunctionRef<void(double Time, double Duration, uint32 Value)> Callback) const = 0;
+
+	// Enumerates the Alloc Events timeline points in the inclusive index interval [StartIndex, EndIndex].
+	virtual void EnumerateAllocEventsTimeline(int32 StartIndex, int32 EndIndex, TFunctionRef<void(double Time, double Duration, uint32 Value)> Callback) const = 0;
+
+	// Enumerates the Free Events timeline points in the inclusive index interval [StartIndex, EndIndex].
+	virtual void EnumerateFreeEventsTimeline(int32 StartIndex, int32 EndIndex, TFunctionRef<void(double Time, double Duration, uint32 Value)> Callback) const = 0;
+
 	virtual FQueryHandle StartQuery(const FQueryParams& Params) const = 0;
 	virtual void CancelQuery(FQueryHandle Query) const = 0;
 	virtual const FQueryStatus PollQuery(FQueryHandle Query) const = 0;
