[Insights] In Memory Insights, the time marker text box value is commited for all commit types

diff --git a/Engine/Source/Developer/TraceInsights/Private/Insights/MemoryProfiler/Widgets/SMemTagTreeView.cpp b/Engine/Source/Developer/TraceInsights/Private/Insights/MemoryProfiler/Widgets/SMemTagTreeView.cpp
--- a/Engine/Source/Developer/TraceInsights/Private/Insights/MemoryProfiler/Widgets/SMemTagTreeView.cpp
+++ b/Engine/Source/Developer/TraceInsights/Private/Insights/MemoryProfiler/Widgets/SMemTagTreeView.cpp
@@ -393,11 +393,8 @@ TSharedRef<SWidget> SMemTagTreeView::ConstructTimeMarkerWidget(uint32 TimeMarker
 				})
 			.OnTextCommitted_Lambda([TimeMarker](const FText& InText, ETextCommit::Type InCommitType)
 				{
-					if (InCommitType == ETextCommit::OnEnter)
-					{
-						const double Time = FCString::Atod(*InText.ToString());
-						TimeMarker->SetTime(Time);
-					}
+					const double Time = FCString::Atod(*InText.ToString());
+					TimeMarker->SetTime(Time);
 				})
 		]
 
