[Insights] Memory Insights query polling optimization

diff --git a/Engine/Source/Developer/TraceInsights/Private/Insights/MemoryProfiler/Widgets/SMemAllocTableTreeView.cpp b/Engine/Source/Developer/TraceInsights/Private/Insights/MemoryProfiler/Widgets/SMemAllocTableTreeView.cpp
--- a/Engine/Source/Developer/TraceInsights/Private/Insights/MemoryProfiler/Widgets/SMemAllocTableTreeView.cpp
+++ b/Engine/Source/Developer/TraceInsights/Private/Insights/MemoryProfiler/Widgets/SMemAllocTableTreeView.cpp
@@ -72,17 +72,10 @@ void SMemAllocTableTreeView::UpdateSourceTable(TSharedPtr<TraceServices::IMemAll
 void SMemAllocTableTreeView::Tick(const FGeometry& AllottedGeometry, const double InCurrentTime, const float InDeltaTime)
 {
 	STableTreeView::Tick(AllottedGeometry, InCurrentTime, InDeltaTime);
-	// We need to check if the list of LLM tags has changed.
-	// But, ensure we do not check too often.
-	static uint64 NextTimestamp = 0;
-	uint64 Time = FPlatformTime::Cycles64();
-	if (!bIsUpdateRunning && Time > NextTimestamp)
+
+	if (!bIsUpdateRunning)
 	{
 		RebuildTree(false);
-
-		const double WaitTimeSec = 0.5;
-		const uint64 WaitTime = static_cast<uint64>(WaitTimeSec / FPlatformTime::GetSecondsPerCycle64());
-		NextTimestamp = Time + WaitTime;
 	}
 }
 
@@ -251,90 +244,101 @@ void SMemAllocTableTreeView::UpdateQuery(TraceServices::IAllocationsProvider::EQ
 	const TraceServices::IAllocationsProvider& Provider = *AllocationsProvider;
 	TraceServices::IAllocationsProvider::FReadScopeLock _(Provider);
 
-	TraceServices::IAllocationsProvider::FQueryStatus Status = Provider.PollQuery(Query);
+	constexpr double MaxPollTime = 0.03; // Stop getting results after 30 ms so we don't tank the frame rate too much.
+	FStopwatch TotalStopwatch;
+	TotalStopwatch.Start();
 
-	OutStatus = Status.Status;
-	if (Status.Status <= TraceServices::IAllocationsProvider::EQueryStatus::Done)
-	{
-		UE_LOG(MemoryProfiler, Log, TEXT("[MemAlloc] Query completed."));
-		Query = 0;
-		QueryStopwatch.Stop();
-		return;
-	}
-
-	TSharedPtr<Insights::FMemAllocTable> MemAllocTable = GetMemAllocTable();
-	if (MemAllocTable)
+	do
 	{
-		TArray<FMemoryAlloc>& Allocs = MemAllocTable->GetAllocs();
+		TraceServices::IAllocationsProvider::FQueryStatus Status = Provider.PollQuery(Query);
 
-		FStopwatch TotalStopwatch;
-		FStopwatch PageStopwatch;
-		TotalStopwatch.Start();
-		uint32 PageCount = 0;
-		uint32 TotatAllocCount = 0;
+		OutStatus = Status.Status;
+		if (Status.Status <= TraceServices::IAllocationsProvider::EQueryStatus::Done)
+		{
+			UE_LOG(MemoryProfiler, Log, TEXT("[MemAlloc] Query completed."));
+			Query = 0;
+			QueryStopwatch.Stop();
+			return;
+		}
 
-		// Multiple 'pages' of results will be returned. No guarantees are made
-		// about the order of pages or the allocations they report.
-		TraceServices::IAllocationsProvider::FQueryResult Result = Status.NextResult();
-		while (Result.IsValid())
+		TSharedPtr<Insights::FMemAllocTable> MemAllocTable = GetMemAllocTable();
+		if (MemAllocTable)
 		{
-			++PageCount;
-			PageStopwatch.Restart();
+			TArray<FMemoryAlloc>& Allocs = MemAllocTable->GetAllocs();
 
-			const uint32 AllocCount = Result->Num();
-			TotatAllocCount += AllocCount;
-#if 0
-			Allocs.Reserve(Allocs.Num() + AllocCount);
+			FStopwatch ResultStopwatch;
+			FStopwatch PageStopwatch;
+			ResultStopwatch.Start();
+			uint32 PageCount = 0;
+			uint32 TotatAllocCount = 0;
 
-			for (uint32 AllocIndex = 0; AllocIndex < AllocCount; ++AllocIndex)
+			// Multiple 'pages' of results will be returned. No guarantees are made
+			// about the order of pages or the allocations they report.
+			TraceServices::IAllocationsProvider::FQueryResult Result = Status.NextResult();
+			while (Result.IsValid())
 			{
-				const IAllocationsProvider::FAllocation* Allocation = Result->Get(AllocIndex);
+				++PageCount;
+				PageStopwatch.Restart();
 
-				const double AllocStartTime = Allocation->GetStartTime();
-				const double AllocEndTime = Allocation->GetEndTime();
-				const uint64 AllocAddress = Allocation->GetAddress();
-				const uint64 AllocationSize = Allocation->GetSize();
-				const FMemoryTagId AllocTag = static_cast<FMemoryTagId>(Allocation->GetTag());
-				const uint64 BacktraceId = Allocation->GetBacktraceId();
+				const uint32 AllocCount = Result->Num();
+				TotatAllocCount += AllocCount;
+#if 0
+				Allocs.Reserve(Allocs.Num() + AllocCount);
 
-				//Allocs.Add(FMemoryAlloc(AllocStartTime, AllocEndTime, AllocAddress, AllocationSize, AllocTag, BacktraceId));
-				Allocs.Emplace(AllocStartTime, AllocEndTime, AllocAddress, AllocationSize, AllocTag, BacktraceId);
-			}
+				for (uint32 AllocIndex = 0; AllocIndex < AllocCount; ++AllocIndex)
+				{
+					const IAllocationsProvider::FAllocation* Allocation = Result->Get(AllocIndex);
+
+					const double AllocStartTime = Allocation->GetStartTime();
+					const double AllocEndTime = Allocation->GetEndTime();
+					const uint64 AllocAddress = Allocation->GetAddress();
+					const uint64 AllocationSize = Allocation->GetSize();
+					const FMemoryTagId AllocTag = static_cast<FMemoryTagId>(Allocation->GetTag());
+					const uint64 BacktraceId = Allocation->GetBacktraceId();
+
+					//Allocs.Add(FMemoryAlloc(AllocStartTime, AllocEndTime, AllocAddress, AllocationSize, AllocTag, BacktraceId));
+					Allocs.Emplace(AllocStartTime, AllocEndTime, AllocAddress, AllocationSize, AllocTag, BacktraceId);
+				}
 #else // 14% faster, but uses more mem
-			uint64 AllocsDestIndex = Allocs.Num();
-			Allocs.AddUninitialized(AllocCount);
-			for (uint32 AllocIndex = 0; AllocIndex < AllocCount; ++AllocIndex, ++AllocsDestIndex)
-			{
-				const TraceServices::IAllocationsProvider::FAllocation* Allocation = Result->Get(AllocIndex);
-				FMemoryAlloc& Alloc = Allocs[AllocsDestIndex];
-				Alloc.StartTime = Allocation->GetStartTime();
-				Alloc.EndTime = Allocation->GetEndTime();
-				Alloc.Address = Allocation->GetAddress();
-				Alloc.Size = Allocation->GetSize();
-				Alloc.MemTag = static_cast<FMemoryTagId>(Allocation->GetTag());
-				Alloc.BacktraceId = Allocation->GetBacktraceId();
-			}
+				uint64 AllocsDestIndex = Allocs.Num();
+				Allocs.AddUninitialized(AllocCount);
+				for (uint32 AllocIndex = 0; AllocIndex < AllocCount; ++AllocIndex, ++AllocsDestIndex)
+				{
+					const TraceServices::IAllocationsProvider::FAllocation* Allocation = Result->Get(AllocIndex);
+					FMemoryAlloc& Alloc = Allocs[AllocsDestIndex];
+					Alloc.StartTime = Allocation->GetStartTime();
+					Alloc.EndTime = Allocation->GetEndTime();
+					Alloc.Address = Allocation->GetAddress();
+					Alloc.Size = Allocation->GetSize();
+					Alloc.MemTag = static_cast<FMemoryTagId>(Allocation->GetTag());
+					Alloc.BacktraceId = Allocation->GetBacktraceId();
+				}
 #endif
 
-			PageStopwatch.Stop();
-			const double PageTime = PageStopwatch.GetAccumulatedTime();
-			if (PageTime > 0.01)
-			{
-				const double Speed = (PageTime * 1000000.0) / AllocCount;
-				UE_LOG(MemoryProfiler, Log, TEXT("[MemAlloc] Query result for page %u (%u allocs, slack=%u) retrieved in %.3fs (speed: %.3f seconds per 1M allocs)."), PageCount, AllocCount, Allocs.GetSlack(), PageTime, Speed);
+				PageStopwatch.Stop();
+				const double PageTime = PageStopwatch.GetAccumulatedTime();
+				if (PageTime > 0.01)
+				{
+					const double Speed = (PageTime * 1000000.0) / AllocCount;
+					UE_LOG(MemoryProfiler, Log, TEXT("[MemAlloc] Query result for page %u (%u allocs, slack=%u) retrieved in %.3fs (speed: %.3f seconds per 1M allocs)."), PageCount, AllocCount, Allocs.GetSlack(), PageTime, Speed);
+				}
+
+				Result = Status.NextResult();
 			}
 
-			Result = Status.NextResult();
+			ResultStopwatch.Stop();
+			const double TotalTime = ResultStopwatch.GetAccumulatedTime();
+			if (TotalTime > 0.01)
+			{
+				const double Speed = (TotalTime * 1000000.0) / TotatAllocCount;
+				UE_LOG(MemoryProfiler, Log, TEXT("[MemAlloc] Query results (%u pages, %u allocs, slack=%u) retrieved in %.3fs (speed: %.3f seconds per 1M allocs)."), PageCount, TotatAllocCount, Allocs.GetSlack(), TotalTime, Speed);
+			}
 		}
 
-		TotalStopwatch.Stop();
-		const double TotalTime = TotalStopwatch.GetAccumulatedTime();
-		if (TotalTime > 0.01)
-		{
-			const double Speed = (TotalTime * 1000000.0) / TotatAllocCount;
-			UE_LOG(MemoryProfiler, Log, TEXT("[MemAlloc] Query results (%u pages, %u allocs, slack=%u) retrieved in %.3fs (speed: %.3f seconds per 1M allocs)."), PageCount, TotatAllocCount, Allocs.GetSlack(), TotalTime, Speed);
-		}
-	}
+		TotalStopwatch.Update();
+	} while (OutStatus == TraceServices::IAllocationsProvider::EQueryStatus::Available && TotalStopwatch.GetAccumulatedTime() < MaxPollTime);
+
+	TotalStopwatch.Stop();
 }
 
 ////////////////////////////////////////////////////////////////////////////////////////////////////
