[Insights] Allow Sorting and Filtering settings to be changed during async operations in Mem Alloc Tree View

diff --git a/Engine/Source/Developer/TraceInsights/Private/Insights/Table/Widgets/STableTreeView.cpp b/Engine/Source/Developer/TraceInsights/Private/Insights/Table/Widgets/STableTreeView.cpp
--- a/Engine/Source/Developer/TraceInsights/Private/Insights/Table/Widgets/STableTreeView.cpp
+++ b/Engine/Source/Developer/TraceInsights/Private/Insights/Table/Widgets/STableTreeView.cpp
@@ -230,7 +230,15 @@ void STableTreeView::ConstructWidget(TSharedPtr<FTable> InTablePtr)
 	// Create the search filters: text based, type based etc.
 	TextFilter = MakeShared<FTableTreeNodeTextFilter>(FTableTreeNodeTextFilter::FItemToStringArray::CreateStatic(&STableTreeView::HandleItemToStringArray));
 	Filters = MakeShared<FTableTreeNodeFilterCollection>();
-	Filters->Add(TextFilter);
+	if (bRunInAsyncMode)
+	{
+		CurrentAsyncOpTextFilter = MakeShared<FTableTreeNodeTextFilter>(FTableTreeNodeTextFilter::FItemToStringArray::CreateStatic(&STableTreeView::HandleItemToStringArray));
+		Filters->Add(CurrentAsyncOpTextFilter);
+	}
+	else
+	{
+		Filters->Add(TextFilter);
+	}
 
 	//BindCommands();
 
@@ -628,7 +636,10 @@ void STableTreeView::UpdateTree()
 	else
 	{
 		CreateGroups(CurrentGroupings);
-		SortTreeNodes();
+		if (CurrentSorter.IsValid())
+		{
+			SortTreeNodes(CurrentSorter.Get(), ColumnSortMode);
+		}
 		ApplyFiltering();
 	}
 
@@ -905,11 +916,6 @@ void STableTreeView::SearchBox_OnTextChanged(const FText& InFilterText)
 
 bool STableTreeView::SearchBox_IsEnabled() const
 {
-	if (bIsUpdateRunning)
-	{
-		return false;
-	}
-
 	return TableTreeNodes.Num() > 0;
 }
 
@@ -1218,7 +1224,10 @@ void STableTreeView::PostChangeGroupings()
 	else
 	{
 		CreateGroups(CurrentGroupings);
-		SortTreeNodes();
+		if (CurrentSorter.IsValid())
+		{
+			SortTreeNodes(CurrentSorter.Get(), ColumnSortMode);
+		}
 		ApplyFiltering();
 	}
 
@@ -1635,29 +1644,26 @@ void STableTreeView::UpdateCurrentSortingByColumn()
 
 ////////////////////////////////////////////////////////////////////////////////////////////////////
 
-void STableTreeView::SortTreeNodes()
+void STableTreeView::SortTreeNodes(ITableCellValueSorter* InSorter, EColumnSortMode::Type InColumnSortMode)
 {
-	if (CurrentSorter.IsValid())
-	{
-		FStopwatch Stopwatch;
-		Stopwatch.Start();
+	FStopwatch Stopwatch;
+	Stopwatch.Start();
 
-		SortTreeNodesRec(*Root, *CurrentSorter);
+	SortTreeNodesRec(*Root, *InSorter, InColumnSortMode);
 
-		Stopwatch.Stop();
-		const double TotalTime = Stopwatch.GetAccumulatedTime();
-		if (TotalTime > 0.1)
-		{
-			UE_LOG(TraceInsights, Log, TEXT("[Tree - %s] Sorting completed in %.3fs."), *Table->GetDisplayName().ToString(), TotalTime);
-		}
+	Stopwatch.Stop();
+	const double TotalTime = Stopwatch.GetAccumulatedTime();
+	if (TotalTime > 0.1)
+	{
+		UE_LOG(TraceInsights, Log, TEXT("[Tree - %s] Sorting completed in %.3fs."), *Table->GetDisplayName().ToString(), TotalTime);
 	}
 }
 
 ////////////////////////////////////////////////////////////////////////////////////////////////////
 
-void STableTreeView::SortTreeNodesRec(FTableTreeNode& GroupNode, const ITableCellValueSorter& Sorter)
+void STableTreeView::SortTreeNodesRec(FTableTreeNode& GroupNode, const ITableCellValueSorter& Sorter, EColumnSortMode::Type InColumnSortMode)
 {
-	if (ColumnSortMode == EColumnSortMode::Type::Descending)
+	if (InColumnSortMode == EColumnSortMode::Type::Descending)
 	{
 		GroupNode.SortChildrenDescending(Sorter);
 	}
@@ -1670,7 +1676,7 @@ void STableTreeView::SortTreeNodesRec(FTableTreeNode& GroupNode, const ITableCel
 	{
 		if (ChildPtr->IsGroup())
 		{
-			SortTreeNodesRec(*StaticCastSharedPtr<FTableTreeNode>(ChildPtr), Sorter);
+			SortTreeNodesRec(*StaticCastSharedPtr<FTableTreeNode>(ChildPtr), Sorter, InColumnSortMode);
 		}
 	}
 }
@@ -1707,7 +1713,10 @@ void STableTreeView::SetSortModeForColumn(const FName& ColumnId, const EColumnSo
 	}
 	else
 	{
-		SortTreeNodes();
+		if (CurrentSorter.IsValid())
+		{
+			SortTreeNodes(CurrentSorter.Get(), ColumnSortMode);
+		}
 		ApplyFiltering();
 	}
 }
@@ -1716,11 +1725,6 @@ void STableTreeView::SetSortModeForColumn(const FName& ColumnId, const EColumnSo
 
 void STableTreeView::OnSortModeChanged(const EColumnSortPriority::Type SortPriority, const FName& ColumnId, const EColumnSortMode::Type SortMode)
 {
-	if (bIsUpdateRunning)
-	{
-		return;
-	}
-
 	SetSortModeForColumn(ColumnId, SortMode);
 	TreeView_Refresh();
 }
@@ -1738,11 +1742,6 @@ bool STableTreeView::HeaderMenu_SortMode_IsChecked(const FName ColumnId, const E
 
 bool STableTreeView::HeaderMenu_SortMode_CanExecute(const FName ColumnId, const EColumnSortMode::Type InSortMode) const
 {
-	if (bIsUpdateRunning)
-	{
-		return false;
-	}
-
 	const FTableColumn& Column = *Table->FindColumnChecked(ColumnId);
 	return Column.CanBeSorted();
 }
@@ -1768,11 +1767,6 @@ bool STableTreeView::ContextMenu_SortMode_IsChecked(const EColumnSortMode::Type
 
 bool STableTreeView::ContextMenu_SortMode_CanExecute(const EColumnSortMode::Type InSortMode) const
 {
-	if (bIsUpdateRunning)
-	{
-		return false;
-	}
-
 	return true; //ColumnSortMode != InSortMode;
 }
 
@@ -1797,11 +1791,6 @@ bool STableTreeView::ContextMenu_SortByColumn_IsChecked(const FName ColumnId)
 
 bool STableTreeView::ContextMenu_SortByColumn_CanExecute(const FName ColumnId) const
 {
-	if (bIsUpdateRunning)
-	{
-		return false;
-	}
-
 	return true; //ColumnId != ColumnBeingSorted;
 }
 
@@ -2093,17 +2082,25 @@ void STableTreeView::OnPostAsyncUpdate()
 
 FGraphEventRef STableTreeView::StartSortTreeNodesTask(FGraphEventRef Prerequisite)
 {
+	if (!CurrentSorter.IsValid())
+	{
+		return Prerequisite;
+	}
+
 	AddInProgressAsyncOperation(EAsyncOperationType::SortingOp);
 
+	CurrentAsyncOpSorter = CurrentSorter.Get();
+	CurrentAsyncOpColumnSortMode = ColumnSortMode;
+
 	if (Prerequisite.IsValid())
 	{
 		FGraphEventArray Prerequisites;
 		Prerequisites.Add(Prerequisite);
-		return TGraphTask<FTableTreeViewSortAsyncTask>::CreateTask(&Prerequisites).ConstructAndDispatchWhenReady(this);
+		return TGraphTask<FTableTreeViewSortAsyncTask>::CreateTask(&Prerequisites).ConstructAndDispatchWhenReady(this, CurrentAsyncOpSorter, CurrentAsyncOpColumnSortMode);
 	}
 	else
 	{
-		return TGraphTask<FTableTreeViewSortAsyncTask>::CreateTask().ConstructAndDispatchWhenReady(this);
+		return TGraphTask<FTableTreeViewSortAsyncTask>::CreateTask().ConstructAndDispatchWhenReady(this, CurrentAsyncOpSorter, CurrentAsyncOpColumnSortMode);
 	}
 }
 
@@ -2134,6 +2131,8 @@ FGraphEventRef STableTreeView::StartApplyFiltersTask(FGraphEventRef Prerequisite
 {
 	AddInProgressAsyncOperation(EAsyncOperationType::FilteringOp);
 
+	CurrentAsyncOpTextFilter->SetRawFilterText(TextFilter->GetRawFilterText());
+
 	if (Prerequisite.IsValid())
 	{
 		FGraphEventArray Prerequisites;
@@ -2179,6 +2178,7 @@ double STableTreeView::GetAllOperationsDuration()
 
 void STableTreeView::StartPendingAsyncOperations()
 {
+	// Check if grouping settings have changed. If they did, a full refresh (Grouping, Sorting and Filtering) is scheduled.
 	bool bGroupingsHaveChanged = CurrentGroupings.Num() != CurrentAsyncOpGroupings.Num();
 
 	if (!bGroupingsHaveChanged)
@@ -2200,6 +2200,36 @@ void STableTreeView::StartPendingAsyncOperations()
 		FGraphEventRef CompletedEvent = StartCreateGroupsTask();
 		CompletedEvent = StartSortTreeNodesTask(CompletedEvent);
 		InProgressAsyncOperationEvent = StartApplyFiltersTask(CompletedEvent);
+
+		return;
+	}
+
+	// Check if sorting settings have changed. If they did, a Sorting and Filtering Refresh is scheduled.
+	bool bSortingHasChanged = (CurrentSorter.IsValid() && CurrentAsyncOpSorter == nullptr) || (!CurrentSorter.IsValid() && CurrentAsyncOpSorter != nullptr);
+	if (!bSortingHasChanged && CurrentSorter.IsValid())
+	{
+		bSortingHasChanged = CurrentSorter.Get() != CurrentAsyncOpSorter || ColumnSortMode != CurrentAsyncOpColumnSortMode;
+	}
+
+	if (bSortingHasChanged)
+	{
+		OnPreAsyncUpdate();
+
+		FGraphEventRef CompletedEvent = StartSortTreeNodesTask();
+		InProgressAsyncOperationEvent = StartApplyFiltersTask(CompletedEvent);
+
+		return;
+	}
+
+	// Check if the text filter has changed. If it has, schedule a new Filtering Refresh.
+	bool bTextFilterHasChanged = TextFilter->GetRawFilterText().CompareTo(CurrentAsyncOpTextFilter->GetRawFilterText()) != 0;
+	if (bTextFilterHasChanged)
+	{
+		OnPreAsyncUpdate();
+
+		InProgressAsyncOperationEvent = StartApplyFiltersTask();
+
+		return;
 	}
 }
 
diff --git a/Engine/Source/Developer/TraceInsights/Private/Insights/Table/Widgets/STableTreeView.h b/Engine/Source/Developer/TraceInsights/Private/Insights/Table/Widgets/STableTreeView.h
--- a/Engine/Source/Developer/TraceInsights/Private/Insights/Table/Widgets/STableTreeView.h
+++ b/Engine/Source/Developer/TraceInsights/Private/Insights/Table/Widgets/STableTreeView.h
@@ -237,8 +237,8 @@ protected:
 	void CreateSortings();
 
 	void UpdateCurrentSortingByColumn();
-	void SortTreeNodes();
-	void SortTreeNodesRec(FTableTreeNode& GroupNode, const ITableCellValueSorter& Sorter);
+	void SortTreeNodes(ITableCellValueSorter* InSorter, EColumnSortMode::Type InColumnSortMode);
+	void SortTreeNodesRec(FTableTreeNode& GroupNode, const ITableCellValueSorter& Sorter, EColumnSortMode::Type InColumnSortMode);
 
 	EColumnSortMode::Type GetSortModeForColumn(const FName ColumnId) const;
 	void SetSortModeForColumn(const FName& ColumnId, EColumnSortMode::Type SortMode);
@@ -403,6 +403,9 @@ protected:
 	FStopwatch AsyncUpdateStopwatch;
 
 	TArray<TSharedPtr<FTreeNodeGrouping>> CurrentAsyncOpGroupings;
+	ITableCellValueSorter* CurrentAsyncOpSorter = nullptr;
+	EColumnSortMode::Type CurrentAsyncOpColumnSortMode;
+	TSharedPtr<FTableTreeNodeTextFilter> CurrentAsyncOpTextFilter;
 
 	//////////////////////////////////////////////////
 
@@ -441,9 +444,11 @@ private:
 class FTableTreeViewSortAsyncTask
 {
 public:
-	FTableTreeViewSortAsyncTask(STableTreeView* InPtr)
+	FTableTreeViewSortAsyncTask(STableTreeView* InPtr, ITableCellValueSorter* InSorter, EColumnSortMode::Type InColumnSortMode)
 	{
 		TableTreeViewPtr = InPtr;
+		Sorter = InSorter;
+		ColumnSortMode = InColumnSortMode;
 	}
 
 	FORCEINLINE TStatId GetStatId() const { RETURN_QUICK_DECLARE_CYCLE_STAT(FTableTreeViewSortAsyncTask, STATGROUP_TaskGraphTasks); }
@@ -454,12 +459,14 @@ public:
 	{
 		if (TableTreeViewPtr)
 		{
-			TableTreeViewPtr->SortTreeNodes();
+			TableTreeViewPtr->SortTreeNodes(Sorter, ColumnSortMode);
 		}
 	}
 
 private:
-	STableTreeView* TableTreeViewPtr = nullptr;
+	STableTreeView* TableTreeViewPtr;
+	ITableCellValueSorter* Sorter;
+	EColumnSortMode::Type ColumnSortMode;
 };
 
 ////////////////////////////////////////////////////////////////////////////////////////////////////
