UnrealInsights: Add the shortcut 'Ctrl-A' to select the entire timeline

diff --git a/Engine/Source/Developer/TraceInsights/Private/Insights/Widgets/STimingView.cpp b/Engine/Source/Developer/TraceInsights/Private/Insights/Widgets/STimingView.cpp
--- a/Engine/Source/Developer/TraceInsights/Private/Insights/Widgets/STimingView.cpp
+++ b/Engine/Source/Developer/TraceInsights/Private/Insights/Widgets/STimingView.cpp
@@ -2579,6 +2579,14 @@ FReply STimingView::OnKeyDown(const FGeometry& MyGeometry, const FKeyEvent& InKe
 		FileActivitySharedState->ToggleBackgroundEvents();
 		return FReply::Handled();
 	}
+	else if (InKeyEvent.GetKey() == EKeys::A)
+	{
+		if (InKeyEvent.GetModifierKeys().IsControlDown())
+		{
+			SelectTimeInterval(0, Viewport.GetMaxValidTime());
+			return FReply::Handled();
+		}
+	}
 	else if (InKeyEvent.GetKey() == EKeys::One)
 	{
 		LoadingSharedState->SetColorSchema(0);
