[Insights] Alloc Tree View async operation cancelation

diff --git a/Engine/Source/Developer/TraceInsights/Private/Insights/Table/ViewModels/TreeNodeGrouping.cpp b/Engine/Source/Developer/TraceInsights/Private/Insights/Table/ViewModels/TreeNodeGrouping.cpp
--- a/Engine/Source/Developer/TraceInsights/Private/Insights/Table/ViewModels/TreeNodeGrouping.cpp
+++ b/Engine/Source/Developer/TraceInsights/Private/Insights/Table/ViewModels/TreeNodeGrouping.cpp
@@ -22,7 +22,7 @@ FTreeNodeGrouping::FTreeNodeGrouping(const FText& InShortName, const FText& InTi
 
 ////////////////////////////////////////////////////////////////////////////////////////////////////
 
-void FTreeNodeGrouping::GroupNodes(const TArray<FTableTreeNodePtr>& Nodes, FTableTreeNode& ParentGroup, TWeakPtr<FTable> InParentTable) const
+void FTreeNodeGrouping::GroupNodes(const TArray<FTableTreeNodePtr>& Nodes, FTableTreeNode& ParentGroup, TWeakPtr<FTable> InParentTable, std::atomic<bool>& bCancelGrouping) const
 {
 	TMap<FName, FTableTreeNodePtr> GroupMap;
 
@@ -32,6 +32,11 @@ void FTreeNodeGrouping::GroupNodes(const TArray<FTableTreeNodePtr>& Nodes, FTabl
 	{
 		ensure(!NodePtr->IsGroup());
 
+		if (bCancelGrouping)
+		{
+			return;
+		}
+
 		FTableTreeNodePtr GroupPtr = nullptr;
 
 		FTreeNodeGroupInfo GroupInfo = GetGroupForNode(NodePtr);
@@ -77,7 +82,7 @@ FTreeNodeGroupInfo FTreeNodeGroupingFlat::GetGroupForNode(const FBaseTreeNodePtr
 
 ////////////////////////////////////////////////////////////////////////////////////////////////////
 
-void FTreeNodeGroupingFlat::GroupNodes(const TArray<FTableTreeNodePtr>& Nodes, FTableTreeNode& ParentGroup, TWeakPtr<FTable> InParentTable) const
+void FTreeNodeGroupingFlat::GroupNodes(const TArray<FTableTreeNodePtr>& Nodes, FTableTreeNode& ParentGroup, TWeakPtr<FTable> InParentTable, std::atomic<bool>& bCancelGrouping) const
 {
 	ParentGroup.ClearChildren(1);
 
@@ -88,6 +93,11 @@ void FTreeNodeGroupingFlat::GroupNodes(const TArray<FTableTreeNodePtr>& Nodes, F
 	GroupPtr->ClearChildren(Nodes.Num());
 	for (FTableTreeNodePtr NodePtr : Nodes)
 	{
+		if (bCancelGrouping)
+		{
+			return;
+		}
+
 		GroupPtr->AddChildAndSetGroupPtr(NodePtr);
 	}
 }
diff --git a/Engine/Source/Developer/TraceInsights/Private/Insights/Table/ViewModels/TreeNodeGrouping.h b/Engine/Source/Developer/TraceInsights/Private/Insights/Table/ViewModels/TreeNodeGrouping.h
--- a/Engine/Source/Developer/TraceInsights/Private/Insights/Table/ViewModels/TreeNodeGrouping.h
+++ b/Engine/Source/Developer/TraceInsights/Private/Insights/Table/ViewModels/TreeNodeGrouping.h
@@ -11,6 +11,8 @@
 #include "Insights/Table/ViewModels/TableColumn.h"
 #include "Insights/Table/ViewModels/TableTreeNode.h"
 
+#include <atomic>
+
 namespace Insights
 {
 
@@ -57,7 +59,7 @@ public:
 	virtual FName GetColumnId() const override { return NAME_None; }
 
 	virtual FTreeNodeGroupInfo GetGroupForNode(const FBaseTreeNodePtr InNode) const { return { FName(), false }; }
-	virtual void GroupNodes(const TArray<FTableTreeNodePtr>& Nodes, FTableTreeNode& ParentGroup, TWeakPtr<FTable> InParentTable) const;
+	virtual void GroupNodes(const TArray<FTableTreeNodePtr>& Nodes, FTableTreeNode& ParentGroup, TWeakPtr<FTable> InParentTable, std::atomic<bool>& bCancelGrouping) const;
 
 protected:
 	FText ShortName;
@@ -77,7 +79,7 @@ public:
 	virtual ~FTreeNodeGroupingFlat() {}
 
 	virtual FTreeNodeGroupInfo GetGroupForNode(const FBaseTreeNodePtr InNode) const override;
-	virtual void GroupNodes(const TArray<FTableTreeNodePtr>& Nodes, FTableTreeNode& ParentGroup, TWeakPtr<FTable> InParentTable) const override;
+	virtual void GroupNodes(const TArray<FTableTreeNodePtr>& Nodes, FTableTreeNode& ParentGroup, TWeakPtr<FTable> InParentTable, std::atomic<bool>& bCancelGrouping) const override;
 };
 
 ////////////////////////////////////////////////////////////////////////////////////////////////////
diff --git a/Engine/Source/Developer/TraceInsights/Private/Insights/Table/Widgets/STableTreeView.cpp b/Engine/Source/Developer/TraceInsights/Private/Insights/Table/Widgets/STableTreeView.cpp
--- a/Engine/Source/Developer/TraceInsights/Private/Insights/Table/Widgets/STableTreeView.cpp
+++ b/Engine/Source/Developer/TraceInsights/Private/Insights/Table/Widgets/STableTreeView.cpp
@@ -632,6 +632,10 @@ void STableTreeView::UpdateTree()
 			CompletedEvent = StartSortTreeNodesTask(CompletedEvent);
 			InProgressAsyncOperationEvent = StartApplyFiltersTask(CompletedEvent);
 		}
+		else
+		{
+			CancelCurrentAsyncOp();
+		}
 	}
 	else
 	{
@@ -735,6 +739,10 @@ bool STableTreeView::ApplyFilteringForNode(FTableTreeNodePtr NodePtr)
 		int32 NumVisibleChildren = 0;
 		for (int32 Cx = 0; Cx < NumChildren; ++Cx)
 		{
+			if (bCancelCurrentAsyncOp)
+			{
+				break;
+			}
 			// Add a child.
 			const FTableTreeNodePtr& ChildNodePtr = StaticCastSharedPtr<FTableTreeNode>(GroupChildren[Cx]);
 			if (ApplyFilteringForNode(ChildNodePtr))
@@ -905,6 +913,10 @@ void STableTreeView::SearchBox_OnTextChanged(const FText& InFilterText)
 
 			InProgressAsyncOperationEvent = StartApplyFiltersTask();
 		}
+		else
+		{
+			CancelCurrentAsyncOp();
+		}
 	}
 	else
 	{
@@ -968,11 +980,16 @@ void STableTreeView::CreateGroups(const TArray<TSharedPtr<FTreeNodeGrouping>>& G
 
 void STableTreeView::GroupNodesRec(const TArray<FTableTreeNodePtr>& Nodes, FTableTreeNode& ParentGroup, int32 GroupingDepth, const TArray<TSharedPtr<FTreeNodeGrouping>>& Groupings)
 {
+	if (bIsUpdateRunning && bCancelCurrentAsyncOp)
+	{
+		return;
+	}
+
 	ensure(Groupings.Num() > 0);
 
 	FTreeNodeGrouping& Grouping = *Groupings[GroupingDepth];
 
-	Grouping.GroupNodes(Nodes, ParentGroup, Table);
+	Grouping.GroupNodes(Nodes, ParentGroup, Table, bCancelCurrentAsyncOp);
 
 	for (FBaseTreeNodePtr GroupPtr : ParentGroup.GetChildren())
 	{
@@ -1220,6 +1237,10 @@ void STableTreeView::PostChangeGroupings()
 			CompletedEvent = StartSortTreeNodesTask(CompletedEvent);
 			InProgressAsyncOperationEvent = StartApplyFiltersTask(CompletedEvent);
 		}
+		else
+		{
+			CancelCurrentAsyncOp();
+		}
 	}
 	else
 	{
@@ -1674,6 +1695,10 @@ void STableTreeView::SortTreeNodesRec(FTableTreeNode& GroupNode, const ITableCel
 
 	for (FBaseTreeNodePtr ChildPtr : GroupNode.GetChildren())
 	{
+		if (bCancelCurrentAsyncOp)
+		{
+			break;
+		}
 		if (ChildPtr->IsGroup())
 		{
 			SortTreeNodesRec(*StaticCastSharedPtr<FTableTreeNode>(ChildPtr), Sorter, InColumnSortMode);
@@ -1710,6 +1735,10 @@ void STableTreeView::SetSortModeForColumn(const FName& ColumnId, const EColumnSo
 			FGraphEventRef CompletedEvent = StartSortTreeNodesTask();
 			InProgressAsyncOperationEvent = StartApplyFiltersTask(CompletedEvent);
 		}
+		else
+		{
+			CancelCurrentAsyncOp();
+		}
 	}
 	else
 	{
@@ -2041,6 +2070,8 @@ void STableTreeView::OnPreAsyncUpdate()
 {
 	check(!bIsUpdateRunning);
 
+	ClearInProgressAsyncOperations();
+
 	AsyncUpdateStopwatch.Restart();
 	bIsUpdateRunning = true;
 
@@ -2059,22 +2090,26 @@ void STableTreeView::OnPostAsyncUpdate()
 
 	bIsUpdateRunning = false;
 
-	TreeView->SetTreeItemsSource(&FilteredGroupNodes);
+	if (!bCancelCurrentAsyncOp)
+	{
+		TreeView->SetTreeItemsSource(&FilteredGroupNodes);
 
-	TreeView->ClearExpandedItems();
+		TreeView->ClearExpandedItems();
 
-	// Grouping can result in old group nodes no longer existing in the tree, so we don't keep the expanded list.
-	if (!HasInProgressAsyncOperation(EAsyncOperationType::GroupingOp))
-	{
-		for (auto It = ExpandedNodes.CreateConstIterator(); It; ++It)
+		// Grouping can result in old group nodes no longer existing in the tree, so we don't keep the expanded list.
+		if (!HasInProgressAsyncOperation(EAsyncOperationType::GroupingOp))
 		{
-			TreeView->SetItemExpansion(*It, true);
+			for (auto It = ExpandedNodes.CreateConstIterator(); It; ++It)
+			{
+				TreeView->SetItemExpansion(*It, true);
+			}
 		}
-	}
 
-	ClearInProgressAsyncOperations();
-	TreeView_Refresh();
+		ClearInProgressAsyncOperations();
+		TreeView_Refresh();
+	}
 
+	bCancelCurrentAsyncOp = false;
 	AsyncUpdateStopwatch.Stop();
 }
 
@@ -2152,6 +2187,7 @@ void STableTreeView::OnClose()
 	if (bIsUpdateRunning && InProgressAsyncOperationEvent.IsValid() && !InProgressAsyncOperationEvent->IsComplete())
 	{
 		bIsCloseScheduled = true;
+		CancelCurrentAsyncOp();
 
 		FGraphEventArray Prerequisites;
 		Prerequisites.Add(InProgressAsyncOperationEvent);
@@ -2179,7 +2215,8 @@ double STableTreeView::GetAllOperationsDuration()
 void STableTreeView::StartPendingAsyncOperations()
 {
 	// Check if grouping settings have changed. If they did, a full refresh (Grouping, Sorting and Filtering) is scheduled.
-	bool bGroupingsHaveChanged = CurrentGroupings.Num() != CurrentAsyncOpGroupings.Num();
+	bool bGroupingsHaveChanged = HasInProgressAsyncOperation(EAsyncOperationType::GroupingOp);
+	bGroupingsHaveChanged |= CurrentGroupings.Num() != CurrentAsyncOpGroupings.Num();
 
 	if (!bGroupingsHaveChanged)
 	{
@@ -2205,7 +2242,8 @@ void STableTreeView::StartPendingAsyncOperations()
 	}
 
 	// Check if sorting settings have changed. If they did, a Sorting and Filtering Refresh is scheduled.
-	bool bSortingHasChanged = (CurrentSorter.IsValid() && CurrentAsyncOpSorter == nullptr) || (!CurrentSorter.IsValid() && CurrentAsyncOpSorter != nullptr);
+	bool bSortingHasChanged = HasInProgressAsyncOperation(EAsyncOperationType::SortingOp);
+	bSortingHasChanged |= (CurrentSorter.IsValid() && CurrentAsyncOpSorter == nullptr) || (!CurrentSorter.IsValid() && CurrentAsyncOpSorter != nullptr);
 	if (!bSortingHasChanged && CurrentSorter.IsValid())
 	{
 		bSortingHasChanged = CurrentSorter.Get() != CurrentAsyncOpSorter || ColumnSortMode != CurrentAsyncOpColumnSortMode;
@@ -2222,7 +2260,8 @@ void STableTreeView::StartPendingAsyncOperations()
 	}
 
 	// Check if the text filter has changed. If it has, schedule a new Filtering Refresh.
-	bool bTextFilterHasChanged = TextFilter->GetRawFilterText().CompareTo(CurrentAsyncOpTextFilter->GetRawFilterText()) != 0;
+	bool bTextFilterHasChanged = HasInProgressAsyncOperation(EAsyncOperationType::FilteringOp);
+	bTextFilterHasChanged |= TextFilter->GetRawFilterText().CompareTo(CurrentAsyncOpTextFilter->GetRawFilterText()) != 0;
 	if (bTextFilterHasChanged)
 	{
 		OnPreAsyncUpdate();
@@ -2235,6 +2274,16 @@ void STableTreeView::StartPendingAsyncOperations()
 
 ////////////////////////////////////////////////////////////////////////////////////////////////////
 
+void STableTreeView::CancelCurrentAsyncOp()
+{
+	if (bIsUpdateRunning)
+	{
+		bCancelCurrentAsyncOp = true;
+	}
+}
+
+////////////////////////////////////////////////////////////////////////////////////////////////////
+
 } // namespace Insights
 
 #undef LOCTEXT_NAMESPACE
diff --git a/Engine/Source/Developer/TraceInsights/Private/Insights/Table/Widgets/STableTreeView.h b/Engine/Source/Developer/TraceInsights/Private/Insights/Table/Widgets/STableTreeView.h
--- a/Engine/Source/Developer/TraceInsights/Private/Insights/Table/Widgets/STableTreeView.h
+++ b/Engine/Source/Developer/TraceInsights/Private/Insights/Table/Widgets/STableTreeView.h
@@ -21,6 +21,8 @@
 #include "Insights/Common/Stopwatch.h"
 #include "Insights/Table/ViewModels/TableTreeNode.h"
 
+#include <atomic>
+
 class FMenuBuilder;
 
 namespace TraceServices
@@ -300,6 +302,8 @@ protected:
 
 	void StartPendingAsyncOperations();
 
+	void CancelCurrentAsyncOp();
+
 	////////////////////////////////////////////////////////////////////////////////////////////////////
 
 protected:
@@ -406,6 +410,7 @@ protected:
 	ITableCellValueSorter* CurrentAsyncOpSorter = nullptr;
 	EColumnSortMode::Type CurrentAsyncOpColumnSortMode;
 	TSharedPtr<FTableTreeNodeTextFilter> CurrentAsyncOpTextFilter;
+	std::atomic<bool> bCancelCurrentAsyncOp { false };
 
 	//////////////////////////////////////////////////
 
