Non-PCH/Unity compile fixes.

#rnx

diff --git a/Engine/Source/Developer/TraceInsights/Private/Insights/Tests/InsightsTestRunner.cpp b/Engine/Source/Developer/TraceInsights/Private/Insights/Tests/InsightsTestRunner.cpp
--- a/Engine/Source/Developer/TraceInsights/Private/Insights/Tests/InsightsTestRunner.cpp
+++ b/Engine/Source/Developer/TraceInsights/Private/Insights/Tests/InsightsTestRunner.cpp
@@ -4,12 +4,14 @@
 
 #if !UE_BUILD_SHIPPING && !WITH_EDITOR
 
+#include "Misc/App.h"
 #include "Misc/AutomationTest.h"
 #include "IAutomationControllerModule.h"
 #include "IAutomationWindowModule.h"
 #include "IAutomationWorkerModule.h"
 #include "ISessionManager.h"
 #include "ISessionServicesModule.h"
+#include "Widgets/Docking/SDockTab.h"
 
 // Insights
 #include "Insights/Common/InsightsMenuBuilder.h"
