[Insights] Add the filter configurator to the Memory Alloc Tree View

diff --git a/Engine/Source/Developer/TraceInsights/Private/Insights/InsightsManager.cpp b/Engine/Source/Developer/TraceInsights/Private/Insights/InsightsManager.cpp
--- a/Engine/Source/Developer/TraceInsights/Private/Insights/InsightsManager.cpp
+++ b/Engine/Source/Developer/TraceInsights/Private/Insights/InsightsManager.cpp
@@ -142,7 +142,7 @@ void FInsightsManager::Initialize(IUnrealInsightsModule& InsightsModule)
 
 	InsightsMenuBuilder = MakeShared<FInsightsMenuBuilder>();
 
-	FFilterService::CreateInstance();
+	Insights::FFilterService::CreateInstance();
 
 	// Register tick functions.
 	OnTick = FTickerDelegate::CreateSP(this, &FInsightsManager::Tick);
diff --git a/Engine/Source/Developer/TraceInsights/Private/Insights/MemoryProfiler/ViewModels/MemAllocTable.cpp b/Engine/Source/Developer/TraceInsights/Private/Insights/MemoryProfiler/ViewModels/MemAllocTable.cpp
--- a/Engine/Source/Developer/TraceInsights/Private/Insights/MemoryProfiler/ViewModels/MemAllocTable.cpp
+++ b/Engine/Source/Developer/TraceInsights/Private/Insights/MemoryProfiler/ViewModels/MemAllocTable.cpp
@@ -104,7 +104,7 @@ void FMemAllocTable::AddDefaultColumns()
 		TSharedRef<FTableColumn> ColumnRef = MakeShared<FTableColumn>(TEXT("AllocEndTime"));
 		FTableColumn& Column = *ColumnRef;
 
-		Column.SetIndex(0);
+		Column.SetIndex(1);
 
 		Column.SetShortName(LOCTEXT("EndTimeColumnName", "End Time"));
 		Column.SetTitleName(LOCTEXT("EndTimeColumnTitle", "End Time"));
@@ -145,7 +145,7 @@ void FMemAllocTable::AddDefaultColumns()
 		TSharedRef<FTableColumn> ColumnRef = MakeShared<FTableColumn>(TEXT("AllocDuration"));
 		FTableColumn& Column = *ColumnRef;
 
-		Column.SetIndex(0);
+		Column.SetIndex(2);
 
 		Column.SetShortName(LOCTEXT("DurationColumnName", "Duration"));
 		Column.SetTitleName(LOCTEXT("DurationColumnTitle", "Duration"));
@@ -186,7 +186,7 @@ void FMemAllocTable::AddDefaultColumns()
 		TSharedRef<FTableColumn> ColumnRef = MakeShared<FTableColumn>(TEXT("AllocAddress"));
 		FTableColumn& Column = *ColumnRef;
 
-		Column.SetIndex(0);
+		Column.SetIndex(3);
 
 		Column.SetShortName(LOCTEXT("AddressColumnName", "Address"));
 		Column.SetTitleName(LOCTEXT("AddressColumnTitle", "Address"));
@@ -227,7 +227,7 @@ void FMemAllocTable::AddDefaultColumns()
 		TSharedRef<FTableColumn> ColumnRef = MakeShared<FTableColumn>(TEXT("AllocSize"));
 		FTableColumn& Column = *ColumnRef;
 
-		Column.SetIndex(0);
+		Column.SetIndex(4);
 
 		Column.SetShortName(LOCTEXT("SizeColumnName", "Size"));
 		Column.SetTitleName(LOCTEXT("SizeColumnTitle", "Size"));
@@ -268,7 +268,7 @@ void FMemAllocTable::AddDefaultColumns()
 		TSharedRef<FTableColumn> ColumnRef = MakeShared<FTableColumn>(TEXT("AllocLlmTag"));
 		FTableColumn& Column = *ColumnRef;
 
-		Column.SetIndex(0);
+		Column.SetIndex(5);
 
 		Column.SetShortName(LOCTEXT("LlmTagColumnName", "LLM Tag"));
 		Column.SetTitleName(LOCTEXT("LlmTagColumnTitle", "LLM Tag"));
@@ -346,7 +346,7 @@ void FMemAllocTable::AddDefaultColumns()
 		TSharedRef<FTableColumn> ColumnRef = MakeShared<FTableColumn>(TEXT("AllocBacktrace"));
 		FTableColumn& Column = *ColumnRef;
 
-		Column.SetIndex(0);
+		Column.SetIndex(6);
 
 		Column.SetShortName(LOCTEXT("BacktraceColumnName", "Backtrace"));
 		Column.SetTitleName(LOCTEXT("BacktraceColumnTitle", "Backtrace"));
diff --git a/Engine/Source/Developer/TraceInsights/Private/Insights/Table/Widgets/STableTreeView.cpp b/Engine/Source/Developer/TraceInsights/Private/Insights/Table/Widgets/STableTreeView.cpp
--- a/Engine/Source/Developer/TraceInsights/Private/Insights/Table/Widgets/STableTreeView.cpp
+++ b/Engine/Source/Developer/TraceInsights/Private/Insights/Table/Widgets/STableTreeView.cpp
@@ -29,6 +29,8 @@
 #include "Insights/Table/Widgets/STableTreeViewTooltip.h"
 #include "Insights/Table/Widgets/STableTreeViewRow.h"
 #include "Insights/TimingProfilerCommon.h"
+#include "Insights/ViewModels/Filters.h"
+#include "Insights/ViewModels/FilterConfigurator.h"
 #include "Insights/Widgets/SAsyncOperationStatus.h"
 
 #define LOCTEXT_NAMESPACE "STableTreeView"
@@ -118,11 +120,28 @@ void STableTreeView::ConstructWidget(TSharedPtr<FTable> InTablePtr)
 				.Padding(2.0f)
 				.AutoHeight()
 				[
-					SAssignNew(SearchBox, SSearchBox)
-					.HintText(LOCTEXT("SearchBoxHint", "Search"))
-					.OnTextChanged(this, &STableTreeView::SearchBox_OnTextChanged)
-					.IsEnabled(this, &STableTreeView::SearchBox_IsEnabled)
-					.ToolTipText(LOCTEXT("FilterSearchHint", "Type here to search the tree hierarchy by item or group name"))
+					SNew(SHorizontalBox)
+
+					+ SHorizontalBox::Slot()
+					.FillWidth(1.0)
+					.VAlign(VAlign_Center)
+					[
+						SAssignNew(SearchBox, SSearchBox)
+						.HintText(LOCTEXT("SearchBoxHint", "Search"))
+						.OnTextChanged(this, &STableTreeView::SearchBox_OnTextChanged)
+						.IsEnabled(this, &STableTreeView::SearchBox_IsEnabled)
+						.ToolTipText(LOCTEXT("FilterSearchHint", "Type here to search the tree hierarchy by item or group name"))
+					]
+
+					+ SHorizontalBox::Slot()
+						.AutoWidth()
+						.VAlign(VAlign_Center)
+						[
+							SNew(SButton)
+							.Text(LOCTEXT("AdvancedFilters", "Advanced Filters"))
+							.ToolTipText(LOCTEXT("AdvancedFilters", "Opens the filter configurator window."))
+							.OnClicked(this, &STableTreeView::OnAdvancedFiltersClicked)
+						]
 				]
 
 				// Group by
@@ -728,7 +747,7 @@ void STableTreeView::ApplyFiltering()
 
 bool STableTreeView::ApplyFilteringForNode(FTableTreeNodePtr NodePtr)
 {
-	const bool bIsNodeVisible = Filters->PassesAllFilters(NodePtr);
+	bool bIsNodeVisible = Filters->PassesAllFilters(NodePtr);
 
 	if (NodePtr->IsGroup())
 	{
@@ -768,6 +787,7 @@ bool STableTreeView::ApplyFilteringForNode(FTableTreeNodePtr NodePtr)
 	}
 	else
 	{
+		bIsNodeVisible &= ApplyAdvancedFilters(NodePtr);
 		return bIsNodeVisible;
 	}
 }
@@ -2284,6 +2304,85 @@ void STableTreeView::CancelCurrentAsyncOp()
 
 ////////////////////////////////////////////////////////////////////////////////////////////////////
 
+FReply STableTreeView::OnAdvancedFiltersClicked()
+{
+	if (!FilterConfigurator.IsValid())
+	{
+		FilterConfigurator = MakeShared<FFilterConfigurator>();
+		TSharedPtr<TArray<TSharedPtr<struct FFilter>>>& AvailableFilters = FilterConfigurator->GetAvailableFilters();
+
+		for (const TSharedRef<FTableColumn>& Column : Table->GetColumns())
+		{
+			switch (Column->GetDataType())
+			{
+				case ETableCellDataType::Int64:
+				{
+					AvailableFilters->Add(MakeShared<FFilter>(Column->GetIndex(), Column->GetTitleName(), Column->GetDescription(), EFilterDataType::Int64, FFilterService::Get()->GetIntegerOperators()));
+					break;
+				}
+				case ETableCellDataType::Double:
+				{
+					AvailableFilters->Add(MakeShared<FFilter>(Column->GetIndex(), Column->GetTitleName(), Column->GetDescription(), EFilterDataType::Double, FFilterService::Get()->GetDoubleOperators()));
+					break;
+				}
+			}
+		}
+
+		OnFilterChangesCommitedHandle = FilterConfigurator->GetOnChangesCommitedEvent().AddSP(this, &STableTreeView::OnAdvancedFiltersChangesCommited);
+	}
+
+	FFilterService::Get()->CreateFilterConfiguratorWidget(FilterConfigurator);
+
+	return FReply::Handled();
+}
+
+////////////////////////////////////////////////////////////////////////////////////////////////////
+
+bool STableTreeView::ApplyAdvancedFilters(const FTableTreeNodePtr& NodePtr)
+{
+	if (!FilterConfigurator.IsValid())
+	{
+		return true;
+	}
+
+	FFilterContext Context;
+	for (const TSharedRef<FTableColumn>& Column : Table->GetColumns())
+	{
+		switch (Column->GetDataType())
+		{
+			case ETableCellDataType::Int64:
+			{
+				Context.AddFilterData<int64>(Column->GetIndex(), Column->GetValue(*NodePtr)->AsInt64());
+				break;
+			}
+			case ETableCellDataType::Double:
+			{
+				Context.AddFilterData<double>(Column->GetIndex(), Column->GetValue(*NodePtr)->AsDouble());
+				break;
+			}
+		}
+	}
+
+	return FilterConfigurator->ApplyFilters(Context);
+}
+
+////////////////////////////////////////////////////////////////////////////////////////////////////
+
+void STableTreeView::OnAdvancedFiltersChangesCommited()
+{
+	if (!bIsUpdateRunning)
+	{
+		OnPreAsyncUpdate();
+		InProgressAsyncOperationEvent = StartApplyFiltersTask();
+	}
+	else
+	{
+		CancelCurrentAsyncOp();
+	}
+}
+
+////////////////////////////////////////////////////////////////////////////////////////////////////
+
 } // namespace Insights
 
 #undef LOCTEXT_NAMESPACE
diff --git a/Engine/Source/Developer/TraceInsights/Private/Insights/Table/Widgets/STableTreeView.h b/Engine/Source/Developer/TraceInsights/Private/Insights/Table/Widgets/STableTreeView.h
--- a/Engine/Source/Developer/TraceInsights/Private/Insights/Table/Widgets/STableTreeView.h
+++ b/Engine/Source/Developer/TraceInsights/Private/Insights/Table/Widgets/STableTreeView.h
@@ -304,6 +304,10 @@ protected:
 
 	void CancelCurrentAsyncOp();
 
+	FReply OnAdvancedFiltersClicked();
+	void OnAdvancedFiltersChangesCommited();
+	bool ApplyAdvancedFilters(const FTableTreeNodePtr& NodePtr);
+
 	////////////////////////////////////////////////////////////////////////////////////////////////////
 
 protected:
@@ -413,6 +417,8 @@ protected:
 	std::atomic<bool> bCancelCurrentAsyncOp { false };
 
 	//////////////////////////////////////////////////
+	TSharedPtr<class FFilterConfigurator> FilterConfigurator;
+	FDelegateHandle OnFilterChangesCommitedHandle;
 
 	double StatsStartTime;
 	double StatsEndTime;
diff --git a/Engine/Source/Developer/TraceInsights/Private/Insights/ViewModels/FilterConfigurator.cpp b/Engine/Source/Developer/TraceInsights/Private/Insights/ViewModels/FilterConfigurator.cpp
--- a/Engine/Source/Developer/TraceInsights/Private/Insights/ViewModels/FilterConfigurator.cpp
+++ b/Engine/Source/Developer/TraceInsights/Private/Insights/ViewModels/FilterConfigurator.cpp
@@ -4,6 +4,9 @@
 
 #include "Insights/Widgets/SFilterConfigurator.h"
 
+namespace Insights
+{
+
 ////////////////////////////////////////////////////////////////////////////////////////////////////
 
 FFilterConfigurator::FFilterConfigurator()
@@ -51,3 +54,5 @@ bool FFilterConfigurator::ApplyFilters(const FFilterContext& Context) const
 }
 
 ////////////////////////////////////////////////////////////////////////////////////////////////////
+
+} // namespace Insights
\ No newline at end of file
diff --git a/Engine/Source/Developer/TraceInsights/Private/Insights/ViewModels/FilterConfigurator.h b/Engine/Source/Developer/TraceInsights/Private/Insights/ViewModels/FilterConfigurator.h
--- a/Engine/Source/Developer/TraceInsights/Private/Insights/ViewModels/FilterConfigurator.h
+++ b/Engine/Source/Developer/TraceInsights/Private/Insights/ViewModels/FilterConfigurator.h
@@ -5,6 +5,9 @@
 #include "CoreMinimal.h"
 #include "Insights/ViewModels/Filters.h"
 #include "Insights/ViewModels/FilterConfiguratorNode.h"
+
+namespace Insights
+{
  
 ////////////////////////////////////////////////////////////////////////////////////////////////////
 
@@ -18,15 +21,12 @@ public:
 
 	~FFilterConfigurator();
 
-	void AddFilter(EFilterField FilterKey)
-	{
-		AvailableFilters->Add(FFilterService::Get()->GetFilter(FilterKey));
-	}
-
 	FFilterConfiguratorNodePtr GetRootNode() { return RootNode; }
 
 	bool ApplyFilters(const FFilterContext& Context) const;
 
+	TSharedPtr<TArray<TSharedPtr<struct FFilter>>>& GetAvailableFilters() { return AvailableFilters; }
+
 	////////////////////////////////////////////////////////////////////////////////////////////////////
 	// OnDestroyedEvent
 public:
@@ -56,4 +56,6 @@ private:
 	TSharedPtr<TArray<TSharedPtr<struct FFilter>>> AvailableFilters;
 };
 
-////////////////////////////////////////////////////////////////////////////////////////////////////
\ No newline at end of file
+////////////////////////////////////////////////////////////////////////////////////////////////////
+
+} // namespace Insights
\ No newline at end of file
diff --git a/Engine/Source/Developer/TraceInsights/Private/Insights/ViewModels/FilterConfiguratorNode.cpp b/Engine/Source/Developer/TraceInsights/Private/Insights/ViewModels/FilterConfiguratorNode.cpp
--- a/Engine/Source/Developer/TraceInsights/Private/Insights/ViewModels/FilterConfiguratorNode.cpp
+++ b/Engine/Source/Developer/TraceInsights/Private/Insights/ViewModels/FilterConfiguratorNode.cpp
@@ -7,6 +7,9 @@
 
 #define LOCTEXT_NAMESPACE "FilterConfiguratorNode"
 
+namespace Insights
+{
+
 ////////////////////////////////////////////////////////////////////////////////////////////////////
 
 const FName FFilterConfiguratorNode::TypeName(TEXT("FFilterConfiguratorNode"));
@@ -53,7 +56,7 @@ FFilterConfiguratorNode& FFilterConfiguratorNode::operator=(const FFilterConfigu
 
 ////////////////////////////////////////////////////////////////////////////////////////////////////
 
-TArray<TSharedPtr<struct FFilterGroupOperator>>& FFilterConfiguratorNode::GetFilterGroupOperators()
+const TArray<TSharedPtr<struct FFilterGroupOperator>>& FFilterConfiguratorNode::GetFilterGroupOperators()
 {
 	return FFilterService::Get()->GetFilterGroupOperators();
 }
@@ -125,7 +128,7 @@ void FFilterConfiguratorNode::ProcessFilter()
 		}
 		case EFilterDataType::Int64:
 		{
-			FilterValue.Set<double>(FCString::Atoi64(*TextBoxValue));
+			FilterValue.Set<int64>(FCString::Atoi64(*TextBoxValue));
 			break;
 		}
 		}
@@ -195,4 +198,6 @@ bool FFilterConfiguratorNode::ApplyFilters(const FFilterContext& Context) const
 
 ////////////////////////////////////////////////////////////////////////////////////////////////////
 
+} // namespace Insights
+
 #undef LOCTEXT_NAMESPACE
diff --git a/Engine/Source/Developer/TraceInsights/Private/Insights/ViewModels/FilterConfiguratorNode.h b/Engine/Source/Developer/TraceInsights/Private/Insights/ViewModels/FilterConfiguratorNode.h
--- a/Engine/Source/Developer/TraceInsights/Private/Insights/ViewModels/FilterConfiguratorNode.h
+++ b/Engine/Source/Developer/TraceInsights/Private/Insights/ViewModels/FilterConfiguratorNode.h
@@ -8,6 +8,9 @@
 #include "Insights/Table/ViewModels/BaseTreeNode.h"
 #include "Insights/ViewModels/Filters.h"
 
+namespace Insights
+{
+
 ////////////////////////////////////////////////////////////////////////////////////////////////////
 
 enum class EFilterConfiguratorNodeType
@@ -62,8 +65,7 @@ public:
 
 	void SetSelectedFilterOperator(TSharedPtr<class IFilterOperator> InSelectedFilterOperator) { SelectedFilterOperator = InSelectedFilterOperator; }
 	TSharedPtr<IFilterOperator> GetSelectedFilterOperator() const { return SelectedFilterOperator;	}
-
-	TArray<TSharedPtr<struct FFilterGroupOperator>>& GetFilterGroupOperators();
+	const TArray<TSharedPtr<struct FFilterGroupOperator>>& GetFilterGroupOperators();
 
 	void SetSelectedFilterGroupOperator(TSharedPtr<struct FFilterGroupOperator> InSelectedFilterGroupOperator) { SelectedFilterGroupOperator = InSelectedFilterGroupOperator; }
 	TSharedPtr<struct FFilterGroupOperator> GetSelectedFilterGroupOperator() const { return SelectedFilterGroupOperator; }
@@ -97,3 +99,5 @@ private:
 };
 
 ////////////////////////////////////////////////////////////////////////////////////////////////////
+
+} // namespace Insights
diff --git a/Engine/Source/Developer/TraceInsights/Private/Insights/ViewModels/Filters.cpp b/Engine/Source/Developer/TraceInsights/Private/Insights/ViewModels/Filters.cpp
--- a/Engine/Source/Developer/TraceInsights/Private/Insights/ViewModels/Filters.cpp
+++ b/Engine/Source/Developer/TraceInsights/Private/Insights/ViewModels/Filters.cpp
@@ -9,6 +9,9 @@
 
 #define LOCTEXT_NAMESPACE "SFilterService"
 
+namespace Insights
+{
+
 ////////////////////////////////////////////////////////////////////////////////////////////////////
 
 TSharedPtr<FFilterService> FFilterService::Instance;
@@ -69,4 +72,6 @@ void FFilterService::RegisterTabSpawner()
 
 ////////////////////////////////////////////////////////////////////////////////////////////////////
 
+} // namespace Insights
+
 #undef LOCTEXT_NAMESPACE
\ No newline at end of file
diff --git a/Engine/Source/Developer/TraceInsights/Private/Insights/ViewModels/Filters.h b/Engine/Source/Developer/TraceInsights/Private/Insights/ViewModels/Filters.h
--- a/Engine/Source/Developer/TraceInsights/Private/Insights/ViewModels/Filters.h
+++ b/Engine/Source/Developer/TraceInsights/Private/Insights/ViewModels/Filters.h
@@ -9,9 +9,12 @@
 
 #define LOCTEXT_NAMESPACE "Filters"
 
+namespace Insights
+{
+
 ////////////////////////////////////////////////////////////////////////////////////////////////////
 
-enum class EFilterField : uint32
+enum class EFilterField : int32
 {
 	StartTime = 0,
 	EndTime = 1,
@@ -102,7 +105,7 @@ struct FFilter
 {
 	typedef TSharedPtr<const TArray<TSharedPtr<IFilterOperator>>> SupportedOperatorsArrayPtr;
 
-	FFilter(EFilterField InKey, FText InName, FText InDesc, EFilterDataType InDataType, SupportedOperatorsArrayPtr InSupportedOperators)
+	FFilter(int32 InKey, FText InName, FText InDesc, EFilterDataType InDataType, SupportedOperatorsArrayPtr InSupportedOperators)
 		: Key(InKey)
 		, Name(InName)
 		, Desc(InDesc)
@@ -112,7 +115,7 @@ struct FFilter
 
 	SupportedOperatorsArrayPtr GetSupportedOperators() { return SupportedOperators;	}
 
-	EFilterField Key;
+	int32 Key;
 	FText Name;
 	FText Desc;
 	EFilterDataType DataType;
@@ -126,35 +129,34 @@ class FFilterStorage
 public:
 	FFilterStorage()
 	{
-		TSharedPtr<TArray<TSharedPtr<IFilterOperator>>> DoubleOperators = MakeShared<TArray<TSharedPtr<IFilterOperator>>>();
+		DoubleOperators = MakeShared<TArray<TSharedPtr<IFilterOperator>>>();
 		DoubleOperators->Add(StaticCastSharedRef<IFilterOperator>(MakeShared<FFilterOperator<double>>(EFilterOperator::Lt, TEXT("<"), std::less<>{})));
 		DoubleOperators->Add(StaticCastSharedRef<IFilterOperator>(MakeShared<FFilterOperator<double>>(EFilterOperator::Lte, TEXT("<="), std::less_equal<>())));
 		DoubleOperators->Add(StaticCastSharedRef<IFilterOperator>(MakeShared<FFilterOperator<double>>(EFilterOperator::Eq, TEXT("="), std::equal_to<>())));
 		DoubleOperators->Add(StaticCastSharedRef<IFilterOperator>(MakeShared<FFilterOperator<double>>(EFilterOperator::Gt, TEXT(">"), std::greater<>())));
 		DoubleOperators->Add(StaticCastSharedRef<IFilterOperator>(MakeShared<FFilterOperator<double>>(EFilterOperator::Gte, TEXT(">="), std::greater_equal<>())));
 
-		TSharedPtr<TArray<TSharedPtr<IFilterOperator>>> IntegerOperators = MakeShared<TArray<TSharedPtr<IFilterOperator>>>();
+		IntegerOperators = MakeShared<TArray<TSharedPtr<IFilterOperator>>>();
 		IntegerOperators->Add(StaticCastSharedRef<IFilterOperator>(MakeShared<FFilterOperator<int64>>(EFilterOperator::Lt, TEXT("<"), std::less<>{})));
 		IntegerOperators->Add(StaticCastSharedRef<IFilterOperator>(MakeShared<FFilterOperator<int64>>(EFilterOperator::Lte, TEXT("<="), std::less_equal<>())));
 		IntegerOperators->Add(StaticCastSharedRef<IFilterOperator>(MakeShared<FFilterOperator<int64>>(EFilterOperator::Eq, TEXT("="), std::equal_to<>())));
 		IntegerOperators->Add(StaticCastSharedRef<IFilterOperator>(MakeShared<FFilterOperator<int64>>(EFilterOperator::Gt, TEXT(">"), std::greater<>())));
 		IntegerOperators->Add(StaticCastSharedRef<IFilterOperator>(MakeShared<FFilterOperator<int64>>(EFilterOperator::Gte, TEXT(">="), std::greater_equal<>())));
 
-		AllAvailableFilters.Add(EFilterField::StartTime, MakeShared<FFilter>(EFilterField::StartTime, LOCTEXT("StartTime", "Start Time"), LOCTEXT("StartTime", "Start Time"), EFilterDataType::Double, DoubleOperators));
-		AllAvailableFilters.Add(EFilterField::EndTime, MakeShared<FFilter>(EFilterField::EndTime, LOCTEXT("EndTime", "End Time"), LOCTEXT("EndTime", "End Time"), EFilterDataType::Double, DoubleOperators));
-		AllAvailableFilters.Add(EFilterField::Duration, MakeShared<FFilter>(EFilterField::Duration, LOCTEXT("Duration", "Duration"), LOCTEXT("Duration", "Duration"), EFilterDataType::Double, DoubleOperators));
-		AllAvailableFilters.Add(EFilterField::EventType, MakeShared<FFilter>(EFilterField::EventType, LOCTEXT("Type", "Type"), LOCTEXT("Type", "Type"), EFilterDataType::Int64, IntegerOperators));
-
 		FilterGroupOperators.Add(MakeShared<FFilterGroupOperator>(EFilterGroupOperator::And, LOCTEXT("AllOf", "All Of"), LOCTEXT("AllOfDesc", "All of the children must be true for the group to return true. Equivalent to an AND operation.")));
 		FilterGroupOperators.Add(MakeShared<FFilterGroupOperator>(EFilterGroupOperator::Or, LOCTEXT("AnyOf", "Any Of"), LOCTEXT("AnyOfDesc", "Any of the children must be true for the group to return true. Equivalent to an OR operation.")));
 	}
 
-	TSharedPtr<FFilter> GetFilter(EFilterField FilterKey) { return *AllAvailableFilters.Find(FilterKey); }
-	TArray<TSharedPtr<FFilterGroupOperator>>& GetFilterGroupOperators() { return FilterGroupOperators; }
+	const TArray<TSharedPtr<FFilterGroupOperator>>& GetFilterGroupOperators() { return FilterGroupOperators; }
+
+	TSharedPtr<TArray<TSharedPtr<IFilterOperator>>> GetDoubleOperators() { return DoubleOperators; }
+	TSharedPtr<TArray<TSharedPtr<IFilterOperator>>> GetIntegerOperators() { return IntegerOperators; }
 
 private:
-	TMap<EFilterField, TSharedPtr<FFilter>> AllAvailableFilters;
 	TArray<TSharedPtr<FFilterGroupOperator>> FilterGroupOperators;
+
+	TSharedPtr<TArray<TSharedPtr<IFilterOperator>>> DoubleOperators;
+	TSharedPtr<TArray<TSharedPtr<IFilterOperator>>> IntegerOperators;
 };
 
 ////////////////////////////////////////////////////////////////////////////////////////////////////
@@ -171,8 +173,10 @@ public:
 	void RegisterTabSpawner();
 	TSharedRef<class SDockTab> SpawnTab(const class FSpawnTabArgs& Args);
 
-	TSharedPtr<FFilter> GetFilter(EFilterField FilterKey) { return FilterStorage.GetFilter(FilterKey); }
-	TArray<TSharedPtr<FFilterGroupOperator>>& GetFilterGroupOperators() { return FilterStorage.GetFilterGroupOperators(); }
+	const TArray<TSharedPtr<FFilterGroupOperator>>& GetFilterGroupOperators() { return FilterStorage.GetFilterGroupOperators(); }
+
+	TSharedPtr<TArray<TSharedPtr<IFilterOperator>>> GetDoubleOperators() { return FilterStorage.GetDoubleOperators(); }
+	TSharedPtr<TArray<TSharedPtr<IFilterOperator>>> GetIntegerOperators() { return FilterStorage.GetIntegerOperators(); }
 
 	TSharedPtr<class SWidget> CreateFilterConfiguratorWidget(TSharedPtr<class FFilterConfigurator> FilterConfiguratorViewModel);
 
@@ -187,14 +191,13 @@ private:
 
 ////////////////////////////////////////////////////////////////////////////////////////////////////
 
-
 class FFilterContext
 {
 public:
 	typedef TVariant<double, int64> ContextData;
 
 	template<typename T>
-	void AddFilterData(EFilterField Key, const T& InData)
+	void AddFilterData(int32 Key, const T& InData)
 	{
 		ContextData VariantData;
 		VariantData.Set<T>(InData);
@@ -202,13 +205,13 @@ public:
 	}
 
 	template<typename T>
-	void SetFilterData(EFilterField Key, const T& InData)
+	void SetFilterData(int32 Key, const T& InData)
 	{
 		DataMap[Key].Set<T>(InData);
 	}
 
 	template<typename T>
-	void GetFilterData(EFilterField Key, T& OutData) const
+	void GetFilterData(int32 Key, T& OutData) const
 	{
 		const ContextData* Data = DataMap.Find(Key);
 		check(Data);
@@ -218,9 +221,11 @@ public:
 	}
 
 private:
-	TMap<EFilterField, ContextData> DataMap;
+	TMap<int32, ContextData> DataMap;
 };
 
 ////////////////////////////////////////////////////////////////////////////////////////////////////
 
+} // namespace Insights
+
 #undef LOCTEXT_NAMESPACE
diff --git a/Engine/Source/Developer/TraceInsights/Private/Insights/ViewModels/ThreadTimingTrack.cpp b/Engine/Source/Developer/TraceInsights/Private/Insights/ViewModels/ThreadTimingTrack.cpp
--- a/Engine/Source/Developer/TraceInsights/Private/Insights/ViewModels/ThreadTimingTrack.cpp
+++ b/Engine/Source/Developer/TraceInsights/Private/Insights/ViewModels/ThreadTimingTrack.cpp
@@ -27,6 +27,8 @@
 
 #define LOCTEXT_NAMESPACE "ThreadTimingTrack"
 
+using namespace Insights;
+
 ////////////////////////////////////////////////////////////////////////////////////////////////////
 
 static void AppendMetadataToTooltip(FTooltipDrawState& Tooltip, TArrayView<const uint8>& Metadata)
@@ -840,10 +842,10 @@ void FThreadTimingTrack::BuildFilteredDrawState(ITimingEventsTrackDrawStateBuild
 						FilterContexts.AddDefaulted(NumTasks);
 						for(FFilterContext& Context : FilterContexts)
 						{
-							Context.AddFilterData<double>(EFilterField::StartTime, 0.0f);
-							Context.AddFilterData<double>(EFilterField::EndTime, 0.0f);
-							Context.AddFilterData<double>(EFilterField::Duration, 0.0f);
-							Context.AddFilterData<int64>(EFilterField::EventType, 0);
+							Context.AddFilterData<double>(static_cast<int32>(EFilterField::StartTime), 0.0f);
+							Context.AddFilterData<double>(static_cast<int32>(EFilterField::EndTime), 0.0f);
+							Context.AddFilterData<double>(static_cast<int32>(EFilterField::Duration), 0.0f);
+							Context.AddFilterData<int64>(static_cast<int32>(EFilterField::EventType), 0);
 						}
 					};
 					Params.Callback = [this, &Builder, TimerReader, &FilteredEvents, &FilterContexts](double StartTime, double EndTime, uint32 Depth, const TraceServices::FTimingProfilerEvent& Event, uint32 TaskIndex)
@@ -852,10 +854,10 @@ void FThreadTimingTrack::BuildFilteredDrawState(ITimingEventsTrackDrawStateBuild
 						if (ensure(Timer != nullptr))
 						{
 							FFilterContext& Context = FilterContexts[TaskIndex];
-							Context.SetFilterData<double>(EFilterField::StartTime, StartTime);
-							Context.SetFilterData<double>(EFilterField::EndTime, EndTime);
-							Context.SetFilterData<double>(EFilterField::Duration, EndTime - StartTime);
-							Context.SetFilterData<int64>(EFilterField::EventType, Timer->Id);
+							Context.SetFilterData<double>(static_cast<int32>(EFilterField::StartTime), StartTime);
+							Context.SetFilterData<double>(static_cast<int32>(EFilterField::EndTime), EndTime);
+							Context.SetFilterData<double>(static_cast<int32>(EFilterField::Duration), EndTime - StartTime);
+							Context.SetFilterData<int64>(static_cast<int32>(EFilterField::EventType), Timer->Id);
 
 							if (FilterConfigurator->ApplyFilters(Context))
 							{
@@ -1327,11 +1329,12 @@ void FThreadTimingTrack::OnFilterTrackClicked()
 	if (!FilterConfigurator.IsValid())
 	{
 		FilterConfigurator = MakeShared<FFilterConfigurator>();
+		TSharedPtr<TArray<TSharedPtr<struct FFilter>>>& AvailableFilters = FilterConfigurator->GetAvailableFilters();
 
-		FilterConfigurator->AddFilter(EFilterField::StartTime);
-		FilterConfigurator->AddFilter(EFilterField::EndTime);
-		FilterConfigurator->AddFilter(EFilterField::Duration);
-		FilterConfigurator->AddFilter(EFilterField::EventType);
+		AvailableFilters->Add(MakeShared<FFilter>(static_cast<int32>(EFilterField::StartTime), LOCTEXT("StartTime", "Start Time"), LOCTEXT("StartTime", "Start Time"), EFilterDataType::Double, FFilterService::Get()->GetDoubleOperators()));
+		AvailableFilters->Add(MakeShared<FFilter>(static_cast<int32>(EFilterField::EndTime), LOCTEXT("EndTime", "End Time"), LOCTEXT("EndTime", "End Time"), EFilterDataType::Double, FFilterService::Get()->GetDoubleOperators()));
+		AvailableFilters->Add(MakeShared<FFilter>(static_cast<int32>(EFilterField::Duration), LOCTEXT("Duration", "Duration"), LOCTEXT("Duration", "Duration"), EFilterDataType::Double, FFilterService::Get()->GetDoubleOperators()));
+		AvailableFilters->Add(MakeShared<FFilter>(static_cast<int32>(EFilterField::EventType), LOCTEXT("Type", "Type"), LOCTEXT("Type", "Type"), EFilterDataType::Int64, FFilterService::Get()->GetIntegerOperators()));
 
 		OnFilterChangesCommitedHandle = FilterConfigurator->GetOnChangesCommitedEvent().AddLambda([this]()
 			{
@@ -1358,4 +1361,6 @@ bool FThreadTimingTrack::HasCustomFilter() const
 	return false;
 }
 
+////////////////////////////////////////////////////////////////////////////////////////////////////
+
 #undef LOCTEXT_NAMESPACE
diff --git a/Engine/Source/Developer/TraceInsights/Private/Insights/ViewModels/ThreadTimingTrack.h b/Engine/Source/Developer/TraceInsights/Private/Insights/ViewModels/ThreadTimingTrack.h
--- a/Engine/Source/Developer/TraceInsights/Private/Insights/ViewModels/ThreadTimingTrack.h
+++ b/Engine/Source/Developer/TraceInsights/Private/Insights/ViewModels/ThreadTimingTrack.h
@@ -18,6 +18,11 @@ class FCpuTimingTrack;
 class STimingView;
 struct FSlateBrush;
 
+namespace Insights
+{
+	class FFilterConfigurator;
+}
+
 ////////////////////////////////////////////////////////////////////////////////////////////////////
 
 class FThreadTimingSharedState : public Insights::ITimingViewExtender, public TSharedFromThis<FThreadTimingSharedState>
@@ -164,7 +169,7 @@ private:
 	uint32 TimelineIndex;
 	uint32 ThreadId;
 
-	TSharedPtr<class FFilterConfigurator> FilterConfigurator;
+	TSharedPtr<Insights::FFilterConfigurator> FilterConfigurator;
 	FDelegateHandle OnFilterChangesCommitedHandle;
 
 	// Search cache
diff --git a/Engine/Source/Developer/TraceInsights/Private/Insights/Widgets/SFilterConfigurator.cpp b/Engine/Source/Developer/TraceInsights/Private/Insights/Widgets/SFilterConfigurator.cpp
--- a/Engine/Source/Developer/TraceInsights/Private/Insights/Widgets/SFilterConfigurator.cpp
+++ b/Engine/Source/Developer/TraceInsights/Private/Insights/Widgets/SFilterConfigurator.cpp
@@ -20,6 +20,9 @@
 
 #define LOCTEXT_NAMESPACE "SFilterConfigurator"
 
+namespace Insights
+{
+
 ////////////////////////////////////////////////////////////////////////////////////////////////////
 // SFilterConfigurator
 ////////////////////////////////////////////////////////////////////////////////////////////////////
@@ -236,4 +239,6 @@ void SFilterConfigurator::RequestClose()
 
 ////////////////////////////////////////////////////////////////////////////////////////////////////
 
+} // namespace Insights
+
 #undef LOCTEXT_NAMESPACE
diff --git a/Engine/Source/Developer/TraceInsights/Private/Insights/Widgets/SFilterConfigurator.h b/Engine/Source/Developer/TraceInsights/Private/Insights/Widgets/SFilterConfigurator.h
--- a/Engine/Source/Developer/TraceInsights/Private/Insights/Widgets/SFilterConfigurator.h
+++ b/Engine/Source/Developer/TraceInsights/Private/Insights/Widgets/SFilterConfigurator.h
@@ -11,9 +11,11 @@
 #include "Insights/ViewModels/Filters.h"
 #include "Insights/ViewModels/FilterConfiguratorNode.h"
 
-////////////////////////////////////////////////////////////////////////////////////////////////////
 class SDockTab;
 
+namespace Insights
+{
+
 ////////////////////////////////////////////////////////////////////////////////////////////////////
 
 /**
@@ -85,3 +87,5 @@ private:
 };
 
 ////////////////////////////////////////////////////////////////////////////////////////////////////
+
+} // namespace Insights
diff --git a/Engine/Source/Developer/TraceInsights/Private/Insights/Widgets/SFilterConfiguratorRow.cpp b/Engine/Source/Developer/TraceInsights/Private/Insights/Widgets/SFilterConfiguratorRow.cpp
--- a/Engine/Source/Developer/TraceInsights/Private/Insights/Widgets/SFilterConfiguratorRow.cpp
+++ b/Engine/Source/Developer/TraceInsights/Private/Insights/Widgets/SFilterConfiguratorRow.cpp
@@ -19,6 +19,9 @@
 
 #define LOCTEXT_NAMESPACE "SFilterConfiguratorRow"
 
+namespace Insights
+{
+
 ////////////////////////////////////////////////////////////////////////////////////////////////////
 
 BEGIN_SLATE_FUNCTION_BUILD_OPTIMIZATION
@@ -304,7 +307,7 @@ FText SFilterConfiguratorRow::AvailableFilterOperators_GetSelectionText() const
 
 ///////////////////////////////////////////////////////////////////////////////////////////////////
 
-TArray<TSharedPtr<FFilterGroupOperator>>* SFilterConfiguratorRow::GetFilterGroupOperators()
+const TArray<TSharedPtr<FFilterGroupOperator>>* SFilterConfiguratorRow::GetFilterGroupOperators()
 {
 	return &FilterConfiguratorNodePtr->GetFilterGroupOperators();
 }
@@ -402,4 +405,6 @@ void SFilterConfiguratorRow::OnTextBoxValueCommitted(const FText& InNewText, ETe
 
 ///////////////////////////////////////////////////////////////////////////////////////////////////
 
+} // namespace Insights
+
 #undef LOCTEXT_NAMESPACE
diff --git a/Engine/Source/Developer/TraceInsights/Private/Insights/Widgets/SFilterConfiguratorRow.h b/Engine/Source/Developer/TraceInsights/Private/Insights/Widgets/SFilterConfiguratorRow.h
--- a/Engine/Source/Developer/TraceInsights/Private/Insights/Widgets/SFilterConfiguratorRow.h
+++ b/Engine/Source/Developer/TraceInsights/Private/Insights/Widgets/SFilterConfiguratorRow.h
@@ -10,6 +10,9 @@
 // Insights
 #include "Insights/ViewModels/FilterConfiguratorNode.h"
 
+namespace Insights
+{
+
 /** Widget that represents a table row in the Filter's tree control. Generates widgets for each column on demand. */
 class SFilterConfiguratorRow : public SMultiColumnTableRow<FFilterConfiguratorNodePtr>
 {
@@ -43,7 +46,7 @@ private:
 
 	FText AvailableFilterOperators_GetSelectionText() const;
 
-	TArray<TSharedPtr<struct FFilterGroupOperator>>* GetFilterGroupOperators();
+	const TArray<TSharedPtr<struct FFilterGroupOperator>>* GetFilterGroupOperators();
 
 	TSharedRef<SWidget> FilterGroupOperators_OnGenerateWidget(TSharedPtr<struct FFilterGroupOperator> InFilter);
 
@@ -69,3 +72,5 @@ private:
 
 	TSharedPtr<SComboBox<TSharedPtr<struct FFilterGroupOperator>>> FilterGroupOperatorComboBox;
 };
+
+} // namnespace Insights
\ No newline at end of file
