[Insights] Fix non unity build errors #rb none #rnx

diff --git a/Engine/Source/Developer/TraceInsights/Private/Insights/ViewModels/Filters.h b/Engine/Source/Developer/TraceInsights/Private/Insights/ViewModels/Filters.h
--- a/Engine/Source/Developer/TraceInsights/Private/Insights/ViewModels/Filters.h
+++ b/Engine/Source/Developer/TraceInsights/Private/Insights/ViewModels/Filters.h
@@ -9,6 +9,10 @@
 
 #define LOCTEXT_NAMESPACE "Filters"
 
+class FSpawnTabArgs;
+class SDockTab;
+class SWidget;
+
 namespace Insights
 {
 
@@ -171,14 +175,14 @@ public:
 	static TSharedPtr<FFilterService> Get() { return Instance; }
 
 	void RegisterTabSpawner();
-	TSharedRef<class SDockTab> SpawnTab(const class FSpawnTabArgs& Args);
+	TSharedRef<SDockTab> SpawnTab(const FSpawnTabArgs& Args);
 
 	const TArray<TSharedPtr<FFilterGroupOperator>>& GetFilterGroupOperators() { return FilterStorage.GetFilterGroupOperators(); }
 
 	TSharedPtr<TArray<TSharedPtr<IFilterOperator>>> GetDoubleOperators() { return FilterStorage.GetDoubleOperators(); }
 	TSharedPtr<TArray<TSharedPtr<IFilterOperator>>> GetIntegerOperators() { return FilterStorage.GetIntegerOperators(); }
 
-	TSharedPtr<class SWidget> CreateFilterConfiguratorWidget(TSharedPtr<class FFilterConfigurator> FilterConfiguratorViewModel);
+	TSharedPtr<SWidget> CreateFilterConfiguratorWidget(TSharedPtr<class FFilterConfigurator> FilterConfiguratorViewModel);
 
 private:
 	static const FName FilterConfiguratorTabId;
