[Insights] Filter Configurator settings can now be safely modified during an async operation
Small performance optimization and refactoring related to the Filter Configurator system.

diff --git a/Engine/Source/Developer/TraceInsights/Private/Insights/Table/Widgets/STableTreeView.cpp b/Engine/Source/Developer/TraceInsights/Private/Insights/Table/Widgets/STableTreeView.cpp
--- a/Engine/Source/Developer/TraceInsights/Private/Insights/Table/Widgets/STableTreeView.cpp
+++ b/Engine/Source/Developer/TraceInsights/Private/Insights/Table/Widgets/STableTreeView.cpp
@@ -29,7 +29,6 @@
 #include "Insights/Table/Widgets/STableTreeViewTooltip.h"
 #include "Insights/Table/Widgets/STableTreeViewRow.h"
 #include "Insights/TimingProfilerCommon.h"
-#include "Insights/ViewModels/Filters.h"
 #include "Insights/ViewModels/FilterConfigurator.h"
 #include "Insights/Widgets/SAsyncOperationStatus.h"
 
@@ -80,6 +79,12 @@ STableTreeView::~STableTreeView()
 	{
 		FInsightsManager::Get()->GetSessionChangedEvent().RemoveAll(this);
 	}
+
+	if (CurrentAsyncOpFilterConfigurator)
+	{
+		delete CurrentAsyncOpFilterConfigurator;
+		CurrentAsyncOpFilterConfigurator = nullptr;
+	}
 }
 
 ////////////////////////////////////////////////////////////////////////////////////////////////////
@@ -2187,6 +2192,10 @@ FGraphEventRef STableTreeView::StartApplyFiltersTask(FGraphEventRef Prerequisite
 	AddInProgressAsyncOperation(EAsyncOperationType::FilteringOp);
 
 	CurrentAsyncOpTextFilter->SetRawFilterText(TextFilter->GetRawFilterText());
+	if (FilterConfigurator.IsValid() && CurrentAsyncOpFilterConfigurator)
+	{
+		*CurrentAsyncOpFilterConfigurator = *FilterConfigurator;
+	}
 
 	if (Prerequisite.IsValid())
 	{
@@ -2280,9 +2289,14 @@ void STableTreeView::StartPendingAsyncOperations()
 	}
 
 	// Check if the text filter has changed. If it has, schedule a new Filtering Refresh.
-	bool bTextFilterHasChanged = HasInProgressAsyncOperation(EAsyncOperationType::FilteringOp);
-	bTextFilterHasChanged |= TextFilter->GetRawFilterText().CompareTo(CurrentAsyncOpTextFilter->GetRawFilterText()) != 0;
-	if (bTextFilterHasChanged)
+	bool bFiltersHaveChanged = HasInProgressAsyncOperation(EAsyncOperationType::FilteringOp);
+	bFiltersHaveChanged |= TextFilter->GetRawFilterText().CompareTo(CurrentAsyncOpTextFilter->GetRawFilterText()) != 0;
+	if (FilterConfigurator.IsValid() && CurrentAsyncOpFilterConfigurator && !bFiltersHaveChanged)
+	{
+		bFiltersHaveChanged |= *FilterConfigurator != *CurrentAsyncOpFilterConfigurator;
+	}
+
+	if (bFiltersHaveChanged)
 	{
 		OnPreAsyncUpdate();
 
@@ -2318,16 +2332,19 @@ FReply STableTreeView::OnAdvancedFiltersClicked()
 				case ETableCellDataType::Int64:
 				{
 					AvailableFilters->Add(MakeShared<FFilter>(Column->GetIndex(), Column->GetTitleName(), Column->GetDescription(), EFilterDataType::Int64, FFilterService::Get()->GetIntegerOperators()));
+					Context.AddFilterData<int64>(Column->GetIndex(), 0);
 					break;
 				}
 				case ETableCellDataType::Double:
 				{
 					AvailableFilters->Add(MakeShared<FFilter>(Column->GetIndex(), Column->GetTitleName(), Column->GetDescription(), EFilterDataType::Double, FFilterService::Get()->GetDoubleOperators()));
+					Context.AddFilterData<double>(Column->GetIndex(), 0.0);
 					break;
 				}
 			}
 		}
 
+		CurrentAsyncOpFilterConfigurator = new FFilterConfigurator(*FilterConfigurator);
 		OnFilterChangesCommitedHandle = FilterConfigurator->GetOnChangesCommitedEvent().AddSP(this, &STableTreeView::OnAdvancedFiltersChangesCommited);
 	}
 
@@ -2340,36 +2357,46 @@ FReply STableTreeView::OnAdvancedFiltersClicked()
 
 bool STableTreeView::ApplyAdvancedFilters(const FTableTreeNodePtr& NodePtr)
 {
-	if (!FilterConfigurator.IsValid())
+	if (CurrentAsyncOpFilterConfigurator == nullptr)
+	{
+		return true;
+	}
+
+	if (CurrentAsyncOpFilterConfigurator->GetRootNode()->GetChildren().Num() == 0)
 	{
 		return true;
 	}
 
-	FFilterContext Context;
 	for (const TSharedRef<FTableColumn>& Column : Table->GetColumns())
 	{
 		switch (Column->GetDataType())
 		{
 			case ETableCellDataType::Int64:
 			{
-				Context.AddFilterData<int64>(Column->GetIndex(), Column->GetValue(*NodePtr)->AsInt64());
+				Context.SetFilterData<int64>(Column->GetIndex(), Column->GetValue(*NodePtr)->AsInt64());
 				break;
 			}
 			case ETableCellDataType::Double:
 			{
-				Context.AddFilterData<double>(Column->GetIndex(), Column->GetValue(*NodePtr)->AsDouble());
+				Context.SetFilterData<double>(Column->GetIndex(), Column->GetValue(*NodePtr)->AsDouble());
 				break;
 			}
 		}
 	}
 
-	return FilterConfigurator->ApplyFilters(Context);
+	return CurrentAsyncOpFilterConfigurator->ApplyFilters(Context);
 }
 
 ////////////////////////////////////////////////////////////////////////////////////////////////////
 
 void STableTreeView::OnAdvancedFiltersChangesCommited()
 {
+	// If settings have not changed, do nothing.
+	if (*FilterConfigurator == *CurrentAsyncOpFilterConfigurator)
+	{
+		return;
+	}
+
 	if (!bIsUpdateRunning)
 	{
 		OnPreAsyncUpdate();
diff --git a/Engine/Source/Developer/TraceInsights/Private/Insights/Table/Widgets/STableTreeView.h b/Engine/Source/Developer/TraceInsights/Private/Insights/Table/Widgets/STableTreeView.h
--- a/Engine/Source/Developer/TraceInsights/Private/Insights/Table/Widgets/STableTreeView.h
+++ b/Engine/Source/Developer/TraceInsights/Private/Insights/Table/Widgets/STableTreeView.h
@@ -20,6 +20,7 @@
 #include "Insights/Common/InsightsAsyncWorkUtils.h"
 #include "Insights/Common/Stopwatch.h"
 #include "Insights/Table/ViewModels/TableTreeNode.h"
+#include "Insights/ViewModels/Filters.h"
 
 #include <atomic>
 
@@ -414,11 +415,13 @@ protected:
 	ITableCellValueSorter* CurrentAsyncOpSorter = nullptr;
 	EColumnSortMode::Type CurrentAsyncOpColumnSortMode;
 	TSharedPtr<FTableTreeNodeTextFilter> CurrentAsyncOpTextFilter;
+	FFilterConfigurator* CurrentAsyncOpFilterConfigurator = nullptr;
 	std::atomic<bool> bCancelCurrentAsyncOp { false };
 
 	//////////////////////////////////////////////////
-	TSharedPtr<class FFilterConfigurator> FilterConfigurator;
+	TSharedPtr<FFilterConfigurator> FilterConfigurator;
 	FDelegateHandle OnFilterChangesCommitedHandle;
+	FFilterContext Context;
 
 	double StatsStartTime;
 	double StatsEndTime;
diff --git a/Engine/Source/Developer/TraceInsights/Private/Insights/ViewModels/FilterConfigurator.cpp b/Engine/Source/Developer/TraceInsights/Private/Insights/ViewModels/FilterConfigurator.cpp
--- a/Engine/Source/Developer/TraceInsights/Private/Insights/ViewModels/FilterConfigurator.cpp
+++ b/Engine/Source/Developer/TraceInsights/Private/Insights/ViewModels/FilterConfigurator.cpp
@@ -48,6 +48,16 @@ FFilterConfigurator& FFilterConfigurator::operator=(const FFilterConfigurator& O
 
 ////////////////////////////////////////////////////////////////////////////////////////////////////
 
+bool FFilterConfigurator::operator==(const FFilterConfigurator& Other)
+{
+	bool bIsEqual = AvailableFilters.Get() == Other.AvailableFilters.Get();
+	bIsEqual &= *RootNode == *Other.RootNode;
+
+	return bIsEqual;
+}
+
+////////////////////////////////////////////////////////////////////////////////////////////////////
+
 bool FFilterConfigurator::ApplyFilters(const FFilterContext& Context) const
 {
 	return RootNode->ApplyFilters(Context);
diff --git a/Engine/Source/Developer/TraceInsights/Private/Insights/ViewModels/FilterConfigurator.h b/Engine/Source/Developer/TraceInsights/Private/Insights/ViewModels/FilterConfigurator.h
--- a/Engine/Source/Developer/TraceInsights/Private/Insights/ViewModels/FilterConfigurator.h
+++ b/Engine/Source/Developer/TraceInsights/Private/Insights/ViewModels/FilterConfigurator.h
@@ -19,6 +19,10 @@ public:
 	FFilterConfigurator(const FFilterConfigurator& Other);
 	FFilterConfigurator& operator=(const FFilterConfigurator& Other);
 
+	bool operator==(const FFilterConfigurator& Other);
+
+	bool operator!=(const FFilterConfigurator& Other) { return !(*this == Other); }
+
 	~FFilterConfigurator();
 
 	FFilterConfiguratorNodePtr GetRootNode() { return RootNode; }
diff --git a/Engine/Source/Developer/TraceInsights/Private/Insights/ViewModels/FilterConfiguratorNode.cpp b/Engine/Source/Developer/TraceInsights/Private/Insights/ViewModels/FilterConfiguratorNode.cpp
--- a/Engine/Source/Developer/TraceInsights/Private/Insights/ViewModels/FilterConfiguratorNode.cpp
+++ b/Engine/Source/Developer/TraceInsights/Private/Insights/ViewModels/FilterConfiguratorNode.cpp
@@ -56,6 +56,29 @@ FFilterConfiguratorNode& FFilterConfiguratorNode::operator=(const FFilterConfigu
 
 ////////////////////////////////////////////////////////////////////////////////////////////////////
 
+bool FFilterConfiguratorNode::operator==(const FFilterConfiguratorNode& Other)
+{
+	bool bIsEqual = true;
+	bIsEqual &= AvailableFilters.Get() == Other.AvailableFilters.Get();
+	bIsEqual &= SelectedFilter.Get() == Other.SelectedFilter.Get();
+	bIsEqual &= SelectedFilterOperator.Get() == Other.SelectedFilterOperator.Get();
+	bIsEqual &= SelectedFilterGroupOperator.Get() == Other.SelectedFilterGroupOperator.Get();
+	bIsEqual &= TextBoxValue == Other.TextBoxValue;
+	bIsEqual &= GetChildren().Num() == Other.GetChildren().Num();
+
+	if (bIsEqual)
+	{
+		for (int32 Index = 0; Index < GetChildren().Num(); ++Index)
+		{
+			bIsEqual &= *StaticCastSharedPtr<FFilterConfiguratorNode>(GetChildren()[Index]) == *StaticCastSharedPtr<FFilterConfiguratorNode>(Other.GetChildren()[Index]);
+		}
+	}
+
+	return bIsEqual;
+}
+
+////////////////////////////////////////////////////////////////////////////////////////////////////
+
 const TArray<TSharedPtr<struct FFilterGroupOperator>>& FFilterConfiguratorNode::GetFilterGroupOperators()
 {
 	return FFilterService::Get()->GetFilterGroupOperators();
@@ -155,9 +178,12 @@ bool FFilterConfiguratorNode::ApplyFilters(const FFilterContext& Context) const
 		}
 		case EFilterGroupOperator::Or:
 		{
-			Ret = false;
-
 			auto& ChildrenArr = GetChildren();
+			if (ChildrenArr.Num() > 0)
+			{
+				Ret = false;
+			}
+
 			for (int Index = 0; Index < ChildrenArr.Num() && !Ret; ++Index)
 			{
 				Ret |= ((FFilterConfiguratorNode*)&*ChildrenArr[Index])->ApplyFilters(Context);
diff --git a/Engine/Source/Developer/TraceInsights/Private/Insights/ViewModels/FilterConfiguratorNode.h b/Engine/Source/Developer/TraceInsights/Private/Insights/ViewModels/FilterConfiguratorNode.h
--- a/Engine/Source/Developer/TraceInsights/Private/Insights/ViewModels/FilterConfiguratorNode.h
+++ b/Engine/Source/Developer/TraceInsights/Private/Insights/ViewModels/FilterConfiguratorNode.h
@@ -53,6 +53,8 @@ public:
 
 	FFilterConfiguratorNode(const FFilterConfiguratorNode& Other);
 
+	bool operator==(const FFilterConfiguratorNode& Other);
+
 	virtual ~FFilterConfiguratorNode() {}
 
 	virtual const FName& GetTypeName() const override { return TypeName; }
diff --git a/Engine/Source/Developer/TraceInsights/Private/Insights/ViewModels/Filters.h b/Engine/Source/Developer/TraceInsights/Private/Insights/ViewModels/Filters.h
--- a/Engine/Source/Developer/TraceInsights/Private/Insights/ViewModels/Filters.h
+++ b/Engine/Source/Developer/TraceInsights/Private/Insights/ViewModels/Filters.h
@@ -18,16 +18,6 @@ namespace Insights
 
 ////////////////////////////////////////////////////////////////////////////////////////////////////
 
-enum class EFilterField : int32
-{
-	StartTime = 0,
-	EndTime = 1,
-	Duration = 2,
-	EventType = 3
-};
-
-////////////////////////////////////////////////////////////////////////////////////////////////////
-
 enum class EFilterDataType : uint32
 {
 	Int64,
diff --git a/Engine/Source/Developer/TraceInsights/Private/Insights/ViewModels/ThreadTimingTrack.h b/Engine/Source/Developer/TraceInsights/Private/Insights/ViewModels/ThreadTimingTrack.h
--- a/Engine/Source/Developer/TraceInsights/Private/Insights/ViewModels/ThreadTimingTrack.h
+++ b/Engine/Source/Developer/TraceInsights/Private/Insights/ViewModels/ThreadTimingTrack.h
@@ -98,6 +98,17 @@ private:
 
 ////////////////////////////////////////////////////////////////////////////////////////////////////
 
+
+enum class EFilterField : int32
+{
+	StartTime = 0,
+	EndTime = 1,
+	Duration = 2,
+	EventType = 3
+};
+
+////////////////////////////////////////////////////////////////////////////////////////////////////
+
 class FThreadTimingTrack : public FTimingEventsTrack
 {
 	INSIGHTS_DECLARE_RTTI(FThreadTimingTrack, FTimingEventsTrack)
