[Insights] Add aggregation to the Size, StartTime and EndTime columns of AllocTreeView.

diff --git a/Engine/Source/Developer/TraceInsights/Private/Insights/MemoryProfiler/ViewModels/MemAllocTable.cpp b/Engine/Source/Developer/TraceInsights/Private/Insights/MemoryProfiler/ViewModels/MemAllocTable.cpp
--- a/Engine/Source/Developer/TraceInsights/Private/Insights/MemoryProfiler/ViewModels/MemAllocTable.cpp
+++ b/Engine/Source/Developer/TraceInsights/Private/Insights/MemoryProfiler/ViewModels/MemAllocTable.cpp
@@ -71,7 +71,7 @@ void FMemAllocTable::AddDefaultColumns()
 
 		Column.SetFlags(ETableColumnFlags::ShouldBeVisible | ETableColumnFlags::CanBeHidden | ETableColumnFlags::CanBeFiltered);
 
-		Column.SetHorizontalAlignment(HAlign_Right);
+		Column.SetHorizontalAlignment(HAlign_Left);
 		Column.SetInitialWidth(100.0f);
 
 		Column.SetDataType(ETableCellDataType::Double);
@@ -96,6 +96,8 @@ void FMemAllocTable::AddDefaultColumns()
 		TSharedRef<ITableCellValueSorter> Sorter = MakeShared<FSorterByDoubleValue>(ColumnRef);
 		Column.SetValueSorter(Sorter);
 
+		Column.SetAggregation(ETableColumnAggregation::Min);
+
 		AddColumn(ColumnRef);
 	}
 	//////////////////////////////////////////////////
@@ -112,7 +114,7 @@ void FMemAllocTable::AddDefaultColumns()
 
 		Column.SetFlags(ETableColumnFlags::ShouldBeVisible | ETableColumnFlags::CanBeHidden | ETableColumnFlags::CanBeFiltered);
 
-		Column.SetHorizontalAlignment(HAlign_Right);
+		Column.SetHorizontalAlignment(HAlign_Left);
 		Column.SetInitialWidth(100.0f);
 
 		Column.SetDataType(ETableCellDataType::Double);
@@ -137,6 +139,8 @@ void FMemAllocTable::AddDefaultColumns()
 		TSharedRef<ITableCellValueSorter> Sorter = MakeShared<FSorterByDoubleValue>(ColumnRef);
 		Column.SetValueSorter(Sorter);
 
+		Column.SetAggregation(ETableColumnAggregation::Max);
+
 		AddColumn(ColumnRef);
 	}
 	//////////////////////////////////////////////////
@@ -153,7 +157,7 @@ void FMemAllocTable::AddDefaultColumns()
 
 		Column.SetFlags(ETableColumnFlags::ShouldBeVisible | ETableColumnFlags::CanBeHidden | ETableColumnFlags::CanBeFiltered);
 
-		Column.SetHorizontalAlignment(HAlign_Right);
+		Column.SetHorizontalAlignment(HAlign_Left);
 		Column.SetInitialWidth(100.0f);
 
 		Column.SetDataType(ETableCellDataType::Double);
@@ -235,7 +239,7 @@ void FMemAllocTable::AddDefaultColumns()
 
 		Column.SetFlags(ETableColumnFlags::ShouldBeVisible | ETableColumnFlags::CanBeHidden | ETableColumnFlags::CanBeFiltered);
 
-		Column.SetHorizontalAlignment(HAlign_Right);
+		Column.SetHorizontalAlignment(HAlign_Left);
 		Column.SetInitialWidth(100.0f);
 
 		Column.SetDataType(ETableCellDataType::Int64);
@@ -260,6 +264,8 @@ void FMemAllocTable::AddDefaultColumns()
 		TSharedRef<ITableCellValueSorter> Sorter = MakeShared<FSorterByInt64Value>(ColumnRef);
 		Column.SetValueSorter(Sorter);
 
+		Column.SetAggregation(ETableColumnAggregation::Sum);
+
 		AddColumn(ColumnRef);
 	}
 	//////////////////////////////////////////////////
diff --git a/Engine/Source/Developer/TraceInsights/Private/Insights/Table/ViewModels/TableCellValue.h b/Engine/Source/Developer/TraceInsights/Private/Insights/Table/ViewModels/TableCellValue.h
--- a/Engine/Source/Developer/TraceInsights/Private/Insights/Table/ViewModels/TableCellValue.h
+++ b/Engine/Source/Developer/TraceInsights/Private/Insights/Table/ViewModels/TableCellValue.h
@@ -115,7 +115,7 @@ public:
 		}
 	}
 
-	float AsDouble() const
+	double AsDouble() const
 	{
 		switch (DataType)
 		{
diff --git a/Engine/Source/Developer/TraceInsights/Private/Insights/Table/ViewModels/TableColumn.h b/Engine/Source/Developer/TraceInsights/Private/Insights/Table/ViewModels/TableColumn.h
--- a/Engine/Source/Developer/TraceInsights/Private/Insights/Table/ViewModels/TableColumn.h
+++ b/Engine/Source/Developer/TraceInsights/Private/Insights/Table/ViewModels/TableColumn.h
@@ -36,8 +36,8 @@ enum class ETableColumnAggregation : uint32
 {
 	None = 0,
 	Sum,
-	//Min,
-	//Max,
+	Min,
+	Max,
 	//Average,
 	//Median,
 };
diff --git a/Engine/Source/Developer/TraceInsights/Private/Insights/Table/ViewModels/TableTreeNode.h b/Engine/Source/Developer/TraceInsights/Private/Insights/Table/ViewModels/TableTreeNode.h
--- a/Engine/Source/Developer/TraceInsights/Private/Insights/Table/ViewModels/TableTreeNode.h
+++ b/Engine/Source/Developer/TraceInsights/Private/Insights/Table/ViewModels/TableTreeNode.h
@@ -75,17 +75,23 @@ public:
 	int32 GetRowIndex() const { return RowId.RowIndex; }
 
 	void ResetAggregatedValues() { AggregatedValues.Reset(); }
+	void ResetAggregatedValues(const FName& ColumnId) { AggregatedValues.Remove(ColumnId); }
 	bool HasAggregatedValue(const FName& ColumnId) const { return AggregatedValues.Contains(ColumnId); }
 	const FTableCellValue* FindAggregatedValue(const FName& ColumnId) const { return AggregatedValues.Find(ColumnId); }
 	const FTableCellValue& GetAggregatedValue(const FName& ColumnId) const { return AggregatedValues.FindChecked(ColumnId); }
 	void AddAggregatedValue(const FName& ColumnId, const FTableCellValue& Value) { AggregatedValues.Add(ColumnId, Value); }
 	void SetAggregatedValue(const FName& ColumnId, const FTableCellValue& Value) { AggregatedValues[ColumnId] = Value; }
+	virtual bool IsFiltered() const override { return bIsFiltered; }
+	virtual void SetIsFiltered(bool InValue) { bIsFiltered = InValue; }
 
 protected:
 	TWeakPtr<FTable> ParentTable;
 	FTableRowId RowId;
 
 	TMap<FName, FTableCellValue> AggregatedValues;
+
+private:
+	bool bIsFiltered = false;
 };
 
 ////////////////////////////////////////////////////////////////////////////////////////////////////
diff --git a/Engine/Source/Developer/TraceInsights/Private/Insights/Table/ViewModels/UntypedTable.cpp b/Engine/Source/Developer/TraceInsights/Private/Insights/Table/ViewModels/UntypedTable.cpp
--- a/Engine/Source/Developer/TraceInsights/Private/Insights/Table/ViewModels/UntypedTable.cpp
+++ b/Engine/Source/Developer/TraceInsights/Private/Insights/Table/ViewModels/UntypedTable.cpp
@@ -222,7 +222,7 @@ void FUntypedTable::CreateColumns(const TraceServices::ITableLayout& TableLayout
 
 		case TraceServices::TableColumnType_Int:
 			Column.SetDataType(ETableCellDataType::Int64);
-			Aggregation = ETableColumnAggregation::Sum;
+			//Aggregation = ETableColumnAggregation::Sum;
 			HorizontalAlignment = HAlign_Right;
 			InitialColumnWidth = 60.0f;
 			//TODO: if (Hint == AsMemory)
@@ -236,7 +236,7 @@ void FUntypedTable::CreateColumns(const TraceServices::ITableLayout& TableLayout
 
 		case TraceServices::TableColumnType_Float:
 			Column.SetDataType(ETableCellDataType::Float);
-			Aggregation = ETableColumnAggregation::Sum;
+			//Aggregation = ETableColumnAggregation::Sum;
 			HorizontalAlignment = HAlign_Right;
 			InitialColumnWidth = 60.0f;
 			//TODO: if (Hint == AsTimeMs)
@@ -247,7 +247,7 @@ void FUntypedTable::CreateColumns(const TraceServices::ITableLayout& TableLayout
 
 		case TraceServices::TableColumnType_Double:
 			Column.SetDataType(ETableCellDataType::Double);
-			Aggregation = ETableColumnAggregation::Sum;
+			//Aggregation = ETableColumnAggregation::Sum;
 			HorizontalAlignment = HAlign_Right;
 			InitialColumnWidth = 80.0f;
 			//TODO: if (Hint == AsTimeMs)
diff --git a/Engine/Source/Developer/TraceInsights/Private/Insights/Table/Widgets/STableTreeView.cpp b/Engine/Source/Developer/TraceInsights/Private/Insights/Table/Widgets/STableTreeView.cpp
--- a/Engine/Source/Developer/TraceInsights/Private/Insights/Table/Widgets/STableTreeView.cpp
+++ b/Engine/Source/Developer/TraceInsights/Private/Insights/Table/Widgets/STableTreeView.cpp
@@ -21,8 +21,6 @@
 #include "Insights/InsightsManager.h"
 #include "Insights/Log.h"
 #include "Insights/Table/ViewModels/Table.h"
-#include "Insights/Table/ViewModels/TableColumn.h"
-#include "Insights/Table/ViewModels/TableTreeNode.h"
 #include "Insights/Table/ViewModels/TreeNodeGrouping.h"
 #include "Insights/Table/ViewModels/TreeNodeSorting.h"
 #include "Insights/Table/ViewModels/UntypedTable.h"
@@ -32,6 +30,8 @@
 #include "Insights/ViewModels/FilterConfigurator.h"
 #include "Insights/Widgets/SAsyncOperationStatus.h"
 
+#include <limits>
+
 #define LOCTEXT_NAMESPACE "STableTreeView"
 
 namespace Insights
@@ -702,6 +702,8 @@ void STableTreeView::ApplyFiltering()
 		}
 	}
 
+	UpdateAggregatedValues(*Root);
+
 	// Cannot call TreeView functions from other threads than MainThread and SlateThread. 
 	if (!bRunInAsyncMode)
 	{
@@ -793,6 +795,7 @@ bool STableTreeView::ApplyFilteringForNode(FTableTreeNodePtr NodePtr)
 	else
 	{
 		bIsNodeVisible &= ApplyAdvancedFilters(NodePtr);
+		NodePtr->SetIsFiltered(!bIsNodeVisible);
 		return bIsNodeVisible;
 	}
 }
@@ -970,28 +973,7 @@ void STableTreeView::CreateGroups(const TArray<TSharedPtr<FTreeNodeGrouping>>& G
 	Stopwatch.Update();
 	const double Time1 = Stopwatch.GetAccumulatedTime();
 
-	ResetAggregatedValuesRec(*Root);
-	for (const TSharedRef<FTableColumn>& ColumnRef : Table->GetColumns())
-	{
-		FTableColumn& Column = *ColumnRef;
-		if (Column.GetAggregation() == ETableColumnAggregation::Sum)
-		{
-			switch (Column.GetDataType())
-			{
-				case ETableCellDataType::Int64:
-					UpdateInt64SumAggregationRec(Column, *Root);
-					break;
-
-				case ETableCellDataType::Float:
-					UpdateFloatSumAggregationRec(Column, *Root);
-					break;
-
-				case ETableCellDataType::Double:
-					UpdateDoubleSumAggregationRec(Column, *Root);
-					break;
-			}
-		}
-	}
+	UpdateAggregatedValues(*Root);
 
 	Stopwatch.Stop();
 	const double TotalTime = Stopwatch.GetAccumulatedTime();
@@ -1049,94 +1031,39 @@ void STableTreeView::GroupNodesRec(const TArray<FTableTreeNodePtr>& Nodes, FTabl
 
 ////////////////////////////////////////////////////////////////////////////////////////////////////
 
-void STableTreeView::ResetAggregatedValuesRec(FTableTreeNode& GroupNode)
+void STableTreeView::UpdateAggregatedValues(FTableTreeNode& GroupNode)
 {
-	ensure(GroupNode.IsGroup());
-
-	GroupNode.ResetAggregatedValues();
-
-	for (FBaseTreeNodePtr ChildPtr : GroupNode.GetChildren())
-	{
-		if (ChildPtr->IsGroup())
-		{
-			ResetAggregatedValuesRec(*StaticCastSharedPtr<FTableTreeNode>(ChildPtr));
-		}
-	}
-}
-
-////////////////////////////////////////////////////////////////////////////////////////////////////
-
-void STableTreeView::UpdateInt64SumAggregationRec(FTableColumn& Column, FTableTreeNode& GroupNode)
-{
-	int64 AggregatedValue = 0;
-
-	for (FBaseTreeNodePtr NodePtr : GroupNode.GetChildren())
+	for (const TSharedRef<FTableColumn>& ColumnRef : Table->GetColumns())
 	{
-		FTableTreeNode& TableNode = *StaticCastSharedPtr<FTableTreeNode>(NodePtr);
-		if (TableNode.IsGroup())
-		{
-			UpdateInt64SumAggregationRec(Column, TableNode);
-		}
-
-		const TOptional<FTableCellValue> OptionalValue = Column.GetValue(TableNode);
-		if (OptionalValue.IsSet())
+		FTableColumn& Column = *ColumnRef;
+		switch (Column.GetAggregation())
 		{
-			ensure(OptionalValue.GetValue().DataType == ETableCellDataType::Int64);
-			AggregatedValue += OptionalValue.GetValue().Int64;
-		}
-	}
-
-	GroupNode.AddAggregatedValue(Column.GetId(), FTableCellValue(AggregatedValue));
-}
-
-////////////////////////////////////////////////////////////////////////////////////////////////////
-
-void STableTreeView::UpdateFloatSumAggregationRec(FTableColumn& Column, FTableTreeNode& GroupNode)
-{
-	float AggregatedValue = 0.0f;
-
-	for (FBaseTreeNodePtr NodePtr : GroupNode.GetChildren())
-	{
-		FTableTreeNode& TableNode = *StaticCastSharedPtr<FTableTreeNode>(NodePtr);
-		if (TableNode.IsGroup())
+		case ETableColumnAggregation::Sum:
 		{
-			UpdateFloatSumAggregationRec(Column, TableNode);
+			STableTreeView::UpdateAggregationRec<int64>(Column, GroupNode, 0, true, [](int64 InValue, TOptional<FTableCellValue> InTableCellValue)
+				{
+					return InValue + InTableCellValue->AsInt64();
+				});
+			break;
 		}
-
-		const TOptional<FTableCellValue> OptionalValue = Column.GetValue(TableNode);
-		if (OptionalValue.IsSet())
+		case ETableColumnAggregation::Min:
 		{
-			ensure(OptionalValue.GetValue().DataType == ETableCellDataType::Float);
-			AggregatedValue += OptionalValue.GetValue().Float;
+			STableTreeView::UpdateAggregationRec<double>(Column, GroupNode, std::numeric_limits<double>::max(), false, [](double InValue, TOptional<FTableCellValue> InTableCellValue)
+				{
+					return FMath::Min(InValue, InTableCellValue->AsDouble());
+				});
+			break;
 		}
-	}
-
-	GroupNode.AddAggregatedValue(Column.GetId(), FTableCellValue(AggregatedValue));
-}
-
-////////////////////////////////////////////////////////////////////////////////////////////////////
-
-void STableTreeView::UpdateDoubleSumAggregationRec(FTableColumn& Column, FTableTreeNode& GroupNode)
-{
-	double AggregatedValue = 0.0;
-
-	for (FBaseTreeNodePtr NodePtr : GroupNode.GetChildren())
-	{
-		FTableTreeNode& TableNode = *StaticCastSharedPtr<FTableTreeNode>(NodePtr);
-		if (TableNode.IsGroup())
+		case ETableColumnAggregation::Max:
 		{
-			UpdateDoubleSumAggregationRec(Column, TableNode);
+			STableTreeView::UpdateAggregationRec<double>(Column, GroupNode, std::numeric_limits<double>::lowest(), false, [](double InValue, TOptional<FTableCellValue> InTableCellValue)
+				{
+					return FMath::Max(InValue, InTableCellValue->AsDouble());
+				});
+			break;
 		}
-
-		const TOptional<FTableCellValue> OptionalValue = Column.GetValue(TableNode);
-		if (OptionalValue.IsSet())
-		{
-			ensure(OptionalValue.GetValue().DataType == ETableCellDataType::Double);
-			AggregatedValue += OptionalValue.GetValue().Double;
 		}
 	}
-
-	GroupNode.AddAggregatedValue(Column.GetId(), FTableCellValue(AggregatedValue));
 }
 
 ////////////////////////////////////////////////////////////////////////////////////////////////////
diff --git a/Engine/Source/Developer/TraceInsights/Private/Insights/Table/Widgets/STableTreeView.h b/Engine/Source/Developer/TraceInsights/Private/Insights/Table/Widgets/STableTreeView.h
--- a/Engine/Source/Developer/TraceInsights/Private/Insights/Table/Widgets/STableTreeView.h
+++ b/Engine/Source/Developer/TraceInsights/Private/Insights/Table/Widgets/STableTreeView.h
@@ -19,6 +19,7 @@
 // Insights
 #include "Insights/Common/InsightsAsyncWorkUtils.h"
 #include "Insights/Common/Stopwatch.h"
+#include "Insights/Table/ViewModels/TableColumn.h"
 #include "Insights/Table/ViewModels/TableTreeNode.h"
 #include "Insights/ViewModels/Filters.h"
 
@@ -207,10 +208,44 @@ protected:
 	void CreateGroups(const TArray<TSharedPtr<FTreeNodeGrouping>>& Groupings);
 	void GroupNodesRec(const TArray<FTableTreeNodePtr>& Nodes, FTableTreeNode& ParentGroup, int32 GroupingDepth, const TArray<TSharedPtr<FTreeNodeGrouping>>& Groupings);
 
-	void ResetAggregatedValuesRec(FTableTreeNode& GroupNode);
-	void UpdateInt64SumAggregationRec(FTableColumn& Column, FTableTreeNode& GroupNode);
-	void UpdateFloatSumAggregationRec(FTableColumn& Column, FTableTreeNode& GroupNode);
-	void UpdateDoubleSumAggregationRec(FTableColumn& Column, FTableTreeNode& GroupNode);
+	void UpdateAggregatedValues(FTableTreeNode& GroupNode);
+	
+	template<typename T>
+	static void UpdateAggregationRec(FTableColumn& Column, FTableTreeNode& GroupNode, T InitialAggregatedValue, bool bSetInitialValue, TFunctionRef<T(T, TOptional<FTableCellValue>)> ValueGetterFunc)
+	{
+		T AggregatedValue = InitialAggregatedValue;
+
+		for (FBaseTreeNodePtr NodePtr : GroupNode.GetChildren())
+		{
+			if (NodePtr->IsFiltered())
+			{
+				continue;
+			}
+			FTableTreeNode& TableNode = *StaticCastSharedPtr<FTableTreeNode>(NodePtr);
+			if (TableNode.IsGroup())
+			{
+				TableNode.ResetAggregatedValues(Column.GetId());
+				UpdateAggregationRec(Column, TableNode, InitialAggregatedValue, bSetInitialValue, ValueGetterFunc);
+				if (TableNode.HasAggregatedValue(Column.GetId()))
+				{
+					AggregatedValue = ValueGetterFunc(AggregatedValue, TableNode.GetAggregatedValue(Column.GetId()));
+				}
+			}
+			else
+			{
+				const TOptional<FTableCellValue> OptionalValue = Column.GetValue(TableNode);
+				if (OptionalValue.IsSet())
+				{
+					AggregatedValue = ValueGetterFunc(AggregatedValue, OptionalValue);
+				}
+			}
+		}
+
+		if (bSetInitialValue || InitialAggregatedValue != AggregatedValue)
+		{
+			GroupNode.AddAggregatedValue(Column.GetId(), FTableCellValue(AggregatedValue));
+		}
+	}
 
 	void RebuildGroupingCrumbs();
 	void OnGroupingCrumbClicked(const TSharedPtr<FTreeNodeGrouping>& InEntry);
diff --git a/Engine/Source/Developer/TraceInsights/Private/Insights/Table/Widgets/STableTreeViewCell.cpp b/Engine/Source/Developer/TraceInsights/Private/Insights/Table/Widgets/STableTreeViewCell.cpp
--- a/Engine/Source/Developer/TraceInsights/Private/Insights/Table/Widgets/STableTreeViewCell.cpp
+++ b/Engine/Source/Developer/TraceInsights/Private/Insights/Table/Widgets/STableTreeViewCell.cpp
@@ -13,6 +13,7 @@
 
 // Insights
 #include "Insights/Table/ViewModels/Table.h"
+#include "Insights/Table/ViewModels/TableCellValueFormatter.h"
 #include "Insights/Table/ViewModels/TableColumn.h"
 #include "Insights/Table/Widgets/STableTreeViewRow.h"
 
@@ -123,7 +124,22 @@ TSharedRef<SWidget> STableTreeViewCell::GenerateWidgetForTableColumn(const FArgu
 	// Note: For performance reason, init the cell text (value) only once.
 	//       If we'll need to update values without recreating the table row/cell widgets, bind .Text to STableTreeViewCell::GetValueAsText
 	//       or add API to explicitly update the text block.
-	const FText CellText = ColumnPtr->GetValueAsText(*TableTreeNodePtr);
+	FText CellText;
+	if (TableTreeNodePtr->IsGroup())
+	{
+		if (TableTreeNodePtr->HasAggregatedValue(ColumnPtr->GetId()))
+		{
+			CellText = ColumnPtr->GetValueFormatter()->FormatValue(TableTreeNodePtr->GetAggregatedValue(ColumnPtr->GetId()));
+		}
+		else
+		{
+			CellText = FText::FromString(TEXT("-"));
+		}
+	}
+	else
+	{
+		CellText = ColumnPtr->GetValueAsText(*TableTreeNodePtr);
+	}
 
 	return
 		SNew(SHorizontalBox)
