[Insights] Fix Mem Alloc Tree View memory leak caused by circular reference.

diff --git a/Engine/Source/Developer/TraceInsights/Private/Insights/Widgets/SAsyncOperationStatus.cpp b/Engine/Source/Developer/TraceInsights/Private/Insights/Widgets/SAsyncOperationStatus.cpp
--- a/Engine/Source/Developer/TraceInsights/Private/Insights/Widgets/SAsyncOperationStatus.cpp
+++ b/Engine/Source/Developer/TraceInsights/Private/Insights/Widgets/SAsyncOperationStatus.cpp
@@ -70,7 +70,13 @@ void SAsyncOperationStatus::Construct(const FArguments& InArgs, TSharedRef<IAsyn
 
 EVisibility SAsyncOperationStatus::GetContentVisibility() const
 {
-	return StatusProvider->IsRunning() ? EVisibility::Visible : EVisibility::Collapsed;
+	TSharedPtr<IAsyncOperationStatusProvider> StatusProviderSharedPtr = StatusProvider.Pin();
+	if (StatusProviderSharedPtr.IsValid())
+	{
+		return StatusProviderSharedPtr->IsRunning() ? EVisibility::Visible : EVisibility::Collapsed;
+	}
+	
+	return EVisibility::Collapsed;
 }
 
 ////////////////////////////////////////////////////////////////////////////////////////////////////
@@ -93,7 +99,8 @@ FSlateColor SAsyncOperationStatus::GetTextColorAndOpacity() const
 
 float SAsyncOperationStatus::ComputeOpacity() const
 {
-	const double TotalDuration = StatusProvider->GetAllOperationsDuration();
+	TSharedPtr<IAsyncOperationStatusProvider> StatusProviderSharedPtr = StatusProvider.Pin();
+	const double TotalDuration = StatusProviderSharedPtr.IsValid() ? StatusProviderSharedPtr->GetAllOperationsDuration() : 0;
 	constexpr double FadeInStartTime = 0.1; // [second]
 	constexpr double FadeInEndTime = 3.0; // [second]
 	constexpr double FadeInDuration = FadeInEndTime - FadeInStartTime;
@@ -112,7 +119,8 @@ FText SAsyncOperationStatus::GetText() const
 
 FText SAsyncOperationStatus::GetAnimatedText() const
 {
-	const double TotalDuration = StatusProvider->GetAllOperationsDuration();
+	TSharedPtr<IAsyncOperationStatusProvider> StatusProviderSharedPtr = StatusProvider.Pin();
+	const double TotalDuration = StatusProviderSharedPtr.IsValid() ? StatusProviderSharedPtr->GetAllOperationsDuration() : 0;
 	const TCHAR* Anim[] = { TEXT(""), TEXT("."), TEXT(".."), TEXT("..."), };
 	int32 AnimIndex = static_cast<int32>(TotalDuration / 0.2) % UE_ARRAY_COUNT(Anim);
 	return FText::FromString(FString(Anim[AnimIndex]));
@@ -122,25 +130,31 @@ FText SAsyncOperationStatus::GetAnimatedText() const
 
 FText SAsyncOperationStatus::GetTooltipText() const
 {
-	const double TotalDuration = StatusProvider->GetAllOperationsDuration();
-	const uint32 OperationCount = StatusProvider->GetOperationCount();
-	FText CurrentOpName = StatusProvider->GetCurrentOperationName();
-
-	if (OperationCount == 1)
+	TSharedPtr<IAsyncOperationStatusProvider> StatusProviderSharedPtr = StatusProvider.Pin();
+	if (StatusProviderSharedPtr.IsValid())
 	{
-		return FText::Format(LOCTEXT("DefaultTooltip_Fmt2", "{0}...\nElapsed Time: {1}"),
-			CurrentOpName,
-			FText::FromString(TimeUtils::FormatTime(TotalDuration, TimeUtils::Second)));
-	}
-	else
-	{
-		const double Duration = StatusProvider->GetCurrentOperationDuration();
-		return FText::Format(LOCTEXT("DefaultTooltip_Fmt3", "{0}...\nElapsed Time: {1}\nElapsed Time (op {2}): {3}"),
-			CurrentOpName,
-			FText::FromString(TimeUtils::FormatTime(TotalDuration, TimeUtils::Second)),
-			FText::AsNumber(OperationCount),
-			FText::FromString(TimeUtils::FormatTime(Duration, TimeUtils::Second)));
+		const double TotalDuration = StatusProviderSharedPtr->GetAllOperationsDuration();
+		const uint32 OperationCount = StatusProviderSharedPtr->GetOperationCount();
+		FText CurrentOpName = StatusProviderSharedPtr->GetCurrentOperationName();
+
+		if (OperationCount == 1)
+		{
+			return FText::Format(LOCTEXT("DefaultTooltip_Fmt2", "{0}...\nElapsed Time: {1}"),
+				CurrentOpName,
+				FText::FromString(TimeUtils::FormatTime(TotalDuration, TimeUtils::Second)));
+		}
+		else
+		{
+			const double Duration = StatusProviderSharedPtr->GetCurrentOperationDuration();
+			return FText::Format(LOCTEXT("DefaultTooltip_Fmt3", "{0}...\nElapsed Time: {1}\nElapsed Time (op {2}): {3}"),
+				CurrentOpName,
+				FText::FromString(TimeUtils::FormatTime(TotalDuration, TimeUtils::Second)),
+				FText::AsNumber(OperationCount),
+				FText::FromString(TimeUtils::FormatTime(Duration, TimeUtils::Second)));
+		}
 	}
+
+	return FText::FromString(TEXT(""));
 }
 
 ////////////////////////////////////////////////////////////////////////////////////////////////////
diff --git a/Engine/Source/Developer/TraceInsights/Private/Insights/Widgets/SAsyncOperationStatus.h b/Engine/Source/Developer/TraceInsights/Private/Insights/Widgets/SAsyncOperationStatus.h
--- a/Engine/Source/Developer/TraceInsights/Private/Insights/Widgets/SAsyncOperationStatus.h
+++ b/Engine/Source/Developer/TraceInsights/Private/Insights/Widgets/SAsyncOperationStatus.h
@@ -33,7 +33,7 @@ private:
 	FText GetTooltipText() const;
 
 private:
-	TSharedPtr<IAsyncOperationStatusProvider> StatusProvider;
+	TWeakPtr<IAsyncOperationStatusProvider> StatusProvider;
 };
 
 } // namespace Insights
