[Insights] Dispatch the async tasks in MemAllocTreeView only after the tree refresh is complete to eliminate concurent access to shared pointers.

diff --git a/Engine/Source/Developer/TraceInsights/Private/Insights/Table/Widgets/STableTreeView.cpp b/Engine/Source/Developer/TraceInsights/Private/Insights/Table/Widgets/STableTreeView.cpp
--- a/Engine/Source/Developer/TraceInsights/Private/Insights/Table/Widgets/STableTreeView.cpp
+++ b/Engine/Source/Developer/TraceInsights/Private/Insights/Table/Widgets/STableTreeView.cpp
@@ -630,11 +630,24 @@ void STableTreeView::Tick(const FGeometry& AllottedGeometry, const double InCurr
 {
 	if (bRunInAsyncMode && bIsUpdateRunning && !bIsCloseScheduled)
 	{
-		check(InProgressAsyncOperationEvent.IsValid());
-		if (InProgressAsyncOperationEvent->IsComplete())
+		if (DispatchEvent.IsValid() && !DispatchEvent->IsComplete())
 		{
-			OnPostAsyncUpdate();
-			StartPendingAsyncOperations();
+			// We wait for the TreeView to be refreshed before dispatching the tasks.
+			// This should make the TreeView release all of it's shared pointers to nodes to prevent 
+			// the TreeView (MainThread) and the tasks from accesing the non-thread safe shared pointers at the same time.
+			if (!TreeView->IsPendingRefresh())
+			{
+				DispatchEvent->DispatchSubsequents();
+			}
+		}
+		else
+		{
+			check(InProgressAsyncOperationEvent.IsValid());
+			if (InProgressAsyncOperationEvent->IsComplete())
+			{
+				OnPostAsyncUpdate();
+				StartPendingAsyncOperations();
+			}
 		}
 	}
 }
@@ -2032,6 +2045,8 @@ void STableTreeView::OnPreAsyncUpdate()
 
 	TreeView->SetTreeItemsSource(&DummyGroupNodes);
 	TreeView_Refresh();
+
+	DispatchEvent = FGraphEvent::CreateGraphEvent();
 }
 
 ////////////////////////////////////////////////////////////////////////////////////////////////////
@@ -2079,16 +2094,17 @@ FGraphEventRef STableTreeView::StartSortTreeNodesTask(FGraphEventRef Prerequisit
 	CurrentAsyncOpSorter = CurrentSorter.Get();
 	CurrentAsyncOpColumnSortMode = ColumnSortMode;
 
+	FGraphEventArray Prerequisites;
 	if (Prerequisite.IsValid())
 	{
-		FGraphEventArray Prerequisites;
 		Prerequisites.Add(Prerequisite);
-		return TGraphTask<FTableTreeViewSortAsyncTask>::CreateTask(&Prerequisites).ConstructAndDispatchWhenReady(this, CurrentAsyncOpSorter, CurrentAsyncOpColumnSortMode);
 	}
 	else
 	{
-		return TGraphTask<FTableTreeViewSortAsyncTask>::CreateTask().ConstructAndDispatchWhenReady(this, CurrentAsyncOpSorter, CurrentAsyncOpColumnSortMode);
+		Prerequisites.Add(DispatchEvent);
 	}
+
+	return TGraphTask<FTableTreeViewSortAsyncTask>::CreateTask(&Prerequisites).ConstructAndDispatchWhenReady(this, CurrentAsyncOpSorter, CurrentAsyncOpColumnSortMode);
 }
 
 ////////////////////////////////////////////////////////////////////////////////////////////////////
@@ -2100,16 +2116,17 @@ FGraphEventRef STableTreeView::StartCreateGroupsTask(FGraphEventRef Prerequisite
 	CurrentAsyncOpGroupings.Empty();
 	CurrentAsyncOpGroupings.Insert(CurrentGroupings, 0);
 
+	FGraphEventArray Prerequisites;
 	if (Prerequisite.IsValid())
 	{
-		FGraphEventArray Prerequisites;
 		Prerequisites.Add(Prerequisite);
-		return TGraphTask<FTableTreeViewGroupAsyncTask>::CreateTask(&Prerequisites).ConstructAndDispatchWhenReady(this, &CurrentAsyncOpGroupings);
 	}
 	else
 	{
-		return TGraphTask<FTableTreeViewGroupAsyncTask>::CreateTask().ConstructAndDispatchWhenReady(this, &CurrentAsyncOpGroupings);
+		Prerequisites.Add(DispatchEvent);
 	}
+
+	return TGraphTask<FTableTreeViewGroupAsyncTask>::CreateTask(&Prerequisites).ConstructAndDispatchWhenReady(this, &CurrentAsyncOpGroupings);
 }
 
 ////////////////////////////////////////////////////////////////////////////////////////////////////
@@ -2124,16 +2141,17 @@ FGraphEventRef STableTreeView::StartApplyFiltersTask(FGraphEventRef Prerequisite
 		*CurrentAsyncOpFilterConfigurator = *FilterConfigurator;
 	}
 
+	FGraphEventArray Prerequisites;
 	if (Prerequisite.IsValid())
 	{
-		FGraphEventArray Prerequisites;
 		Prerequisites.Add(Prerequisite);
-		return TGraphTask<FTableTreeViewFilterAsyncTask>::CreateTask(&Prerequisites).ConstructAndDispatchWhenReady(this);
 	}
 	else
 	{
-		return TGraphTask<FTableTreeViewFilterAsyncTask>::CreateTask().ConstructAndDispatchWhenReady(this);
+		Prerequisites.Add(DispatchEvent);
 	}
+
+	return TGraphTask<FTableTreeViewFilterAsyncTask>::CreateTask(&Prerequisites).ConstructAndDispatchWhenReady(this);
 }
 
 ////////////////////////////////////////////////////////////////////////////////////////////////////
diff --git a/Engine/Source/Developer/TraceInsights/Private/Insights/Table/Widgets/STableTreeView.h b/Engine/Source/Developer/TraceInsights/Private/Insights/Table/Widgets/STableTreeView.h
--- a/Engine/Source/Developer/TraceInsights/Private/Insights/Table/Widgets/STableTreeView.h
+++ b/Engine/Source/Developer/TraceInsights/Private/Insights/Table/Widgets/STableTreeView.h
@@ -452,6 +452,7 @@ protected:
 	TSharedPtr<FTableTreeNodeTextFilter> CurrentAsyncOpTextFilter;
 	FFilterConfigurator* CurrentAsyncOpFilterConfigurator = nullptr;
 	std::atomic<bool> bCancelCurrentAsyncOp { false };
+	FGraphEventRef DispatchEvent;
 
 	//////////////////////////////////////////////////
 	TSharedPtr<FFilterConfigurator> FilterConfigurator;
