[Insights] Remove memory insights from the editor menu

diff --git a/Engine/Source/Developer/TraceInsights/Private/Insights/MemoryProfiler/MemoryProfilerManager.cpp b/Engine/Source/Developer/TraceInsights/Private/Insights/MemoryProfiler/MemoryProfilerManager.cpp
--- a/Engine/Source/Developer/TraceInsights/Private/Insights/MemoryProfiler/MemoryProfilerManager.cpp
+++ b/Engine/Source/Developer/TraceInsights/Private/Insights/MemoryProfiler/MemoryProfilerManager.cpp
@@ -138,6 +138,7 @@ void FMemoryProfilerManager::BindCommands()
 
 void FMemoryProfilerManager::RegisterMajorTabs(IUnrealInsightsModule& InsightsModule)
 {
+#if !WITH_EDITOR
 	const FInsightsMajorTabConfig& Config = InsightsModule.FindMajorTabConfig(FInsightsManagerTabs::MemoryProfilerTabId);
 
 	if (Config.bIsAvailable)
@@ -152,6 +153,7 @@ void FMemoryProfilerManager::RegisterMajorTabs(IUnrealInsightsModule& InsightsMo
 		TSharedRef<FWorkspaceItem> Group = Config.WorkspaceGroup.IsValid() ? Config.WorkspaceGroup.ToSharedRef() : FInsightsManager::Get()->GetInsightsMenuBuilder()->GetInsightsToolsGroup();
 		TabSpawnerEntry.SetGroup(Group);
 	}
+#endif
 }
 
 ////////////////////////////////////////////////////////////////////////////////////////////////////
