[Insights] Support for TableTreeView columns with dynamic values
Make the backtrace column from AllocTreeTableView dynamic.

diff --git a/Engine/Source/Developer/TraceInsights/Private/Insights/MemoryProfiler/ViewModels/MemAllocTable.cpp b/Engine/Source/Developer/TraceInsights/Private/Insights/MemoryProfiler/ViewModels/MemAllocTable.cpp
--- a/Engine/Source/Developer/TraceInsights/Private/Insights/MemoryProfiler/ViewModels/MemAllocTable.cpp
+++ b/Engine/Source/Developer/TraceInsights/Private/Insights/MemoryProfiler/ViewModels/MemAllocTable.cpp
@@ -365,6 +365,8 @@ void FMemAllocTable::AddDefaultColumns()
 
 		Column.SetDataType(ETableCellDataType::CString);
 
+		Column.SetIsDynamic(true);
+
 		class FBacktraceValueGetter : public FTableCellValueGetter
 		{
 		public:
diff --git a/Engine/Source/Developer/TraceInsights/Private/Insights/Table/ViewModels/TableColumn.h b/Engine/Source/Developer/TraceInsights/Private/Insights/Table/ViewModels/TableColumn.h
--- a/Engine/Source/Developer/TraceInsights/Private/Insights/Table/ViewModels/TableColumn.h
+++ b/Engine/Source/Developer/TraceInsights/Private/Insights/Table/ViewModels/TableColumn.h
@@ -61,6 +61,7 @@ public:
 		, TitleName()
 		, Description()
 		, bIsVisible(false)
+		, bIsDynamic(false)
 		, Flags(ETableColumnFlags::None)
 		, HorizontalAlignment(HAlign_Left)
 		, InitialWidth(60.0f)
@@ -186,6 +187,11 @@ public:
 	TWeakPtr<FTable> GetParentTable() const { return ParentTable; }
 	void SetParentTable(TWeakPtr<FTable> InParentTable) { ParentTable = InParentTable; }
 
+	//////////////////////////////////////////////////
+
+	void SetIsDynamic(bool InValue) { bIsDynamic = InValue; }
+	bool IsDynamic() const { return bIsDynamic; }
+
 private:
 	/** Id of the column. */
 	FName Id;
@@ -205,6 +211,9 @@ private:
 	/** Is this column visible? */
 	bool bIsVisible;
 
+	/** If it is not dynamic, the column's value can be cached. */
+	bool bIsDynamic;
+
 	/** Other on/off switches. */
 	ETableColumnFlags Flags;
 
diff --git a/Engine/Source/Developer/TraceInsights/Private/Insights/Table/Widgets/STableTreeViewCell.cpp b/Engine/Source/Developer/TraceInsights/Private/Insights/Table/Widgets/STableTreeViewCell.cpp
--- a/Engine/Source/Developer/TraceInsights/Private/Insights/Table/Widgets/STableTreeViewCell.cpp
+++ b/Engine/Source/Developer/TraceInsights/Private/Insights/Table/Widgets/STableTreeViewCell.cpp
@@ -121,24 +121,35 @@ FText STableTreeViewCell::GetValueAsText() const
 
 TSharedRef<SWidget> STableTreeViewCell::GenerateWidgetForTableColumn(const FArguments& InArgs, const TSharedRef<ITableRow>& TableRow)
 {
-	// Note: For performance reason, init the cell text (value) only once.
-	//       If we'll need to update values without recreating the table row/cell widgets, bind .Text to STableTreeViewCell::GetValueAsText
-	//       or add API to explicitly update the text block.
-	FText CellText;
-	if (TableTreeNodePtr->IsGroup())
+	TSharedRef<STextBlock> TextBox = SNew(STextBlock)
+		.TextStyle(FEditorStyle::Get(), TEXT("Profiler.Tooltip"))
+		.ColorAndOpacity(this, &STableTreeViewCell::GetStatsColorAndOpacity)
+		.ShadowColorAndOpacity(this, &STableTreeViewCell::GetShadowColorAndOpacity);
+
+	if (ColumnPtr->IsDynamic())
+	{
+		TextBox->SetText(TAttribute<FText>::Create(TAttribute<FText>::FGetter::CreateRaw(this, &STableTreeViewCell::GetValueAsText)));
+	}
+	else
 	{
-		if (TableTreeNodePtr->HasAggregatedValue(ColumnPtr->GetId()))
+		FText CellText;
+		if (TableTreeNodePtr->IsGroup())
 		{
-			CellText = ColumnPtr->GetValueFormatter()->FormatValue(TableTreeNodePtr->GetAggregatedValue(ColumnPtr->GetId()));
+			if (TableTreeNodePtr->HasAggregatedValue(ColumnPtr->GetId()))
+			{
+				CellText = ColumnPtr->GetValueFormatter()->FormatValue(TableTreeNodePtr->GetAggregatedValue(ColumnPtr->GetId()));
+			}
+			else
+			{
+				CellText = FText::FromString(TEXT("-"));
+			}
 		}
 		else
 		{
-			CellText = FText::FromString(TEXT("-"));
+			CellText = ColumnPtr->GetValueAsText(*TableTreeNodePtr);
 		}
-	}
-	else
-	{
-		CellText = ColumnPtr->GetValueAsText(*TableTreeNodePtr);
+
+		TextBox->SetText(CellText);
 	}
 
 	return
@@ -151,11 +162,7 @@ TSharedRef<SWidget> STableTreeViewCell::GenerateWidgetForTableColumn(const FArgu
 		.HAlign(ColumnPtr->GetHorizontalAlignment())
 		.Padding(FMargin(2.0f, 0.0f))
 		[
-			SNew(STextBlock)
-			.Text(CellText)
-			.TextStyle(FEditorStyle::Get(), TEXT("Profiler.Tooltip"))
-			.ColorAndOpacity(this, &STableTreeViewCell::GetStatsColorAndOpacity)
-			.ShadowColorAndOpacity(this, &STableTreeViewCell::GetShadowColorAndOpacity)
+			TextBox
 		]
 	;
 }
