Fixing PVS 7.7 Issues:

warning V1047: Lifetime of the lambda is greater than lifetime of the local variable 'Column' captured by reference.

diff --git a/Engine/Source/Developer/TraceInsights/Private/Insights/Table/ViewModels/TableCellValueSorter.cpp b/Engine/Source/Developer/TraceInsights/Private/Insights/Table/ViewModels/TableCellValueSorter.cpp
--- a/Engine/Source/Developer/TraceInsights/Private/Insights/Table/ViewModels/TableCellValueSorter.cpp
+++ b/Engine/Source/Developer/TraceInsights/Private/Insights/Table/ViewModels/TableCellValueSorter.cpp
@@ -140,11 +140,9 @@ FSorterByTypeName::FSorterByTypeName(TSharedRef<FTableColumn> InColumnRef)
 FSorterByBoolValue::FSorterByBoolValue(TSharedRef<FTableColumn> InColumnRef)
 	: FBaseTableColumnSorter(InColumnRef)
 {
-	const FTableColumn& Column = *InColumnRef;
+	ensure(InColumnRef->GetDataType() == ETableCellDataType::Bool);
 
-	ensure(Column.GetDataType() == ETableCellDataType::Bool);
-
-	AscendingCompareDelegate = [&Column](const FBaseTreeNodePtr& A, const FBaseTreeNodePtr& B) -> bool
+	AscendingCompareDelegate = [&Column = *InColumnRef](const FBaseTreeNodePtr& A, const FBaseTreeNodePtr& B) -> bool
 	{
 		const TOptional<FTableCellValue> OptionalValueA = Column.GetValue(*A);
 		const TOptional<FTableCellValue> OptionalValueB = Column.GetValue(*B);
@@ -163,7 +161,7 @@ FSorterByBoolValue::FSorterByBoolValue(TSharedRef<FTableColumn> InColumnRef)
 		}
 	};
 
-	DescendingCompareDelegate = [&Column](const FBaseTreeNodePtr& A, const FBaseTreeNodePtr& B) -> bool
+	DescendingCompareDelegate = [&Column = * InColumnRef](const FBaseTreeNodePtr& A, const FBaseTreeNodePtr& B) -> bool
 	{
 		const TOptional<FTableCellValue> OptionalValueA = Column.GetValue(*A);
 		const TOptional<FTableCellValue> OptionalValueB = Column.GetValue(*B);
@@ -190,11 +188,9 @@ FSorterByBoolValue::FSorterByBoolValue(TSharedRef<FTableColumn> InColumnRef)
 FSorterByInt64Value::FSorterByInt64Value(TSharedRef<FTableColumn> InColumnRef)
 	: FBaseTableColumnSorter(InColumnRef)
 {
-	const FTableColumn& Column = *InColumnRef;
-
-	ensure(Column.GetDataType() == ETableCellDataType::Int64);
+	ensure(InColumnRef->GetDataType() == ETableCellDataType::Int64);
 
-	AscendingCompareDelegate = [&Column](const FBaseTreeNodePtr& A, const FBaseTreeNodePtr& B) -> bool
+	AscendingCompareDelegate = [&Column = *InColumnRef](const FBaseTreeNodePtr& A, const FBaseTreeNodePtr& B) -> bool
 	{
 		const TOptional<FTableCellValue> OptionalValueA = Column.GetValue(*A);
 		const TOptional<FTableCellValue> OptionalValueB = Column.GetValue(*B);
@@ -213,7 +209,7 @@ FSorterByInt64Value::FSorterByInt64Value(TSharedRef<FTableColumn> InColumnRef)
 		}
 	};
 
-	DescendingCompareDelegate = [&Column](const FBaseTreeNodePtr& A, const FBaseTreeNodePtr& B) -> bool
+	DescendingCompareDelegate = [&Column = *InColumnRef](const FBaseTreeNodePtr& A, const FBaseTreeNodePtr& B) -> bool
 	{
 		const TOptional<FTableCellValue> OptionalValueA = Column.GetValue(*A);
 		const TOptional<FTableCellValue> OptionalValueB = Column.GetValue(*B);
@@ -305,11 +301,9 @@ void FSorterByInt64Value::Sort(TArray<FBaseTreeNodePtr>& NodesToSort, ESortMode
 FSorterByFloatValue::FSorterByFloatValue(TSharedRef<FTableColumn> InColumnRef)
 	: FBaseTableColumnSorter(InColumnRef)
 {
-	const FTableColumn& Column = *InColumnRef;
-
-	ensure(Column.GetDataType() == ETableCellDataType::Float);
+	ensure(InColumnRef->GetDataType() == ETableCellDataType::Float);
 
-	AscendingCompareDelegate = [&Column](const FBaseTreeNodePtr& A, const FBaseTreeNodePtr& B) -> bool
+	AscendingCompareDelegate = [&Column = *InColumnRef](const FBaseTreeNodePtr& A, const FBaseTreeNodePtr& B) -> bool
 	{
 		const TOptional<FTableCellValue> OptionalValueA = Column.GetValue(*A);
 		const TOptional<FTableCellValue> OptionalValueB = Column.GetValue(*B);
@@ -328,7 +322,7 @@ FSorterByFloatValue::FSorterByFloatValue(TSharedRef<FTableColumn> InColumnRef)
 		}
 	};
 
-	DescendingCompareDelegate = [&Column](const FBaseTreeNodePtr& A, const FBaseTreeNodePtr& B) -> bool
+	DescendingCompareDelegate = [&Column = *InColumnRef](const FBaseTreeNodePtr& A, const FBaseTreeNodePtr& B) -> bool
 	{
 		const TOptional<FTableCellValue> OptionalValueA = Column.GetValue(*A);
 		const TOptional<FTableCellValue> OptionalValueB = Column.GetValue(*B);
@@ -420,11 +414,9 @@ void FSorterByFloatValue::Sort(TArray<FBaseTreeNodePtr>& NodesToSort, ESortMode
 FSorterByDoubleValue::FSorterByDoubleValue(TSharedRef<FTableColumn> InColumnRef)
 	: FBaseTableColumnSorter(InColumnRef)
 {
-	const FTableColumn& Column = *InColumnRef;
+	ensure(InColumnRef->GetDataType() == ETableCellDataType::Double);
 
-	ensure(Column.GetDataType() == ETableCellDataType::Double);
-
-	AscendingCompareDelegate = [&Column](const FBaseTreeNodePtr& A, const FBaseTreeNodePtr& B) -> bool
+	AscendingCompareDelegate = [&Column = *InColumnRef](const FBaseTreeNodePtr& A, const FBaseTreeNodePtr& B) -> bool
 	{
 		const TOptional<FTableCellValue> OptionalValueA = Column.GetValue(*A);
 		const TOptional<FTableCellValue> OptionalValueB = Column.GetValue(*B);
@@ -443,7 +435,7 @@ FSorterByDoubleValue::FSorterByDoubleValue(TSharedRef<FTableColumn> InColumnRef)
 		}
 	};
 
-	DescendingCompareDelegate = [&Column](const FBaseTreeNodePtr& A, const FBaseTreeNodePtr& B) -> bool
+	DescendingCompareDelegate = [&Column = *InColumnRef](const FBaseTreeNodePtr& A, const FBaseTreeNodePtr& B) -> bool
 	{
 		const TOptional<FTableCellValue> OptionalValueA = Column.GetValue(*A);
 		const TOptional<FTableCellValue> OptionalValueB = Column.GetValue(*B);
@@ -535,11 +527,9 @@ void FSorterByDoubleValue::Sort(TArray<FBaseTreeNodePtr>& NodesToSort, ESortMode
 FSorterByCStringValue::FSorterByCStringValue(TSharedRef<FTableColumn> InColumnRef)
 	: FBaseTableColumnSorter(InColumnRef)
 {
-	const FTableColumn& Column = *InColumnRef;
-
-	ensure(Column.GetDataType() == ETableCellDataType::CString);
+	ensure(InColumnRef->GetDataType() == ETableCellDataType::CString);
 
-	AscendingCompareDelegate = [&Column](const FBaseTreeNodePtr& A, const FBaseTreeNodePtr& B) -> bool
+	AscendingCompareDelegate = [&Column = *InColumnRef](const FBaseTreeNodePtr& A, const FBaseTreeNodePtr& B) -> bool
 	{
 		const TOptional<FTableCellValue> OptionalValueA = Column.GetValue(*A);
 		const TOptional<FTableCellValue> OptionalValueB = Column.GetValue(*B);
@@ -558,7 +548,7 @@ FSorterByCStringValue::FSorterByCStringValue(TSharedRef<FTableColumn> InColumnRe
 		}
 	};
 
-	DescendingCompareDelegate = [&Column](const FBaseTreeNodePtr& A, const FBaseTreeNodePtr& B) -> bool
+	DescendingCompareDelegate = [&Column = *InColumnRef](const FBaseTreeNodePtr& A, const FBaseTreeNodePtr& B) -> bool
 	{
 		const TOptional<FTableCellValue> OptionalValueA = Column.GetValue(*A);
 		const TOptional<FTableCellValue> OptionalValueB = Column.GetValue(*B);
@@ -585,11 +575,9 @@ FSorterByCStringValue::FSorterByCStringValue(TSharedRef<FTableColumn> InColumnRe
 FSorterByTextValue::FSorterByTextValue(TSharedRef<FTableColumn> InColumnRef)
 	: FBaseTableColumnSorter(InColumnRef)
 {
-	const FTableColumn& Column = *InColumnRef;
-
-	ensure(Column.GetDataType() == ETableCellDataType::Text);
+	ensure(InColumnRef->GetDataType() == ETableCellDataType::Text);
 
-	AscendingCompareDelegate = [&Column](const FBaseTreeNodePtr& A, const FBaseTreeNodePtr& B) -> bool
+	AscendingCompareDelegate = [&Column = *InColumnRef](const FBaseTreeNodePtr& A, const FBaseTreeNodePtr& B) -> bool
 	{
 		const TOptional<FTableCellValue> OptionalValueA = Column.GetValue(*A);
 		const TOptional<FTableCellValue> OptionalValueB = Column.GetValue(*B);
@@ -619,7 +607,7 @@ FSorterByTextValue::FSorterByTextValue(TSharedRef<FTableColumn> InColumnRef)
 		}
 	};
 
-	DescendingCompareDelegate = [&Column](const FBaseTreeNodePtr& A, const FBaseTreeNodePtr& B) -> bool
+	DescendingCompareDelegate = [&Column = *InColumnRef](const FBaseTreeNodePtr& A, const FBaseTreeNodePtr& B) -> bool
 	{
 		const TOptional<FTableCellValue> OptionalValueA = Column.GetValue(*A);
 		const TOptional<FTableCellValue> OptionalValueB = Column.GetValue(*B);
