[Insights] Sets empty label to "-" instead

#rnx

diff --git a/Engine/Source/Developer/TraceInsights/Private/Insights/MemoryProfiler/ViewModels/MemAllocTable.cpp b/Engine/Source/Developer/TraceInsights/Private/Insights/MemoryProfiler/ViewModels/MemAllocTable.cpp
--- a/Engine/Source/Developer/TraceInsights/Private/Insights/MemoryProfiler/ViewModels/MemAllocTable.cpp
+++ b/Engine/Source/Developer/TraceInsights/Private/Insights/MemoryProfiler/ViewModels/MemAllocTable.cpp
@@ -375,7 +375,7 @@ void FMemAllocTable::AddDefaultColumns()
 				static const TCHAR* DisplayStrings[] = {
 					TEXT("Pending..."),
 					TEXT("Not found"),
-					TEXT("")
+					TEXT("-")
 				};
 				const TCHAR* Value = DisplayStrings[0];
 				const FMemAllocNode& MemAllocNode = static_cast<const FMemAllocNode&>(Node);
