[Insights] Display allocations missing callstack as "Not found"

#rnx

diff --git a/Engine/Source/Developer/TraceInsights/Private/Insights/MemoryProfiler/ViewModels/MemAllocTable.cpp b/Engine/Source/Developer/TraceInsights/Private/Insights/MemoryProfiler/ViewModels/MemAllocTable.cpp
--- a/Engine/Source/Developer/TraceInsights/Private/Insights/MemoryProfiler/ViewModels/MemAllocTable.cpp
+++ b/Engine/Source/Developer/TraceInsights/Private/Insights/MemoryProfiler/ViewModels/MemAllocTable.cpp
@@ -386,7 +386,7 @@ void FMemAllocTable::AddDefaultColumns()
 
 					if (Callstack)
 					{
-						const TraceServices::FStackFrame* Frame = Callstack->Frame(2);
+						const TraceServices::FStackFrame* Frame = Callstack->Frame(FMath::Min(2u, Callstack->Num()));
 						const TraceServices::QueryResult Result = Frame->Symbol->Result.load(std::memory_order_acquire);
 						switch (Result)
 						{
@@ -401,6 +401,10 @@ void FMemAllocTable::AddDefaultColumns()
 							break;
 						}
 					}
+					else
+					{
+						Value = DisplayStrings[1];
+					}
 				}
 				else
 				{
