[Insights] Fix aggregatated values in MemAllocView #rb Ionut.Matasaru

diff --git a/Engine/Source/Developer/TraceInsights/Private/Insights/MemoryProfiler/ViewModels/MemAllocTable.cpp b/Engine/Source/Developer/TraceInsights/Private/Insights/MemoryProfiler/ViewModels/MemAllocTable.cpp
--- a/Engine/Source/Developer/TraceInsights/Private/Insights/MemoryProfiler/ViewModels/MemAllocTable.cpp
+++ b/Engine/Source/Developer/TraceInsights/Private/Insights/MemoryProfiler/ViewModels/MemAllocTable.cpp
@@ -81,10 +81,25 @@ void FMemAllocTable::AddDefaultColumns()
 		public:
 			virtual const TOptional<FTableCellValue> GetValue(const FTableColumn& Column, const FBaseTreeNode& Node) const override
 			{
-				const FMemAllocNode& MemAllocNode = static_cast<const FMemAllocNode&>(Node);
-				const FMemoryAlloc* Alloc = MemAllocNode.GetMemAlloc();
-				const double AllocStartTime = Alloc ? Alloc->GetStartTime() : 0.0;
-				return TOptional<FTableCellValue>(FTableCellValue(AllocStartTime));
+				if (Node.IsGroup())
+				{
+					const FTableTreeNode& NodePtr = static_cast<const FTableTreeNode&>(Node);
+					if (NodePtr.HasAggregatedValue(Column.GetId()))
+					{
+						return NodePtr.GetAggregatedValue(Column.GetId());
+					}
+					else
+					{
+						return FTableCellValue(TEXT("-"));
+					}
+				}
+				else
+				{
+					const FMemAllocNode& MemAllocNode = static_cast<const FMemAllocNode&>(Node);
+					const FMemoryAlloc* Alloc = MemAllocNode.GetMemAlloc();
+					const double AllocStartTime = Alloc ? Alloc->GetStartTime() : 0.0;
+					return TOptional<FTableCellValue>(FTableCellValue(AllocStartTime));
+				}
 			}
 		};
 		TSharedRef<ITableCellValueGetter> Getter = MakeShared<FMemAllocStartTimeValueGetter>();
@@ -124,10 +139,25 @@ void FMemAllocTable::AddDefaultColumns()
 		public:
 			virtual const TOptional<FTableCellValue> GetValue(const FTableColumn& Column, const FBaseTreeNode& Node) const override
 			{
-				const FMemAllocNode& MemAllocNode = static_cast<const FMemAllocNode&>(Node);
-				const FMemoryAlloc* Alloc = MemAllocNode.GetMemAlloc();
-				const double AllocEndTime = Alloc ? Alloc->GetEndTime() : 0.0;
-				return TOptional<FTableCellValue>(FTableCellValue(AllocEndTime));
+				if (Node.IsGroup())
+				{
+					const FTableTreeNode& NodePtr = static_cast<const FTableTreeNode&>(Node);
+					if (NodePtr.HasAggregatedValue(Column.GetId()))
+					{
+						return NodePtr.GetAggregatedValue(Column.GetId());
+					}
+					else
+					{
+						return FTableCellValue(TEXT("-"));
+					}
+				}
+				else
+				{
+					const FMemAllocNode& MemAllocNode = static_cast<const FMemAllocNode&>(Node);
+					const FMemoryAlloc* Alloc = MemAllocNode.GetMemAlloc();
+					const double AllocEndTime = Alloc ? Alloc->GetEndTime() : 0.0;
+					return TOptional<FTableCellValue>(FTableCellValue(AllocEndTime));
+				}
 			}
 		};
 		TSharedRef<ITableCellValueGetter> Getter = MakeShared<FMemAllocStartTimeValueGetter>();
@@ -249,10 +279,25 @@ void FMemAllocTable::AddDefaultColumns()
 		public:
 			virtual const TOptional<FTableCellValue> GetValue(const FTableColumn& Column, const FBaseTreeNode& Node) const override
 			{
-				const FMemAllocNode& MemAllocNode = static_cast<const FMemAllocNode&>(Node);
-				const FMemoryAlloc* Alloc = MemAllocNode.GetMemAlloc();
-				const uint64 AllocSize = Alloc ? Alloc->GetSize() : 0;
-				return TOptional<FTableCellValue>(FTableCellValue(static_cast<int64>(AllocSize)));
+				if (Node.IsGroup())
+				{
+					const FTableTreeNode& NodePtr = static_cast<const FTableTreeNode&>(Node);
+					if (NodePtr.HasAggregatedValue(Column.GetId()))
+					{
+						return NodePtr.GetAggregatedValue(Column.GetId());
+					}
+					else
+					{
+						return FTableCellValue(TEXT("-"));
+					}
+				}
+				else
+				{
+					const FMemAllocNode& MemAllocNode = static_cast<const FMemAllocNode&>(Node);
+					const FMemoryAlloc* Alloc = MemAllocNode.GetMemAlloc();
+					const uint64 AllocSize = Alloc ? Alloc->GetSize() : 0;
+					return TOptional<FTableCellValue>(FTableCellValue(static_cast<int64>(AllocSize)));
+				}
 			}
 		};
 		TSharedRef<ITableCellValueGetter> Getter = MakeShared<FMemAllocSizeValueGetter>();
diff --git a/Engine/Source/Developer/TraceInsights/Private/Insights/Table/Widgets/STableTreeViewCell.cpp b/Engine/Source/Developer/TraceInsights/Private/Insights/Table/Widgets/STableTreeViewCell.cpp
--- a/Engine/Source/Developer/TraceInsights/Private/Insights/Table/Widgets/STableTreeViewCell.cpp
+++ b/Engine/Source/Developer/TraceInsights/Private/Insights/Table/Widgets/STableTreeViewCell.cpp
@@ -132,23 +132,7 @@ TSharedRef<SWidget> STableTreeViewCell::GenerateWidgetForTableColumn(const FArgu
 	}
 	else
 	{
-		FText CellText;
-		if (TableTreeNodePtr->IsGroup())
-		{
-			if (TableTreeNodePtr->HasAggregatedValue(ColumnPtr->GetId()))
-			{
-				CellText = ColumnPtr->GetValueFormatter()->FormatValue(TableTreeNodePtr->GetAggregatedValue(ColumnPtr->GetId()));
-			}
-			else
-			{
-				CellText = FText::FromString(TEXT("-"));
-			}
-		}
-		else
-		{
-			CellText = ColumnPtr->GetValueAsText(*TableTreeNodePtr);
-		}
-
+		FText CellText = ColumnPtr->GetValueAsText(*TableTreeNodePtr);
 		TextBox->SetText(CellText);
 	}
 
