[Insights] Cap lookup to Num - 1.

#rnx

diff --git a/Engine/Source/Developer/TraceInsights/Private/Insights/MemoryProfiler/ViewModels/MemAllocTable.cpp b/Engine/Source/Developer/TraceInsights/Private/Insights/MemoryProfiler/ViewModels/MemAllocTable.cpp
--- a/Engine/Source/Developer/TraceInsights/Private/Insights/MemoryProfiler/ViewModels/MemAllocTable.cpp
+++ b/Engine/Source/Developer/TraceInsights/Private/Insights/MemoryProfiler/ViewModels/MemAllocTable.cpp
@@ -431,7 +431,7 @@ void FMemAllocTable::AddDefaultColumns()
 
 					if (Callstack)
 					{
-						const TraceServices::FStackFrame* Frame = Callstack->Frame(FMath::Min(2u, Callstack->Num()));
+						const TraceServices::FStackFrame* Frame = Callstack->Frame(FMath::Min(2u, Callstack->Num()-1));
 						const TraceServices::QueryResult Result = Frame->Symbol->Result.load(std::memory_order_acquire);
 						switch (Result)
 						{
