[Insights]
  - Tree Views: Improved speed of "Unique Value" groupings for all data types.
  - Memory: Fixed LLM Tag column (names and gropuping) in Mem Alloc Tables.
  - Memory: Fixed aggregation (for group nodes) for all columns in Mem Alloc Tables.
  - Made FTableCellValue constructors explicit.

diff --git a/Engine/Source/Developer/TraceInsights/Private/Insights/MemoryProfiler/ViewModels/MemAllocNode.cpp b/Engine/Source/Developer/TraceInsights/Private/Insights/MemoryProfiler/ViewModels/MemAllocNode.cpp
--- a/Engine/Source/Developer/TraceInsights/Private/Insights/MemoryProfiler/ViewModels/MemAllocNode.cpp
+++ b/Engine/Source/Developer/TraceInsights/Private/Insights/MemoryProfiler/ViewModels/MemAllocNode.cpp
@@ -2,11 +2,6 @@
 
 #include "MemAllocNode.h"
 
-// Insights
-#include "Insights/MemoryProfiler/MemoryProfilerManager.h"
-#include "Insights/MemoryProfiler/ViewModels/MemorySharedState.h"
-#include "Insights/MemoryProfiler/Widgets/SMemoryProfilerWindow.h"
-
 #define LOCTEXT_NAMESPACE "MemAllocNode"
 
 namespace Insights
@@ -16,37 +11,6 @@ const FName FMemAllocNode::TypeName(TEXT("FMemAllocNode"));
 
 ////////////////////////////////////////////////////////////////////////////////////////////////////
 
-FText FMemAllocNode::GetAllocMemTagAsText() const
-{
-	const FMemoryAlloc* Alloc = GetMemAlloc();
-	if (Alloc)
-	{
-		return FText::FromString(GetAllocMemTagAsString(Alloc->GetMemTag()));
-	}
-	return FText::GetEmpty();
-}
-
-////////////////////////////////////////////////////////////////////////////////////////////////////
-
-FString FMemAllocNode::GetAllocMemTagAsString(FMemoryTagId InMemTag)
-{
-	TSharedPtr<SMemoryProfilerWindow> ProfilerWindow = FMemoryProfilerManager::Get()->GetProfilerWindow();
-	if (ProfilerWindow)
-	{
-		auto& SharedState = ProfilerWindow->GetSharedState();
-		const Insights::FMemoryTagList& TagList = SharedState.GetTagList();
-
-		const FMemoryTag* Tag = TagList.GetTagById(InMemTag);
-		if (Tag)
-		{
-			return Tag->GetStatName();
-		}
-	}
-	return FString();
-}
-
-////////////////////////////////////////////////////////////////////////////////////////////////////
-
 } // namespace Insights
 
 #undef LOCTEXT_NAMESPACE
diff --git a/Engine/Source/Developer/TraceInsights/Private/Insights/MemoryProfiler/ViewModels/MemAllocNode.h b/Engine/Source/Developer/TraceInsights/Private/Insights/MemoryProfiler/ViewModels/MemAllocNode.h
--- a/Engine/Source/Developer/TraceInsights/Private/Insights/MemoryProfiler/ViewModels/MemAllocNode.h
+++ b/Engine/Source/Developer/TraceInsights/Private/Insights/MemoryProfiler/ViewModels/MemAllocNode.h
@@ -86,9 +86,6 @@ public:
 	const FMemoryAlloc* GetMemAlloc() const { return GetMemTableChecked().GetMemAlloc(GetRowIndex()); }
 	const FMemoryAlloc& GetMemAllocChecked() const { return GetMemTableChecked().GetMemAllocChecked(GetRowIndex()); }
 
-	FText GetAllocMemTagAsText() const;
-	static FString GetAllocMemTagAsString(FMemoryTagId InMemTag);
-
 private:
 	const EMemAllocNodeType Type;
 };
diff --git a/Engine/Source/Developer/TraceInsights/Private/Insights/MemoryProfiler/ViewModels/MemAllocTable.cpp b/Engine/Source/Developer/TraceInsights/Private/Insights/MemoryProfiler/ViewModels/MemAllocTable.cpp
--- a/Engine/Source/Developer/TraceInsights/Private/Insights/MemoryProfiler/ViewModels/MemAllocTable.cpp
+++ b/Engine/Source/Developer/TraceInsights/Private/Insights/MemoryProfiler/ViewModels/MemAllocTable.cpp
@@ -88,18 +88,18 @@ void FMemAllocTable::AddDefaultColumns()
 					{
 						return NodePtr.GetAggregatedValue(Column.GetId());
 					}
-					else
-					{
-						return FTableCellValue(TEXT("-"));
-					}
 				}
-				else
+				else //if (Node->Is<FMemAllocNode>())
 				{
 					const FMemAllocNode& MemAllocNode = static_cast<const FMemAllocNode&>(Node);
 					const FMemoryAlloc* Alloc = MemAllocNode.GetMemAlloc();
-					const double AllocStartTime = Alloc ? Alloc->GetStartTime() : 0.0;
-					return TOptional<FTableCellValue>(FTableCellValue(AllocStartTime));
+					if (Alloc)
+					{
+						return FTableCellValue(Alloc->GetStartTime());
+					}
 				}
+
+				return TOptional<FTableCellValue>();
 			}
 		};
 		TSharedRef<ITableCellValueGetter> Getter = MakeShared<FMemAllocStartTimeValueGetter>();
@@ -146,18 +146,18 @@ void FMemAllocTable::AddDefaultColumns()
 					{
 						return NodePtr.GetAggregatedValue(Column.GetId());
 					}
-					else
-					{
-						return FTableCellValue(TEXT("-"));
-					}
 				}
-				else
+				else //if (Node->Is<FMemAllocNode>())
 				{
 					const FMemAllocNode& MemAllocNode = static_cast<const FMemAllocNode&>(Node);
 					const FMemoryAlloc* Alloc = MemAllocNode.GetMemAlloc();
-					const double AllocEndTime = Alloc ? Alloc->GetEndTime() : 0.0;
-					return TOptional<FTableCellValue>(FTableCellValue(AllocEndTime));
+					if (Alloc)
+					{
+						return FTableCellValue(Alloc->GetEndTime());
+					}
 				}
+
+				return TOptional<FTableCellValue>();
 			}
 		};
 		TSharedRef<ITableCellValueGetter> Getter = MakeShared<FMemAllocStartTimeValueGetter>();
@@ -197,10 +197,25 @@ void FMemAllocTable::AddDefaultColumns()
 		public:
 			virtual const TOptional<FTableCellValue> GetValue(const FTableColumn& Column, const FBaseTreeNode& Node) const override
 			{
-				const FMemAllocNode& MemAllocNode = static_cast<const FMemAllocNode&>(Node);
-				const FMemoryAlloc* Alloc = MemAllocNode.GetMemAlloc();
-				const double AllocDuration = Alloc ? Alloc->GetDuration() : 0.0;
-				return TOptional<FTableCellValue>(FTableCellValue(AllocDuration));
+				if (Node.IsGroup())
+				{
+					const FTableTreeNode& NodePtr = static_cast<const FTableTreeNode&>(Node);
+					if (NodePtr.HasAggregatedValue(Column.GetId()))
+					{
+						return NodePtr.GetAggregatedValue(Column.GetId());
+					}
+				}
+				else //if (Node->Is<FMemAllocNode>())
+				{
+					const FMemAllocNode& MemAllocNode = static_cast<const FMemAllocNode&>(Node);
+					const FMemoryAlloc* Alloc = MemAllocNode.GetMemAlloc();
+					if (Alloc)
+					{
+						return FTableCellValue(Alloc->GetDuration());
+					}
+				}
+
+				return TOptional<FTableCellValue>();
 			}
 		};
 		TSharedRef<ITableCellValueGetter> Getter = MakeShared<FMemAllocStartTimeValueGetter>();
@@ -238,10 +253,25 @@ void FMemAllocTable::AddDefaultColumns()
 		public:
 			virtual const TOptional<FTableCellValue> GetValue(const FTableColumn& Column, const FBaseTreeNode& Node) const override
 			{
-				const FMemAllocNode& MemAllocNode = static_cast<const FMemAllocNode&>(Node);
-				const FMemoryAlloc* Alloc = MemAllocNode.GetMemAlloc();
-				const uint64 AllocAddress = Alloc ? Alloc->GetAddress() : 0;
-				return TOptional<FTableCellValue>(FTableCellValue(static_cast<int64>(AllocAddress)));
+				if (Node.IsGroup())
+				{
+					const FTableTreeNode& NodePtr = static_cast<const FTableTreeNode&>(Node);
+					if (NodePtr.HasAggregatedValue(Column.GetId()))
+					{
+						return NodePtr.GetAggregatedValue(Column.GetId());
+					}
+				}
+				else //if (Node->Is<FMemAllocNode>())
+				{
+					const FMemAllocNode& MemAllocNode = static_cast<const FMemAllocNode&>(Node);
+					const FMemoryAlloc* Alloc = MemAllocNode.GetMemAlloc();
+					if (Alloc)
+					{
+						return FTableCellValue(static_cast<int64>(Alloc->GetAddress()));
+					}
+				}
+
+				return TOptional<FTableCellValue>();
 			}
 		};
 		TSharedRef<ITableCellValueGetter> Getter = MakeShared<FMemAllocAddressValueGetter>();
@@ -286,18 +316,18 @@ void FMemAllocTable::AddDefaultColumns()
 					{
 						return NodePtr.GetAggregatedValue(Column.GetId());
 					}
-					else
-					{
-						return FTableCellValue(TEXT("-"));
-					}
 				}
-				else
+				else //if (Node->Is<FMemAllocNode>())
 				{
 					const FMemAllocNode& MemAllocNode = static_cast<const FMemAllocNode&>(Node);
 					const FMemoryAlloc* Alloc = MemAllocNode.GetMemAlloc();
-					const uint64 AllocSize = Alloc ? Alloc->GetSize() : 0;
-					return TOptional<FTableCellValue>(FTableCellValue(static_cast<int64>(AllocSize)));
+					if (Alloc)
+					{
+						return FTableCellValue(static_cast<int64>(Alloc->GetSize()));
+					}
 				}
+
+				return TOptional<FTableCellValue>();
 			}
 		};
 		TSharedRef<ITableCellValueGetter> Getter = MakeShared<FMemAllocSizeValueGetter>();
@@ -330,63 +360,41 @@ void FMemAllocTable::AddDefaultColumns()
 		Column.SetHorizontalAlignment(HAlign_Left);
 		Column.SetInitialWidth(100.0f);
 
-		Column.SetDataType(ETableCellDataType::Int64);
+		Column.SetDataType(ETableCellDataType::CString);
 
 		class FMemAllocLlmTagValueGetter : public FTableCellValueGetter
 		{
 		public:
 			virtual const TOptional<FTableCellValue> GetValue(const FTableColumn& Column, const FBaseTreeNode& Node) const
 			{
-				//if (Node->Is<FMemAllocNode>())
+				if (Node.IsGroup())
+				{
+					const FTableTreeNode& NodePtr = static_cast<const FTableTreeNode&>(Node);
+					if (NodePtr.HasAggregatedValue(Column.GetId()))
+					{
+						return NodePtr.GetAggregatedValue(Column.GetId());
+					}
+				}
+				else //if (Node->Is<FMemAllocNode>())
 				{
 					const FMemAllocNode& MemAllocNode = static_cast<const FMemAllocNode&>(Node);
 					const FMemoryAlloc* Alloc = MemAllocNode.GetMemAlloc();
-					const FMemoryTagId AllocMemTagId = Alloc ? Alloc->GetMemTag() : 0;
-					return TOptional<FTableCellValue>(FTableCellValue(static_cast<int64>(AllocMemTagId)));
+					if (Alloc)
+					{
+						return FTableCellValue(Alloc->GetTag());
+					}
 				}
-				//else
-				//{
-				//	return TOptional<FTableCellValue>(FText::GetEmpty());
-				//}
+
+				return TOptional<FTableCellValue>();
 			}
 		};
 		TSharedRef<ITableCellValueGetter> Getter = MakeShared<FMemAllocLlmTagValueGetter>();
 		Column.SetValueGetter(Getter);
 
-		class FMemAllocLlmTagValueFormatter : public FTableCellValueFormatter
-		{
-		public:
-			virtual FText FormatValue(const TOptional<FTableCellValue>& InValue) const override
-			{
-				if (InValue.IsSet())
-				{
-					const FMemoryTagId MemTagId = static_cast<FMemoryTagId>(InValue.GetValue().Int64);
-					const FString& TagName = FMemAllocNode::GetAllocMemTagAsString(MemTagId);
-					if (TagName.IsEmpty())
-					{
-						//FText::AsNumber();
-					}
-					return FText::FromString(TagName);
-				}
-				return FText::GetEmpty();
-			}
-			virtual FText FormatValueForTooltip(const TOptional<FTableCellValue>& InValue) const override
-			{
-				if (InValue.IsSet())
-				{
-					const FMemoryTagId MemTagId = static_cast<FMemoryTagId>(InValue.GetValue().Int64);
-					const FString& TagName = FMemAllocNode::GetAllocMemTagAsString(MemTagId);
-					return FText::FromString(FString::Printf(TEXT("%lli (%s)"), MemTagId, *TagName));
-				}
-				return FText::GetEmpty();
-			}
-			virtual FText FormatValue(const FTableColumn& Column, const FBaseTreeNode& Node) const override { return FormatValue(Column.GetValue(Node)); }
-			virtual FText FormatValueForTooltip(const FTableColumn& Column, const FBaseTreeNode& Node) const override { return FormatValueForTooltip(Column.GetValue(Node)); }
-		};
-		TSharedRef<ITableCellValueFormatter> Formatter = MakeShared<FMemAllocLlmTagValueFormatter>();
+		TSharedRef<ITableCellValueFormatter> Formatter = MakeShared<FCStringValueFormatterAsText>();
 		Column.SetValueFormatter(Formatter);
 
-		TSharedRef<ITableCellValueSorter> Sorter = MakeShared<FSorterByInt64Value>(ColumnRef);
+		TSharedRef<ITableCellValueSorter> Sorter = MakeShared<FSorterByCStringValue>(ColumnRef);
 		Column.SetValueSorter(Sorter);
 
 		AddColumn(ColumnRef);
@@ -401,7 +409,7 @@ void FMemAllocTable::AddDefaultColumns()
 
 		Column.SetShortName(LOCTEXT("BacktraceColumnName", "Backtrace"));
 		Column.SetTitleName(LOCTEXT("BacktraceColumnTitle", "Backtrace"));
-		Column.SetDescription(LOCTEXT("BacktraceColumnDesc", "Backtrace id of allocation"));
+		Column.SetDescription(LOCTEXT("BacktraceColumnDesc", "Backtrace of allocation"));
 
 		Column.SetFlags(ETableColumnFlags::ShouldBeVisible | ETableColumnFlags::CanBeHidden | ETableColumnFlags::CanBeFiltered);
 
@@ -417,46 +425,56 @@ void FMemAllocTable::AddDefaultColumns()
 		public:
 			virtual const TOptional<FTableCellValue> GetValue(const FTableColumn& Column, const FBaseTreeNode& Node) const override
 			{
-				static const TCHAR* DisplayStrings[] = {
-					TEXT("Pending..."),
-					TEXT("Not found"),
-					TEXT("-")
-				};
-				const TCHAR* Value = DisplayStrings[0];
-				const FMemAllocNode& MemAllocNode = static_cast<const FMemAllocNode&>(Node);
-				const FMemoryAlloc* Alloc = MemAllocNode.GetMemAlloc();
-				if (Alloc)
+				if (Node.IsGroup())
 				{
-					const TraceServices::FCallstack* Callstack = Alloc->GetCallstack();
-
-					if (Callstack)
-					{
-						const TraceServices::FStackFrame* Frame = Callstack->Frame(FMath::Min(2u, Callstack->Num()-1));
-						const TraceServices::QueryResult Result = Frame->Symbol->Result.load(std::memory_order_acquire);
-						switch (Result)
-						{
-						case TraceServices::QueryResult::QR_NotLoaded:
-							Value = DisplayStrings[0];
-							break;
-						case TraceServices::QueryResult::QR_NotFound:
-							Value = DisplayStrings[1];
-							break;
-						case TraceServices::QueryResult::QR_OK:
-							Value = Frame->Symbol->Name;
-							break;
-						}
-					}
-					else
+					const FTableTreeNode& NodePtr = static_cast<const FTableTreeNode&>(Node);
+					if (NodePtr.HasAggregatedValue(Column.GetId()))
 					{
-						Value = DisplayStrings[1];
+						return NodePtr.GetAggregatedValue(Column.GetId());
 					}
 				}
-				else
+				else //if (Node->Is<FMemAllocNode>())
 				{
-					// Parent nodes (i.e. not actual allocation)
-					Value = DisplayStrings[2];
+					const FMemAllocNode& MemAllocNode = static_cast<const FMemAllocNode&>(Node);
+					const FMemoryAlloc* Alloc = MemAllocNode.GetMemAlloc();
+					if (Alloc)
+					{
+						static const TCHAR* DisplayStrings[] = {
+							TEXT("Pending..."),
+							TEXT("Not found"),
+							TEXT("N/A"),
+						};
+
+						const TCHAR* Value = DisplayStrings[2]; // not available
+
+						const TraceServices::FCallstack* Callstack = Alloc->GetCallstack();
+
+						if (Callstack)
+						{
+							check(Callstack->Num() > 0);
+							const TraceServices::FStackFrame* Frame = Callstack->Frame(FMath::Min(2u, Callstack->Num() - 1));
+							check(Frame != nullptr);
+
+							const TraceServices::QueryResult Result = Frame->Symbol->Result.load(std::memory_order_acquire);
+							switch (Result)
+							{
+							case TraceServices::QueryResult::QR_NotLoaded:
+								Value = DisplayStrings[0]; // pending
+								break;
+							case TraceServices::QueryResult::QR_NotFound:
+								Value = DisplayStrings[1]; // not found
+								break;
+							case TraceServices::QueryResult::QR_OK:
+								Value = Frame->Symbol->Name;
+								break;
+							}
+						}
+
+						return FTableCellValue(Value);
+					}
 				}
-				return TOptional<FTableCellValue>(FTableCellValue(Value));
+
+				return TOptional<FTableCellValue>();
 			}
 		};
 		TSharedRef<ITableCellValueGetter> Getter = MakeShared<FBacktraceValueGetter>();
diff --git a/Engine/Source/Developer/TraceInsights/Private/Insights/MemoryProfiler/ViewModels/MemTagTreeViewColumnFactory.cpp b/Engine/Source/Developer/TraceInsights/Private/Insights/MemoryProfiler/ViewModels/MemTagTreeViewColumnFactory.cpp
--- a/Engine/Source/Developer/TraceInsights/Private/Insights/MemoryProfiler/ViewModels/MemTagTreeViewColumnFactory.cpp
+++ b/Engine/Source/Developer/TraceInsights/Private/Insights/MemoryProfiler/ViewModels/MemTagTreeViewColumnFactory.cpp
@@ -101,7 +101,7 @@ TSharedRef<Insights::FTableColumn> FMemTagTreeViewColumnFactory::CreateTypeColum
 		{
 			ensure(Column.GetId() == FMemTagTreeViewColumns::TypeColumnID);
 			const FMemTagNode& MemTagNode = static_cast<const FMemTagNode&>(Node);
-			return TOptional<FTableCellValue>(FTableCellValue(MemTagNodeTypeHelper::ToText(MemTagNode.GetType())));
+			return FTableCellValue(MemTagNodeTypeHelper::ToText(MemTagNode.GetType()));
 		}
 	};
 
@@ -147,7 +147,7 @@ TSharedRef<Insights::FTableColumn> FMemTagTreeViewColumnFactory::CreateTrackerCo
 		{
 			ensure(Column.GetId() == FMemTagTreeViewColumns::TrackerColumnID);
 			const FMemTagNode& MemTagNode = static_cast<const FMemTagNode&>(Node);
-			return TOptional<FTableCellValue>(FTableCellValue(MemTagNode.GetTrackerText()));
+			return FTableCellValue(MemTagNode.GetTrackerText());
 		}
 	};
 
@@ -192,7 +192,7 @@ TSharedRef<Insights::FTableColumn> FMemTagTreeViewColumnFactory::CreateInstanceC
 		{
 			ensure(Column.GetId() == FMemTagTreeViewColumns::InstanceCountColumnID);
 			const FMemTagNode& MemTagNode = static_cast<const FMemTagNode&>(Node);
-			return TOptional<FTableCellValue>(FTableCellValue(static_cast<int64>(MemTagNode.GetAggregatedStats().InstanceCount)));
+			return FTableCellValue(static_cast<int64>(MemTagNode.GetAggregatedStats().InstanceCount));
 		}
 	};
 
@@ -238,7 +238,7 @@ TSharedRef<Insights::FTableColumn> FMemTagTreeViewColumnFactory::CreateMinValueC
 		{
 			ensure(Column.GetId() == FMemTagTreeViewColumns::MinValueColumnID);
 			const FMemTagNode& MemTagNode = static_cast<const FMemTagNode&>(Node);
-			return TOptional<FTableCellValue>(FTableCellValue(static_cast<int64>(MemTagNode.GetAggregatedStats().Min)));
+			return FTableCellValue(static_cast<int64>(MemTagNode.GetAggregatedStats().Min));
 		}
 	};
 
@@ -283,7 +283,7 @@ TSharedRef<Insights::FTableColumn> FMemTagTreeViewColumnFactory::CreateMaxValueC
 		{
 			ensure(Column.GetId() == FMemTagTreeViewColumns::MaxValueColumnID);
 			const FMemTagNode& MemTagNode = static_cast<const FMemTagNode&>(Node);
-			return TOptional<FTableCellValue>(FTableCellValue(static_cast<int64>(MemTagNode.GetAggregatedStats().Max)));
+			return FTableCellValue(static_cast<int64>(MemTagNode.GetAggregatedStats().Max));
 		}
 	};
 
@@ -328,7 +328,7 @@ TSharedRef<Insights::FTableColumn> FMemTagTreeViewColumnFactory::CreateAverageVa
 		{
 			ensure(Column.GetId() == FMemTagTreeViewColumns::AverageValueColumnID);
 			const FMemTagNode& MemTagNode = static_cast<const FMemTagNode&>(Node);
-			return TOptional<FTableCellValue>(FTableCellValue(static_cast<int64>(MemTagNode.GetAggregatedStats().Average)));
+			return FTableCellValue(static_cast<int64>(MemTagNode.GetAggregatedStats().Average));
 		}
 	};
 
diff --git a/Engine/Source/Developer/TraceInsights/Private/Insights/MemoryProfiler/ViewModels/MemoryAlloc.cpp b/Engine/Source/Developer/TraceInsights/Private/Insights/MemoryProfiler/ViewModels/MemoryAlloc.cpp
--- a/Engine/Source/Developer/TraceInsights/Private/Insights/MemoryProfiler/ViewModels/MemoryAlloc.cpp
+++ b/Engine/Source/Developer/TraceInsights/Private/Insights/MemoryProfiler/ViewModels/MemoryAlloc.cpp
@@ -7,12 +7,12 @@ namespace Insights
 
 ////////////////////////////////////////////////////////////////////////////////////////////////////
 
-FMemoryAlloc::FMemoryAlloc(double InStartTime, double InEndTime, uint64 InAddress, uint64 InSize, FMemoryTagId InMemTag, const TraceServices::FCallstack* InCallstack)
+FMemoryAlloc::FMemoryAlloc(double InStartTime, double InEndTime, uint64 InAddress, uint64 InSize, const TCHAR* InTag, const TraceServices::FCallstack* InCallstack)
 : StartTime(InStartTime)
 , EndTime(InEndTime)
 , Address(InAddress)
 , Size(InSize)
-, MemTag(InMemTag)
+, Tag(InTag)
 , Callstack(InCallstack)
 {
 }
diff --git a/Engine/Source/Developer/TraceInsights/Private/Insights/MemoryProfiler/ViewModels/MemoryAlloc.h b/Engine/Source/Developer/TraceInsights/Private/Insights/MemoryProfiler/ViewModels/MemoryAlloc.h
--- a/Engine/Source/Developer/TraceInsights/Private/Insights/MemoryProfiler/ViewModels/MemoryAlloc.h
+++ b/Engine/Source/Developer/TraceInsights/Private/Insights/MemoryProfiler/ViewModels/MemoryAlloc.h
@@ -4,7 +4,6 @@
 
 #include "CoreMinimal.h"
 #include "TraceServices/Model/Callstack.h"
-#include "Insights/MemoryProfiler/ViewModels/MemoryTag.h"
 
 namespace Insights
 {
@@ -16,7 +15,7 @@ class FMemoryAlloc
 	friend class SMemAllocTableTreeView;
 
 public:
-	FMemoryAlloc(double InStartTime, double InEndTime, uint64 InAddress, uint64 InSize, FMemoryTagId InMemTag, const TraceServices::FCallstack* InCallstack);
+	FMemoryAlloc(double InStartTime, double InEndTime, uint64 InAddress, uint64 InSize, const TCHAR* InTag, const TraceServices::FCallstack* InCallstack);
 	~FMemoryAlloc();
 
 	double GetStartTime() const { return StartTime; }
@@ -24,7 +23,7 @@ public:
 	double GetDuration() const { return EndTime - StartTime; }
 	uint64 GetAddress() const { return Address; }
 	uint64 GetSize() const { return Size; }
-	FMemoryTagId GetMemTag() const { return MemTag; }
+	const TCHAR* GetTag() const { return Tag; }
 	const TraceServices::FCallstack* GetCallstack() const { return Callstack; }
 
 private:
@@ -32,7 +31,7 @@ private:
 	double EndTime;
 	uint64 Address;
 	uint64 Size;
-	FMemoryTagId MemTag;
+	const TCHAR* Tag;
 	const TraceServices::FCallstack* Callstack;
 };
 
diff --git a/Engine/Source/Developer/TraceInsights/Private/Insights/MemoryProfiler/Widgets/SMemAllocTableTreeView.cpp b/Engine/Source/Developer/TraceInsights/Private/Insights/MemoryProfiler/Widgets/SMemAllocTableTreeView.cpp
--- a/Engine/Source/Developer/TraceInsights/Private/Insights/MemoryProfiler/Widgets/SMemAllocTableTreeView.cpp
+++ b/Engine/Source/Developer/TraceInsights/Private/Insights/MemoryProfiler/Widgets/SMemAllocTableTreeView.cpp
@@ -291,24 +291,7 @@ void SMemAllocTableTreeView::UpdateQuery(TraceServices::IAllocationsProvider::EQ
 
 				const uint32 AllocCount = Result->Num();
 				TotalAllocCount += AllocCount;
-#if 0
-				Allocs.Reserve(Allocs.Num() + AllocCount);
 
-				for (uint32 AllocIndex = 0; AllocIndex < AllocCount; ++AllocIndex)
-				{
-					const IAllocationsProvider::FAllocation* Allocation = Result->Get(AllocIndex);
-
-					const double AllocStartTime = Allocation->GetStartTime();
-					const double AllocEndTime = Allocation->GetEndTime();
-					const uint64 AllocAddress = Allocation->GetAddress();
-					const uint64 AllocationSize = Allocation->GetSize();
-					const FMemoryTagId AllocTag = static_cast<FMemoryTagId>(Allocation->GetTag());
-					const uint64 BacktraceId = Allocation->GetBacktraceId();
-
-					//Allocs.Add(FMemoryAlloc(AllocStartTime, AllocEndTime, AllocAddress, AllocationSize, AllocTag, BacktraceId));
-					Allocs.Emplace(AllocStartTime, AllocEndTime, AllocAddress, AllocationSize, AllocTag, BacktraceId);
-				}
-#else // 14% faster, but uses more mem
 				uint64 AllocsDestIndex = Allocs.Num();
 				Allocs.AddUninitialized(AllocCount);
 				for (uint32 AllocIndex = 0; AllocIndex < AllocCount; ++AllocIndex, ++AllocsDestIndex)
@@ -319,10 +302,9 @@ void SMemAllocTableTreeView::UpdateQuery(TraceServices::IAllocationsProvider::EQ
 					Alloc.EndTime = Allocation->GetEndTime();
 					Alloc.Address = Allocation->GetAddress();
 					Alloc.Size = Allocation->GetSize();
-					Alloc.MemTag = static_cast<FMemoryTagId>(Allocation->GetTag());
+					Alloc.Tag = Provider.GetTagName(Allocation->GetTag());
 					Alloc.Callstack = Allocation->GetCallstack();
 				}
-#endif
 
 				PageStopwatch.Stop();
 				const double PageTime = PageStopwatch.GetAccumulatedTime();
@@ -371,95 +353,6 @@ void SMemAllocTableTreeView::CancelQuery()
 
 ////////////////////////////////////////////////////////////////////////////////////////////////////
 
-void SMemAllocTableTreeView::OnPreAsyncUpdate()
-{
-	auto LLMColumn = Table->FindColumn(TEXT("AllocLlmTag"));
-	OriginalLlmTagValueFormatter = LLMColumn->GetValueFormatter();
-	LLMColumn->SetValueFormatter(CreateCachedLlmTagValueFormatter());
-
-	STableTreeView::OnPreAsyncUpdate();
-}
-
-////////////////////////////////////////////////////////////////////////////////////////////////////
-
-void SMemAllocTableTreeView::OnPostAsyncUpdate()
-{
-	if (OriginalLlmTagValueFormatter.IsValid())
-	{
-		auto LLMColumn = Table->FindColumn(TEXT("AllocLlmTag"));
-		LLMColumn->SetValueFormatter(OriginalLlmTagValueFormatter.ToSharedRef());
-		OriginalLlmTagValueFormatter = nullptr;
-	}
-
-	STableTreeView::OnPostAsyncUpdate();
-}
-
-////////////////////////////////////////////////////////////////////////////////////////////////////
-
-TSharedRef<ITableCellValueFormatter> SMemAllocTableTreeView::CreateCachedLlmTagValueFormatter()
-{
-	class FCachedMemAllocLlmTagValueFormatter : public FTableCellValueFormatter
-	{
-	public:
-		virtual FText FormatValue(const TOptional<FTableCellValue>& InValue) const override
-		{
-			if (InValue.IsSet())
-			{
-				const FMemoryTagId MemTagId = static_cast<FMemoryTagId>(InValue.GetValue().Int64);
-
-				if (TagIdToStatNameMap.Contains(MemTagId))
-				{
-					FText::FromString(TagIdToStatNameMap[MemTagId]);
-				}
-				return FText::FromString(FString());
-			}
-			return FText::GetEmpty();
-		}
-		virtual FText FormatValueForTooltip(const TOptional<FTableCellValue>& InValue) const override
-		{
-			if (InValue.IsSet())
-			{
-				const FMemoryTagId MemTagId = static_cast<FMemoryTagId>(InValue.GetValue().Int64);
-				if (TagIdToStatNameMap.Contains(MemTagId))
-				{
-					return FText::FromString(FString::Printf(TEXT("%lli (%s)"), MemTagId, *TagIdToStatNameMap[MemTagId]));
-				}
-
-				return FText::FromString(FString::Printf(TEXT("%lli ()"), MemTagId));
-			}
-			return FText::GetEmpty();
-		}
-		virtual FText FormatValue(const FTableColumn& Column, const FBaseTreeNode& Node) const override { return FormatValue(Column.GetValue(Node)); }
-		virtual FText FormatValueForTooltip(const FTableColumn& Column, const FBaseTreeNode& Node) const override { return FormatValueForTooltip(Column.GetValue(Node)); }
-
-		void AddTagToCache(FMemoryTagId TagId, const FString& StatName)
-		{
-			TagIdToStatNameMap.Add(TagId, StatName);
-		}
-
-	private:
-		TMap<FMemoryTagId, FString> TagIdToStatNameMap;
-	};
-
-	TSharedRef<FCachedMemAllocLlmTagValueFormatter> Formatter = MakeShared<FCachedMemAllocLlmTagValueFormatter>();
-
-	TSharedPtr<SMemoryProfilerWindow> ProfilerWindow = FMemoryProfilerManager::Get()->GetProfilerWindow();
-	if (ProfilerWindow)
-	{
-		auto& SharedState = ProfilerWindow->GetSharedState();
-		const Insights::FMemoryTagList& TagList = SharedState.GetTagList();
-		for (FMemoryTag* CurrentTag : TagList.GetTags())
-		{
-			Formatter->AddTagToCache(CurrentTag->GetId(), CurrentTag->GetStatName());
-		}
-
-	}
-
-	return Formatter;
-}
-
-////////////////////////////////////////////////////////////////////////////////////////////////////
-
 bool SMemAllocTableTreeView::IsRunning() const
 {
 	return Query != 0 || STableTreeView::IsRunning();
diff --git a/Engine/Source/Developer/TraceInsights/Private/Insights/MemoryProfiler/Widgets/SMemAllocTableTreeView.h b/Engine/Source/Developer/TraceInsights/Private/Insights/MemoryProfiler/Widgets/SMemAllocTableTreeView.h
--- a/Engine/Source/Developer/TraceInsights/Private/Insights/MemoryProfiler/Widgets/SMemAllocTableTreeView.h
+++ b/Engine/Source/Developer/TraceInsights/Private/Insights/MemoryProfiler/Widgets/SMemAllocTableTreeView.h
@@ -81,26 +81,18 @@ public:
 
 	////////////////////////////////////////////////////////////////////////////////////////////////////
 
-protected:
-	virtual void OnPreAsyncUpdate() override;
-	virtual void OnPostAsyncUpdate() override;
-
 private:
 	void OnQueryInvalidated();
 	void StartQuery();
 	void UpdateQuery(TraceServices::IAllocationsProvider::EQueryStatus& OutStatus);
 	void CancelQuery();
 
-	TSharedRef<class ITableCellValueFormatter> CreateCachedLlmTagValueFormatter();
-
 private:
 	int32 TabIndex = -1;
 	TSharedPtr<FMemoryRuleSpec> Rule = nullptr;
 	double TimeMarkers[4];
 	TraceServices::IAllocationsProvider::FQueryHandle Query = 0;
 
-	TSharedPtr<class ITableCellValueFormatter> OriginalLlmTagValueFormatter;
-
 	FStopwatch QueryStopwatch;
 };
 
diff --git a/Engine/Source/Developer/TraceInsights/Private/Insights/Table/ViewModels/TableCellValue.h b/Engine/Source/Developer/TraceInsights/Private/Insights/Table/ViewModels/TableCellValue.h
--- a/Engine/Source/Developer/TraceInsights/Private/Insights/Table/ViewModels/TableCellValue.h
+++ b/Engine/Source/Developer/TraceInsights/Private/Insights/Table/ViewModels/TableCellValue.h
@@ -57,7 +57,6 @@ private:
 	const FText Text;
 };
 
-
 ////////////////////////////////////////////////////////////////////////////////////////////////////
 
 struct FTableCellValue
@@ -65,13 +64,13 @@ struct FTableCellValue
 public:
 	FTableCellValue() : DataType(ETableCellDataType::Unknown) {}
 
-	FTableCellValue(bool Value) : DataType(ETableCellDataType::Bool), Bool(Value) {}
-	FTableCellValue(int64 Value) : DataType(ETableCellDataType::Int64), Int64(Value) {}
-	FTableCellValue(float Value) : DataType(ETableCellDataType::Float), Float(Value) {}
-	FTableCellValue(double Value) : DataType(ETableCellDataType::Double), Double(Value) {}
-	FTableCellValue(const TCHAR* Value) : DataType(ETableCellDataType::CString), CString(Value) {}
-	FTableCellValue(TSharedPtr<ICustomTableCellValue> Value) : DataType(ETableCellDataType::Custom), Custom(Value) {}
-	FTableCellValue(const FText& Value) : DataType(ETableCellDataType::Custom_Text), Custom(MakeShared<FTextCustomTableCellValue>(Value)) {}
+	explicit FTableCellValue(bool Value) : DataType(ETableCellDataType::Bool), Bool(Value) {}
+	explicit FTableCellValue(int64 Value) : DataType(ETableCellDataType::Int64), Int64(Value) {}
+	explicit FTableCellValue(float Value) : DataType(ETableCellDataType::Float), Float(Value) {}
+	explicit FTableCellValue(double Value) : DataType(ETableCellDataType::Double), Double(Value) {}
+	explicit FTableCellValue(const TCHAR* Value) : DataType(ETableCellDataType::CString), CString(Value) {}
+	explicit FTableCellValue(TSharedPtr<ICustomTableCellValue> Value) : DataType(ETableCellDataType::Custom), Custom(Value) {}
+	explicit FTableCellValue(const FText& Value) : DataType(ETableCellDataType::Custom_Text), Custom(MakeShared<FTextCustomTableCellValue>(Value)) {}
 
 	bool AsBool() const
 	{
diff --git a/Engine/Source/Developer/TraceInsights/Private/Insights/Table/ViewModels/TableCellValueFormatter.cpp b/Engine/Source/Developer/TraceInsights/Private/Insights/Table/ViewModels/TableCellValueFormatter.cpp
--- a/Engine/Source/Developer/TraceInsights/Private/Insights/Table/ViewModels/TableCellValueFormatter.cpp
+++ b/Engine/Source/Developer/TraceInsights/Private/Insights/Table/ViewModels/TableCellValueFormatter.cpp
@@ -216,7 +216,10 @@ FText FCStringValueFormatterAsText::FormatValue(const TOptional<FTableCellValue>
 	if (InValue.IsSet())
 	{
 		const TCHAR* Value = InValue.GetValue().CString;
-		return FText::FromString(Value);
+		if (Value)
+		{
+			return FText::FromString(Value);
+		}
 	}
 	return FText::GetEmpty();
 }
diff --git a/Engine/Source/Developer/TraceInsights/Private/Insights/Table/ViewModels/TreeNodeGrouping.cpp b/Engine/Source/Developer/TraceInsights/Private/Insights/Table/ViewModels/TreeNodeGrouping.cpp
--- a/Engine/Source/Developer/TraceInsights/Private/Insights/Table/ViewModels/TreeNodeGrouping.cpp
+++ b/Engine/Source/Developer/TraceInsights/Private/Insights/Table/ViewModels/TreeNodeGrouping.cpp
@@ -44,11 +44,9 @@ void FTreeNodeGrouping::GroupNodes(const TArray<FTableTreeNodePtr>& Nodes, FTabl
 		if (!GroupPtrPtr)
 		{
 			GroupPtr = MakeShared<FTableTreeNode>(GroupInfo.Name, InParentTable);
-
-			GroupMap.Add(GroupInfo.Name, GroupPtr);
-			ParentGroup.AddChildAndSetGroupPtr(GroupPtr);
-
 			GroupPtr->SetExpansion(GroupInfo.IsExpanded);
+			ParentGroup.AddChildAndSetGroupPtr(GroupPtr);
+			GroupMap.Add(GroupInfo.Name, GroupPtr);
 		}
 		else
 		{
@@ -125,6 +123,18 @@ FTreeNodeGroupInfo FTreeNodeGroupingByUniqueValue::GetGroupForNode(const FBaseTr
 	return { FName(*ColumnRef->GetValueAsText(*TableTreeNodePtr).ToString()), false };
 }
 
+////////////////////////////////////////////////////////////////////////////////////////////////////
+// TTreeNodeGroupingByUniqueValue specializations
+////////////////////////////////////////////////////////////////////////////////////////////////////
+
+template<> bool TTreeNodeGroupingByUniqueValue<bool>::GetValue(const FTableCellValue& CellValue) { return CellValue.Bool; }
+template<> int64 TTreeNodeGroupingByUniqueValue<int64>::GetValue(const FTableCellValue& CellValue) { return CellValue.Int64; }
+template<> float TTreeNodeGroupingByUniqueValue<float>::GetValue(const FTableCellValue& CellValue) { return CellValue.Float; }
+template<> double TTreeNodeGroupingByUniqueValue<double>::GetValue(const FTableCellValue& CellValue) { return CellValue.Double; }
+template<> const TCHAR* TTreeNodeGroupingByUniqueValue<const TCHAR*>::GetValue(const FTableCellValue& CellValue) { return CellValue.CString; }
+
+////////////////////////////////////////////////////////////////////////////////////////////////////
+// FTreeNodeGroupingByNameFirstLetter
 ////////////////////////////////////////////////////////////////////////////////////////////////////
 
 FTreeNodeGroupingByNameFirstLetter::FTreeNodeGroupingByNameFirstLetter()
@@ -144,6 +154,8 @@ FTreeNodeGroupInfo FTreeNodeGroupingByNameFirstLetter::GetGroupForNode(const FBa
 	return { *InNode->GetName().GetPlainNameString().Left(1), false };
 }
 
+////////////////////////////////////////////////////////////////////////////////////////////////////
+// FTreeNodeGroupingByType
 ////////////////////////////////////////////////////////////////////////////////////////////////////
 
 FTreeNodeGroupingByType::FTreeNodeGroupingByType()
diff --git a/Engine/Source/Developer/TraceInsights/Private/Insights/Table/ViewModels/TreeNodeGrouping.h b/Engine/Source/Developer/TraceInsights/Private/Insights/Table/ViewModels/TreeNodeGrouping.h
--- a/Engine/Source/Developer/TraceInsights/Private/Insights/Table/ViewModels/TreeNodeGrouping.h
+++ b/Engine/Source/Developer/TraceInsights/Private/Insights/Table/ViewModels/TreeNodeGrouping.h
@@ -84,7 +84,7 @@ public:
 
 ////////////////////////////////////////////////////////////////////////////////////////////////////
 
-/** Creates a group for unique column value of node. */
+/** Creates a group for each unique value. */
 class FTreeNodeGroupingByUniqueValue : public FTreeNodeGrouping
 {
 public:
@@ -102,6 +102,81 @@ private:
 
 ////////////////////////////////////////////////////////////////////////////////////////////////////
 
+/** Creates a group for each unique value (assumes data type of cell values is a simple type). */
+template<typename Type>
+class TTreeNodeGroupingByUniqueValue : public FTreeNodeGroupingByUniqueValue
+{
+public:
+	TTreeNodeGroupingByUniqueValue(TSharedRef<FTableColumn> InColumnRef) : FTreeNodeGroupingByUniqueValue(InColumnRef) {}
+	virtual ~TTreeNodeGroupingByUniqueValue() {}
+
+	virtual void GroupNodes(const TArray<FTableTreeNodePtr>& Nodes, FTableTreeNode& ParentGroup, TWeakPtr<FTable> InParentTable, std::atomic<bool>& bCancelGrouping) const override;
+
+private:
+	static Type GetValue(const FTableCellValue& CellValue);
+};
+
+template<typename Type>
+void TTreeNodeGroupingByUniqueValue<Type>::GroupNodes(const TArray<FTableTreeNodePtr>& Nodes, FTableTreeNode& ParentGroup, TWeakPtr<FTable> InParentTable, std::atomic<bool>& bCancelGrouping) const
+{
+	TMap<Type, FTableTreeNodePtr> GroupMap;
+	FTableTreeNodePtr UnsetGroupPtr = nullptr;
+
+	ParentGroup.ClearChildren();
+
+	for (FTableTreeNodePtr NodePtr : Nodes)
+	{
+		ensure(!NodePtr->IsGroup());
+
+		if (bCancelGrouping)
+		{
+			return;
+		}
+
+		FTableTreeNodePtr GroupPtr = nullptr;
+
+		const TOptional<FTableCellValue> CellValue = GetColumn()->GetValue(*NodePtr);
+		if (CellValue.IsSet())
+		{
+			const Type Value = GetValue(CellValue.GetValue());
+
+			FTableTreeNodePtr* GroupPtrPtr = GroupMap.Find(Value);
+			if (!GroupPtrPtr)
+			{
+				FName GroupName(GetColumn()->GetValueAsText(*NodePtr).ToString());
+				GroupPtr = MakeShared<FTableTreeNode>(GroupName, InParentTable);
+				GroupPtr->SetExpansion(false);
+				ParentGroup.AddChildAndSetGroupPtr(GroupPtr);
+				GroupMap.Add(Value, GroupPtr);
+			}
+			else
+			{
+				GroupPtr = *GroupPtrPtr;
+			}
+		}
+		else
+		{
+			if (!UnsetGroupPtr)
+			{
+				UnsetGroupPtr = MakeShared<FTableTreeNode>(FName(TEXT("<unset>")), InParentTable);
+				UnsetGroupPtr->SetExpansion(false);
+				ParentGroup.AddChildAndSetGroupPtr(UnsetGroupPtr);
+			}
+			GroupPtr = UnsetGroupPtr;
+		}
+
+		GroupPtr->AddChildAndSetGroupPtr(NodePtr);
+	}
+}
+
+typedef TTreeNodeGroupingByUniqueValue<bool> FTreeNodeGroupingByUniqueValueBool;
+typedef TTreeNodeGroupingByUniqueValue<int64> FTreeNodeGroupingByUniqueValueInt64;
+typedef TTreeNodeGroupingByUniqueValue<float> FTreeNodeGroupingByUniqueValueFloat;
+typedef TTreeNodeGroupingByUniqueValue<double> FTreeNodeGroupingByUniqueValueDouble;
+typedef TTreeNodeGroupingByUniqueValue<const TCHAR*> FTreeNodeGroupingByUniqueValueCString;
+
+////////////////////////////////////////////////////////////////////////////////////////////////////
+
 /** Creates a group for each first letter of node names. */
 class FTreeNodeGroupingByNameFirstLetter : public FTreeNodeGrouping
 {
diff --git a/Engine/Source/Developer/TraceInsights/Private/Insights/Table/ViewModels/UntypedTable.cpp b/Engine/Source/Developer/TraceInsights/Private/Insights/Table/ViewModels/UntypedTable.cpp
--- a/Engine/Source/Developer/TraceInsights/Private/Insights/Table/ViewModels/UntypedTable.cpp
+++ b/Engine/Source/Developer/TraceInsights/Private/Insights/Table/ViewModels/UntypedTable.cpp
@@ -42,11 +42,11 @@ public:
 					const int32 ColumnIndex = Column.GetIndex();
 					switch (DataType)
 					{
-						case ETableCellDataType::Bool:    return TOptional<FTableCellValue>(Reader->GetValueBool(ColumnIndex));
-						case ETableCellDataType::Int64:   return TOptional<FTableCellValue>(Reader->GetValueInt(ColumnIndex));
-						case ETableCellDataType::Float:   return TOptional<FTableCellValue>(Reader->GetValueFloat(ColumnIndex));
-						case ETableCellDataType::Double:  return TOptional<FTableCellValue>(Reader->GetValueDouble(ColumnIndex));
-						case ETableCellDataType::CString: return TOptional<FTableCellValue>(Reader->GetValueCString(ColumnIndex));
+						case ETableCellDataType::Bool:    return TOptional<FTableCellValue>(FTableCellValue(Reader->GetValueBool(ColumnIndex)));
+						case ETableCellDataType::Int64:   return TOptional<FTableCellValue>(FTableCellValue(Reader->GetValueInt(ColumnIndex)));
+						case ETableCellDataType::Float:   return TOptional<FTableCellValue>(FTableCellValue(Reader->GetValueFloat(ColumnIndex)));
+						case ETableCellDataType::Double:  return TOptional<FTableCellValue>(FTableCellValue(Reader->GetValueDouble(ColumnIndex)));
+						case ETableCellDataType::CString: return TOptional<FTableCellValue>(FTableCellValue(Reader->GetValueCString(ColumnIndex)));
 					}
 				}
 			}
diff --git a/Engine/Source/Developer/TraceInsights/Private/Insights/Table/Widgets/STableTreeView.cpp b/Engine/Source/Developer/TraceInsights/Private/Insights/Table/Widgets/STableTreeView.cpp
--- a/Engine/Source/Developer/TraceInsights/Private/Insights/Table/Widgets/STableTreeView.cpp
+++ b/Engine/Source/Developer/TraceInsights/Private/Insights/Table/Widgets/STableTreeView.cpp
@@ -1094,7 +1094,31 @@ void STableTreeView::CreateGroupings()
 	{
 		if (!ColumnRef->IsHierarchy())
 		{
-			AvailableGroupings.Add(MakeShared<FTreeNodeGroupingByUniqueValue>(ColumnRef));
+			switch(ColumnRef->GetDataType())
+			{
+				case ETableCellDataType::Bool:
+					AvailableGroupings.Add(MakeShared<FTreeNodeGroupingByUniqueValueBool>(ColumnRef));
+					break;
+
+				case ETableCellDataType::Int64:
+					AvailableGroupings.Add(MakeShared<FTreeNodeGroupingByUniqueValueInt64>(ColumnRef));
+					break;
+
+				case ETableCellDataType::Float:
+					AvailableGroupings.Add(MakeShared<FTreeNodeGroupingByUniqueValueFloat>(ColumnRef));
+					break;
+
+				case ETableCellDataType::Double:
+					AvailableGroupings.Add(MakeShared<FTreeNodeGroupingByUniqueValueDouble>(ColumnRef));
+					break;
+
+				case ETableCellDataType::CString:
+					AvailableGroupings.Add(MakeShared<FTreeNodeGroupingByUniqueValueCString>(ColumnRef));
+					break;
+
+				default:
+					AvailableGroupings.Add(MakeShared<FTreeNodeGroupingByUniqueValue>(ColumnRef));
+			}
 		}
 	}
 
diff --git a/Engine/Source/Developer/TraceServices/Private/Model/AllocationsProvider.h b/Engine/Source/Developer/TraceServices/Private/Model/AllocationsProvider.h
--- a/Engine/Source/Developer/TraceServices/Private/Model/AllocationsProvider.h
+++ b/Engine/Source/Developer/TraceServices/Private/Model/AllocationsProvider.h
@@ -166,6 +166,8 @@ public:
 
 	void DebugPrint() const;
 
+	virtual const TCHAR* GetTagName(int32 Tag) const { ReadAccessCheck(); return TagTracker.GetTagString(Tag); }
+
 	//////////////////////////////////////////////////
 	// Edit operations
 
diff --git a/Engine/Source/Developer/TraceServices/Public/TraceServices/Model/AllocationsProvider.h b/Engine/Source/Developer/TraceServices/Public/TraceServices/Model/AllocationsProvider.h
--- a/Engine/Source/Developer/TraceServices/Public/TraceServices/Model/AllocationsProvider.h
+++ b/Engine/Source/Developer/TraceServices/Public/TraceServices/Model/AllocationsProvider.h
@@ -156,6 +156,10 @@ public:
 	virtual FQueryHandle StartQuery(const FQueryParams& Params) const = 0;
 	virtual void CancelQuery(FQueryHandle Query) const = 0;
 	virtual const FQueryStatus PollQuery(FQueryHandle Query) const = 0;
+
+	// Returns the display name of the specified LLM tag.
+	// Lifetime of returned string matches the session lifetime.
+	virtual const TCHAR* GetTagName(int32 Tag) const = 0;
 };
 
 TRACESERVICES_API const IAllocationsProvider* ReadAllocationsProvider(const IAnalysisSession& Session);
