[Insights]
  - Counters: Improved precission of values for aggregates stats (in the list view and also in the tooltip of a counter).
  - Counters: Fixed tooltip of graph series of memory counters with negative values.

diff --git a/Engine/Source/Developer/TraceInsights/Private/Insights/ViewModels/StatsNode.cpp b/Engine/Source/Developer/TraceInsights/Private/Insights/ViewModels/StatsNode.cpp
--- a/Engine/Source/Developer/TraceInsights/Private/Insights/ViewModels/StatsNode.cpp
+++ b/Engine/Source/Developer/TraceInsights/Private/Insights/ViewModels/StatsNode.cpp
@@ -90,7 +90,7 @@ const FName FStatsNodeMetaGroupName::Memory(TEXT("Memory"));
 
 ////////////////////////////////////////////////////////////////////////////////////////////////////
 
-const FText FStatsNode::FormatAggregatedStatsValue(double ValueDbl, int64 ValueInt) const
+const FText FStatsNode::FormatAggregatedStatsValue(double ValueDbl, int64 ValueInt, bool bForTooltip) const
 {
 	if (AggregatedStats.Count > 0)
 	{
@@ -99,11 +99,20 @@ const FText FStatsNode::FormatAggregatedStatsValue(double ValueDbl, int64 ValueI
 			//TODO: if (GetDisplayHint() == FStatsNodeDisplayHint::Seconds)
 			if (GetMetaGroupName() == FStatsNodeMetaGroupName::Time)
 			{
-				return FText::FromString(TimeUtils::FormatTimeAuto(ValueDbl));
+				if (bForTooltip)
+				{
+					return FText::FromString(TimeUtils::FormatTimeAuto(ValueDbl, 2));
+				}
+				else
+				{
+					return FText::FromString(TimeUtils::FormatTimeAuto(ValueDbl, 1));
+				}
 			}
 			else
 			{
-				return FText::AsNumber(ValueDbl);
+				FNumberFormattingOptions FormattingOptions;
+				FormattingOptions.MaximumFractionalDigits = bForTooltip ? 12 : 6;
+				return FText::AsNumber(ValueDbl, &FormattingOptions);
 			}
 		}
 		else
@@ -114,7 +123,25 @@ const FText FStatsNode::FormatAggregatedStatsValue(double ValueDbl, int64 ValueI
 			{
 				if (ValueInt > 0)
 				{
-					return FText::AsMemory(ValueInt);
+					if (bForTooltip)
+					{
+						if (ValueInt < 1024)
+						{
+							return FText::Format(LOCTEXT("Counter_MemValueFmt1", "{0} bytes"), FText::AsNumber(ValueInt));
+						}
+						else
+						{
+							FNumberFormattingOptions FormattingOptions;
+							FormattingOptions.MaximumFractionalDigits = 2;
+							return FText::Format(LOCTEXT("Counter_MemValueFmt2", "{0} ({1} bytes)"), FText::AsMemory(ValueInt, &FormattingOptions), FText::AsNumber(ValueInt));
+						}
+					}
+					else
+					{
+						FNumberFormattingOptions FormattingOptions;
+						FormattingOptions.MaximumFractionalDigits = 1;
+						return FText::AsMemory(ValueInt, &FormattingOptions);
+					}
 				}
 				else if (ValueInt == 0)
 				{
@@ -122,7 +149,25 @@ const FText FStatsNode::FormatAggregatedStatsValue(double ValueDbl, int64 ValueI
 				}
 				else
 				{
-					return FText::FromString(FString::Printf(TEXT("-%s"), *FText::AsMemory(-ValueInt).ToString()));
+					if (bForTooltip)
+					{
+						if (-ValueInt < 1024)
+						{
+							return FText::Format(LOCTEXT("Counter_NegMemValueFmt1", "-{0} bytes"), FText::AsNumber(-ValueInt));
+						}
+						else
+						{
+							FNumberFormattingOptions FormattingOptions;
+							FormattingOptions.MaximumFractionalDigits = 2;
+							return FText::Format(LOCTEXT("Counter_NegMemValueFmt2", "-{0} (-{1} bytes)"), FText::AsMemory(-ValueInt, &FormattingOptions), FText::AsNumber(-ValueInt));
+						}
+					}
+					else
+					{
+						FNumberFormattingOptions FormattingOptions;
+						FormattingOptions.MaximumFractionalDigits = 1;
+						return FText::Format(LOCTEXT("Counter_NegMemValueFmt3", "-{0}"), FText::AsMemory(-ValueInt, &FormattingOptions));
+					}
 				}
 			}
 			else
@@ -137,51 +182,51 @@ const FText FStatsNode::FormatAggregatedStatsValue(double ValueDbl, int64 ValueI
 
 ////////////////////////////////////////////////////////////////////////////////////////////////////
 
-const FText FStatsNode::GetTextForAggregatedStatsSum() const
+const FText FStatsNode::GetTextForAggregatedStatsSum(bool bForTooltip) const
 {
-	return FormatAggregatedStatsValue(AggregatedStats.DoubleStats.Sum, AggregatedStats.Int64Stats.Sum);
+	return FormatAggregatedStatsValue(AggregatedStats.DoubleStats.Sum, AggregatedStats.Int64Stats.Sum, bForTooltip);
 }
 
 ////////////////////////////////////////////////////////////////////////////////////////////////////
 
-const FText FStatsNode::GetTextForAggregatedStatsMin() const
+const FText FStatsNode::GetTextForAggregatedStatsMin(bool bForTooltip) const
 {
-	return FormatAggregatedStatsValue(AggregatedStats.DoubleStats.Min, AggregatedStats.Int64Stats.Min);
+	return FormatAggregatedStatsValue(AggregatedStats.DoubleStats.Min, AggregatedStats.Int64Stats.Min, bForTooltip);
 }
 
 ////////////////////////////////////////////////////////////////////////////////////////////////////
 
-const FText FStatsNode::GetTextForAggregatedStatsMax() const
+const FText FStatsNode::GetTextForAggregatedStatsMax(bool bForTooltip) const
 {
-	return FormatAggregatedStatsValue(AggregatedStats.DoubleStats.Max, AggregatedStats.Int64Stats.Max);
+	return FormatAggregatedStatsValue(AggregatedStats.DoubleStats.Max, AggregatedStats.Int64Stats.Max, bForTooltip);
 }
 
 ////////////////////////////////////////////////////////////////////////////////////////////////////
 
-const FText FStatsNode::GetTextForAggregatedStatsAverage() const
+const FText FStatsNode::GetTextForAggregatedStatsAverage(bool bForTooltip) const
 {
-	return FormatAggregatedStatsValue(AggregatedStats.DoubleStats.Average, AggregatedStats.Int64Stats.Average);
+	return FormatAggregatedStatsValue(AggregatedStats.DoubleStats.Average, AggregatedStats.Int64Stats.Average, bForTooltip);
 }
 
 ////////////////////////////////////////////////////////////////////////////////////////////////////
 
-const FText FStatsNode::GetTextForAggregatedStatsMedian() const
+const FText FStatsNode::GetTextForAggregatedStatsMedian(bool bForTooltip) const
 {
-	return FormatAggregatedStatsValue(AggregatedStats.DoubleStats.Median, AggregatedStats.Int64Stats.Median);
+	return FormatAggregatedStatsValue(AggregatedStats.DoubleStats.Median, AggregatedStats.Int64Stats.Median, bForTooltip);
 }
 
 ////////////////////////////////////////////////////////////////////////////////////////////////////
 
-const FText FStatsNode::GetTextForAggregatedStatsLowerQuartile() const
+const FText FStatsNode::GetTextForAggregatedStatsLowerQuartile(bool bForTooltip) const
 {
-	return FormatAggregatedStatsValue(AggregatedStats.DoubleStats.LowerQuartile, AggregatedStats.Int64Stats.LowerQuartile);
+	return FormatAggregatedStatsValue(AggregatedStats.DoubleStats.LowerQuartile, AggregatedStats.Int64Stats.LowerQuartile, bForTooltip);
 }
 
 ////////////////////////////////////////////////////////////////////////////////////////////////////
 
-const FText FStatsNode::GetTextForAggregatedStatsUpperQuartile() const
+const FText FStatsNode::GetTextForAggregatedStatsUpperQuartile(bool bForTooltip) const
 {
-	return FormatAggregatedStatsValue(AggregatedStats.DoubleStats.UpperQuartile, AggregatedStats.Int64Stats.UpperQuartile);
+	return FormatAggregatedStatsValue(AggregatedStats.DoubleStats.UpperQuartile, AggregatedStats.Int64Stats.UpperQuartile, bForTooltip);
 }
 
 ////////////////////////////////////////////////////////////////////////////////////////////////////
diff --git a/Engine/Source/Developer/TraceInsights/Private/Insights/ViewModels/StatsNode.h b/Engine/Source/Developer/TraceInsights/Private/Insights/ViewModels/StatsNode.h
--- a/Engine/Source/Developer/TraceInsights/Private/Insights/ViewModels/StatsNode.h
+++ b/Engine/Source/Developer/TraceInsights/Private/Insights/ViewModels/StatsNode.h
@@ -236,18 +236,15 @@ public:
 	void UpdateAggregatedStatsInt64FromDouble();
 	void UpdateAggregatedStatsDoubleFromInt64();
 
-	const FText FormatValue(double Value) const;
-	const FText FormatValue(int64 Value) const;
-
-	const FText FormatAggregatedStatsValue(double ValueDbl, int64 ValueInt) const;
-
-	const FText GetTextForAggregatedStatsSum() const;
-	const FText GetTextForAggregatedStatsMin() const;
-	const FText GetTextForAggregatedStatsMax() const;
-	const FText GetTextForAggregatedStatsAverage() const;
-	const FText GetTextForAggregatedStatsMedian() const;
-	const FText GetTextForAggregatedStatsLowerQuartile() const;
-	const FText GetTextForAggregatedStatsUpperQuartile() const;
+	const FText FormatAggregatedStatsValue(double ValueDbl, int64 ValueInt, bool bForTooltip = false) const;
+
+	const FText GetTextForAggregatedStatsSum(bool bForTooltip = false) const;
+	const FText GetTextForAggregatedStatsMin(bool bForTooltip = false) const;
+	const FText GetTextForAggregatedStatsMax(bool bForTooltip = false) const;
+	const FText GetTextForAggregatedStatsAverage(bool bForTooltip = false) const;
+	const FText GetTextForAggregatedStatsMedian(bool bForTooltip = false) const;
+	const FText GetTextForAggregatedStatsLowerQuartile(bool bForTooltip = false) const;
+	const FText GetTextForAggregatedStatsUpperQuartile(bool bForTooltip = false) const;
 
 	/** Sorts children using the specified class instance. */
 	template<typename TSortingClass>
diff --git a/Engine/Source/Developer/TraceInsights/Private/Insights/ViewModels/TimingGraphTrack.cpp b/Engine/Source/Developer/TraceInsights/Private/Insights/ViewModels/TimingGraphTrack.cpp
--- a/Engine/Source/Developer/TraceInsights/Private/Insights/ViewModels/TimingGraphTrack.cpp
+++ b/Engine/Source/Developer/TraceInsights/Private/Insights/ViewModels/TimingGraphTrack.cpp
@@ -61,7 +61,32 @@ FString FTimingGraphSeries::FormatValue(double Value) const
 		else if (bIsMemory)
 		{
 			const int64 MemValue = static_cast<int64>(Value);
-			return FString::Printf(TEXT("%s (%s bytes)"), *FText::AsMemory(MemValue).ToString(), *FText::AsNumber(MemValue).ToString());
+			if (MemValue >= 0)
+			{
+				if (MemValue < 1024)
+				{
+					return FString::Printf(TEXT("%s bytes"), *FText::AsNumber(MemValue).ToString());
+				}
+				else
+				{
+					FNumberFormattingOptions FormattingOptions;
+					FormattingOptions.MaximumFractionalDigits = 2;
+					return FString::Printf(TEXT("%s (%s bytes)"), *FText::AsMemory(MemValue, &FormattingOptions).ToString(), *FText::AsNumber(MemValue).ToString());
+				}
+			}
+			else
+			{
+				if (-MemValue < 1024)
+				{
+					return FString::Printf(TEXT("-%s bytes"), *FText::AsNumber(-MemValue).ToString());
+				}
+				else
+				{
+					FNumberFormattingOptions FormattingOptions;
+					FormattingOptions.MaximumFractionalDigits = 2;
+					return FString::Printf(TEXT("-%s (-%s bytes)"), *FText::AsMemory(-MemValue, &FormattingOptions).ToString(), *FText::AsNumber(-MemValue).ToString());
+				}
+			}
 		}
 		else if (bIsFloatingPoint)
 		{
diff --git a/Engine/Source/Developer/TraceInsights/Private/Insights/Widgets/SStatsViewTooltip.cpp b/Engine/Source/Developer/TraceInsights/Private/Insights/Widgets/SStatsViewTooltip.cpp
--- a/Engine/Source/Developer/TraceInsights/Private/Insights/Widgets/SStatsViewTooltip.cpp
+++ b/Engine/Source/Developer/TraceInsights/Private/Insights/Widgets/SStatsViewTooltip.cpp
@@ -92,13 +92,13 @@ TSharedPtr<SToolTip> SStatsViewTooltip::GetRowTooltip(const TSharedPtr<FStatsNod
 {
 	const FText InstanceCountText = FText::AsNumber(StatsNodePtr->GetAggregatedStats().Count);
 
-	FText SumText = StatsNodePtr->GetTextForAggregatedStatsSum();
-	FText MinText = StatsNodePtr->GetTextForAggregatedStatsMin();
-	FText MaxText = StatsNodePtr->GetTextForAggregatedStatsMax();
-	FText AvgText = StatsNodePtr->GetTextForAggregatedStatsAverage();
-	FText MedText = StatsNodePtr->GetTextForAggregatedStatsMedian();
-	//FText LowText = StatsNodePtr->GetTextForAggregatedStatsLowerQuartile();
-	//FText UppText = StatsNodePtr->GetTextForAggregatedStatsUpperQuartile();
+	FText SumText = StatsNodePtr->GetTextForAggregatedStatsSum(true);
+	FText MinText = StatsNodePtr->GetTextForAggregatedStatsMin(true);
+	FText MaxText = StatsNodePtr->GetTextForAggregatedStatsMax(true);
+	FText AvgText = StatsNodePtr->GetTextForAggregatedStatsAverage(true);
+	FText MedText = StatsNodePtr->GetTextForAggregatedStatsMedian(true);
+	//FText LowText = StatsNodePtr->GetTextForAggregatedStatsLowerQuartile(true);
+	//FText UppText = StatsNodePtr->GetTextForAggregatedStatsUpperQuartile(true);
 
 	TSharedPtr<SGridPanel> GridPanel;
 	TSharedPtr<SHorizontalBox> HBox;
@@ -277,11 +277,12 @@ void SStatsViewTooltip::AddAggregatedStatsRow(TSharedPtr<SGridPanel> Grid, int32
 
 	Grid->AddSlot(1, Row)
 		.Padding(2.0f)
-		.HAlign(HAlign_Right)
+		.HAlign(HAlign_Left)
 		[
 			SNew(STextBlock)
 			.Text(Value)
 			.TextStyle(FEditorStyle::Get(), TEXT("Profiler.Tooltip"))
+			.ColorAndOpacity(Row == 1 ? FLinearColor::Gray : FLinearColor::White)
 		];
 
 	Row++;
