[Insights] Memory: Increased default width of LLM Tags and Backtrace columns in Allocs Table tree view. Also made Size values to be right aligned.

diff --git a/Engine/Source/Developer/TraceInsights/Private/Insights/MemoryProfiler/ViewModels/MemAllocTable.cpp b/Engine/Source/Developer/TraceInsights/Private/Insights/MemoryProfiler/ViewModels/MemAllocTable.cpp
--- a/Engine/Source/Developer/TraceInsights/Private/Insights/MemoryProfiler/ViewModels/MemAllocTable.cpp
+++ b/Engine/Source/Developer/TraceInsights/Private/Insights/MemoryProfiler/ViewModels/MemAllocTable.cpp
@@ -299,7 +299,7 @@ void FMemAllocTable::AddDefaultColumns()
 
 		Column.SetFlags(ETableColumnFlags::ShouldBeVisible | ETableColumnFlags::CanBeHidden | ETableColumnFlags::CanBeFiltered);
 
-		Column.SetHorizontalAlignment(HAlign_Left);
+		Column.SetHorizontalAlignment(HAlign_Right);
 		Column.SetInitialWidth(100.0f);
 
 		Column.SetDataType(ETableCellDataType::Int64);
@@ -358,7 +358,7 @@ void FMemAllocTable::AddDefaultColumns()
 		Column.SetFlags(ETableColumnFlags::ShouldBeVisible | ETableColumnFlags::CanBeHidden | ETableColumnFlags::CanBeFiltered);
 
 		Column.SetHorizontalAlignment(HAlign_Left);
-		Column.SetInitialWidth(100.0f);
+		Column.SetInitialWidth(120.0f);
 
 		Column.SetDataType(ETableCellDataType::CString);
 
@@ -414,7 +414,7 @@ void FMemAllocTable::AddDefaultColumns()
 		Column.SetFlags(ETableColumnFlags::ShouldBeVisible | ETableColumnFlags::CanBeHidden | ETableColumnFlags::CanBeFiltered);
 
 		Column.SetHorizontalAlignment(HAlign_Left);
-		Column.SetInitialWidth(120.0f);
+		Column.SetInitialWidth(550.0f);
 
 		Column.SetDataType(ETableCellDataType::CString);
 
