[Insights] Show symbol resolution status in MemAllocTreeView

diff --git a/Engine/Source/Developer/TraceInsights/Private/Insights/MemoryProfiler/Widgets/SMemAllocTableTreeView.cpp b/Engine/Source/Developer/TraceInsights/Private/Insights/MemoryProfiler/Widgets/SMemAllocTableTreeView.cpp
--- a/Engine/Source/Developer/TraceInsights/Private/Insights/MemoryProfiler/Widgets/SMemAllocTableTreeView.cpp
+++ b/Engine/Source/Developer/TraceInsights/Private/Insights/MemoryProfiler/Widgets/SMemAllocTableTreeView.cpp
@@ -5,6 +5,7 @@
 #include "TraceServices/AnalysisService.h"
 #include "TraceServices/Model/AllocationsProvider.h"
 #include "TraceServices/Model/Callstack.h"
+#include "TraceServices/Model/Modules.h"
 
 // Insights
 #include "Insights/MemoryProfiler/MemoryProfilerManager.h"
@@ -385,6 +386,34 @@ FText SMemAllocTableTreeView::GetCurrentOperationName() const
 
 ////////////////////////////////////////////////////////////////////////////////////////////////////
 
+TSharedPtr<SWidget> SMemAllocTableTreeView::ConstructFooter()
+{
+	return
+		SNew(SHorizontalBox)
+
+		+ SHorizontalBox::Slot()
+		.HAlign(HAlign_Right)
+		.Padding(2.0f)
+		[
+			SNew(STextBlock)
+			.Text(this, &SMemAllocTableTreeView::GetSymbolResolutionStatus)
+		];
+}
+
+////////////////////////////////////////////////////////////////////////////////////////////////////
+
+FText SMemAllocTableTreeView::GetSymbolResolutionStatus() const
+{
+	auto ModuleProvider = Session->ReadProvider<IModuleProvider>(FName("ModuleProvider"));
+
+	IModuleProvider::FStats Stats;
+	ModuleProvider->GetStats(&Stats);
+
+	return FText::Format(LOCTEXT("SymbolsResolved", "{0} / {1} symbols resolved. {2} failed."), Stats.SymbolsResolved, Stats.SymbolsDiscovered, Stats.SymbolsFailed);
+}
+
+////////////////////////////////////////////////////////////////////////////////////////////////////
+
 } // namespace Insights
 
 #undef LOCTEXT_NAMESPACE
diff --git a/Engine/Source/Developer/TraceInsights/Private/Insights/MemoryProfiler/Widgets/SMemAllocTableTreeView.h b/Engine/Source/Developer/TraceInsights/Private/Insights/MemoryProfiler/Widgets/SMemAllocTableTreeView.h
--- a/Engine/Source/Developer/TraceInsights/Private/Insights/MemoryProfiler/Widgets/SMemAllocTableTreeView.h
+++ b/Engine/Source/Developer/TraceInsights/Private/Insights/MemoryProfiler/Widgets/SMemAllocTableTreeView.h
@@ -36,6 +36,8 @@ public:
 	 */
 	void Construct(const FArguments& InArgs, TSharedPtr<FMemAllocTable> InTablePtr);
 
+	virtual TSharedPtr<SWidget> ConstructFooter() override;
+
 	TSharedPtr<FMemAllocTable> GetMemAllocTable() { return StaticCastSharedPtr<FMemAllocTable>(GetTable()); }
 	const TSharedPtr<FMemAllocTable> GetMemAllocTable() const { return StaticCastSharedPtr<FMemAllocTable>(GetTable()); }
 
@@ -86,6 +88,7 @@ private:
 	void StartQuery();
 	void UpdateQuery(TraceServices::IAllocationsProvider::EQueryStatus& OutStatus);
 	void CancelQuery();
+	FText GetSymbolResolutionStatus() const;
 
 private:
 	int32 TabIndex = -1;
diff --git a/Engine/Source/Developer/TraceInsights/Private/Insights/Table/Widgets/STableTreeView.cpp b/Engine/Source/Developer/TraceInsights/Private/Insights/Table/Widgets/STableTreeView.cpp
--- a/Engine/Source/Developer/TraceInsights/Private/Insights/Table/Widgets/STableTreeView.cpp
+++ b/Engine/Source/Developer/TraceInsights/Private/Insights/Table/Widgets/STableTreeView.cpp
@@ -105,152 +105,166 @@ void STableTreeView::ConstructWidget(TSharedPtr<FTable> InTablePtr)
 	SAssignNew(ExternalScrollbar, SScrollBar)
 	.AlwaysShowScrollbar(true);
 
-	ChildSlot
-	[
-		SNew(SVerticalBox)
+	auto WidgetContent = 
+	SNew(SVerticalBox)
 
-		+ SVerticalBox::Slot()
-		.VAlign(VAlign_Center)
-		.AutoHeight()
+	+ SVerticalBox::Slot()
+	.VAlign(VAlign_Center)
+	.AutoHeight()
+	[
+		SNew(SBorder)
+		.BorderImage(FEditorStyle::GetBrush("ToolPanel.GroupBorder"))
+		.Padding(2.0f)
 		[
-			SNew(SBorder)
-			.BorderImage(FEditorStyle::GetBrush("ToolPanel.GroupBorder"))
+			SNew(SVerticalBox)
+
+			// Search box
+			+ SVerticalBox::Slot()
+			.VAlign(VAlign_Center)
 			.Padding(2.0f)
+			.AutoHeight()
 			[
-				SNew(SVerticalBox)
+				SNew(SHorizontalBox)
 
-				// Search box
-				+ SVerticalBox::Slot()
+				+ SHorizontalBox::Slot()
+				.FillWidth(1.0)
 				.VAlign(VAlign_Center)
-				.Padding(2.0f)
-				.AutoHeight()
 				[
-					SNew(SHorizontalBox)
-
-					+ SHorizontalBox::Slot()
-					.FillWidth(1.0)
-					.VAlign(VAlign_Center)
-					[
-						SAssignNew(SearchBox, SSearchBox)
-						.HintText(LOCTEXT("SearchBoxHint", "Search"))
-						.OnTextChanged(this, &STableTreeView::SearchBox_OnTextChanged)
-						.IsEnabled(this, &STableTreeView::SearchBox_IsEnabled)
-						.ToolTipText(LOCTEXT("FilterSearchHint", "Type here to search the tree hierarchy by item or group name"))
-					]
-
-					+ SHorizontalBox::Slot()
-						.AutoWidth()
-						.VAlign(VAlign_Center)
-						[
-							SNew(SButton)
-							.Text(LOCTEXT("AdvancedFilters", "Advanced Filters"))
-							.ToolTipText(LOCTEXT("AdvancedFilters", "Opens the filter configurator window."))
-							.OnClicked(this, &STableTreeView::OnAdvancedFiltersClicked)
-						]
+					SAssignNew(SearchBox, SSearchBox)
+					.HintText(LOCTEXT("SearchBoxHint", "Search"))
+					.OnTextChanged(this, &STableTreeView::SearchBox_OnTextChanged)
+					.IsEnabled(this, &STableTreeView::SearchBox_IsEnabled)
+					.ToolTipText(LOCTEXT("FilterSearchHint", "Type here to search the tree hierarchy by item or group name"))
 				]
 
-				// Group by
-				+ SVerticalBox::Slot()
+				+ SHorizontalBox::Slot()
+				.AutoWidth()
 				.VAlign(VAlign_Center)
-				.Padding(2.0f)
-				.AutoHeight()
 				[
-					SNew(SHorizontalBox)
+					SNew(SButton)
+					.Text(LOCTEXT("AdvancedFilters", "Advanced Filters"))
+					.ToolTipText(LOCTEXT("AdvancedFilters", "Opens the filter configurator window."))
+					.OnClicked(this, &STableTreeView::OnAdvancedFiltersClicked)
+				]
+			]
 
-					+ SHorizontalBox::Slot()
-					.AutoWidth()
-					.VAlign(VAlign_Center)
-					[
-						SNew(STextBlock)
-						.Text(LOCTEXT("GroupByText", "Hierarchy:"))
-						.Margin(FMargin(0.0f, 0.0f, 4.0f, 0.0f))
-					]
+			// Group by
+			+ SVerticalBox::Slot()
+			.VAlign(VAlign_Center)
+			.Padding(2.0f)
+			.AutoHeight()
+			[
+				SNew(SHorizontalBox)
 
-					+ SHorizontalBox::Slot()
-					.FillWidth(1.0f)
-					.VAlign(VAlign_Center)
-					[
-						SAssignNew(GroupingBreadcrumbTrail, SBreadcrumbTrail<TSharedPtr<FTreeNodeGrouping>>)
-						.ButtonContentPadding(FMargin(1.0f, 1.0f))
-						//.DelimiterImage(FEditorStyle::GetBrush("SlateFileDialogs.PathDelimiter"))
-						//.TextStyle(FEditorStyle::Get(), "Tutorials.Browser.PathText")
-						//.ShowLeadingDelimiter(true)
-						//.PersistentBreadcrumbs(true)
-						.InvertTextColorOnHover(true)
-						.OnCrumbClicked(this, &STableTreeView::OnGroupingCrumbClicked)
-						.GetCrumbMenuContent(this, &STableTreeView::GetGroupingCrumbMenuContent)
-					]
+				+ SHorizontalBox::Slot()
+				.AutoWidth()
+				.VAlign(VAlign_Center)
+				[
+					SNew(STextBlock)
+					.Text(LOCTEXT("GroupByText", "Hierarchy:"))
+					.Margin(FMargin(0.0f, 0.0f, 4.0f, 0.0f))
+				]
+
+				+ SHorizontalBox::Slot()
+				.FillWidth(1.0f)
+				.VAlign(VAlign_Center)
+				[
+					SAssignNew(GroupingBreadcrumbTrail, SBreadcrumbTrail<TSharedPtr<FTreeNodeGrouping>>)
+					.ButtonContentPadding(FMargin(1.0f, 1.0f))
+					//.DelimiterImage(FEditorStyle::GetBrush("SlateFileDialogs.PathDelimiter"))
+					//.TextStyle(FEditorStyle::Get(), "Tutorials.Browser.PathText")
+					//.ShowLeadingDelimiter(true)
+					//.PersistentBreadcrumbs(true)
+					.InvertTextColorOnHover(true)
+					.OnCrumbClicked(this, &STableTreeView::OnGroupingCrumbClicked)
+					.GetCrumbMenuContent(this, &STableTreeView::GetGroupingCrumbMenuContent)
 				]
 			]
 		]
+	]
 
-		// Tree view
-		+ SVerticalBox::Slot()
-		.FillHeight(1.0f)
-		.Padding(0.0f, 6.0f, 0.0f, 0.0f)
+	// Tree view
+	+ SVerticalBox::Slot()
+	.FillHeight(1.0f)
+	.Padding(0.0f, 6.0f, 0.0f, 0.0f)
+	[
+		SNew(SHorizontalBox)
+
+		+ SHorizontalBox::Slot()
+		.FillWidth(1.0f)
+		.Padding(0.0f)
 		[
-			SNew(SHorizontalBox)
+			SNew(SScrollBox)
+			.Orientation(Orient_Horizontal)
 
-			+ SHorizontalBox::Slot()
-			.FillWidth(1.0f)
-			.Padding(0.0f)
+			+ SScrollBox::Slot()
 			[
-				SNew(SScrollBox)
-				.Orientation(Orient_Horizontal)
+				SNew(SOverlay)
 
-				+ SScrollBox::Slot()
+				+ SOverlay::Slot()
+				.HAlign(HAlign_Fill)
+				.VAlign(VAlign_Fill)
 				[
-					SNew(SOverlay)
-
-					+ SOverlay::Slot()
-					.HAlign(HAlign_Fill)
-					.VAlign(VAlign_Fill)
+					SNew(SBorder)
+					.BorderImage(FEditorStyle::GetBrush("ToolPanel.GroupBorder"))
+					.Padding(0.0f)
 					[
-						SNew(SBorder)
-						.BorderImage(FEditorStyle::GetBrush("ToolPanel.GroupBorder"))
-						.Padding(0.0f)
-						[
-							SAssignNew(TreeView, STreeView<FTableTreeNodePtr>)
-							.ExternalScrollbar(ExternalScrollbar)
-							.SelectionMode(ESelectionMode::Multi)
-							.TreeItemsSource(&FilteredGroupNodes)
-							.OnGetChildren(this, &STableTreeView::TreeView_OnGetChildren)
-							.OnGenerateRow(this, &STableTreeView::TreeView_OnGenerateRow)
-							.OnSelectionChanged(this, &STableTreeView::TreeView_OnSelectionChanged)
-							.OnMouseButtonDoubleClick(this, &STableTreeView::TreeView_OnMouseButtonDoubleClick)
-							.OnContextMenuOpening(FOnContextMenuOpening::CreateSP(this, &STableTreeView::TreeView_GetMenuContent))
-							.ItemHeight(12.0f)
-							.HeaderRow
-							(
-								SAssignNew(TreeViewHeaderRow, SHeaderRow)
-								.Visibility(EVisibility::Visible)
-							)
-						]
+						SAssignNew(TreeView, STreeView<FTableTreeNodePtr>)
+						.ExternalScrollbar(ExternalScrollbar)
+						.SelectionMode(ESelectionMode::Multi)
+						.TreeItemsSource(&FilteredGroupNodes)
+						.OnGetChildren(this, &STableTreeView::TreeView_OnGetChildren)
+						.OnGenerateRow(this, &STableTreeView::TreeView_OnGenerateRow)
+						.OnSelectionChanged(this, &STableTreeView::TreeView_OnSelectionChanged)
+						.OnMouseButtonDoubleClick(this, &STableTreeView::TreeView_OnMouseButtonDoubleClick)
+						.OnContextMenuOpening(FOnContextMenuOpening::CreateSP(this, &STableTreeView::TreeView_GetMenuContent))
+						.ItemHeight(12.0f)
+						.HeaderRow
+						(
+							SAssignNew(TreeViewHeaderRow, SHeaderRow)
+							.Visibility(EVisibility::Visible)
+						)
 					]
+				]
 
-					+ SOverlay::Slot()
-					.HAlign(HAlign_Right)
-					.VAlign(VAlign_Bottom)
-					.Padding(16.0f)
-					[
-						SAssignNew(AsyncOperationStatus, Insights::SAsyncOperationStatus, SharedThis(this))
-					]
+				+ SOverlay::Slot()
+				.HAlign(HAlign_Right)
+				.VAlign(VAlign_Bottom)
+				.Padding(16.0f)
+				[
+					SAssignNew(AsyncOperationStatus, Insights::SAsyncOperationStatus, SharedThis(this))
 				]
 			]
+		]
 
-			+ SHorizontalBox::Slot()
-			.AutoWidth()
-			.Padding(0.0f)
+		+ SHorizontalBox::Slot()
+		.AutoWidth()
+		.Padding(0.0f)
+		[
+			SNew(SBox)
+			.WidthOverride(FOptionalSize(13.0f))
 			[
-				SNew(SBox)
-				.WidthOverride(FOptionalSize(13.0f))
-				[
-					ExternalScrollbar.ToSharedRef()
-				]
+				ExternalScrollbar.ToSharedRef()
 			]
 		]
 	];
 
+	TSharedPtr<SWidget> Footer = ConstructFooter();
+	if (Footer.IsValid())
+	{
+		WidgetContent->AddSlot()
+		.AutoHeight()
+		.Padding(0.0f, 6.0f, 0.0f, 0.0f)
+		[
+			Footer.ToSharedRef()
+		];
+	}
+
+	ChildSlot
+	[
+		WidgetContent
+	];
+
 	// Create the search filters: text based, type based etc.
 	TextFilter = MakeShared<FTableTreeNodeTextFilter>(FTableTreeNodeTextFilter::FItemToStringArray::CreateStatic(&STableTreeView::HandleItemToStringArray));
 	Filters = MakeShared<FTableTreeNodeFilterCollection>();
@@ -2379,6 +2393,13 @@ void STableTreeView::OnAdvancedFiltersChangesCommited()
 
 ////////////////////////////////////////////////////////////////////////////////////////////////////
 
+TSharedPtr<SWidget> STableTreeView::ConstructFooter()
+{
+	return nullptr;
+}
+
+////////////////////////////////////////////////////////////////////////////////////////////////////
+
 } // namespace Insights
 
 #undef LOCTEXT_NAMESPACE
diff --git a/Engine/Source/Developer/TraceInsights/Private/Insights/Table/Widgets/STableTreeView.h b/Engine/Source/Developer/TraceInsights/Private/Insights/Table/Widgets/STableTreeView.h
--- a/Engine/Source/Developer/TraceInsights/Private/Insights/Table/Widgets/STableTreeView.h
+++ b/Engine/Source/Developer/TraceInsights/Private/Insights/Table/Widgets/STableTreeView.h
@@ -127,6 +127,8 @@ public:
 
 protected:
 	void ConstructWidget(TSharedPtr<FTable> InTablePtr);
+	virtual TSharedPtr<SWidget> ConstructFooter();
+
 	void UpdateTree();
 
 	/** Called when the analysis session has changed. */
diff --git a/Engine/Source/Developer/TraceServices/Private/Model/ModuleProvider.cpp b/Engine/Source/Developer/TraceServices/Private/Model/ModuleProvider.cpp
--- a/Engine/Source/Developer/TraceServices/Private/Model/ModuleProvider.cpp
+++ b/Engine/Source/Developer/TraceServices/Private/Model/ModuleProvider.cpp
@@ -45,7 +45,7 @@ public:
 
 	//Query interface
 	const FResolvedSymbol*		GetSymbol(uint64 Address) override;
-	void						GetStats(FStats* OutStats) override;
+	void						GetStats(FStats* OutStats) const override;
 	
 	
 	//Analysis interface
@@ -113,7 +113,7 @@ const FResolvedSymbol* TModuleProvider<SymbolProvider>::GetSymbol(uint64 Address
 
 /////////////////////////////////////////////////////////////////////
 template<typename SymbolResolver>
-void TModuleProvider<SymbolResolver>::GetStats(FStats* OutStats)
+void TModuleProvider<SymbolResolver>::GetStats(FStats* OutStats) const
 {
 	Resolver.GetStats(OutStats);
 }
diff --git a/Engine/Source/Developer/TraceServices/Private/Model/PbdSymbols.inl b/Engine/Source/Developer/TraceServices/Private/Model/PbdSymbols.inl
--- a/Engine/Source/Developer/TraceServices/Private/Model/PbdSymbols.inl
+++ b/Engine/Source/Developer/TraceServices/Private/Model/PbdSymbols.inl
@@ -91,7 +91,7 @@ public:
 		ResolveQueue.Enqueue(QueuedAddress{Address, Symbol});
 	}
 
-	void GetStats(IModuleProvider::FStats* OutStats)
+	void GetStats(IModuleProvider::FStats* OutStats) const
 	{
 		OutStats->ModulesDiscovered = ModulesDiscovered.load();
 		OutStats->ModulesFailed = ModulesFailed.load();
diff --git a/Engine/Source/Developer/TraceServices/Public/TraceServices/Model/Modules.h b/Engine/Source/Developer/TraceServices/Public/TraceServices/Model/Modules.h
--- a/Engine/Source/Developer/TraceServices/Public/TraceServices/Model/Modules.h
+++ b/Engine/Source/Developer/TraceServices/Public/TraceServices/Model/Modules.h
@@ -49,7 +49,7 @@ public:
 	virtual const FResolvedSymbol* GetSymbol(uint64 Address) = 0;
 
 	/** Gets statistics from provider */
-	virtual void GetStats(FStats* OutStats) = 0;
+	virtual void GetStats(FStats* OutStats) const = 0;
 };
 
 ////////////////////////////////////////////////////////////////////////////////
