[Insights] Add tooltip to the Backtrace column in MemAllocTreeView that shows the full callstack.

diff --git a/Engine/Source/Developer/TraceInsights/Private/Insights/MemoryProfiler/ViewModels/MemAllocNode.cpp b/Engine/Source/Developer/TraceInsights/Private/Insights/MemoryProfiler/ViewModels/MemAllocNode.cpp
--- a/Engine/Source/Developer/TraceInsights/Private/Insights/MemoryProfiler/ViewModels/MemAllocNode.cpp
+++ b/Engine/Source/Developer/TraceInsights/Private/Insights/MemoryProfiler/ViewModels/MemAllocNode.cpp
@@ -11,6 +11,18 @@ const FName FMemAllocNode::TypeName(TEXT("FMemAllocNode"));
 
 ////////////////////////////////////////////////////////////////////////////////////////////////////
 
+FText FMemAllocNode::GetFullCallstack() const
+{
+	if (IsValidMemAlloc())
+	{
+		return GetMemAllocChecked().GetFullCallstack();
+	}
+
+	return FText();
+}
+
+////////////////////////////////////////////////////////////////////////////////////////////////////
+
 } // namespace Insights
 
 #undef LOCTEXT_NAMESPACE
diff --git a/Engine/Source/Developer/TraceInsights/Private/Insights/MemoryProfiler/ViewModels/MemAllocNode.h b/Engine/Source/Developer/TraceInsights/Private/Insights/MemoryProfiler/ViewModels/MemAllocNode.h
--- a/Engine/Source/Developer/TraceInsights/Private/Insights/MemoryProfiler/ViewModels/MemAllocNode.h
+++ b/Engine/Source/Developer/TraceInsights/Private/Insights/MemoryProfiler/ViewModels/MemAllocNode.h
@@ -85,6 +85,7 @@ public:
 	bool IsValidMemAlloc() const { return GetMemTableChecked().IsValidRowIndex(GetRowIndex()); }
 	const FMemoryAlloc* GetMemAlloc() const { return GetMemTableChecked().GetMemAlloc(GetRowIndex()); }
 	const FMemoryAlloc& GetMemAllocChecked() const { return GetMemTableChecked().GetMemAllocChecked(GetRowIndex()); }
+	FText GetFullCallstack() const;
 
 private:
 	const EMemAllocNodeType Type;
diff --git a/Engine/Source/Developer/TraceInsights/Private/Insights/MemoryProfiler/ViewModels/MemAllocTable.cpp b/Engine/Source/Developer/TraceInsights/Private/Insights/MemoryProfiler/ViewModels/MemAllocTable.cpp
--- a/Engine/Source/Developer/TraceInsights/Private/Insights/MemoryProfiler/ViewModels/MemAllocTable.cpp
+++ b/Engine/Source/Developer/TraceInsights/Private/Insights/MemoryProfiler/ViewModels/MemAllocTable.cpp
@@ -9,6 +9,9 @@
 #include "Insights/Table/ViewModels/TableCellValueGetter.h"
 #include "Insights/Table/ViewModels/TableCellValueSorter.h"
 #include "Insights/Table/ViewModels/TableColumn.h"
+#include "Widgets/SBoxPanel.h"
+#include "Widgets/SToolTip.h"
+#include "Widgets/Text/STextBlock.h"
 
 #define LOCTEXT_NAMESPACE "Insights::FMemAllocTable"
 
@@ -480,7 +483,29 @@ void FMemAllocTable::AddDefaultColumns()
 		TSharedRef<ITableCellValueGetter> Getter = MakeShared<FBacktraceValueGetter>();
 		Column.SetValueGetter(Getter);
 
-		TSharedRef<ITableCellValueFormatter> Formatter = MakeShared<FCStringValueFormatterAsText>();
+		class BackTraceValueFormatter : public FCStringValueFormatterAsText
+		{
+		public:
+			virtual TSharedPtr<IToolTip> GetCustomTooltip(const FTableColumn& Column, const FBaseTreeNode& Node) const override
+			{
+				const FMemAllocNode& MemAllocNode = static_cast<const FMemAllocNode&>(Node);
+
+				return SNew(SToolTip)
+						[
+							SNew(SVerticalBox)
+
+							+ SVerticalBox::Slot()
+							.AutoHeight()
+							.Padding(2.0f)
+							[
+								SNew(STextBlock)
+								.Text(&MemAllocNode, &FMemAllocNode::GetFullCallstack)
+							]
+						];
+			}
+		};
+
+		TSharedRef<ITableCellValueFormatter> Formatter = MakeShared<BackTraceValueFormatter>();
 		Column.SetValueFormatter(Formatter);
 
 		TSharedRef<ITableCellValueSorter> Sorter = MakeShared<FSorterByCStringValue>(ColumnRef);
diff --git a/Engine/Source/Developer/TraceInsights/Private/Insights/MemoryProfiler/ViewModels/MemoryAlloc.cpp b/Engine/Source/Developer/TraceInsights/Private/Insights/MemoryProfiler/ViewModels/MemoryAlloc.cpp
--- a/Engine/Source/Developer/TraceInsights/Private/Insights/MemoryProfiler/ViewModels/MemoryAlloc.cpp
+++ b/Engine/Source/Developer/TraceInsights/Private/Insights/MemoryProfiler/ViewModels/MemoryAlloc.cpp
@@ -25,4 +25,48 @@ FMemoryAlloc::~FMemoryAlloc()
 
 ////////////////////////////////////////////////////////////////////////////////////////////////////
 
+FText FMemoryAlloc::GetFullCallstack() const
+{
+	static const TCHAR* DisplayStrings[] = {
+	TEXT("Pending..."),
+	TEXT("Not found"),
+	TEXT("N/A"),
+	};
+
+	FString Tooltip;
+	if (Callstack)
+	{
+		check(Callstack->Num() > 0);
+		for (uint32 Index = 0; Index < Callstack->Num(); ++Index)
+		{
+			const TraceServices::FStackFrame* Frame = Callstack->Frame(Index);
+			check(Frame != nullptr);
+
+			const TraceServices::QueryResult Result = Frame->Symbol->Result.load(std::memory_order_acquire);
+			switch (Result)
+			{
+			case TraceServices::QueryResult::QR_NotLoaded:
+				Tooltip.Appendf(TEXT("%s\n"), DisplayStrings[0]); // pending
+				break;
+			case TraceServices::QueryResult::QR_NotFound:
+				Tooltip.Appendf(TEXT("%s\n"), DisplayStrings[1]); // not found
+				break;
+			case TraceServices::QueryResult::QR_OK:
+				Tooltip.Appendf(TEXT("%s\n"), Frame->Symbol->Name);
+				break;
+			}
+		}
+
+		// Remove the last "\n"
+		Tooltip.RemoveAt(Tooltip.Len() - 1, 1, false);
+	}
+
+	if (Tooltip.Len() == 0)
+	{
+		Tooltip.Append(DisplayStrings[2]);
+	}
+
+	return FText::FromString(Tooltip);
+}
+
 } // namespace Insights
diff --git a/Engine/Source/Developer/TraceInsights/Private/Insights/MemoryProfiler/ViewModels/MemoryAlloc.h b/Engine/Source/Developer/TraceInsights/Private/Insights/MemoryProfiler/ViewModels/MemoryAlloc.h
--- a/Engine/Source/Developer/TraceInsights/Private/Insights/MemoryProfiler/ViewModels/MemoryAlloc.h
+++ b/Engine/Source/Developer/TraceInsights/Private/Insights/MemoryProfiler/ViewModels/MemoryAlloc.h
@@ -25,6 +25,7 @@ public:
 	uint64 GetSize() const { return Size; }
 	const TCHAR* GetTag() const { return Tag; }
 	const TraceServices::FCallstack* GetCallstack() const { return Callstack; }
+	FText GetFullCallstack() const;
 
 private:
 	double StartTime;
diff --git a/Engine/Source/Developer/TraceInsights/Private/Insights/Table/ViewModels/TableCellValueFormatter.h b/Engine/Source/Developer/TraceInsights/Private/Insights/Table/ViewModels/TableCellValueFormatter.h
--- a/Engine/Source/Developer/TraceInsights/Private/Insights/Table/ViewModels/TableCellValueFormatter.h
+++ b/Engine/Source/Developer/TraceInsights/Private/Insights/Table/ViewModels/TableCellValueFormatter.h
@@ -6,6 +6,8 @@
 
 #include "Insights/Table/ViewModels/TableCellValue.h"
 
+class IToolTip;
+
 namespace Insights
 {
 
@@ -22,6 +24,8 @@ public:
 
 	virtual FText FormatValue(const FTableColumn& Column, const FBaseTreeNode& Node) const = 0;
 	virtual FText FormatValueForTooltip(const FTableColumn& Column, const FBaseTreeNode& Node) const = 0;
+
+	virtual TSharedPtr<IToolTip> GetCustomTooltip(const FTableColumn& Column, const FBaseTreeNode& Node) const = 0;
 };
 
 ////////////////////////////////////////////////////////////////////////////////////////////////////
@@ -37,6 +41,8 @@ public:
 
 	virtual FText FormatValue(const FTableColumn& Column, const FBaseTreeNode& Node) const override; // { return FormatValue(Column.GetValue(Node)); }
 	virtual FText FormatValueForTooltip(const FTableColumn& Column, const FBaseTreeNode& Node) const override; // { return FormatValueForTooltip(Column.GetValue(Node)); }
+
+	virtual TSharedPtr<IToolTip> GetCustomTooltip(const FTableColumn& Column, const FBaseTreeNode& Node) const override { return nullptr; }
 };
 
 ////////////////////////////////////////////////////////////////////////////////////////////////////
diff --git a/Engine/Source/Developer/TraceInsights/Private/Insights/Table/Widgets/STableTreeViewCell.cpp b/Engine/Source/Developer/TraceInsights/Private/Insights/Table/Widgets/STableTreeViewCell.cpp
--- a/Engine/Source/Developer/TraceInsights/Private/Insights/Table/Widgets/STableTreeViewCell.cpp
+++ b/Engine/Source/Developer/TraceInsights/Private/Insights/Table/Widgets/STableTreeViewCell.cpp
@@ -138,8 +138,8 @@ TSharedRef<SWidget> STableTreeViewCell::GenerateWidgetForTableColumn(const FArgu
 
 	return
 		SNew(SHorizontalBox)
+		.ToolTip(GetTooltip())
 
-		// Value
 		+ SHorizontalBox::Slot()
 		.FillWidth(1.0f)
 		.VAlign(VAlign_Center)
@@ -147,14 +147,20 @@ TSharedRef<SWidget> STableTreeViewCell::GenerateWidgetForTableColumn(const FArgu
 		.Padding(FMargin(2.0f, 0.0f))
 		[
 			TextBox
-		]
-	;
+		];
 }
 
 END_SLATE_FUNCTION_BUILD_OPTIMIZATION
 
 ////////////////////////////////////////////////////////////////////////////////////////////////////
 
+TSharedPtr<IToolTip> STableTreeViewCell::GetTooltip() const
+{
+	return ColumnPtr->GetValueFormatter()->GetCustomTooltip(*ColumnPtr, *TableTreeNodePtr);
+}
+
+////////////////////////////////////////////////////////////////////////////////////////////////////
+
 } // namespace Insights
 
 #undef LOCTEXT_NAMESPACE
diff --git a/Engine/Source/Developer/TraceInsights/Private/Insights/Table/Widgets/STableTreeViewCell.h b/Engine/Source/Developer/TraceInsights/Private/Insights/Table/Widgets/STableTreeViewCell.h
--- a/Engine/Source/Developer/TraceInsights/Private/Insights/Table/Widgets/STableTreeViewCell.h
+++ b/Engine/Source/Developer/TraceInsights/Private/Insights/Table/Widgets/STableTreeViewCell.h
@@ -10,6 +10,7 @@
 #include "Insights/Table/ViewModels/TableTreeNode.h"
 
 class ITableRow;
+class IToolTip;
 
 namespace Insights
 {
@@ -138,6 +139,8 @@ protected:
 		return ShadowColor;
 	}
 
+	TSharedPtr<IToolTip> GetTooltip() const;
+
 protected:
 	/** A shared pointer to the table view model. */
 	TSharedPtr<FTable> TablePtr;
