[Insights]
  - Tree views: Improved formatting for "memory" values in tree table views (ex. Size of alloc).
  - Tree views: Added tooltip for all columns of all data types (default FTableCellValueFormatter).

diff --git a/Engine/Source/Developer/TraceInsights/Private/Insights/Table/ViewModels/TableCellValueFormatter.cpp b/Engine/Source/Developer/TraceInsights/Private/Insights/Table/ViewModels/TableCellValueFormatter.cpp
--- a/Engine/Source/Developer/TraceInsights/Private/Insights/Table/ViewModels/TableCellValueFormatter.cpp
+++ b/Engine/Source/Developer/TraceInsights/Private/Insights/Table/ViewModels/TableCellValueFormatter.cpp
@@ -1,9 +1,12 @@
 // Copyright Epic Games, Inc. All Rights Reserved.
 
 #include "TableCellValueFormatter.h"
+#include "Widgets/SBoxPanel.h"
+#include "Widgets/SToolTip.h"
+#include "Widgets/Text/STextBlock.h"
 
 #include "Insights/Common/TimeUtils.h"
-//#include "Insights/Table/ViewModels/BaseTreeNode.h"
+#include "Insights/Table/ViewModels/TableCellValueGetter.h"
 #include "Insights/Table/ViewModels/TableColumn.h"
 
 #define LOCTEXT_NAMESPACE "TableCellValueFormatter"
@@ -27,6 +30,23 @@ FText FTableCellValueFormatter::FormatValueForTooltip(const FTableColumn& Column
 
 ////////////////////////////////////////////////////////////////////////////////////////////////////
 
+TSharedPtr<IToolTip> FTableCellValueFormatter::GetCustomTooltip(const FTableColumn& Column, const FBaseTreeNode& Node) const
+{
+	return SNew(SToolTip)
+		[
+			SNew(SVerticalBox)
+			+ SVerticalBox::Slot()
+			.AutoHeight()
+			.Padding(2.0f)
+			[
+				SNew(STextBlock)
+				.Text(FormatValueForTooltip(Column.GetValueGetter()->GetValue(Column, Node)))
+			]
+		];
+}
+
+////////////////////////////////////////////////////////////////////////////////////////////////////
+
 FText FBoolValueFormatterAsTrueFalse::FormatValue(const TOptional<FTableCellValue>& InValue) const
 {
 	if (InValue.IsSet())
@@ -54,7 +74,22 @@ FText FInt64ValueFormatterAsMemory::FormatValue(const TOptional<FTableCellValue>
 	if (InValue.IsSet())
 	{
 		const int64 Value = InValue.GetValue().Int64;
-		return FText::AsMemory(Value);
+		if (Value > 0)
+		{
+			FNumberFormattingOptions FormattingOptions;
+			FormattingOptions.MaximumFractionalDigits = 1;
+			return FText::AsMemory(Value, &FormattingOptions);
+		}
+		else if (Value == 0)
+		{
+			return LOCTEXT("AsMemory_ZeroValue", "0");
+		}
+		else // Value < 0
+		{
+			FNumberFormattingOptions FormattingOptions;
+			FormattingOptions.MaximumFractionalDigits = 1;
+			return FText::Format(LOCTEXT("AsMemory_NegativeValue_Fmt1", "-{0}"), FText::AsMemory(-Value, &FormattingOptions));
+		}
 	}
 	return FText::GetEmpty();
 }
@@ -66,13 +101,35 @@ FText FInt64ValueFormatterAsMemory::FormatValueForTooltip(const TOptional<FTable
 	if (InValue.IsSet())
 	{
 		const int64 Value = InValue.GetValue().Int64;
-		if (Value == 0)
+		if (Value > 0)
 		{
-			return FText::FromString(TEXT("0"));
+			if (Value < 1024)
+			{
+				return FText::Format(LOCTEXT("AsMemory_PositiveValue_TooltipFmt1", "{0} bytes"), FText::AsNumber(Value));
+			}
+			else
+			{
+				FNumberFormattingOptions FormattingOptions;
+				FormattingOptions.MaximumFractionalDigits = 2;
+				return FText::Format(LOCTEXT("AsMemory_PositiveValue_TooltipFmt2", "{0} ({1} bytes)"), FText::AsMemory(Value, &FormattingOptions), FText::AsNumber(Value));
+			}
 		}
-		else
+		else if (Value == 0)
+		{
+			return LOCTEXT("AsMemory_ZeroValue", "0");
+		}
+		else // Value < 0
 		{
-			return FText::FromString(FString::Printf(TEXT("%d bytes (%s)"), Value, *FText::AsMemory(Value).ToString()));
+			if (-Value < 1024)
+			{
+				return FText::Format(LOCTEXT("AsMemory_NegativeValue_TooltipFmt1", "-{0} bytes"), FText::AsNumber(-Value));
+			}
+			else
+			{
+				FNumberFormattingOptions FormattingOptions;
+				FormattingOptions.MaximumFractionalDigits = 2;
+				return FText::Format(LOCTEXT("AsMemory_NegativeValue_TooltipFmt2", "-{0} (-{1} bytes)"), FText::AsMemory(-Value, &FormattingOptions), FText::AsNumber(-Value));
+			}
 		}
 	}
 	return FText::GetEmpty();
diff --git a/Engine/Source/Developer/TraceInsights/Private/Insights/Table/ViewModels/TableCellValueFormatter.h b/Engine/Source/Developer/TraceInsights/Private/Insights/Table/ViewModels/TableCellValueFormatter.h
--- a/Engine/Source/Developer/TraceInsights/Private/Insights/Table/ViewModels/TableCellValueFormatter.h
+++ b/Engine/Source/Developer/TraceInsights/Private/Insights/Table/ViewModels/TableCellValueFormatter.h
@@ -42,7 +42,7 @@ public:
 	virtual FText FormatValue(const FTableColumn& Column, const FBaseTreeNode& Node) const override; // { return FormatValue(Column.GetValue(Node)); }
 	virtual FText FormatValueForTooltip(const FTableColumn& Column, const FBaseTreeNode& Node) const override; // { return FormatValueForTooltip(Column.GetValue(Node)); }
 
-	virtual TSharedPtr<IToolTip> GetCustomTooltip(const FTableColumn& Column, const FBaseTreeNode& Node) const override { return nullptr; }
+	virtual TSharedPtr<IToolTip> GetCustomTooltip(const FTableColumn& Column, const FBaseTreeNode& Node) const override;
 };
 
 ////////////////////////////////////////////////////////////////////////////////////////////////////
diff --git a/Engine/Source/Developer/TraceInsights/Private/Insights/ViewModels/TimingGraphTrack.cpp b/Engine/Source/Developer/TraceInsights/Private/Insights/ViewModels/TimingGraphTrack.cpp
--- a/Engine/Source/Developer/TraceInsights/Private/Insights/ViewModels/TimingGraphTrack.cpp
+++ b/Engine/Source/Developer/TraceInsights/Private/Insights/ViewModels/TimingGraphTrack.cpp
@@ -61,7 +61,7 @@ FString FTimingGraphSeries::FormatValue(double Value) const
 		else if (bIsMemory)
 		{
 			const int64 MemValue = static_cast<int64>(Value);
-			if (MemValue >= 0)
+			if (MemValue > 0)
 			{
 				if (MemValue < 1024)
 				{
@@ -74,6 +74,10 @@ FString FTimingGraphSeries::FormatValue(double Value) const
 					return FString::Printf(TEXT("%s (%s bytes)"), *FText::AsMemory(MemValue, &FormattingOptions).ToString(), *FText::AsNumber(MemValue).ToString());
 				}
 			}
+			else if (MemValue == 0)
+			{
+				return TEXT("0");
+			}
 			else
 			{
 				if (-MemValue < 1024)
