[Insights]
  - Memory: Added grouping by Callstack (also supports "inverted" callstack) for Alloc Tree Table views.
  - Memory: Renamed "Backtrace" column to "Function".
  - Memory: Made Start Time, End Time, Duration and Address columns hidden by default.
  - Memory: Added new "Allocation Count" column.
  - Tree views: Fixed groupings when including children groups.
  - Tree views: Disabled width auto adjustment for Hierarchy column (see bAdjustHierarchyColumnWidth).
  - Tree views: Switched column header to "FillWidth".

diff --git a/Engine/Source/Developer/TraceInsights/Private/Insights/MemoryProfiler/ViewModels/CallstackGrouping.cpp b/Engine/Source/Developer/TraceInsights/Private/Insights/MemoryProfiler/ViewModels/CallstackGrouping.cpp
new file mode 100644
--- /dev/null
+++ b/Engine/Source/Developer/TraceInsights/Private/Insights/MemoryProfiler/ViewModels/CallstackGrouping.cpp
@@ -0,0 +1,154 @@
+// Copyright Epic Games, Inc. All Rights Reserved.
+
+#include "CallstackGrouping.h"
+#include "TraceServices/Model/Callstack.h"
+
+// Insights
+#include "Insights/MemoryProfiler/ViewModels/MemAllocNode.h"
+
+#define LOCTEXT_NAMESPACE "Insights::FCallstackGrouping"
+
+namespace Insights
+{
+
+////////////////////////////////////////////////////////////////////////////////////////////////////
+// FCallstackGrouping
+////////////////////////////////////////////////////////////////////////////////////////////////////
+
+FCallstackGrouping::FCallstackGrouping(bool bInIsInverted)
+	: FTreeNodeGrouping(
+		LOCTEXT("Grouping_Callstack_ShortName", "Callstack"),
+		bInIsInverted ? LOCTEXT("Grouping_Callstack2_TitleName", "By Callstack (Inverted)")
+					  : LOCTEXT("Grouping_Callstack1_TitleName", "By Callstack"),
+		LOCTEXT("Grouping_Callstack_Desc", "Creates a tree based on callstack of each allocation."),
+		TEXT("Profiler.FiltersAndPresets.GroupNameIcon"),
+		nullptr)
+	, bIsInverted(bInIsInverted)
+{
+}
+
+////////////////////////////////////////////////////////////////////////////////////////////////////
+
+FCallstackGrouping::~FCallstackGrouping()
+{
+}
+
+////////////////////////////////////////////////////////////////////////////////////////////////////
+
+void FCallstackGrouping::GroupNodes(const TArray<FTableTreeNodePtr>& Nodes, FTableTreeNode& ParentGroup, TWeakPtr<FTable> InParentTable, std::atomic<bool>& bCancelGrouping) const
+{
+	ParentGroup.ClearChildren();
+
+	struct FCallstackGroup
+	{
+		FTableTreeNode* Node = nullptr;
+		const TraceServices::FStackFrame* Frame = nullptr;
+		TMap<uint64, FCallstackGroup*> GroupMap; // Callstack Frame Address -> FCallstackGroup*
+	};
+	TArray<FCallstackGroup*> CallstackGroups;
+
+	FCallstackGroup* Root = new FCallstackGroup();
+	Root->Node = &ParentGroup;
+	CallstackGroups.Add(Root);
+
+	FTableTreeNodePtr UnsetGroupPtr = nullptr;
+
+	static FName PendingName(TEXT("Pending..."));
+	static FName NotFoundName(TEXT("Not found"));
+	static FName NotAvailableName(TEXT("N/A"));
+
+	for (FTableTreeNodePtr NodePtr : Nodes)
+	{
+		if (bCancelGrouping)
+		{
+			return;
+		}
+
+		if (NodePtr->IsGroup())
+		{
+			ParentGroup.AddChildAndSetGroupPtr(NodePtr);
+			continue;
+		}
+
+		FCallstackGroup* GroupPtr = Root;
+
+		const FMemAllocNode& MemAllocNode = static_cast<const FMemAllocNode&>(*NodePtr);
+		const FMemoryAlloc* Alloc = MemAllocNode.GetMemAlloc();
+		if (Alloc)
+		{
+			const TraceServices::FCallstack* Callstack = Alloc->GetCallstack();
+			if (Callstack)
+			{
+				const int32 NumFrames = static_cast<int32>(Callstack->Num());
+				for (int32 FrameDepth = NumFrames - 1; FrameDepth >= 0; --FrameDepth)
+				{
+					const TraceServices::FStackFrame* Frame = Callstack->Frame(bIsInverted ? NumFrames - FrameDepth - 1 : FrameDepth);
+					check(Frame != nullptr);
+
+					FCallstackGroup** GroupPtrPtr = GroupPtr->GroupMap.Find(Frame->Addr);
+					if (!GroupPtrPtr)
+					{
+						FCallstackGroup* NewGroupPtr = new FCallstackGroup();
+						CallstackGroups.Add(NewGroupPtr);
+						GroupPtr->GroupMap.Add(Frame->Addr, NewGroupPtr);
+
+						FName GroupName = NotAvailableName;
+						const TraceServices::QueryResult Result = Frame->Symbol->Result.load(std::memory_order_acquire);
+						switch (Result)
+						{
+						case TraceServices::QueryResult::QR_NotLoaded:
+							//GroupName = PendingName;
+							GroupName = FName(*FString::Printf(TEXT("0x%X (...)"), Frame->Addr));
+							break;
+						case TraceServices::QueryResult::QR_NotFound:
+							//GroupName = NotFoundName;
+							GroupName = FName(*FString::Printf(TEXT("0x%X"), Frame->Addr));
+							break;
+						case TraceServices::QueryResult::QR_OK:
+							//GroupName = FName(Frame->Symbol->Name);
+							GroupName = FName(*FString::Printf(TEXT("0x%X %s"), Frame->Addr, Frame->Symbol->Name));
+							break;
+						}
+
+						FTableTreeNodePtr NewGroupNodePtr = MakeShared<FTableTreeNode>(GroupName, InParentTable);
+						NewGroupPtr->Node = NewGroupNodePtr.Get();
+						NewGroupNodePtr->SetExpansion(false);
+						GroupPtr->Node->AddChildAndSetGroupPtr(NewGroupNodePtr);
+
+						GroupPtr = NewGroupPtr;
+					}
+					else
+					{
+						GroupPtr = *GroupPtrPtr;
+					}
+				}
+			}
+		}
+
+		if (GroupPtr != Root)
+		{
+			GroupPtr->Node->AddChildAndSetGroupPtr(NodePtr);
+		}
+		else
+		{
+			if (!UnsetGroupPtr)
+			{
+				UnsetGroupPtr = MakeShared<FTableTreeNode>(NotAvailableName, InParentTable);
+				UnsetGroupPtr->SetExpansion(false);
+				ParentGroup.AddChildAndSetGroupPtr(UnsetGroupPtr);
+			}
+			UnsetGroupPtr->AddChildAndSetGroupPtr(NodePtr);
+		}
+	}
+
+	for (FCallstackGroup* Group : CallstackGroups)
+	{
+		delete Group;
+	}
+}
+
+////////////////////////////////////////////////////////////////////////////////////////////////////
+
+} // namespace Insights
+
+#undef LOCTEXT_NAMESPACE
diff --git a/Engine/Source/Developer/TraceInsights/Private/Insights/MemoryProfiler/ViewModels/CallstackGrouping.h b/Engine/Source/Developer/TraceInsights/Private/Insights/MemoryProfiler/ViewModels/CallstackGrouping.h
new file mode 100644
--- /dev/null
+++ b/Engine/Source/Developer/TraceInsights/Private/Insights/MemoryProfiler/ViewModels/CallstackGrouping.h
@@ -0,0 +1,30 @@
+// Copyright Epic Games, Inc. All Rights Reserved.
+
+#pragma once
+
+#include "CoreMinimal.h"
+
+// Insights
+#include "Insights/Table/ViewModels/BaseTreeNode.h"
+#include "Insights/Table/ViewModels/TreeNodeGrouping.h"
+
+namespace Insights
+{
+
+////////////////////////////////////////////////////////////////////////////////////////////////////
+
+class FCallstackGrouping : public FTreeNodeGrouping
+{
+public:
+	FCallstackGrouping(bool bInIsInverted);
+	virtual ~FCallstackGrouping();
+
+	virtual void GroupNodes(const TArray<FTableTreeNodePtr>& Nodes, FTableTreeNode& ParentGroup, TWeakPtr<FTable> InParentTable, std::atomic<bool>& bCancelGrouping) const override;
+
+private:
+	bool bIsInverted;
+};
+
+////////////////////////////////////////////////////////////////////////////////////////////////////
+
+} // namespace Insights
diff --git a/Engine/Source/Developer/TraceInsights/Private/Insights/MemoryProfiler/ViewModels/MemAllocTable.cpp b/Engine/Source/Developer/TraceInsights/Private/Insights/MemoryProfiler/ViewModels/MemAllocTable.cpp
--- a/Engine/Source/Developer/TraceInsights/Private/Insights/MemoryProfiler/ViewModels/MemAllocTable.cpp
+++ b/Engine/Source/Developer/TraceInsights/Private/Insights/MemoryProfiler/ViewModels/MemAllocTable.cpp
@@ -60,19 +60,22 @@ void FMemAllocTable::AddDefaultColumns()
 		ColumnRef->SetTitleName(LOCTEXT("AllocationColumnTitle", "Allocation Hierarchy"));
 		ColumnRef->SetDescription(LOCTEXT("AllocationColumnDesc", "Hierarchy of the allocation's tree"));
 	}
+
+	int32 ColumnIndex = 0;
+
 	//////////////////////////////////////////////////
-	// Start Time
+	// Start Time Column
 	{
 		TSharedRef<FTableColumn> ColumnRef = MakeShared<FTableColumn>(TEXT("AllocStartTime"));
 		FTableColumn& Column = *ColumnRef;
 
-		Column.SetIndex(0);
+		Column.SetIndex(ColumnIndex++);
 
 		Column.SetShortName(LOCTEXT("StartTimeColumnName", "Start Time"));
 		Column.SetTitleName(LOCTEXT("StartTimeColumnTitle", "Start Time"));
 		Column.SetDescription(LOCTEXT("StartTimeColumnDesc", "The time when the allocation was allocated."));
 
-		Column.SetFlags(ETableColumnFlags::ShouldBeVisible | ETableColumnFlags::CanBeHidden | ETableColumnFlags::CanBeFiltered);
+		Column.SetFlags(ETableColumnFlags::CanBeHidden | ETableColumnFlags::CanBeFiltered);
 
 		Column.SetHorizontalAlignment(HAlign_Left);
 		Column.SetInitialWidth(100.0f);
@@ -119,18 +122,18 @@ void FMemAllocTable::AddDefaultColumns()
 		AddColumn(ColumnRef);
 	}
 	//////////////////////////////////////////////////
-	// End Time
+	// End Time Column
 	{
 		TSharedRef<FTableColumn> ColumnRef = MakeShared<FTableColumn>(TEXT("AllocEndTime"));
 		FTableColumn& Column = *ColumnRef;
 
-		Column.SetIndex(1);
+		Column.SetIndex(ColumnIndex++);
 
 		Column.SetShortName(LOCTEXT("EndTimeColumnName", "End Time"));
 		Column.SetTitleName(LOCTEXT("EndTimeColumnTitle", "End Time"));
 		Column.SetDescription(LOCTEXT("EndTimeColumnDesc", "The time when the allocation was freed."));
 
-		Column.SetFlags(ETableColumnFlags::ShouldBeVisible | ETableColumnFlags::CanBeHidden | ETableColumnFlags::CanBeFiltered);
+		Column.SetFlags(ETableColumnFlags::CanBeHidden | ETableColumnFlags::CanBeFiltered);
 
 		Column.SetHorizontalAlignment(HAlign_Left);
 		Column.SetInitialWidth(100.0f);
@@ -177,18 +180,18 @@ void FMemAllocTable::AddDefaultColumns()
 		AddColumn(ColumnRef);
 	}
 	//////////////////////////////////////////////////
-	// Duration
+	// Duration Column
 	{
 		TSharedRef<FTableColumn> ColumnRef = MakeShared<FTableColumn>(TEXT("AllocDuration"));
 		FTableColumn& Column = *ColumnRef;
 
-		Column.SetIndex(2);
+		Column.SetIndex(ColumnIndex++);
 
 		Column.SetShortName(LOCTEXT("DurationColumnName", "Duration"));
 		Column.SetTitleName(LOCTEXT("DurationColumnTitle", "Duration"));
 		Column.SetDescription(LOCTEXT("DurationColumnDesc", "The duration of the allocation's life."));
 
-		Column.SetFlags(ETableColumnFlags::ShouldBeVisible | ETableColumnFlags::CanBeHidden | ETableColumnFlags::CanBeFiltered);
+		Column.SetFlags(ETableColumnFlags::CanBeHidden | ETableColumnFlags::CanBeFiltered);
 
 		Column.SetHorizontalAlignment(HAlign_Left);
 		Column.SetInitialWidth(100.0f);
@@ -238,13 +241,13 @@ void FMemAllocTable::AddDefaultColumns()
 		TSharedRef<FTableColumn> ColumnRef = MakeShared<FTableColumn>(TEXT("AllocAddress"));
 		FTableColumn& Column = *ColumnRef;
 
-		Column.SetIndex(3);
+		Column.SetIndex(ColumnIndex++);
 
 		Column.SetShortName(LOCTEXT("AddressColumnName", "Address"));
 		Column.SetTitleName(LOCTEXT("AddressColumnTitle", "Address"));
 		Column.SetDescription(LOCTEXT("AddressColumnDesc", "Address of allocation"));
 
-		Column.SetFlags(ETableColumnFlags::ShouldBeVisible | ETableColumnFlags::CanBeHidden | ETableColumnFlags::CanBeFiltered);
+		Column.SetFlags(ETableColumnFlags::CanBeHidden | ETableColumnFlags::CanBeFiltered);
 
 		Column.SetHorizontalAlignment(HAlign_Left);
 		Column.SetInitialWidth(120.0f);
@@ -289,12 +292,70 @@ void FMemAllocTable::AddDefaultColumns()
 		AddColumn(ColumnRef);
 	}
 	//////////////////////////////////////////////////
+	// Alloc Count Column
+	{
+		TSharedRef<FTableColumn> ColumnRef = MakeShared<FTableColumn>(TEXT("AllocCount"));
+		FTableColumn& Column = *ColumnRef;
+
+		Column.SetIndex(ColumnIndex++);
+
+		Column.SetShortName(LOCTEXT("CountColumnName", "Count"));
+		Column.SetTitleName(LOCTEXT("CountColumnTitle", "Allocation Count"));
+		Column.SetDescription(LOCTEXT("CountColumnDesc", "Number of allocations"));
+
+		Column.SetFlags(ETableColumnFlags::ShouldBeVisible | ETableColumnFlags::CanBeHidden | ETableColumnFlags::CanBeFiltered);
+
+		Column.SetHorizontalAlignment(HAlign_Right);
+		Column.SetInitialWidth(100.0f);
+
+		Column.SetDataType(ETableCellDataType::Int64);
+
+		class FMemAllocCountValueGetter : public FTableCellValueGetter
+		{
+		public:
+			virtual const TOptional<FTableCellValue> GetValue(const FTableColumn& Column, const FBaseTreeNode& Node) const override
+			{
+				if (Node.IsGroup())
+				{
+					const FTableTreeNode& NodePtr = static_cast<const FTableTreeNode&>(Node);
+					if (NodePtr.HasAggregatedValue(Column.GetId()))
+					{
+						return NodePtr.GetAggregatedValue(Column.GetId());
+					}
+				}
+				else //if (Node->Is<FMemAllocNode>())
+				{
+					const FMemAllocNode& MemAllocNode = static_cast<const FMemAllocNode&>(Node);
+					const FMemoryAlloc* Alloc = MemAllocNode.GetMemAlloc();
+					if (Alloc)
+					{
+						return FTableCellValue(static_cast<int64>(1));
+					}
+				}
+
+				return TOptional<FTableCellValue>();
+			}
+		};
+		TSharedRef<ITableCellValueGetter> Getter = MakeShared<FMemAllocCountValueGetter>();
+		Column.SetValueGetter(Getter);
+
+		TSharedRef<ITableCellValueFormatter> Formatter = MakeShared<FInt64ValueFormatterAsNumber>();
+		Column.SetValueFormatter(Formatter);
+
+		TSharedRef<ITableCellValueSorter> Sorter = MakeShared<FSorterByInt64Value>(ColumnRef);
+		Column.SetValueSorter(Sorter);
+
+		Column.SetAggregation(ETableColumnAggregation::Sum);
+
+		AddColumn(ColumnRef);
+	}
+	//////////////////////////////////////////////////
 	// Size Column
 	{
 		TSharedRef<FTableColumn> ColumnRef = MakeShared<FTableColumn>(TEXT("AllocSize"));
 		FTableColumn& Column = *ColumnRef;
 
-		Column.SetIndex(4);
+		Column.SetIndex(ColumnIndex++);
 
 		Column.SetShortName(LOCTEXT("SizeColumnName", "Size"));
 		Column.SetTitleName(LOCTEXT("SizeColumnTitle", "Size"));
@@ -352,7 +413,7 @@ void FMemAllocTable::AddDefaultColumns()
 		TSharedRef<FTableColumn> ColumnRef = MakeShared<FTableColumn>(TEXT("AllocLlmTag"));
 		FTableColumn& Column = *ColumnRef;
 
-		Column.SetIndex(5);
+		Column.SetIndex(ColumnIndex++);
 
 		Column.SetShortName(LOCTEXT("LlmTagColumnName", "LLM Tag"));
 		Column.SetTitleName(LOCTEXT("LlmTagColumnTitle", "LLM Tag"));
@@ -403,16 +464,16 @@ void FMemAllocTable::AddDefaultColumns()
 		AddColumn(ColumnRef);
 	}
 	//////////////////////////////////////////////////
-	// Backtrace Column
+	// Function Column
 	{
-		TSharedRef<FTableColumn> ColumnRef = MakeShared<FTableColumn>(TEXT("AllocBacktrace"));
+		TSharedRef<FTableColumn> ColumnRef = MakeShared<FTableColumn>(TEXT("AllocFunction"));
 		FTableColumn& Column = *ColumnRef;
 
-		Column.SetIndex(6);
+		Column.SetIndex(ColumnIndex++);
 
-		Column.SetShortName(LOCTEXT("BacktraceColumnName", "Backtrace"));
-		Column.SetTitleName(LOCTEXT("BacktraceColumnTitle", "Backtrace"));
-		Column.SetDescription(LOCTEXT("BacktraceColumnDesc", "Backtrace of allocation"));
+		Column.SetShortName(LOCTEXT("FunctionColumnName", "Function"));
+		Column.SetTitleName(LOCTEXT("FunctionColumnTitle", "Function"));
+		Column.SetDescription(LOCTEXT("FunctionColumnDesc", "Resolved top function from the callstack of allocation"));
 
 		Column.SetFlags(ETableColumnFlags::ShouldBeVisible | ETableColumnFlags::CanBeHidden | ETableColumnFlags::CanBeFiltered);
 
@@ -423,7 +484,7 @@ void FMemAllocTable::AddDefaultColumns()
 
 		Column.SetIsDynamic(true);
 
-		class FBacktraceValueGetter : public FTableCellValueGetter
+		class FFunctionValueGetter : public FTableCellValueGetter
 		{
 		public:
 			virtual const TOptional<FTableCellValue> GetValue(const FTableColumn& Column, const FBaseTreeNode& Node) const override
@@ -480,10 +541,10 @@ void FMemAllocTable::AddDefaultColumns()
 				return TOptional<FTableCellValue>();
 			}
 		};
-		TSharedRef<ITableCellValueGetter> Getter = MakeShared<FBacktraceValueGetter>();
+		TSharedRef<ITableCellValueGetter> Getter = MakeShared<FFunctionValueGetter>();
 		Column.SetValueGetter(Getter);
 
-		class BackTraceValueFormatter : public FCStringValueFormatterAsText
+		class FunctionValueFormatter : public FCStringValueFormatterAsText
 		{
 		public:
 			virtual TSharedPtr<IToolTip> GetCustomTooltip(const FTableColumn& Column, const FBaseTreeNode& Node) const override
@@ -493,7 +554,6 @@ void FMemAllocTable::AddDefaultColumns()
 				return SNew(SToolTip)
 						[
 							SNew(SVerticalBox)
-
 							+ SVerticalBox::Slot()
 							.AutoHeight()
 							.Padding(2.0f)
@@ -505,7 +565,7 @@ void FMemAllocTable::AddDefaultColumns()
 			}
 		};
 
-		TSharedRef<ITableCellValueFormatter> Formatter = MakeShared<BackTraceValueFormatter>();
+		TSharedRef<ITableCellValueFormatter> Formatter = MakeShared<FunctionValueFormatter>();
 		Column.SetValueFormatter(Formatter);
 
 		TSharedRef<ITableCellValueSorter> Sorter = MakeShared<FSorterByCStringValue>(ColumnRef);
diff --git a/Engine/Source/Developer/TraceInsights/Private/Insights/MemoryProfiler/Widgets/SMemAllocTableTreeView.cpp b/Engine/Source/Developer/TraceInsights/Private/Insights/MemoryProfiler/Widgets/SMemAllocTableTreeView.cpp
--- a/Engine/Source/Developer/TraceInsights/Private/Insights/MemoryProfiler/Widgets/SMemAllocTableTreeView.cpp
+++ b/Engine/Source/Developer/TraceInsights/Private/Insights/MemoryProfiler/Widgets/SMemAllocTableTreeView.cpp
@@ -9,6 +9,7 @@
 
 // Insights
 #include "Insights/MemoryProfiler/MemoryProfilerManager.h"
+#include "Insights/MemoryProfiler/ViewModels/CallstackGrouping.h"
 #include "Insights/MemoryProfiler/ViewModels/MemAllocNode.h"
 #include "Insights/MemoryProfiler/ViewModels/MemAllocTable.h"
 #include "Insights/MemoryProfiler/ViewModels/MemorySharedState.h"
@@ -414,6 +415,16 @@ FText SMemAllocTableTreeView::GetSymbolResolutionStatus() const
 
 ////////////////////////////////////////////////////////////////////////////////////////////////////
 
+void SMemAllocTableTreeView::InternalCreateGroupings()
+{
+	STableTreeView::InternalCreateGroupings();
+
+	AvailableGroupings.Add(MakeShared<FCallstackGrouping>(false));
+	AvailableGroupings.Add(MakeShared<FCallstackGrouping>(true));
+}
+
+////////////////////////////////////////////////////////////////////////////////////////////////////
+
 } // namespace Insights
 
 #undef LOCTEXT_NAMESPACE
diff --git a/Engine/Source/Developer/TraceInsights/Private/Insights/MemoryProfiler/Widgets/SMemAllocTableTreeView.h b/Engine/Source/Developer/TraceInsights/Private/Insights/MemoryProfiler/Widgets/SMemAllocTableTreeView.h
--- a/Engine/Source/Developer/TraceInsights/Private/Insights/MemoryProfiler/Widgets/SMemAllocTableTreeView.h
+++ b/Engine/Source/Developer/TraceInsights/Private/Insights/MemoryProfiler/Widgets/SMemAllocTableTreeView.h
@@ -73,7 +73,6 @@ public:
 	int32 GetTabIndex() const { return TabIndex; }
 	void SetTabIndex(int32 InTabIndex) { TabIndex = InTabIndex; }
 
-
 	////////////////////////////////////////////////////////////////////////////////////////////////////
 	// IAsyncOperationStatusProvider implementation
 
@@ -83,6 +82,9 @@ public:
 
 	////////////////////////////////////////////////////////////////////////////////////////////////////
 
+protected:
+	virtual void InternalCreateGroupings() override;
+
 private:
 	void OnQueryInvalidated();
 	void StartQuery();
diff --git a/Engine/Source/Developer/TraceInsights/Private/Insights/Table/ViewModels/TreeNodeGrouping.cpp b/Engine/Source/Developer/TraceInsights/Private/Insights/Table/ViewModels/TreeNodeGrouping.cpp
--- a/Engine/Source/Developer/TraceInsights/Private/Insights/Table/ViewModels/TreeNodeGrouping.cpp
+++ b/Engine/Source/Developer/TraceInsights/Private/Insights/Table/ViewModels/TreeNodeGrouping.cpp
@@ -30,13 +30,17 @@ void FTreeNodeGrouping::GroupNodes(const TArray<FTableTreeNodePtr>& Nodes, FTabl
 
 	for (FTableTreeNodePtr NodePtr : Nodes)
 	{
-		ensure(!NodePtr->IsGroup());
-
 		if (bCancelGrouping)
 		{
 			return;
 		}
 
+		if (NodePtr->IsGroup())
+		{
+			ParentGroup.AddChildAndSetGroupPtr(NodePtr);
+			continue;
+		}
+
 		FTableTreeNodePtr GroupPtr = nullptr;
 
 		FTreeNodeGroupInfo GroupInfo = GetGroupForNode(NodePtr);
@@ -73,13 +77,6 @@ FTreeNodeGroupingFlat::FTreeNodeGroupingFlat()
 
 ////////////////////////////////////////////////////////////////////////////////////////////////////
 
-FTreeNodeGroupInfo FTreeNodeGroupingFlat::GetGroupForNode(const FBaseTreeNodePtr InNode) const
-{
-	return { FName(TEXT("All")), true };
-}
-
-////////////////////////////////////////////////////////////////////////////////////////////////////
-
 void FTreeNodeGroupingFlat::GroupNodes(const TArray<FTableTreeNodePtr>& Nodes, FTableTreeNode& ParentGroup, TWeakPtr<FTable> InParentTable, std::atomic<bool>& bCancelGrouping) const
 {
 	ParentGroup.ClearChildren(1);
@@ -109,7 +106,7 @@ FTreeNodeGroupingByUniqueValue::FTreeNodeGroupingByUniqueValue(TSharedRef<FTable
 		InColumnRef->GetTitleName(),
 		FText::Format(LOCTEXT("Grouping_ByUniqueValue_TitleNameFmt", "Unique Values - {0}"), InColumnRef->GetTitleName()),
 		LOCTEXT("Grouping_ByUniqueValue_Desc", "Creates a group for each unique value."),
-		TEXT("Profiler.FiltersAndPresets.Group1NameIcon"), //TODO: "Icons.Grouping.ByName"
+		TEXT("Profiler.FiltersAndPresets.GroupNameIcon"), //TODO: "Icons.Grouping.ByUniqueValue"
 		nullptr)
 	, ColumnRef(InColumnRef)
 {
@@ -142,7 +139,7 @@ FTreeNodeGroupingByNameFirstLetter::FTreeNodeGroupingByNameFirstLetter()
 		LOCTEXT("Grouping_ByName_ShortName", "Name"),
 		LOCTEXT("Grouping_ByName_TitleName", "Name (First Letter)"),
 		LOCTEXT("Grouping_ByName_Desc", "Creates a group for each first letter of node names."),
-		TEXT("Profiler.FiltersAndPresets.Group1NameIcon"), //TODO: "Icons.Grouping.ByName"
+		TEXT("Profiler.FiltersAndPresets.GroupNameIcon"), //TODO: "Icons.Grouping.ByName"
 		nullptr)
 {
 }
diff --git a/Engine/Source/Developer/TraceInsights/Private/Insights/Table/ViewModels/TreeNodeGrouping.h b/Engine/Source/Developer/TraceInsights/Private/Insights/Table/ViewModels/TreeNodeGrouping.h
--- a/Engine/Source/Developer/TraceInsights/Private/Insights/Table/ViewModels/TreeNodeGrouping.h
+++ b/Engine/Source/Developer/TraceInsights/Private/Insights/Table/ViewModels/TreeNodeGrouping.h
@@ -78,7 +78,6 @@ public:
 	FTreeNodeGroupingFlat();
 	virtual ~FTreeNodeGroupingFlat() {}
 
-	virtual FTreeNodeGroupInfo GetGroupForNode(const FBaseTreeNodePtr InNode) const override;
 	virtual void GroupNodes(const TArray<FTableTreeNodePtr>& Nodes, FTableTreeNode& ParentGroup, TWeakPtr<FTable> InParentTable, std::atomic<bool>& bCancelGrouping) const override;
 };
 
@@ -126,13 +125,17 @@ void TTreeNodeGroupingByUniqueValue<Type>::GroupNodes(const TArray<FTableTreeNod
 
 	for (FTableTreeNodePtr NodePtr : Nodes)
 	{
-		ensure(!NodePtr->IsGroup());
-
 		if (bCancelGrouping)
 		{
 			return;
 		}
 
+		if (NodePtr->IsGroup())
+		{
+			ParentGroup.AddChildAndSetGroupPtr(NodePtr);
+			continue;
+		}
+
 		FTableTreeNodePtr GroupPtr = nullptr;
 
 		const TOptional<FTableCellValue> CellValue = GetColumn()->GetValue(*NodePtr);
diff --git a/Engine/Source/Developer/TraceInsights/Private/Insights/Table/Widgets/STableTreeView.cpp b/Engine/Source/Developer/TraceInsights/Private/Insights/Table/Widgets/STableTreeView.cpp
--- a/Engine/Source/Developer/TraceInsights/Private/Insights/Table/Widgets/STableTreeView.cpp
+++ b/Engine/Source/Developer/TraceInsights/Private/Insights/Table/Widgets/STableTreeView.cpp
@@ -139,6 +139,7 @@ void STableTreeView::ConstructWidget(TSharedPtr<FTable> InTablePtr)
 
 				+ SHorizontalBox::Slot()
 				.AutoWidth()
+				.Padding(4.0f, 0.0f, 0.0f, 0.0f)
 				.VAlign(VAlign_Center)
 				[
 					SNew(SButton)
@@ -194,11 +195,11 @@ void STableTreeView::ConstructWidget(TSharedPtr<FTable> InTablePtr)
 		.FillWidth(1.0f)
 		.Padding(0.0f)
 		[
-			SNew(SScrollBox)
-			.Orientation(Orient_Horizontal)
+			//SNew(SScrollBox)
+			//.Orientation(Orient_Horizontal)
 
-			+ SScrollBox::Slot()
-			[
+			//+ SScrollBox::Slot()
+			//[
 				SNew(SOverlay)
 
 				+ SOverlay::Slot()
@@ -210,7 +211,7 @@ void STableTreeView::ConstructWidget(TSharedPtr<FTable> InTablePtr)
 					.Padding(0.0f)
 					[
 						SAssignNew(TreeView, STreeView<FTableTreeNodePtr>)
-						.ExternalScrollbar(ExternalScrollbar)
+						//.ExternalScrollbar(ExternalScrollbar)
 						.SelectionMode(ESelectionMode::Multi)
 						.TreeItemsSource(&FilteredGroupNodes)
 						.OnGetChildren(this, &STableTreeView::TreeView_OnGetChildren)
@@ -234,7 +235,7 @@ void STableTreeView::ConstructWidget(TSharedPtr<FTable> InTablePtr)
 				[
 					SAssignNew(AsyncOperationStatus, Insights::SAsyncOperationStatus, SharedThis(this))
 				]
-			]
+			//]
 		]
 
 		+ SHorizontalBox::Slot()
@@ -1100,6 +1101,20 @@ void STableTreeView::CreateGroupings()
 	AvailableGroupings.Reset();
 	CurrentGroupings.Reset();
 
+	InternalCreateGroupings();
+
+	if (CurrentGroupings.IsEmpty() && AvailableGroupings.Num() > 0)
+	{
+		CurrentGroupings.Add(AvailableGroupings[0]);
+	}
+
+	RebuildGroupingCrumbs();
+}
+
+////////////////////////////////////////////////////////////////////////////////////////////////////
+
+void STableTreeView::InternalCreateGroupings()
+{
 	AvailableGroupings.Add(MakeShared<FTreeNodeGroupingFlat>());
 	//AvailableGroupings.Add(MakeShared<FTreeNodeGroupingByNameFirstLetter>());
 	//AvailableGroupings.Add(MakeShared<FTreeNodeGroupingByType>());
@@ -1108,7 +1123,7 @@ void STableTreeView::CreateGroupings()
 	{
 		if (!ColumnRef->IsHierarchy())
 		{
-			switch(ColumnRef->GetDataType())
+			switch (ColumnRef->GetDataType())
 			{
 				case ETableCellDataType::Bool:
 					AvailableGroupings.Add(MakeShared<FTreeNodeGroupingByUniqueValueBool>(ColumnRef));
@@ -1135,10 +1150,6 @@ void STableTreeView::CreateGroupings()
 			}
 		}
 	}
-
-	CurrentGroupings.Add(AvailableGroupings[0]);
-
-	RebuildGroupingCrumbs();
 }
 
 ////////////////////////////////////////////////////////////////////////////////////////////////////
@@ -1160,12 +1171,15 @@ void STableTreeView::PreChangeGroupings()
 
 void STableTreeView::PostChangeGroupings()
 {
+	constexpr bool bAdjustHierarchyColumnWidth = false;
+	constexpr bool bAdjustHierarchyColumnName = false;
+
 	constexpr float HierarchyMinWidth = 60.0f;
 	constexpr float HierarchyIndentation = 10.0f;
 	constexpr float DefaultHierarchyColumnWidth = 90.0f;
 
 	float HierarchyColumnWidth = DefaultHierarchyColumnWidth;
-	//FString GroupingStr;
+	FString GroupingStr;
 
 	int32 GroupingDepth = 0;
 	for (TSharedPtr<FTreeNodeGrouping>& GroupingPtr : CurrentGroupings)
@@ -1174,19 +1188,22 @@ void STableTreeView::PostChangeGroupings()
 
 		if (ColumnId != NAME_None)
 		{
-			// Compute width for Hierarchy column based on column used in grouping and its indentation.
-			const int32 NumColumns = TreeViewHeaderRow->GetColumns().Num();
-			for (int32 ColumnIndex = 0; ColumnIndex < NumColumns; ++ColumnIndex)
+			if (bAdjustHierarchyColumnWidth)
 			{
-				const SHeaderRow::FColumn& CurrentColumn = TreeViewHeaderRow->GetColumns()[ColumnIndex];
-				if (CurrentColumn.ColumnId == ColumnId)
+				// Compute width for Hierarchy column based on column used in grouping and its indentation.
+				const int32 NumColumns = TreeViewHeaderRow->GetColumns().Num();
+				for (int32 ColumnIndex = 0; ColumnIndex < NumColumns; ++ColumnIndex)
 				{
-					const float Width = HierarchyMinWidth + GroupingDepth * HierarchyIndentation + CurrentColumn.GetWidth();
-					if (Width > HierarchyColumnWidth)
+					const SHeaderRow::FColumn& CurrentColumn = TreeViewHeaderRow->GetColumns()[ColumnIndex];
+					if (CurrentColumn.ColumnId == ColumnId)
 					{
-						HierarchyColumnWidth = Width;
+						const float Width = HierarchyMinWidth + GroupingDepth * HierarchyIndentation + CurrentColumn.GetWidth();
+						if (Width > HierarchyColumnWidth)
+						{
+							HierarchyColumnWidth = Width;
+						}
+						break;
 					}
-					break;
 				}
 			}
 
@@ -1194,12 +1211,15 @@ void STableTreeView::PostChangeGroupings()
 			HideColumn(ColumnId);
 		}
 
-		// Compute name of the Hierarchy column.
-		//if (!GroupingStr.IsEmpty())
-		//{
-		//	GroupingStr.Append(TEXT(" / "));
-		//}
-		//GroupingStr.Append(GroupingPtr->GetShortName().ToString());
+		if (bAdjustHierarchyColumnName)
+		{
+			// Compute name of the Hierarchy column.
+			if (!GroupingStr.IsEmpty())
+			{
+				GroupingStr.Append(TEXT(" / "));
+			}
+			GroupingStr.Append(GroupingPtr->GetShortName().ToString());
+		}
 
 		++GroupingDepth;
 	}
@@ -1208,22 +1228,29 @@ void STableTreeView::PostChangeGroupings()
 
 	if (TreeViewHeaderRow->GetColumns().Num() > 0)
 	{
-		// Set width for the Hierarchy column.
-		SHeaderRow::FColumn& HierarchyColumn = const_cast<SHeaderRow::FColumn&>(TreeViewHeaderRow->GetColumns()[0]);
-		HierarchyColumn.SetWidth(HierarchyColumnWidth);
+		if (bAdjustHierarchyColumnWidth)
+		{
+			// Set width for the Hierarchy column.
+			SHeaderRow::FColumn& HierarchyColumn = const_cast<SHeaderRow::FColumn&>(TreeViewHeaderRow->GetColumns()[0]);
+			HierarchyColumn.SetWidth(HierarchyColumnWidth);
+		}
 
-		// Set name for the Hierarchy column.
-		//FTableColumn& HierarchyTableColumn = *Table->FindColumnChecked(HierarchyColumn.ColumnId);
-		//if (!GroupingStr.IsEmpty())
-		//{
-		//	const FText HierarchyColumnName = FText::Format(LOCTEXT("HierarchyShortNameFmt", "Hierarchy ({0})"), FText::FromString(GroupingStr));
-		//	HierarchyTableColumn.SetShortName(HierarchyColumnName);
-		//}
-		//else
-		//{
-		//	const FText HierarchyColumnName(LOCTEXT("HierarchyShortName", "Hierarchy"));
-		//	HierarchyTableColumn.SetShortName(HierarchyColumnName);
-		//}
+		if (bAdjustHierarchyColumnName)
+		{
+			// Set name for the Hierarchy column.
+			SHeaderRow::FColumn& HierarchyColumn = const_cast<SHeaderRow::FColumn&>(TreeViewHeaderRow->GetColumns()[0]);
+			FTableColumn& HierarchyTableColumn = *Table->FindColumnChecked(HierarchyColumn.ColumnId);
+			if (!GroupingStr.IsEmpty())
+			{
+				const FText HierarchyColumnName = FText::Format(LOCTEXT("HierarchyShortNameFmt", "Hierarchy ({0})"), FText::FromString(GroupingStr));
+				HierarchyTableColumn.SetShortName(HierarchyColumnName);
+			}
+			else
+			{
+				const FText HierarchyColumnName(LOCTEXT("HierarchyShortName", "Hierarchy"));
+				HierarchyTableColumn.SetShortName(HierarchyColumnName);
+			}
+		}
 	}
 
 	//////////////////////////////////////////////////
@@ -1861,8 +1888,9 @@ void STableTreeView::ShowColumn(const FName ColumnId)
 		.VAlignCell(VAlign_Fill)
 		.SortMode(this, &STableTreeView::GetSortModeForColumn, Column.GetId())
 		.OnSort(this, &STableTreeView::OnSortModeChanged)
-		.ManualWidth(Column.GetInitialWidth())
-		.FixedWidth(Column.IsFixedWidth() ? Column.GetInitialWidth() : TOptional<float>())
+		//.ManualWidth(Column.GetInitialWidth())
+		.FillWidth(Column.GetInitialWidth())
+		//.FixedWidth(Column.IsFixedWidth() ? Column.GetInitialWidth() : TOptional<float>())
 		.HeaderContent()
 		[
 			SNew(SBox)
diff --git a/Engine/Source/Developer/TraceInsights/Private/Insights/Table/Widgets/STableTreeView.h b/Engine/Source/Developer/TraceInsights/Private/Insights/Table/Widgets/STableTreeView.h
--- a/Engine/Source/Developer/TraceInsights/Private/Insights/Table/Widgets/STableTreeView.h
+++ b/Engine/Source/Developer/TraceInsights/Private/Insights/Table/Widgets/STableTreeView.h
@@ -206,6 +206,7 @@ protected:
 	// Grouping
 
 	void CreateGroupings();
+	virtual void InternalCreateGroupings();
 
 	void CreateGroups(const TArray<TSharedPtr<FTreeNodeGrouping>>& Groupings);
 	void GroupNodesRec(const TArray<FTableTreeNodePtr>& Nodes, FTableTreeNode& ParentGroup, int32 GroupingDepth, const TArray<TSharedPtr<FTreeNodeGrouping>>& Groupings);
