[insights] Add copy row to clipboard functionality in STableTreeView.
Tooltips are hidden when the context menu is open in STableTreeview.

diff --git a/Engine/Source/Developer/TraceInsights/Private/Insights/LoadingProfiler/LoadingProfilerManager.cpp b/Engine/Source/Developer/TraceInsights/Private/Insights/LoadingProfiler/LoadingProfilerManager.cpp
--- a/Engine/Source/Developer/TraceInsights/Private/Insights/LoadingProfiler/LoadingProfilerManager.cpp
+++ b/Engine/Source/Developer/TraceInsights/Private/Insights/LoadingProfiler/LoadingProfilerManager.cpp
@@ -2,6 +2,7 @@
 
 #include "LoadingProfilerManager.h"
 
+#include "MessageLog/Public/MessageLogModule.h"
 #include "Modules/ModuleManager.h"
 #include "TraceServices/AnalysisService.h"
 #include "TraceServices/Model/LoadTimeProfiler.h"
@@ -60,6 +61,7 @@ FLoadingProfilerManager::FLoadingProfilerManager(TSharedRef<FUICommandList> InCo
 	, bIsPackageDetailsTreeViewVisible(false)
 	, bIsExportDetailsTreeViewVisible(false)
 	, bIsRequestsTreeViewVisible(false)
+	, LogListingName(TEXT("LoadingInsights"))
 {
 }
 
@@ -97,6 +99,16 @@ void FLoadingProfilerManager::Shutdown()
 	}
 	bIsInitialized = false;
 
+	// If the MessageLog module was already unloaded as part of the global Shutdown process, do not load it again.
+	if (FModuleManager::Get().IsModuleLoaded("MessageLog"))
+	{
+		FMessageLogModule& MessageLogModule = FModuleManager::LoadModuleChecked<FMessageLogModule>("MessageLog");
+		if (MessageLogModule.IsRegisteredLogListing(GetLogListingName()))
+		{
+			MessageLogModule.UnregisterLogListing(GetLogListingName());
+		}
+	}
+
 	FInsightsManager::Get()->GetSessionChangedEvent().RemoveAll(this);
 
 	FLoadingProfilerCommands::Unregister();
@@ -248,6 +260,10 @@ bool FLoadingProfilerManager::Tick(float DeltaTime)
 		{
 			bIsAvailable = true;
 
+			FMessageLogModule& MessageLogModule = FModuleManager::LoadModuleChecked<FMessageLogModule>("MessageLog");
+			MessageLogModule.RegisterLogListing(GetLogListingName(), LOCTEXT("LoadingInsights", "Loading Insights"));
+			MessageLogModule.EnableMessageLogDisplay(true);
+
 			const FName& TabId = FInsightsManagerTabs::LoadingProfilerTabId;
 			if (FGlobalTabmanager::Get()->HasTabSpawner(TabId))
 			{
diff --git a/Engine/Source/Developer/TraceInsights/Private/Insights/LoadingProfiler/LoadingProfilerManager.h b/Engine/Source/Developer/TraceInsights/Private/Insights/LoadingProfiler/LoadingProfilerManager.h
--- a/Engine/Source/Developer/TraceInsights/Private/Insights/LoadingProfiler/LoadingProfilerManager.h
+++ b/Engine/Source/Developer/TraceInsights/Private/Insights/LoadingProfiler/LoadingProfilerManager.h
@@ -118,6 +118,8 @@ public:
 
 	void OnSessionChanged();
 
+	const FName& GetLogListingName() const { return LogListingName; }
+
 private:
 	/** Binds our UI commands to delegates. */
 	void BindCommands();
@@ -171,6 +173,9 @@ private:
 	/** If the Requests tree view is visible or hidden. */
 	bool bIsRequestsTreeViewVisible;
 
+	/** The name of the Loading Insights log listing. */
+	FName LogListingName;
+
 	/** A shared pointer to the global instance of the Loading Profiler manager. */
 	static TSharedPtr<FLoadingProfilerManager> Instance;
 };
diff --git a/Engine/Source/Developer/TraceInsights/Private/Insights/LoadingProfiler/Widgets/SLoadingProfilerWindow.cpp b/Engine/Source/Developer/TraceInsights/Private/Insights/LoadingProfiler/Widgets/SLoadingProfilerWindow.cpp
--- a/Engine/Source/Developer/TraceInsights/Private/Insights/LoadingProfiler/Widgets/SLoadingProfilerWindow.cpp
+++ b/Engine/Source/Developer/TraceInsights/Private/Insights/LoadingProfiler/Widgets/SLoadingProfilerWindow.cpp
@@ -357,6 +357,7 @@ TSharedRef<SDockTab> SLoadingProfilerWindow::SpawnTab_EventAggregationTreeView(c
 			SAssignNew(EventAggregationTreeView, Insights::SUntypedTableTreeView, Table)
 		];
 
+	EventAggregationTreeView->SetLogListingName(FLoadingProfilerManager::Get()->GetLogListingName());
 	UpdateEventAggregationTreeView();
 
 	DockTab->SetOnTabClosed(SDockTab::FOnTabClosedCallback::CreateRaw(this, &SLoadingProfilerWindow::OnEventAggregationTreeViewTabClosed));
@@ -387,6 +388,7 @@ TSharedRef<SDockTab> SLoadingProfilerWindow::SpawnTab_ObjectTypeAggregationTreeV
 			SAssignNew(ObjectTypeAggregationTreeView, Insights::SUntypedTableTreeView, Table)
 		];
 
+	ObjectTypeAggregationTreeView->SetLogListingName(FLoadingProfilerManager::Get()->GetLogListingName());
 	UpdateObjectTypeAggregationTreeView();
 
 	DockTab->SetOnTabClosed(SDockTab::FOnTabClosedCallback::CreateRaw(this, &SLoadingProfilerWindow::OnObjectTypeAggregationTreeViewTabClosed));
@@ -417,6 +419,7 @@ TSharedRef<SDockTab> SLoadingProfilerWindow::SpawnTab_PackageDetailsTreeView(con
 			SAssignNew(PackageDetailsTreeView, Insights::SUntypedTableTreeView, Table)
 		];
 
+	PackageDetailsTreeView->SetLogListingName(FLoadingProfilerManager::Get()->GetLogListingName());
 	UpdatePackageDetailsTreeView();
 
 	DockTab->SetOnTabClosed(SDockTab::FOnTabClosedCallback::CreateRaw(this, &SLoadingProfilerWindow::OnPackageDetailsTreeViewTabClosed));
@@ -447,6 +450,7 @@ TSharedRef<SDockTab> SLoadingProfilerWindow::SpawnTab_ExportDetailsTreeView(cons
 			SAssignNew(ExportDetailsTreeView, Insights::SUntypedTableTreeView, Table)
 		];
 
+	ExportDetailsTreeView->SetLogListingName(FLoadingProfilerManager::Get()->GetLogListingName());
 	UpdateExportDetailsTreeView();
 
 	DockTab->SetOnTabClosed(SDockTab::FOnTabClosedCallback::CreateRaw(this, &SLoadingProfilerWindow::OnExportDetailsTreeViewTabClosed));
@@ -477,6 +481,7 @@ TSharedRef<SDockTab> SLoadingProfilerWindow::SpawnTab_RequestsTreeView(const FSp
 			SAssignNew(RequestsTreeView, Insights::SUntypedTableTreeView, Table)
 		];
 
+	RequestsTreeView->SetLogListingName(FLoadingProfilerManager::Get()->GetLogListingName());
 	UpdateRequestsTreeView();
 
 	DockTab->SetOnTabClosed(SDockTab::FOnTabClosedCallback::CreateRaw(this, &SLoadingProfilerWindow::OnRequestsTreeViewTabClosed));
diff --git a/Engine/Source/Developer/TraceInsights/Private/Insights/MemoryProfiler/ViewModels/MemAllocTable.cpp b/Engine/Source/Developer/TraceInsights/Private/Insights/MemoryProfiler/ViewModels/MemAllocTable.cpp
--- a/Engine/Source/Developer/TraceInsights/Private/Insights/MemoryProfiler/ViewModels/MemAllocTable.cpp
+++ b/Engine/Source/Developer/TraceInsights/Private/Insights/MemoryProfiler/ViewModels/MemAllocTable.cpp
@@ -552,6 +552,7 @@ void FMemAllocTable::AddDefaultColumns()
 				const FMemAllocNode& MemAllocNode = static_cast<const FMemAllocNode&>(Node);
 
 				return SNew(SToolTip)
+						.Visibility(TAttribute<EVisibility>::Create(TAttribute<EVisibility>::FGetter::CreateStatic(&FTableCellValueFormatter::GetTooltipVisibility)))
 						[
 							SNew(SVerticalBox)
 							+ SVerticalBox::Slot()
diff --git a/Engine/Source/Developer/TraceInsights/Private/Insights/MemoryProfiler/Widgets/SMemoryProfilerWindow.cpp b/Engine/Source/Developer/TraceInsights/Private/Insights/MemoryProfiler/Widgets/SMemoryProfilerWindow.cpp
--- a/Engine/Source/Developer/TraceInsights/Private/Insights/MemoryProfiler/Widgets/SMemoryProfilerWindow.cpp
+++ b/Engine/Source/Developer/TraceInsights/Private/Insights/MemoryProfiler/Widgets/SMemoryProfilerWindow.cpp
@@ -324,6 +324,7 @@ TSharedRef<SDockTab> SMemoryProfilerWindow::SpawnTab_MemAllocTableTreeView(const
 			SAssignNew(MemAllocTableTreeView, Insights::SMemAllocTableTreeView, MemAllocTable)
 		];
 
+	MemAllocTableTreeView->SetLogListingName(FMemoryProfilerManager::Get()->GetLogListingName());
 	MemAllocTableTreeView->SetTabIndex(TabIndex);
 	MemAllocTableTreeViews.Add(MemAllocTableTreeView);
 	UpdateMemAllocTableTreeView(MemAllocTableTreeView);
diff --git a/Engine/Source/Developer/TraceInsights/Private/Insights/Table/ViewModels/Table.cpp b/Engine/Source/Developer/TraceInsights/Private/Insights/Table/ViewModels/Table.cpp
--- a/Engine/Source/Developer/TraceInsights/Private/Insights/Table/ViewModels/Table.cpp
+++ b/Engine/Source/Developer/TraceInsights/Private/Insights/Table/ViewModels/Table.cpp
@@ -2,6 +2,8 @@
 
 #include "Table.h"
 
+#include "Logging/MessageLog.h"
+
 // Insights
 #include "Insights/Table/ViewModels/TableCellValueFormatter.h"
 #include "Insights/Table/ViewModels/TableCellValueGetter.h"
@@ -120,7 +122,7 @@ void FTable::GetVisibleColumns(TArray<TSharedRef<FTableColumn>>& InArray) const
 
 ////////////////////////////////////////////////////////////////////////////////////////////////////
 
-void FTable::GetVisibleColumnsData(const TArray<Insights::FBaseTreeNodePtr>& InNodes, FString& OutData) const
+void FTable::GetVisibleColumnsData(const TArray<Insights::FBaseTreeNodePtr>& InNodes, const FName& LogListingName, FString& OutData) const
 {
 	constexpr TCHAR Separator = TEXT('\t');
 
@@ -139,9 +141,21 @@ void FTable::GetVisibleColumnsData(const TArray<Insights::FBaseTreeNodePtr>& InN
 		OutData.AppendChar(TEXT('\n'));
 	}
 
+	static int32 Max_Rows = 100000;
+	int32 NumItems = InNodes.Num();
+	if (NumItems > Max_Rows)
+	{
+		NumItems = Max_Rows;
+
+		FMessageLog ReportMessageLog(LogListingName != NAME_None ? LogListingName : TEXT("Other"));
+		ReportMessageLog.Warning(FText::Format(LOCTEXT("TooManyRows", "Too many rows selected. Only the first {0} will be copied."), NumItems));
+		ReportMessageLog.Notify();
+	}
+
 	// Selected items
-	for (Insights::FBaseTreeNodePtr Node : InNodes)
+	for (int Index = 0; Index < NumItems; Index++)
 	{
+		const Insights::FBaseTreeNodePtr& Node = InNodes[Index];
 		for (const TSharedRef<Insights::FTableColumn>& ColumnRef : VisibleColumns)
 		{
 			FText NodeText = ColumnRef->GetValueAsText(*Node);
diff --git a/Engine/Source/Developer/TraceInsights/Private/Insights/Table/ViewModels/Table.h b/Engine/Source/Developer/TraceInsights/Private/Insights/Table/ViewModels/Table.h
--- a/Engine/Source/Developer/TraceInsights/Private/Insights/Table/ViewModels/Table.h
+++ b/Engine/Source/Developer/TraceInsights/Private/Insights/Table/ViewModels/Table.h
@@ -51,7 +51,7 @@ public:
 
 	int32 GetColumnPositionIndex(const FName& ColumnId) const;
 
-	void GetVisibleColumnsData(const TArray<TSharedPtr<class FBaseTreeNode>>& InNodes, FString& OutData) const;
+	void GetVisibleColumnsData(const TArray<TSharedPtr<class FBaseTreeNode>>& InNodes, const FName& LogListingName, FString& OutData) const;
 
 protected:
 	void ResetColumns() { Columns.Reset(); }
diff --git a/Engine/Source/Developer/TraceInsights/Private/Insights/Table/ViewModels/TableCellValueFormatter.cpp b/Engine/Source/Developer/TraceInsights/Private/Insights/Table/ViewModels/TableCellValueFormatter.cpp
--- a/Engine/Source/Developer/TraceInsights/Private/Insights/Table/ViewModels/TableCellValueFormatter.cpp
+++ b/Engine/Source/Developer/TraceInsights/Private/Insights/Table/ViewModels/TableCellValueFormatter.cpp
@@ -1,6 +1,8 @@
 // Copyright Epic Games, Inc. All Rights Reserved.
 
 #include "TableCellValueFormatter.h"
+
+#include "Framework/Application/SlateApplication.h"
 #include "Widgets/SBoxPanel.h"
 #include "Widgets/SToolTip.h"
 #include "Widgets/Text/STextBlock.h"
@@ -33,6 +35,7 @@ FText FTableCellValueFormatter::FormatValueForTooltip(const FTableColumn& Column
 TSharedPtr<IToolTip> FTableCellValueFormatter::GetCustomTooltip(const FTableColumn& Column, const FBaseTreeNode& Node) const
 {
 	return SNew(SToolTip)
+			.Visibility(TAttribute<EVisibility>::Create(TAttribute<EVisibility>::FGetter::CreateStatic(&FTableCellValueFormatter::GetTooltipVisibility)))
 		[
 			SNew(SVerticalBox)
 			+ SVerticalBox::Slot()
@@ -47,6 +50,13 @@ TSharedPtr<IToolTip> FTableCellValueFormatter::GetCustomTooltip(const FTableColu
 
 ////////////////////////////////////////////////////////////////////////////////////////////////////
 
+EVisibility FTableCellValueFormatter::GetTooltipVisibility()
+{
+	return FSlateApplication::Get().AnyMenusVisible() ? EVisibility::Collapsed : EVisibility::Visible;
+}
+
+////////////////////////////////////////////////////////////////////////////////////////////////////
+
 FText FBoolValueFormatterAsTrueFalse::FormatValue(const TOptional<FTableCellValue>& InValue) const
 {
 	if (InValue.IsSet())
diff --git a/Engine/Source/Developer/TraceInsights/Private/Insights/Table/ViewModels/TableCellValueFormatter.h b/Engine/Source/Developer/TraceInsights/Private/Insights/Table/ViewModels/TableCellValueFormatter.h
--- a/Engine/Source/Developer/TraceInsights/Private/Insights/Table/ViewModels/TableCellValueFormatter.h
+++ b/Engine/Source/Developer/TraceInsights/Private/Insights/Table/ViewModels/TableCellValueFormatter.h
@@ -4,6 +4,8 @@
 
 #include "CoreMinimal.h"
 
+#include "Layout/Visibility.h"
+
 #include "Insights/Table/ViewModels/TableCellValue.h"
 
 class IToolTip;
@@ -43,6 +45,8 @@ public:
 	virtual FText FormatValueForTooltip(const FTableColumn& Column, const FBaseTreeNode& Node) const override; // { return FormatValueForTooltip(Column.GetValue(Node)); }
 
 	virtual TSharedPtr<IToolTip> GetCustomTooltip(const FTableColumn& Column, const FBaseTreeNode& Node) const override;
+
+	static EVisibility GetTooltipVisibility();
 };
 
 ////////////////////////////////////////////////////////////////////////////////////////////////////
diff --git a/Engine/Source/Developer/TraceInsights/Private/Insights/Table/Widgets/STableTreeView.cpp b/Engine/Source/Developer/TraceInsights/Private/Insights/Table/Widgets/STableTreeView.cpp
--- a/Engine/Source/Developer/TraceInsights/Private/Insights/Table/Widgets/STableTreeView.cpp
+++ b/Engine/Source/Developer/TraceInsights/Private/Insights/Table/Widgets/STableTreeView.cpp
@@ -3,8 +3,11 @@
 #include "STableTreeView.h"
 
 #include "EditorStyleSet.h"
+#include "Framework/Commands/Commands.h"
+#include "Framework/Commands/UICommandList.h"
 #include "Framework/Application/SlateApplication.h"
 #include "Framework/MultiBox/MultiBoxBuilder.h"
+#include "HAL/PlatformApplicationMisc.h"
 #include "Layout/WidgetPath.h"
 #include "SlateOptMacros.h"
 #include "TraceServices/AnalysisService.h"
@@ -37,6 +40,33 @@
 namespace Insights
 {
 
+////////////////////////////////////////////////////////////////////////////////////////////////////
+// FTableTreeViewCommands
+////////////////////////////////////////////////////////////////////////////////////////////////////
+
+class FTableTreeViewCommands : public TCommands<FTableTreeViewCommands>
+{
+public:
+	FTableTreeViewCommands()
+		: TCommands<FTableTreeViewCommands>(TEXT("FTableTreeViewCommands"), NSLOCTEXT("FTableTreeViewCommands", "Table Tree View Commands", "Table Tree View Commands"), NAME_None, FEditorStyle::Get().GetStyleSetName())
+	{
+	}
+
+	virtual ~FTableTreeViewCommands()
+	{
+	}
+
+	// UI_COMMAND takes long for the compiler to optimize
+	PRAGMA_DISABLE_OPTIMIZATION
+		virtual void RegisterCommands() override
+	{
+		UI_COMMAND(Command_CopyToClipboard, "Copy To Clipboard", "Copies selection to clipboard", EUserInterfaceActionType::Button, FInputChord(EModifierKey::Control, EKeys::C));
+	}
+	PRAGMA_ENABLE_OPTIMIZATION
+
+	TSharedPtr<FUICommandInfo> Command_CopyToClipboard;
+};
+
 const FName STableTreeView::RootNodeName(TEXT("Root"));
 
 ////////////////////////////////////////////////////////////////////////////////////////////////////
@@ -279,12 +309,12 @@ void STableTreeView::ConstructWidget(TSharedPtr<FTable> InTablePtr)
 		Filters->Add(TextFilter);
 	}
 
-	//BindCommands();
-
 	InitializeAndShowHeaderColumns();
 	CreateGroupings();
 	CreateSortings();
 
+	InitCommandList();
+
 	// Register ourselves with the Insights manager.
 	FInsightsManager::Get()->GetSessionChangedEvent().AddSP(this, &STableTreeView::InsightsManager_OnSessionChanged);
 }
@@ -292,6 +322,15 @@ END_SLATE_FUNCTION_BUILD_OPTIMIZATION
 
 ////////////////////////////////////////////////////////////////////////////////////////////////////
 
+void STableTreeView::InitCommandList()
+{
+	FTableTreeViewCommands::Register();
+	CommandList = MakeShared<FUICommandList>();
+	CommandList->MapAction(FTableTreeViewCommands::Get().Command_CopyToClipboard, FExecuteAction::CreateSP(this, &STableTreeView::ContextMenu_CopySelectedToClipboard_Execute), FCanExecuteAction::CreateSP(this, &STableTreeView::ContextMenu_CopySelectedToClipboard_CanExecute));
+}
+
+////////////////////////////////////////////////////////////////////////////////////////////////////
+
 TSharedPtr<SWidget> STableTreeView::TreeView_GetMenuContent()
 {
 	const TArray<FTableTreeNodePtr> SelectedNodes = TreeView->GetSelectedItems();
@@ -329,7 +368,7 @@ TSharedPtr<SWidget> STableTreeView::TreeView_GetMenuContent()
 	}
 
 	const bool bShouldCloseWindowAfterMenuSelection = true;
-	FMenuBuilder MenuBuilder(bShouldCloseWindowAfterMenuSelection, NULL);
+	FMenuBuilder MenuBuilder(bShouldCloseWindowAfterMenuSelection, CommandList.ToSharedRef());
 
 	// Selection menu
 	MenuBuilder.BeginSection("Selection", LOCTEXT("ContextMenu_Header_Selection", "Selection"));
@@ -355,19 +394,14 @@ TSharedPtr<SWidget> STableTreeView::TreeView_GetMenuContent()
 
 	MenuBuilder.BeginSection("Misc", LOCTEXT("ContextMenu_Header_Misc", "Miscellaneous"));
 	{
-		/*TODO
-		FUIAction Action_CopySelectedToClipboard
-		(
-			FExecuteAction::CreateSP(this, &STableTreeView::ContextMenu_CopySelectedToClipboard_Execute),
-			FCanExecuteAction::CreateSP(this, &STableTreeView::ContextMenu_CopySelectedToClipboard_CanExecute)
-		);
 		MenuBuilder.AddMenuEntry
 		(
-			LOCTEXT("ContextMenu_Header_Misc_CopySelectedToClipboard", "Copy To Clipboard"),
-			LOCTEXT("ContextMenu_Header_Misc_CopySelectedToClipboard_Desc", "Copies selection to clipboard"),
-			FSlateIcon(FEditorStyle::GetStyleSetName(), "Profiler.Misc.CopyToClipboard"), Action_CopySelectedToClipboard, NAME_None, EUserInterfaceActionType::Button
+			FTableTreeViewCommands::Get().Command_CopyToClipboard,
+			NAME_None,
+			TAttribute<FText>(),
+			TAttribute<FText>(),
+			FSlateIcon(FEditorStyle::GetStyleSetName(), "Profiler.Misc.CopyToClipboard")
 		);
-		*/
 
 		MenuBuilder.AddSubMenu
 		(
@@ -2428,6 +2462,55 @@ TSharedPtr<SWidget> STableTreeView::ConstructFooter()
 
 ////////////////////////////////////////////////////////////////////////////////////////////////////
 
+FReply STableTreeView::OnKeyDown(const FGeometry& MyGeometry, const FKeyEvent& InKeyEvent)
+{
+	return CommandList->ProcessCommandBindings(InKeyEvent) == true ? FReply::Handled() : FReply::Unhandled();
+}
+
+////////////////////////////////////////////////////////////////////////////////////////////////////
+
+bool STableTreeView::ContextMenu_CopySelectedToClipboard_CanExecute() const
+{
+	return TreeView->GetNumItemsSelected() > 0;
+}
+
+////////////////////////////////////////////////////////////////////////////////////////////////////
+
+void STableTreeView::ContextMenu_CopySelectedToClipboard_Execute()
+{
+	if (!Table->IsValid())
+	{
+		return;
+	}
+
+	TArray<Insights::FBaseTreeNodePtr> SelectedNodes;
+	for (FTableTreeNodePtr TimerPtr : TreeView->GetSelectedItems())
+	{
+		SelectedNodes.Add(TimerPtr);
+	}
+
+	if (SelectedNodes.Num() == 0)
+	{
+		return;
+	}
+
+	FString ClipboardText;
+
+	if (CurrentSorter.IsValid())
+	{
+		CurrentSorter->Sort(SelectedNodes, ColumnSortMode == EColumnSortMode::Ascending ? Insights::ESortMode::Ascending : Insights::ESortMode::Descending);
+	}
+
+	Table->GetVisibleColumnsData(SelectedNodes, GetLogListingName(), ClipboardText);
+
+	if (ClipboardText.Len() > 0)
+	{
+		FPlatformApplicationMisc::ClipboardCopy(*ClipboardText);
+	}
+}
+
+////////////////////////////////////////////////////////////////////////////////////////////////////
+
 } // namespace Insights
 
 #undef LOCTEXT_NAMESPACE
diff --git a/Engine/Source/Developer/TraceInsights/Private/Insights/Table/Widgets/STableTreeView.h b/Engine/Source/Developer/TraceInsights/Private/Insights/Table/Widgets/STableTreeView.h
--- a/Engine/Source/Developer/TraceInsights/Private/Insights/Table/Widgets/STableTreeView.h
+++ b/Engine/Source/Developer/TraceInsights/Private/Insights/Table/Widgets/STableTreeView.h
@@ -26,6 +26,7 @@
 #include <atomic>
 
 class FMenuBuilder;
+class FUICommandList;
 
 namespace TraceServices
 {
@@ -113,6 +114,8 @@ public:
 
 	void OnClose();
 
+	virtual FReply OnKeyDown(const FGeometry& MyGeometry, const FKeyEvent& InKeyEvent) override;
+
 	////////////////////////////////////////////////////////////////////////////////////////////////////
 	// IAsyncOperationStatusProvider implementation
 
@@ -124,9 +127,14 @@ public:
 	virtual FText GetCurrentOperationName() const override;
 
 	////////////////////////////////////////////////////////////////////////////////////////////////////
+	
+	/** Set a log listing name to be used for any errors or warnings. Must be preregistered by the caller with the MessageLog module*/
+	void SetLogListingName(const FName& InLogListingName) { LogListingName = InLogListingName; }
+	const FName& GetLogListingName() { return LogListingName; }
 
 protected:
 	void ConstructWidget(TSharedPtr<FTable> InTablePtr);
+	void InitCommandList();
 	virtual TSharedPtr<SWidget> ConstructFooter();
 
 	void UpdateTree();
@@ -150,6 +158,9 @@ protected:
 	void TreeView_BuildSortByMenu(FMenuBuilder& MenuBuilder);
 	void TreeView_BuildViewColumnMenu(FMenuBuilder& MenuBuilder);
 
+	bool ContextMenu_CopySelectedToClipboard_CanExecute() const;
+	void ContextMenu_CopySelectedToClipboard_Execute();
+
 	////////////////////////////////////////////////////////////////////////////////////////////////////
 	// Tree View - Columns' Header
 
@@ -380,6 +391,8 @@ protected:
 	/** Name of the tree node that should be drawn as highlighted. */
 	FName HighlightedNodeName;
 
+	TSharedPtr<FUICommandList> CommandList;
+
 	//////////////////////////////////////////////////
 	// Tree Nodes
 
@@ -464,6 +477,12 @@ protected:
 
 	double StatsStartTime;
 	double StatsEndTime;
+
+	//////////////////////////////////////////////////
+	// Logging
+
+	/** A log listing name to be used for any errors or warnings. */
+	FName LogListingName;
 };
 
 ////////////////////////////////////////////////////////////////////////////////////////////////////
diff --git a/Engine/Source/Developer/TraceInsights/Private/Insights/TimingProfilerManager.cpp b/Engine/Source/Developer/TraceInsights/Private/Insights/TimingProfilerManager.cpp
--- a/Engine/Source/Developer/TraceInsights/Private/Insights/TimingProfilerManager.cpp
+++ b/Engine/Source/Developer/TraceInsights/Private/Insights/TimingProfilerManager.cpp
@@ -2,6 +2,7 @@
 
 #include "TimingProfilerManager.h"
 
+#include "MessageLog/Public/MessageLogModule.h"
 #include "Modules/ModuleManager.h"
 #include "TraceServices/AnalysisService.h"
 #include "WorkspaceMenuStructure.h"
@@ -70,6 +71,7 @@ FTimingProfilerManager::FTimingProfilerManager(TSharedRef<FUICommandList> InComm
 	, SelectionEndTime(0.0)
 	, SelectedTimerId(InvalidTimerId)
 	, TimerButterflyAggregator(MakeShared<Insights::FTimerButterflyAggregator>())
+	, LogListingName(TEXT("TimingInsights"))
 {
 }
 
@@ -107,6 +109,16 @@ void FTimingProfilerManager::Shutdown()
 	}
 	bIsInitialized = false;
 
+	// If the MessageLog module was already unloaded as part of the global Shutdown process, do not load it again.
+	if (FModuleManager::Get().IsModuleLoaded("MessageLog"))
+	{
+		FMessageLogModule& MessageLogModule = FModuleManager::LoadModuleChecked<FMessageLogModule>("MessageLog");
+		if (MessageLogModule.IsRegisteredLogListing(GetLogListingName()))
+		{
+			MessageLogModule.UnregisterLogListing(GetLogListingName());
+		}
+	}
+
 	FInsightsManager::Get()->GetSessionChangedEvent().RemoveAll(this);
 
 	FTimingProfilerCommands::Unregister();
@@ -238,6 +250,10 @@ bool FTimingProfilerManager::Tick(float DeltaTime)
 		{
 			bIsAvailable = true;
 
+			FMessageLogModule& MessageLogModule = FModuleManager::LoadModuleChecked<FMessageLogModule>("MessageLog");
+			MessageLogModule.RegisterLogListing(GetLogListingName(), LOCTEXT("TimingInsights", "Timing Insights"));
+			MessageLogModule.EnableMessageLogDisplay(true);
+
 #if !WITH_EDITOR
 			const FName& TabId = FInsightsManagerTabs::TimingProfilerTabId;
 			if (FGlobalTabmanager::Get()->HasTabSpawner(TabId))
diff --git a/Engine/Source/Developer/TraceInsights/Private/Insights/TimingProfilerManager.h b/Engine/Source/Developer/TraceInsights/Private/Insights/TimingProfilerManager.h
--- a/Engine/Source/Developer/TraceInsights/Private/Insights/TimingProfilerManager.h
+++ b/Engine/Source/Developer/TraceInsights/Private/Insights/TimingProfilerManager.h
@@ -143,6 +143,8 @@ public:
 	void UpdateAggregatedTimerStats();
 	void UpdateAggregatedCounterStats();
 
+	const FName& GetLogListingName() const { return LogListingName; }
+
 private:
 	/** Binds our UI commands to delegates. */
 	void BindCommands();
@@ -211,6 +213,9 @@ private:
 
 	TSharedRef<Insights::FTimerButterflyAggregator> TimerButterflyAggregator;
 
+	/** The name of the Timing Insights log listing. */
+	FName LogListingName;
+
 	/** A shared pointer to the global instance of the Timing Profiler manager. */
 	static TSharedPtr<FTimingProfilerManager> Instance;
 };
diff --git a/Engine/Source/Developer/TraceInsights/Private/Insights/Widgets/SStatsView.cpp b/Engine/Source/Developer/TraceInsights/Private/Insights/Widgets/SStatsView.cpp
--- a/Engine/Source/Developer/TraceInsights/Private/Insights/Widgets/SStatsView.cpp
+++ b/Engine/Source/Developer/TraceInsights/Private/Insights/Widgets/SStatsView.cpp
@@ -2010,7 +2010,7 @@ void SStatsView::ContextMenu_CopySelectedToClipboard_Execute()
 		CurrentSorter->Sort(SelectedNodes, ColumnSortMode == EColumnSortMode::Ascending ? Insights::ESortMode::Ascending : Insights::ESortMode::Descending);
 	}
 
-	Table->GetVisibleColumnsData(SelectedNodes, ClipboardText);
+	Table->GetVisibleColumnsData(SelectedNodes, FTimingProfilerManager::Get()->GetLogListingName(), ClipboardText);
 
 	if (ClipboardText.Len() > 0)
 	{
diff --git a/Engine/Source/Developer/TraceInsights/Private/Insights/Widgets/STimerTreeView.cpp b/Engine/Source/Developer/TraceInsights/Private/Insights/Widgets/STimerTreeView.cpp
--- a/Engine/Source/Developer/TraceInsights/Private/Insights/Widgets/STimerTreeView.cpp
+++ b/Engine/Source/Developer/TraceInsights/Private/Insights/Widgets/STimerTreeView.cpp
@@ -1112,7 +1112,7 @@ void STimerTreeView::ContextMenu_CopySelectedToClipboard_Execute()
 		CurrentSorter->Sort(SelectedNodes, ColumnSortMode == EColumnSortMode::Ascending ? Insights::ESortMode::Ascending : Insights::ESortMode::Descending);
 	}
 
-	Table->GetVisibleColumnsData(SelectedNodes, ClipboardText);
+	Table->GetVisibleColumnsData(SelectedNodes, FTimingProfilerManager::Get()->GetLogListingName(), ClipboardText);
 
 	if (ClipboardText.Len() > 0)
 	{
diff --git a/Engine/Source/Developer/TraceInsights/Private/Insights/Widgets/STimersView.cpp b/Engine/Source/Developer/TraceInsights/Private/Insights/Widgets/STimersView.cpp
--- a/Engine/Source/Developer/TraceInsights/Private/Insights/Widgets/STimersView.cpp
+++ b/Engine/Source/Developer/TraceInsights/Private/Insights/Widgets/STimersView.cpp
@@ -2023,7 +2023,7 @@ void STimersView::ContextMenu_CopySelectedToClipboard_Execute()
 		CurrentSorter->Sort(SelectedNodes, ColumnSortMode == EColumnSortMode::Ascending ? Insights::ESortMode::Ascending : Insights::ESortMode::Descending);
 	}
 
-	Table->GetVisibleColumnsData(SelectedNodes, ClipboardText);
+	Table->GetVisibleColumnsData(SelectedNodes, FTimingProfilerManager::Get()->GetLogListingName(), ClipboardText);
 
 	if (ClipboardText.Len() > 0)
 	{
