[Insights]
  - Memory / Mem Allocs Tables: Added predefined "breakdown views":
    * Detailed - shows columns for all allocation properties; flat
    * By Size - group allocations by their size (using predefined power of two thresholds)
    * By Tag - group allocations by their LLM tag
    * By Callstack - generate a tree based on callstack of allocations
    * By [REDACTED_PERSON] - generate a tree based on inverted callstack of allocations
  - Added Simple RTTI to FTreeNodeGrouping hierarchy.
  - Added static FNames to Mem Alloc table columns.
  - Renamed FCallstackGrouping to FMemAllocGroupingByCallstack.

diff --git a/Engine/Source/Developer/TraceInsights/Private/Insights/MemoryProfiler/ViewModels/CallstackGrouping.cpp b/Engine/Source/Developer/TraceInsights/Private/Insights/MemoryProfiler/ViewModels/CallstackGrouping.cpp
deleted file mode 100644
diff --git a/Engine/Source/Developer/TraceInsights/Private/Insights/MemoryProfiler/ViewModels/CallstackGrouping.h b/Engine/Source/Developer/TraceInsights/Private/Insights/MemoryProfiler/ViewModels/CallstackGrouping.h
deleted file mode 100644
diff --git a/Engine/Source/Developer/TraceInsights/Private/Insights/MemoryProfiler/ViewModels/MemAllocGroupingByCallstack.cpp b/Engine/Source/Developer/TraceInsights/Private/Insights/MemoryProfiler/ViewModels/MemAllocGroupingByCallstack.cpp
new file mode 100644
--- /dev/null
+++ b/Engine/Source/Developer/TraceInsights/Private/Insights/MemoryProfiler/ViewModels/MemAllocGroupingByCallstack.cpp
@@ -0,0 +1,159 @@
+// Copyright Epic Games, Inc. All Rights Reserved.
+
+#include "MemAllocGroupingByCallstack.h"
+#include "TraceServices/Model/Callstack.h"
+
+// Insights
+#include "Insights/MemoryProfiler/ViewModels/MemAllocNode.h"
+
+#define LOCTEXT_NAMESPACE "Insights::FMemAllocGroupingByCallstack"
+
+namespace Insights
+{
+
+////////////////////////////////////////////////////////////////////////////////////////////////////
+// FMemAllocGroupingByCallstack
+////////////////////////////////////////////////////////////////////////////////////////////////////
+
+INSIGHTS_IMPLEMENT_RTTI(FMemAllocGroupingByCallstack)
+
+////////////////////////////////////////////////////////////////////////////////////////////////////
+
+FMemAllocGroupingByCallstack::FMemAllocGroupingByCallstack(bool bInIsInverted)
+	: FTreeNodeGrouping(
+		bInIsInverted ? LOCTEXT("Grouping_ByCallstack2_ShortName", "Inverted Callstack")
+					  : LOCTEXT("Grouping_ByCallstack1_ShortName", "Callstack"),
+		bInIsInverted ? LOCTEXT("Grouping_ByCallstack2_TitleName", "By Inverted Callstack")
+					  : LOCTEXT("Grouping_ByCallstack1_TitleName", "By Callstack"),
+		LOCTEXT("Grouping_Callstack_Desc", "Creates a tree based on callstack of each allocation."),
+		TEXT("Profiler.FiltersAndPresets.GroupNameIcon"),
+		nullptr)
+	, bIsInverted(bInIsInverted)
+{
+}
+
+////////////////////////////////////////////////////////////////////////////////////////////////////
+
+FMemAllocGroupingByCallstack::~FMemAllocGroupingByCallstack()
+{
+}
+
+////////////////////////////////////////////////////////////////////////////////////////////////////
+
+void FMemAllocGroupingByCallstack::GroupNodes(const TArray<FTableTreeNodePtr>& Nodes, FTableTreeNode& ParentGroup, TWeakPtr<FTable> InParentTable, std::atomic<bool>& bCancelGrouping) const
+{
+	ParentGroup.ClearChildren();
+
+	struct FCallstackGroup
+	{
+		FTableTreeNode* Node = nullptr;
+		const TraceServices::FStackFrame* Frame = nullptr;
+		TMap<uint64, FCallstackGroup*> GroupMap; // Callstack Frame Address -> FCallstackGroup*
+	};
+	TArray<FCallstackGroup*> CallstackGroups;
+
+	FCallstackGroup* Root = new FCallstackGroup();
+	Root->Node = &ParentGroup;
+	CallstackGroups.Add(Root);
+
+	FTableTreeNodePtr UnsetGroupPtr = nullptr;
+
+	static FName PendingName(TEXT("Pending..."));
+	static FName NotFoundName(TEXT("Not found"));
+	static FName NotAvailableName(TEXT("N/A"));
+
+	for (FTableTreeNodePtr NodePtr : Nodes)
+	{
+		if (bCancelGrouping)
+		{
+			return;
+		}
+
+		if (NodePtr->IsGroup())
+		{
+			ParentGroup.AddChildAndSetGroupPtr(NodePtr);
+			continue;
+		}
+
+		FCallstackGroup* GroupPtr = Root;
+
+		const FMemAllocNode& MemAllocNode = static_cast<const FMemAllocNode&>(*NodePtr);
+		const FMemoryAlloc* Alloc = MemAllocNode.GetMemAlloc();
+		if (Alloc)
+		{
+			const TraceServices::FCallstack* Callstack = Alloc->GetCallstack();
+			if (Callstack)
+			{
+				const int32 NumFrames = static_cast<int32>(Callstack->Num());
+				for (int32 FrameDepth = NumFrames - 1; FrameDepth >= 0; --FrameDepth)
+				{
+					const TraceServices::FStackFrame* Frame = Callstack->Frame(bIsInverted ? NumFrames - FrameDepth - 1 : FrameDepth);
+					check(Frame != nullptr);
+
+					FCallstackGroup** GroupPtrPtr = GroupPtr->GroupMap.Find(Frame->Addr);
+					if (!GroupPtrPtr)
+					{
+						FCallstackGroup* NewGroupPtr = new FCallstackGroup();
+						CallstackGroups.Add(NewGroupPtr);
+						GroupPtr->GroupMap.Add(Frame->Addr, NewGroupPtr);
+
+						FName GroupName = NotAvailableName;
+						const TraceServices::QueryResult Result = Frame->Symbol->Result.load(std::memory_order_acquire);
+						switch (Result)
+						{
+						case TraceServices::QueryResult::QR_NotLoaded:
+							//GroupName = PendingName;
+							GroupName = FName(*FString::Printf(TEXT("0x%X (...)"), Frame->Addr));
+							break;
+						case TraceServices::QueryResult::QR_NotFound:
+							//GroupName = NotFoundName;
+							GroupName = FName(*FString::Printf(TEXT("0x%X"), Frame->Addr));
+							break;
+						case TraceServices::QueryResult::QR_OK:
+							//GroupName = FName(Frame->Symbol->Name);
+							GroupName = FName(*FString::Printf(TEXT("0x%X %s"), Frame->Addr, Frame->Symbol->Name));
+							break;
+						}
+
+						FTableTreeNodePtr NewGroupNodePtr = MakeShared<FTableTreeNode>(GroupName, InParentTable);
+						NewGroupPtr->Node = NewGroupNodePtr.Get();
+						NewGroupNodePtr->SetExpansion(false);
+						GroupPtr->Node->AddChildAndSetGroupPtr(NewGroupNodePtr);
+
+						GroupPtr = NewGroupPtr;
+					}
+					else
+					{
+						GroupPtr = *GroupPtrPtr;
+					}
+				}
+			}
+		}
+
+		if (GroupPtr != Root)
+		{
+			GroupPtr->Node->AddChildAndSetGroupPtr(NodePtr);
+		}
+		else
+		{
+			if (!UnsetGroupPtr)
+			{
+				UnsetGroupPtr = MakeShared<FTableTreeNode>(NotAvailableName, InParentTable);
+				UnsetGroupPtr->SetExpansion(false);
+				ParentGroup.AddChildAndSetGroupPtr(UnsetGroupPtr);
+			}
+			UnsetGroupPtr->AddChildAndSetGroupPtr(NodePtr);
+		}
+	}
+
+	for (FCallstackGroup* Group : CallstackGroups)
+	{
+		delete Group;
+	}
+}
+
+////////////////////////////////////////////////////////////////////////////////////////////////////
+
+} // namespace Insights
+
+#undef LOCTEXT_NAMESPACE
diff --git a/Engine/Source/Developer/TraceInsights/Private/Insights/MemoryProfiler/ViewModels/MemAllocGroupingByCallstack.h b/Engine/Source/Developer/TraceInsights/Private/Insights/MemoryProfiler/ViewModels/MemAllocGroupingByCallstack.h
new file mode 100644
--- /dev/null
+++ b/Engine/Source/Developer/TraceInsights/Private/Insights/MemoryProfiler/ViewModels/MemAllocGroupingByCallstack.h
@@ -0,0 +1,34 @@
+// Copyright Epic Games, Inc. All Rights Reserved.
+
+#pragma once
+
+#include "CoreMinimal.h"
+
+// Insights
+#include "Insights/Table/ViewModels/BaseTreeNode.h"
+#include "Insights/Table/ViewModels/TreeNodeGrouping.h"
+
+namespace Insights
+{
+
+////////////////////////////////////////////////////////////////////////////////////////////////////
+
+class FMemAllocGroupingByCallstack : public FTreeNodeGrouping
+{
+	INSIGHTS_DECLARE_RTTI(FMemAllocGroupingByCallstack, FTreeNodeGrouping)
+
+public:
+	FMemAllocGroupingByCallstack(bool bInIsInverted);
+	virtual ~FMemAllocGroupingByCallstack();
+
+	virtual void GroupNodes(const TArray<FTableTreeNodePtr>& Nodes, FTableTreeNode& ParentGroup, TWeakPtr<FTable> InParentTable, std::atomic<bool>& bCancelGrouping) const override;
+
+	bool IsInverted() const { return bIsInverted; }
+
+private:
+	bool bIsInverted;
+};
+
+////////////////////////////////////////////////////////////////////////////////////////////////////
+
+} // namespace Insights
diff --git a/Engine/Source/Developer/TraceInsights/Private/Insights/MemoryProfiler/ViewModels/MemAllocGroupingBySize.cpp b/Engine/Source/Developer/TraceInsights/Private/Insights/MemoryProfiler/ViewModels/MemAllocGroupingBySize.cpp
new file mode 100644
--- /dev/null
+++ b/Engine/Source/Developer/TraceInsights/Private/Insights/MemoryProfiler/ViewModels/MemAllocGroupingBySize.cpp
@@ -0,0 +1,148 @@
+// Copyright Epic Games, Inc. All Rights Reserved.
+
+#include "MemAllocGroupingBySize.h"
+
+// Insights
+#include "Insights/MemoryProfiler/ViewModels/MemAllocNode.h"
+
+#define LOCTEXT_NAMESPACE "Insights::FMemAllocGroupingBySize"
+
+namespace Insights
+{
+
+////////////////////////////////////////////////////////////////////////////////////////////////////
+// FMemAllocGroupingBySize
+////////////////////////////////////////////////////////////////////////////////////////////////////
+
+INSIGHTS_IMPLEMENT_RTTI(FMemAllocGroupingBySize)
+
+////////////////////////////////////////////////////////////////////////////////////////////////////
+
+FMemAllocGroupingBySize::FMemAllocGroupingBySize()
+	: FTreeNodeGrouping(
+		LOCTEXT("Grouping_BySize_ShortName", "Size"),
+		LOCTEXT("Grouping_BySize_TitleName", "By Size"),
+		LOCTEXT("Grouping_BySize_Desc", "Group allocations based on their size."),
+		TEXT("Profiler.FiltersAndPresets.GroupNameIcon"),
+		nullptr)
+{
+	ResetThresholdsPow2();
+}
+
+////////////////////////////////////////////////////////////////////////////////////////////////////
+
+FMemAllocGroupingBySize::~FMemAllocGroupingBySize()
+{
+}
+
+////////////////////////////////////////////////////////////////////////////////////////////////////
+
+void FMemAllocGroupingBySize::ResetThresholdsPow2()
+{
+	bIsPow2 = true;
+
+	constexpr uint32 NumThresholds = 65;
+
+	FNumberFormattingOptions FormattingOptions;
+	FormattingOptions.MinimumIntegralDigits = 2;
+
+	FNumberFormattingOptions FormattingOptionsMem;
+	FormattingOptionsMem.MaximumFractionalDigits = 1;
+
+	Thresholds.Reset(65);
+
+	Thresholds.Add({ 0, LOCTEXT("Grouping_BySize_0", "00 : 0 bytes") });
+	Thresholds.Add({ 1, LOCTEXT("Grouping_BySize_1", "01 : 1 byte") });
+	Thresholds.Add({ 3, LOCTEXT("Grouping_BySize_2", "02 : 2 or 3 bytes") });
+
+	for (uint32 Index = 3; Index < 64; ++Index)
+	{
+		const uint64 Size1 = 1ULL << (Index - 1);
+		const uint64 Size2 = 1ULL << Index;
+		const FText Name = FText::Format(LOCTEXT("Grouping_BySize_Pow2_Fmt", "{0} : [{1} .. {2})"),
+			FText::AsNumber(Index, &FormattingOptions),
+			FText::AsMemory(Size1, &FormattingOptionsMem),
+			FText::AsMemory(Size2, &FormattingOptionsMem));
+		Thresholds.Add({ Size2 - 1ULL, Name });
+	}
+
+	const FText Name64 = FText::Format(LOCTEXT("Grouping_BySize_64_Fmt", "64 : [{0} .. MAX]"), FText::AsMemory(1ULL << 63, &FormattingOptionsMem));
+	Thresholds.Add({ ~0ULL, Name64 });
+}
+
+////////////////////////////////////////////////////////////////////////////////////////////////////
+
+void FMemAllocGroupingBySize::GroupNodes(const TArray<FTableTreeNodePtr>& Nodes, FTableTreeNode& ParentGroup, TWeakPtr<FTable> InParentTable, std::atomic<bool>& bCancelGrouping) const
+{
+	TArray<FTableTreeNodePtr> GroupMap; // Node for each ThresholdIndex
+	GroupMap.AddDefaulted(Thresholds.Num());
+
+	ParentGroup.ClearChildren();
+
+	FTableTreeNodePtr UnsetGroupPtr = nullptr;
+	static FName NotAvailableName(TEXT("N/A"));
+
+	for (FTableTreeNodePtr NodePtr : Nodes)
+	{
+		if (bCancelGrouping)
+		{
+			return;
+		}
+
+		if (NodePtr->IsGroup())
+		{
+			ParentGroup.AddChildAndSetGroupPtr(NodePtr);
+			continue;
+		}
+
+		FTableTreeNodePtr GroupPtr = nullptr;
+
+		const FMemAllocNode& MemAllocNode = static_cast<const FMemAllocNode&>(*NodePtr);
+		const FMemoryAlloc* Alloc = MemAllocNode.GetMemAlloc();
+		if (Alloc)
+		{
+			const uint64 Size = Alloc->GetSize();
+
+			uint32 ThresholdIndex;
+			if (bIsPow2)
+			{
+				ThresholdIndex = 64 - FMath::CountLeadingZeros64(Size);
+			}
+			else
+			{
+				ThresholdIndex = Algo::UpperBoundBy(Thresholds, Size, &FThreshold::Size) - 1;
+			}
+
+			GroupPtr = GroupMap[ThresholdIndex];
+			if (!GroupPtr)
+			{
+				const FThreshold& Threshold = Thresholds[ThresholdIndex];
+				GroupPtr = MakeShared<FTableTreeNode>(FName(*Threshold.Name.ToString()), InParentTable);
+				GroupPtr->SetExpansion(false);
+				ParentGroup.AddChildAndSetGroupPtr(GroupPtr);
+				GroupMap[ThresholdIndex] = GroupPtr;
+			}
+		}
+
+		if (GroupPtr != nullptr)
+		{
+			GroupPtr->AddChildAndSetGroupPtr(NodePtr);
+		}
+		else
+		{
+			if (!UnsetGroupPtr)
+			{
+				UnsetGroupPtr = MakeShared<FTableTreeNode>(NotAvailableName, InParentTable);
+				UnsetGroupPtr->SetExpansion(false);
+				ParentGroup.AddChildAndSetGroupPtr(UnsetGroupPtr);
+			}
+			UnsetGroupPtr->AddChildAndSetGroupPtr(NodePtr);
+		}
+	}
+}
+
+////////////////////////////////////////////////////////////////////////////////////////////////////
+
+} // namespace Insights
+
+#undef LOCTEXT_NAMESPACE
diff --git a/Engine/Source/Developer/TraceInsights/Private/Insights/MemoryProfiler/ViewModels/MemAllocGroupingBySize.h b/Engine/Source/Developer/TraceInsights/Private/Insights/MemoryProfiler/ViewModels/MemAllocGroupingBySize.h
new file mode 100644
--- /dev/null
+++ b/Engine/Source/Developer/TraceInsights/Private/Insights/MemoryProfiler/ViewModels/MemAllocGroupingBySize.h
@@ -0,0 +1,46 @@
+// Copyright Epic Games, Inc. All Rights Reserved.
+
+#pragma once
+
+#include "CoreMinimal.h"
+
+// Insights
+#include "Insights/Table/ViewModels/BaseTreeNode.h"
+#include "Insights/Table/ViewModels/TreeNodeGrouping.h"
+
+namespace Insights
+{
+
+////////////////////////////////////////////////////////////////////////////////////////////////////
+
+class FMemAllocGroupingBySize : public FTreeNodeGrouping
+{
+	INSIGHTS_DECLARE_RTTI(FMemAllocGroupingBySize, FTreeNodeGrouping)
+
+public:
+	struct FThreshold
+	{
+		uint64 Size; // inclusive upper limit
+		FText Name; // group name
+	};
+
+public:
+	FMemAllocGroupingBySize();
+	virtual ~FMemAllocGroupingBySize();
+
+	const TArray<FThreshold>& GetThresholds() const { return Thresholds; }
+	TArray<FThreshold>& EditThresholds() { bIsPow2 = false; return Thresholds; }
+
+	bool IsPow2() const { return bIsPow2; }
+	void ResetThresholdsPow2();
+
+	virtual void GroupNodes(const TArray<FTableTreeNodePtr>& Nodes, FTableTreeNode& ParentGroup, TWeakPtr<FTable> InParentTable, std::atomic<bool>& bCancelGrouping) const override;
+
+private:
+	TArray<FThreshold> Thresholds;
+	bool bIsPow2; // thresholds are automatically set as power of two
+};
+
+////////////////////////////////////////////////////////////////////////////////////////////////////
+
+} // namespace Insights
diff --git a/Engine/Source/Developer/TraceInsights/Private/Insights/MemoryProfiler/ViewModels/MemAllocTable.cpp b/Engine/Source/Developer/TraceInsights/Private/Insights/MemoryProfiler/ViewModels/MemAllocTable.cpp
--- a/Engine/Source/Developer/TraceInsights/Private/Insights/MemoryProfiler/ViewModels/MemAllocTable.cpp
+++ b/Engine/Source/Developer/TraceInsights/Private/Insights/MemoryProfiler/ViewModels/MemAllocTable.cpp
@@ -18,6 +18,18 @@
 namespace Insights
 {
 
+////////////////////////////////////////////////////////////////////////////////////////////////////
+// Column identifiers
+
+const FName FMemAllocTableColumns::StartTimeColumnId(TEXT("StartTime"));
+const FName FMemAllocTableColumns::EndTimeColumnId(TEXT("EndTime"));
+const FName FMemAllocTableColumns::DurationColumnId(TEXT("Duration"));
+const FName FMemAllocTableColumns::AddressColumnId(TEXT("Address"));
+const FName FMemAllocTableColumns::CountColumnId(TEXT("Count"));
+const FName FMemAllocTableColumns::SizeColumnId(TEXT("Size"));
+const FName FMemAllocTableColumns::TagColumnId(TEXT("Tag"));
+const FName FMemAllocTableColumns::FunctionColumnId(TEXT("Function"));
+
 ////////////////////////////////////////////////////////////////////////////////////////////////////
 // FMemAllocTable
 ////////////////////////////////////////////////////////////////////////////////////////////////////
@@ -50,7 +62,7 @@ void FMemAllocTable::AddDefaultColumns()
 	//////////////////////////////////////////////////
 	// Hierarchy Column
 	{
-		int32 HierarchyColumnIndex = -1;
+		const int32 HierarchyColumnIndex = -1;
 		const TCHAR* HierarchyColumnName = nullptr;
 		AddHierarchyColumn(HierarchyColumnIndex, HierarchyColumnName);
 
@@ -66,7 +78,7 @@ void FMemAllocTable::AddDefaultColumns()
 	//////////////////////////////////////////////////
 	// Start Time Column
 	{
-		TSharedRef<FTableColumn> ColumnRef = MakeShared<FTableColumn>(TEXT("AllocStartTime"));
+		TSharedRef<FTableColumn> ColumnRef = MakeShared<FTableColumn>(FMemAllocTableColumns::StartTimeColumnId);
 		FTableColumn& Column = *ColumnRef;
 
 		Column.SetIndex(ColumnIndex++);
@@ -124,7 +136,7 @@ void FMemAllocTable::AddDefaultColumns()
 	//////////////////////////////////////////////////
 	// End Time Column
 	{
-		TSharedRef<FTableColumn> ColumnRef = MakeShared<FTableColumn>(TEXT("AllocEndTime"));
+		TSharedRef<FTableColumn> ColumnRef = MakeShared<FTableColumn>(FMemAllocTableColumns::EndTimeColumnId);
 		FTableColumn& Column = *ColumnRef;
 
 		Column.SetIndex(ColumnIndex++);
@@ -182,7 +194,7 @@ void FMemAllocTable::AddDefaultColumns()
 	//////////////////////////////////////////////////
 	// Duration Column
 	{
-		TSharedRef<FTableColumn> ColumnRef = MakeShared<FTableColumn>(TEXT("AllocDuration"));
+		TSharedRef<FTableColumn> ColumnRef = MakeShared<FTableColumn>(FMemAllocTableColumns::DurationColumnId);
 		FTableColumn& Column = *ColumnRef;
 
 		Column.SetIndex(ColumnIndex++);
@@ -238,7 +250,7 @@ void FMemAllocTable::AddDefaultColumns()
 	//////////////////////////////////////////////////
 	// Address Column
 	{
-		TSharedRef<FTableColumn> ColumnRef = MakeShared<FTableColumn>(TEXT("AllocAddress"));
+		TSharedRef<FTableColumn> ColumnRef = MakeShared<FTableColumn>(FMemAllocTableColumns::AddressColumnId);
 		FTableColumn& Column = *ColumnRef;
 
 		Column.SetIndex(ColumnIndex++);
@@ -294,7 +306,7 @@ void FMemAllocTable::AddDefaultColumns()
 	//////////////////////////////////////////////////
 	// Alloc Count Column
 	{
-		TSharedRef<FTableColumn> ColumnRef = MakeShared<FTableColumn>(TEXT("AllocCount"));
+		TSharedRef<FTableColumn> ColumnRef = MakeShared<FTableColumn>(FMemAllocTableColumns::CountColumnId);
 		FTableColumn& Column = *ColumnRef;
 
 		Column.SetIndex(ColumnIndex++);
@@ -352,7 +364,7 @@ void FMemAllocTable::AddDefaultColumns()
 	//////////////////////////////////////////////////
 	// Size Column
 	{
-		TSharedRef<FTableColumn> ColumnRef = MakeShared<FTableColumn>(TEXT("AllocSize"));
+		TSharedRef<FTableColumn> ColumnRef = MakeShared<FTableColumn>(FMemAllocTableColumns::SizeColumnId);
 		FTableColumn& Column = *ColumnRef;
 
 		Column.SetIndex(ColumnIndex++);
@@ -410,14 +422,14 @@ void FMemAllocTable::AddDefaultColumns()
 	//////////////////////////////////////////////////
 	// LLM Tag Column
 	{
-		TSharedRef<FTableColumn> ColumnRef = MakeShared<FTableColumn>(TEXT("AllocLlmTag"));
+		TSharedRef<FTableColumn> ColumnRef = MakeShared<FTableColumn>(FMemAllocTableColumns::TagColumnId);
 		FTableColumn& Column = *ColumnRef;
 
 		Column.SetIndex(ColumnIndex++);
 
-		Column.SetShortName(LOCTEXT("LlmTagColumnName", "LLM Tag"));
-		Column.SetTitleName(LOCTEXT("LlmTagColumnTitle", "LLM Tag"));
-		Column.SetDescription(LOCTEXT("LlmTagColumnDesc", "LLM tag of allocation"));
+		Column.SetShortName(LOCTEXT("TagColumnName", "LLM Tag"));
+		Column.SetTitleName(LOCTEXT("TagColumnTitle", "LLM Tag"));
+		Column.SetDescription(LOCTEXT("TagColumnDesc", "LLM tag of allocation"));
 
 		Column.SetFlags(ETableColumnFlags::ShouldBeVisible | ETableColumnFlags::CanBeHidden | ETableColumnFlags::CanBeFiltered);
 
@@ -466,7 +478,7 @@ void FMemAllocTable::AddDefaultColumns()
 	//////////////////////////////////////////////////
 	// Function Column
 	{
-		TSharedRef<FTableColumn> ColumnRef = MakeShared<FTableColumn>(TEXT("AllocFunction"));
+		TSharedRef<FTableColumn> ColumnRef = MakeShared<FTableColumn>(FMemAllocTableColumns::FunctionColumnId);
 		FTableColumn& Column = *ColumnRef;
 
 		Column.SetIndex(ColumnIndex++);
diff --git a/Engine/Source/Developer/TraceInsights/Private/Insights/MemoryProfiler/ViewModels/MemAllocTable.h b/Engine/Source/Developer/TraceInsights/Private/Insights/MemoryProfiler/ViewModels/MemAllocTable.h
--- a/Engine/Source/Developer/TraceInsights/Private/Insights/MemoryProfiler/ViewModels/MemAllocTable.h
+++ b/Engine/Source/Developer/TraceInsights/Private/Insights/MemoryProfiler/ViewModels/MemAllocTable.h
@@ -15,6 +15,21 @@ class FTableColumn;
 
 ////////////////////////////////////////////////////////////////////////////////////////////////////
 
+// Column identifiers
+struct FMemAllocTableColumns
+{
+	static const FName StartTimeColumnId;
+	static const FName EndTimeColumnId;
+	static const FName DurationColumnId;
+	static const FName AddressColumnId;
+	static const FName CountColumnId;
+	static const FName SizeColumnId;
+	static const FName TagColumnId;
+	static const FName FunctionColumnId;
+};
+
+////////////////////////////////////////////////////////////////////////////////////////////////////
+
 class FMemAllocTable : public FTable
 {
 public:
diff --git a/Engine/Source/Developer/TraceInsights/Private/Insights/MemoryProfiler/Widgets/SMemAllocTableTreeView.cpp b/Engine/Source/Developer/TraceInsights/Private/Insights/MemoryProfiler/Widgets/SMemAllocTableTreeView.cpp
--- a/Engine/Source/Developer/TraceInsights/Private/Insights/MemoryProfiler/Widgets/SMemAllocTableTreeView.cpp
+++ b/Engine/Source/Developer/TraceInsights/Private/Insights/MemoryProfiler/Widgets/SMemAllocTableTreeView.cpp
@@ -9,7 +9,8 @@
 
 // Insights
 #include "Insights/MemoryProfiler/MemoryProfilerManager.h"
-#include "Insights/MemoryProfiler/ViewModels/CallstackGrouping.h"
+#include "Insights/MemoryProfiler/ViewModels/MemAllocGroupingByCallstack.h"
+#include "Insights/MemoryProfiler/ViewModels/MemAllocGroupingBySize.h"
 #include "Insights/MemoryProfiler/ViewModels/MemAllocNode.h"
 #include "Insights/MemoryProfiler/ViewModels/MemAllocTable.h"
 #include "Insights/MemoryProfiler/ViewModels/MemorySharedState.h"
@@ -387,6 +388,245 @@ FText SMemAllocTableTreeView::GetCurrentOperationName() const
 
 ////////////////////////////////////////////////////////////////////////////////////////////////////
 
+TSharedPtr<SWidget> SMemAllocTableTreeView::ConstructToolbar()
+{
+	return
+		SNew(SHorizontalBox)
+
+		+ SHorizontalBox::Slot()
+		.AutoWidth()
+		.Padding(4.0f, 0.0f, 0.0f, 0.0f)
+		[
+			SNew(SButton)
+			.Text(LOCTEXT("DetailedBtn_Text", "Detailed"))
+			.ToolTipText(LOCTEXT("DetailedBtn_Tooltip", "Detailed View\nConfigure the tree view to show detailed allocation info."))
+			.OnClicked(this, &SMemAllocTableTreeView::OnDetailedViewClicked)
+		]
+
+		+ SHorizontalBox::Slot()
+		.AutoWidth()
+		.Padding(4.0f, 0.0f, 0.0f, 0.0f)
+		[
+			SNew(SButton)
+			.Text(LOCTEXT("BySizeBtn_Text", "By Size"))
+			.ToolTipText(LOCTEXT("BySizeBtn_Tooltip", "Size Breakdown View\nConfigure the tree view to show a breakdown of allocations by their size."))
+			.OnClicked(this, &SMemAllocTableTreeView::OnSizeViewClicked)
+		]
+
+		+ SHorizontalBox::Slot()
+		.AutoWidth()
+		.Padding(4.0f, 0.0f, 0.0f, 0.0f)
+		[
+			SNew(SButton)
+			.Text(LOCTEXT("ByTagBtn_Text", "By Tag"))
+			.ToolTipText(LOCTEXT("ByTagBtn_Tooltip", "Tag Breakdown View\nConfigure the tree view to show a breakdown of allocations by their LLM tags."))
+			.OnClicked(this, &SMemAllocTableTreeView::OnTagViewClicked)
+		]
+
+		+ SHorizontalBox::Slot()
+		.AutoWidth()
+		.Padding(4.0f, 0.0f, 0.0f, 0.0f)
+		[
+			SNew(SButton)
+			.Text(LOCTEXT("ByCallstackBtn_Text", "By Callstack"))
+			.ToolTipText(LOCTEXT("ByCallstackBtn_Tooltip", "Callstack Breakdown View\nConfigure the tree view to show a breakdown of allocations by callstack."))
+			.OnClicked(this, &SMemAllocTableTreeView::OnCallstackViewClicked, false)
+		]
+
+		+ SHorizontalBox::Slot()
+		.AutoWidth()
+		.Padding(4.0f, 0.0f, 0.0f, 0.0f)
+		[
+			SNew(SButton)
+			.Text(LOCTEXT("ByInvertedCallstackBtn_Text", "By Inverted Callstack"))
+			.ToolTipText(LOCTEXT("ByInvertedCallstackBtn_Tooltip", "Inverted Callstack Breakdown View\nConfigure the tree view to show a breakdown of allocations by inverted callstack."))
+			.OnClicked(this, &SMemAllocTableTreeView::OnCallstackViewClicked, true)
+		]
+	;
+}
+
+////////////////////////////////////////////////////////////////////////////////////////////////////
+
+FReply SMemAllocTableTreeView::OnDetailedViewClicked()
+{
+	PreChangeGroupings();
+
+	CurrentGroupings.Reset();
+
+	check(AvailableGroupings[0]->Is<FTreeNodeGroupingFlat>());
+	CurrentGroupings.Add(AvailableGroupings[0]);
+
+	PostChangeGroupings();
+
+	FColumnConfig Preset[] =
+	{
+		{ FTable::GetHierarchyColumnId(),           true,  200.0f },
+		{ FMemAllocTableColumns::StartTimeColumnId, true,  100.0f },
+		{ FMemAllocTableColumns::EndTimeColumnId,   true,  100.0f },
+		{ FMemAllocTableColumns::DurationColumnId,  true,  100.0f },
+		{ FMemAllocTableColumns::AddressColumnId,   true,  120.0f },
+		{ FMemAllocTableColumns::CountColumnId,     true,  100.0f },
+		{ FMemAllocTableColumns::SizeColumnId,      true,  100.0f },
+		{ FMemAllocTableColumns::TagColumnId,       true,  120.0f },
+		{ FMemAllocTableColumns::FunctionColumnId,  true,  550.0f },
+	};
+	ApplyColumnConfig(TArrayView<FColumnConfig>(Preset, UE_ARRAY_COUNT(Preset)));
+
+	return FReply::Handled();
+}
+
+////////////////////////////////////////////////////////////////////////////////////////////////////
+
+FReply SMemAllocTableTreeView::OnSizeViewClicked()
+{
+	PreChangeGroupings();
+
+	CurrentGroupings.Reset();
+
+	check(AvailableGroupings[0]->Is<FTreeNodeGroupingFlat>());
+	CurrentGroupings.Add(AvailableGroupings[0]);
+
+	TSharedPtr<FTreeNodeGrouping>* SizeGrouping = AvailableGroupings.FindByPredicate(
+		[](TSharedPtr<FTreeNodeGrouping>& Grouping)
+		{
+			return Grouping->Is<FMemAllocGroupingBySize>();
+		});
+	if (SizeGrouping)
+	{
+		CurrentGroupings.Add(*SizeGrouping);
+	}
+
+	PostChangeGroupings();
+
+	FColumnConfig Preset[] =
+	{
+		{ FTable::GetHierarchyColumnId(),           true,  200.0f },
+		{ FMemAllocTableColumns::StartTimeColumnId, false, 0.0f },
+		{ FMemAllocTableColumns::EndTimeColumnId,   false, 0.0f },
+		{ FMemAllocTableColumns::DurationColumnId,  false, 0.0f },
+		{ FMemAllocTableColumns::AddressColumnId,   false, 0.0f },
+		{ FMemAllocTableColumns::CountColumnId,     true,  100.0f },
+		{ FMemAllocTableColumns::SizeColumnId,      true,  100.0f },
+		{ FMemAllocTableColumns::TagColumnId,       true,  120.0f },
+		{ FMemAllocTableColumns::FunctionColumnId,  true,  400.0f },
+	};
+	ApplyColumnConfig(TArrayView<FColumnConfig>(Preset, UE_ARRAY_COUNT(Preset)));
+
+	return FReply::Handled();
+}
+
+////////////////////////////////////////////////////////////////////////////////////////////////////
+
+FReply SMemAllocTableTreeView::OnTagViewClicked()
+{
+	PreChangeGroupings();
+
+	CurrentGroupings.Reset();
+
+	check(AvailableGroupings[0]->Is<FTreeNodeGroupingFlat>());
+	CurrentGroupings.Add(AvailableGroupings[0]);
+
+	TSharedPtr<FTreeNodeGrouping>* TagGrouping = AvailableGroupings.FindByPredicate(
+		[](TSharedPtr<FTreeNodeGrouping>& Grouping)
+		{
+			return Grouping->Is<FTreeNodeGroupingByUniqueValue>() &&
+				   Grouping->As<FTreeNodeGroupingByUniqueValue>().GetColumnId() == FMemAllocTableColumns::TagColumnId;
+		});
+	if (TagGrouping)
+	{
+		CurrentGroupings.Add(*TagGrouping);
+	}
+
+	PostChangeGroupings();
+
+	FColumnConfig Preset[] =
+	{
+		{ FTable::GetHierarchyColumnId(),           true,  200.0f },
+		{ FMemAllocTableColumns::StartTimeColumnId, false, 0.0f },
+		{ FMemAllocTableColumns::EndTimeColumnId,   false, 0.0f },
+		{ FMemAllocTableColumns::DurationColumnId,  false, 0.0f },
+		{ FMemAllocTableColumns::AddressColumnId,   false, 0.0f },
+		{ FMemAllocTableColumns::CountColumnId,     true,  100.0f },
+		{ FMemAllocTableColumns::SizeColumnId,      true,  100.0f },
+		{ FMemAllocTableColumns::TagColumnId,       false, 0.0f },
+		{ FMemAllocTableColumns::FunctionColumnId,  true,  400.0f },
+	};
+	ApplyColumnConfig(TArrayView<FColumnConfig>(Preset, UE_ARRAY_COUNT(Preset)));
+
+	return FReply::Handled();
+}
+
+////////////////////////////////////////////////////////////////////////////////////////////////////
+
+FReply SMemAllocTableTreeView::OnCallstackViewClicked(bool bIsInverted)
+{
+	PreChangeGroupings();
+
+	CurrentGroupings.Reset();
+
+	check(AvailableGroupings[0]->Is<FTreeNodeGroupingFlat>());
+	CurrentGroupings.Add(AvailableGroupings[0]);
+
+	TSharedPtr<FTreeNodeGrouping>* CallstackGrouping = AvailableGroupings.FindByPredicate(
+		[bIsInverted](TSharedPtr<FTreeNodeGrouping>& Grouping)
+		{
+			return Grouping->Is<FMemAllocGroupingByCallstack>() &&
+				   Grouping->As<FMemAllocGroupingByCallstack>().IsInverted() == bIsInverted;
+		});
+	if (CallstackGrouping)
+	{
+		CurrentGroupings.Add(*CallstackGrouping);
+	}
+
+	PostChangeGroupings();
+
+	FColumnConfig Preset[] =
+	{
+		{ FTable::GetHierarchyColumnId(),           true,  400.0f },
+		{ FMemAllocTableColumns::StartTimeColumnId, false, 0.0f },
+		{ FMemAllocTableColumns::EndTimeColumnId,   false, 0.0f },
+		{ FMemAllocTableColumns::DurationColumnId,  false, 0.0f },
+		{ FMemAllocTableColumns::AddressColumnId,   false, 0.0f },
+		{ FMemAllocTableColumns::CountColumnId,     true,  100.0f },
+		{ FMemAllocTableColumns::SizeColumnId,      true,  100.0f },
+		{ FMemAllocTableColumns::TagColumnId,       true,  200.0f },
+		{ FMemAllocTableColumns::FunctionColumnId,  true,  200.0f },
+	};
+	ApplyColumnConfig(TArrayView<FColumnConfig>(Preset, UE_ARRAY_COUNT(Preset)));
+
+	return FReply::Handled();
+}
+
+////////////////////////////////////////////////////////////////////////////////////////////////////
+
+void SMemAllocTableTreeView::ApplyColumnConfig(const TArrayView<FColumnConfig>& Preset)
+{
+	for (const TSharedRef<FTableColumn>& ColumnRef : Table->GetColumns())
+	{
+		FTableColumn& Column = ColumnRef.Get();
+		for (const FColumnConfig& Config : Preset)
+		{
+			if (Column.GetId() == Config.ColumnId)
+			{
+				if (Config.bIsVisible)
+				{
+					ShowColumn(Column);
+					if (Config.Width > 0.0f)
+					{
+						TreeViewHeaderRow->SetColumnWidth(Column.GetId(), Config.Width);
+					}
+				}
+				else
+				{
+					HideColumn(Column);
+				}
+			}
+		}
+	}
+}
+
+////////////////////////////////////////////////////////////////////////////////////////////////////
+
 TSharedPtr<SWidget> SMemAllocTableTreeView::ConstructFooter()
 {
 	return
@@ -419,8 +659,27 @@ void SMemAllocTableTreeView::InternalCreateGroupings()
 {
 	STableTreeView::InternalCreateGroupings();
 
-	AvailableGroupings.Add(MakeShared<FCallstackGrouping>(false));
-	AvailableGroupings.Add(MakeShared<FCallstackGrouping>(true));
+	int32 Index = 1; // after the Flat ("All") grouping
+	
+	AvailableGroupings.Insert(MakeShared<FMemAllocGroupingBySize>(), Index++);
+
+	TSharedPtr<FTreeNodeGrouping>* TagGroupingPtr = AvailableGroupings.FindByPredicate(
+		[](TSharedPtr<FTreeNodeGrouping>& Grouping)
+		{
+			return Grouping->Is<FTreeNodeGroupingByUniqueValue>() &&
+				   Grouping->As<FTreeNodeGroupingByUniqueValue>().GetColumnId() == FMemAllocTableColumns::TagColumnId;
+		});
+	if (TagGroupingPtr)
+	{
+		TSharedPtr<FTreeNodeGroupingByUniqueValue> TagGrouping = StaticCastSharedPtr<FTreeNodeGroupingByUniqueValue>(*TagGroupingPtr);
+		AvailableGroupings.Remove(TagGrouping);
+		//TODO: TagGrouping->SetShortName(LOCTEXT("Grouping_ByTag_ShortName", "LLM Tag"));
+		//TODO: TagGrouping->SetTitleName(LOCTEXT("Grouping_ByTag_TitleName", "By LLM Tag"));
+		AvailableGroupings.Insert(TagGrouping, Index++);
+	}
+
+	AvailableGroupings.Insert(MakeShared<FMemAllocGroupingByCallstack>(false), Index++);
+	AvailableGroupings.Insert(MakeShared<FMemAllocGroupingByCallstack>(true), Index++);
 }
 
 ////////////////////////////////////////////////////////////////////////////////////////////////////
diff --git a/Engine/Source/Developer/TraceInsights/Private/Insights/MemoryProfiler/Widgets/SMemAllocTableTreeView.h b/Engine/Source/Developer/TraceInsights/Private/Insights/MemoryProfiler/Widgets/SMemAllocTableTreeView.h
--- a/Engine/Source/Developer/TraceInsights/Private/Insights/MemoryProfiler/Widgets/SMemAllocTableTreeView.h
+++ b/Engine/Source/Developer/TraceInsights/Private/Insights/MemoryProfiler/Widgets/SMemAllocTableTreeView.h
@@ -20,6 +20,14 @@ class FMemoryRuleSpec;
 
 class SMemAllocTableTreeView : public STableTreeView
 {
+private:
+	struct FColumnConfig
+	{
+		const FName& ColumnId;
+		bool bIsVisible;
+		float Width;
+	};
+
 public:
 	/** Default constructor. */
 	SMemAllocTableTreeView();
@@ -36,6 +44,7 @@ public:
 	 */
 	void Construct(const FArguments& InArgs, TSharedPtr<FMemAllocTable> InTablePtr);
 
+	virtual TSharedPtr<SWidget> ConstructToolbar() override;;
 	virtual TSharedPtr<SWidget> ConstructFooter() override;
 
 	TSharedPtr<FMemAllocTable> GetMemAllocTable() { return StaticCastSharedPtr<FMemAllocTable>(GetTable()); }
@@ -90,8 +99,16 @@ private:
 	void StartQuery();
 	void UpdateQuery(TraceServices::IAllocationsProvider::EQueryStatus& OutStatus);
 	void CancelQuery();
+
+	FReply OnDetailedViewClicked();
+	FReply OnSizeViewClicked();
+	FReply OnTagViewClicked();
+	FReply OnCallstackViewClicked(bool bIsInverted);
+
 	FText GetSymbolResolutionStatus() const;
 
+	void ApplyColumnConfig(const TArrayView<FColumnConfig>& Preset);
+
 private:
 	int32 TabIndex = -1;
 	TSharedPtr<FMemoryRuleSpec> Rule = nullptr;
diff --git a/Engine/Source/Developer/TraceInsights/Private/Insights/Table/ViewModels/Table.cpp b/Engine/Source/Developer/TraceInsights/Private/Insights/Table/ViewModels/Table.cpp
--- a/Engine/Source/Developer/TraceInsights/Private/Insights/Table/ViewModels/Table.cpp
+++ b/Engine/Source/Developer/TraceInsights/Private/Insights/Table/ViewModels/Table.cpp
@@ -72,11 +72,17 @@ void FTable::AddColumn(TSharedRef<FTableColumn> ColumnRef)
 
 ////////////////////////////////////////////////////////////////////////////////////////////////////
 
-void FTable::AddHierarchyColumn(int32 ColumnIndex, const TCHAR* ColumnName)
+const FName FTable::GetHierarchyColumnId()
 {
-	const FName HierarchyColumnId(TEXT("_Hierarchy"));
+	static const FName HierarchyColumnId(TEXT("_Hierarchy"));
+	return HierarchyColumnId;
+}
+
+////////////////////////////////////////////////////////////////////////////////////////////////////
 
-	TSharedRef<FTableColumn> ColumnRef = MakeShared<FTableColumn>(HierarchyColumnId);
+void FTable::AddHierarchyColumn(int32 ColumnIndex, const TCHAR* ColumnName)
+{
+	TSharedRef<FTableColumn> ColumnRef = MakeShared<FTableColumn>(GetHierarchyColumnId());
 	FTableColumn& Column = *ColumnRef;
 
 	Column.SetIndex(ColumnIndex);
diff --git a/Engine/Source/Developer/TraceInsights/Private/Insights/Table/ViewModels/Table.h b/Engine/Source/Developer/TraceInsights/Private/Insights/Table/ViewModels/Table.h
--- a/Engine/Source/Developer/TraceInsights/Private/Insights/Table/ViewModels/Table.h
+++ b/Engine/Source/Developer/TraceInsights/Private/Insights/Table/ViewModels/Table.h
@@ -53,6 +53,8 @@ public:
 
 	void GetVisibleColumnsData(const TArray<TSharedPtr<class FBaseTreeNode>>& InNodes, const FName& LogListingName, FString& OutData) const;
 
+	static const FName GetHierarchyColumnId();
+
 protected:
 	void ResetColumns() { Columns.Reset(); }
 	void AddColumn(TSharedRef<FTableColumn> Column);
diff --git a/Engine/Source/Developer/TraceInsights/Private/Insights/Table/ViewModels/TreeNodeGrouping.cpp b/Engine/Source/Developer/TraceInsights/Private/Insights/Table/ViewModels/TreeNodeGrouping.cpp
--- a/Engine/Source/Developer/TraceInsights/Private/Insights/Table/ViewModels/TreeNodeGrouping.cpp
+++ b/Engine/Source/Developer/TraceInsights/Private/Insights/Table/ViewModels/TreeNodeGrouping.cpp
@@ -9,6 +9,15 @@
 namespace Insights
 {
 
+INSIGHTS_IMPLEMENT_RTTI(ITreeNodeGrouping)
+INSIGHTS_IMPLEMENT_RTTI(FTreeNodeGrouping)
+INSIGHTS_IMPLEMENT_RTTI(FTreeNodeGroupingFlat)
+INSIGHTS_IMPLEMENT_RTTI(FTreeNodeGroupingByUniqueValue)
+INSIGHTS_IMPLEMENT_RTTI(FTreeNodeGroupingByNameFirstLetter)
+INSIGHTS_IMPLEMENT_RTTI(FTreeNodeGroupingByType)
+
+////////////////////////////////////////////////////////////////////////////////////////////////////
+// FTreeNodeGrouping
 ////////////////////////////////////////////////////////////////////////////////////////////////////
 
 FTreeNodeGrouping::FTreeNodeGrouping(const FText& InShortName, const FText& InTitleName, const FText& InDescription, const FName InBrushName, const FSlateBrush* InIcon)
diff --git a/Engine/Source/Developer/TraceInsights/Private/Insights/Table/ViewModels/TreeNodeGrouping.h b/Engine/Source/Developer/TraceInsights/Private/Insights/Table/ViewModels/TreeNodeGrouping.h
--- a/Engine/Source/Developer/TraceInsights/Private/Insights/Table/ViewModels/TreeNodeGrouping.h
+++ b/Engine/Source/Developer/TraceInsights/Private/Insights/Table/ViewModels/TreeNodeGrouping.h
@@ -7,6 +7,7 @@
 #include "Delegates/DelegateCombinations.h"
 
 // Insights
+#include "Insights/Common/SimpleRtti.h"
 #include "Insights/Table/ViewModels/BaseTreeNode.h"
 #include "Insights/Table/ViewModels/TableColumn.h"
 #include "Insights/Table/ViewModels/TableTreeNode.h"
@@ -22,6 +23,8 @@ class FTable;
 
 class ITreeNodeGrouping
 {
+	INSIGHTS_DECLARE_RTTI_BASE(ITreeNodeGrouping)
+
 public:
 	virtual FText GetShortName() const = 0;
 	virtual FText GetTitleName() const = 0;
@@ -45,6 +48,8 @@ struct FTreeNodeGroupInfo
 
 class FTreeNodeGrouping : public ITreeNodeGrouping
 {
+	INSIGHTS_DECLARE_RTTI(FTreeNodeGrouping, ITreeNodeGrouping)
+
 public:
 	FTreeNodeGrouping(const FText& InShortName, const FText& InTitleName, const FText& InDescription, const FName InBrushName, const FSlateBrush* InIcon);
 	virtual ~FTreeNodeGrouping() {}
@@ -74,6 +79,8 @@ protected:
 /** Creates a single group for all nodes. */
 class FTreeNodeGroupingFlat : public FTreeNodeGrouping
 {
+	INSIGHTS_DECLARE_RTTI(FTreeNodeGroupingFlat, FTreeNodeGrouping)
+
 public:
 	FTreeNodeGroupingFlat();
 	virtual ~FTreeNodeGroupingFlat() {}
@@ -86,6 +93,8 @@ public:
 /** Creates a group for each unique value. */
 class FTreeNodeGroupingByUniqueValue : public FTreeNodeGrouping
 {
+	INSIGHTS_DECLARE_RTTI(FTreeNodeGroupingByUniqueValue, FTreeNodeGrouping);
+
 public:
 	FTreeNodeGroupingByUniqueValue(TSharedRef<FTableColumn> InColumnRef);
 	virtual ~FTreeNodeGroupingByUniqueValue() {}
@@ -183,6 +192,8 @@ typedef TTreeNodeGroupingByUniqueValue<const TCHAR*> FTreeNodeGroupingByUniqueVa
 /** Creates a group for each first letter of node names. */
 class FTreeNodeGroupingByNameFirstLetter : public FTreeNodeGrouping
 {
+	INSIGHTS_DECLARE_RTTI(FTreeNodeGroupingByNameFirstLetter, FTreeNodeGrouping);
+
 public:
 	FTreeNodeGroupingByNameFirstLetter();
 	virtual ~FTreeNodeGroupingByNameFirstLetter() {}
@@ -195,6 +206,8 @@ public:
 /** Creates a group for each type. */
 class FTreeNodeGroupingByType : public FTreeNodeGrouping
 {
+	INSIGHTS_DECLARE_RTTI(FTreeNodeGroupingByType, FTreeNodeGrouping);
+
 public:
 	FTreeNodeGroupingByType();
 	virtual ~FTreeNodeGroupingByType() {}
diff --git a/Engine/Source/Developer/TraceInsights/Private/Insights/Table/Widgets/STableTreeView.cpp b/Engine/Source/Developer/TraceInsights/Private/Insights/Table/Widgets/STableTreeView.cpp
--- a/Engine/Source/Developer/TraceInsights/Private/Insights/Table/Widgets/STableTreeView.cpp
+++ b/Engine/Source/Developer/TraceInsights/Private/Insights/Table/Widgets/STableTreeView.cpp
@@ -135,6 +135,8 @@ void STableTreeView::ConstructWidget(TSharedPtr<FTable> InTablePtr)
 	SAssignNew(ExternalScrollbar, SScrollBar)
 	.AlwaysShowScrollbar(true);
 
+	TSharedPtr<SHorizontalBox> ToolbarBox;
+
 	auto WidgetContent = 
 	SNew(SVerticalBox)
 
@@ -161,10 +163,10 @@ void STableTreeView::ConstructWidget(TSharedPtr<FTable> InTablePtr)
 				.VAlign(VAlign_Center)
 				[
 					SAssignNew(SearchBox, SSearchBox)
-					.HintText(LOCTEXT("SearchBoxHint", "Search"))
+					.HintText(LOCTEXT("SearchBox_Hint", "Search"))
 					.OnTextChanged(this, &STableTreeView::SearchBox_OnTextChanged)
 					.IsEnabled(this, &STableTreeView::SearchBox_IsEnabled)
-					.ToolTipText(LOCTEXT("FilterSearchHint", "Type here to search the tree hierarchy by item or group name"))
+					.ToolTipText(LOCTEXT("SearchBox_ToolTip", "Type here to search the tree hierarchy by item or group name"))
 				]
 
 				+ SHorizontalBox::Slot()
@@ -173,8 +175,8 @@ void STableTreeView::ConstructWidget(TSharedPtr<FTable> InTablePtr)
 				.VAlign(VAlign_Center)
 				[
 					SNew(SButton)
-					.Text(LOCTEXT("AdvancedFilters", "Advanced Filters"))
-					.ToolTipText(LOCTEXT("AdvancedFilters", "Opens the filter configurator window."))
+					.Text(LOCTEXT("AdvancedFiltersBtn_Text", "Advanced Filters"))
+					.ToolTipText(LOCTEXT("AdvancedFiltersBtn_ToolTip", "Opens the filter configurator window."))
 					.OnClicked(this, &STableTreeView::OnAdvancedFiltersClicked)
 				]
 			]
@@ -185,7 +187,7 @@ void STableTreeView::ConstructWidget(TSharedPtr<FTable> InTablePtr)
 			.Padding(2.0f)
 			.AutoHeight()
 			[
-				SNew(SHorizontalBox)
+				SAssignNew(ToolbarBox, SHorizontalBox)
 
 				+ SHorizontalBox::Slot()
 				.AutoWidth()
@@ -280,6 +282,18 @@ void STableTreeView::ConstructWidget(TSharedPtr<FTable> InTablePtr)
 		]
 	];
 
+	TSharedPtr<SWidget> Toolbar = ConstructToolbar();
+	if (Toolbar.IsValid())
+	{
+		ToolbarBox->AddSlot()
+			.AutoWidth()
+			.HAlign(HAlign_Right)
+			.Padding(0.0f)
+			[
+				Toolbar.ToSharedRef()
+			];
+	}
+
 	TSharedPtr<SWidget> Footer = ConstructFooter();
 	if (Footer.IsValid())
 	{
@@ -465,8 +479,7 @@ void STableTreeView::TreeView_BuildSortByMenu(FMenuBuilder& MenuBuilder)
 	//TODO: for (Sorting : AvailableSortings)
 	for (const TSharedRef<FTableColumn>& ColumnRef : Table->GetColumns())
 	{
-		const FTableColumn& Column = *ColumnRef;
-
+		const FTableColumn& Column = ColumnRef.Get();
 		if (Column.IsVisible() && Column.CanBeSorted())
 		{
 			FUIAction Action_SortByColumn
@@ -527,8 +540,7 @@ void STableTreeView::TreeView_BuildViewColumnMenu(FMenuBuilder& MenuBuilder)
 
 	for (const TSharedRef<FTableColumn>& ColumnRef : Table->GetColumns())
 	{
-		const FTableColumn& Column = *ColumnRef;
-
+		const FTableColumn& Column = ColumnRef.Get();
 		FUIAction Action_ToggleColumn
 		(
 			FExecuteAction::CreateSP(this, &STableTreeView::ToggleColumnVisibility, Column.GetId()),
@@ -552,9 +564,10 @@ void STableTreeView::InitializeAndShowHeaderColumns()
 {
 	for (const TSharedRef<FTableColumn>& ColumnRef : Table->GetColumns())
 	{
-		if (ColumnRef->ShouldBeVisible())
+		FTableColumn& Column = ColumnRef.Get();
+		if (Column.ShouldBeVisible())
 		{
-			ShowColumn(ColumnRef->GetId());
+			ShowColumn(Column);
 		}
 	}
 }
@@ -1097,33 +1110,29 @@ void STableTreeView::UpdateAggregatedValues(FTableTreeNode& GroupNode)
 {
 	for (const TSharedRef<FTableColumn>& ColumnRef : Table->GetColumns())
 	{
-		FTableColumn& Column = *ColumnRef;
+		FTableColumn& Column = ColumnRef.Get();
 		switch (Column.GetAggregation())
 		{
-		case ETableColumnAggregation::Sum:
-		{
-			STableTreeView::UpdateAggregationRec<int64>(Column, GroupNode, 0, true, [](int64 InValue, TOptional<FTableCellValue> InTableCellValue)
+			case ETableColumnAggregation::Sum:
+				STableTreeView::UpdateAggregationRec<int64>(Column, GroupNode, 0, true, [](int64 InValue, TOptional<FTableCellValue> InTableCellValue)
 				{
 					return InValue + InTableCellValue->AsInt64();
 				});
-			break;
-		}
-		case ETableColumnAggregation::Min:
-		{
-			STableTreeView::UpdateAggregationRec<double>(Column, GroupNode, std::numeric_limits<double>::max(), false, [](double InValue, TOptional<FTableCellValue> InTableCellValue)
+				break;
+
+			case ETableColumnAggregation::Min:
+				STableTreeView::UpdateAggregationRec<double>(Column, GroupNode, std::numeric_limits<double>::max(), false, [](double InValue, TOptional<FTableCellValue> InTableCellValue)
 				{
 					return FMath::Min(InValue, InTableCellValue->AsDouble());
 				});
-			break;
-		}
-		case ETableColumnAggregation::Max:
-		{
-			STableTreeView::UpdateAggregationRec<double>(Column, GroupNode, std::numeric_limits<double>::lowest(), false, [](double InValue, TOptional<FTableCellValue> InTableCellValue)
+				break;
+
+			case ETableColumnAggregation::Max:
+				STableTreeView::UpdateAggregationRec<double>(Column, GroupNode, std::numeric_limits<double>::lowest(), false, [](double InValue, TOptional<FTableCellValue> InTableCellValue)
 				{
 					return FMath::Max(InValue, InTableCellValue->AsDouble());
 				});
-			break;
-		}
+				break;
 		}
 	}
 }
@@ -1155,9 +1164,10 @@ void STableTreeView::InternalCreateGroupings()
 
 	for (const TSharedRef<FTableColumn>& ColumnRef : Table->GetColumns())
 	{
-		if (!ColumnRef->IsHierarchy())
+		const FTableColumn& Column = ColumnRef.Get();
+		if (!Column.IsHierarchy())
 		{
-			switch (ColumnRef->GetDataType())
+			switch (Column.GetDataType())
 			{
 				case ETableCellDataType::Bool:
 					AvailableGroupings.Add(MakeShared<FTreeNodeGroupingByUniqueValueBool>(ColumnRef));
@@ -1706,9 +1716,10 @@ void STableTreeView::CreateSortings()
 
 	for (const TSharedRef<FTableColumn>& ColumnRef : Table->GetColumns())
 	{
-		if (ColumnRef->CanBeSorted())
+		const FTableColumn& Column = ColumnRef.Get();
+		if (Column.CanBeSorted())
 		{
-			TSharedPtr<Insights::ITableCellValueSorter> SorterPtr = ColumnRef->GetValueSorter();
+			TSharedPtr<Insights::ITableCellValueSorter> SorterPtr = Column.GetValueSorter();
 			if (ensure(SorterPtr.IsValid()))
 			{
 				AvailableSorters.Add(SorterPtr);
@@ -1909,52 +1920,62 @@ bool STableTreeView::CanShowColumn(const FName ColumnId) const
 void STableTreeView::ShowColumn(const FName ColumnId)
 {
 	FTableColumn& Column = *Table->FindColumnChecked(ColumnId);
-	Column.Show();
-
-	SHeaderRow::FColumn::FArguments ColumnArgs;
-	ColumnArgs
-		.ColumnId(Column.GetId())
-		.DefaultLabel(Column.GetShortName())
-		.HAlignHeader(HAlign_Fill)
-		.VAlignHeader(VAlign_Fill)
-		.HeaderContentPadding(FMargin(2.0f))
-		.HAlignCell(HAlign_Fill)
-		.VAlignCell(VAlign_Fill)
-		.SortMode(this, &STableTreeView::GetSortModeForColumn, Column.GetId())
-		.OnSort(this, &STableTreeView::OnSortModeChanged)
-		//.ManualWidth(Column.GetInitialWidth())
-		.FillWidth(Column.GetInitialWidth())
-		//.FixedWidth(Column.IsFixedWidth() ? Column.GetInitialWidth() : TOptional<float>())
-		.HeaderContent()
-		[
-			SNew(SBox)
-			.ToolTip(STableTreeViewTooltip::GetColumnTooltip(Column))
-			.HAlign(Column.GetHorizontalAlignment())
-			.VAlign(VAlign_Center)
+	ShowColumn(Column);
+}
+
+////////////////////////////////////////////////////////////////////////////////////////////////////
+
+void STableTreeView::ShowColumn(FTableColumn& Column)
+{
+	if (!Column.IsVisible())
+	{
+		Column.Show();
+
+		SHeaderRow::FColumn::FArguments ColumnArgs;
+		ColumnArgs
+			.ColumnId(Column.GetId())
+			.DefaultLabel(Column.GetShortName())
+			.HAlignHeader(HAlign_Fill)
+			.VAlignHeader(VAlign_Fill)
+			.HeaderContentPadding(FMargin(2.0f))
+			.HAlignCell(HAlign_Fill)
+			.VAlignCell(VAlign_Fill)
+			.SortMode(this, &STableTreeView::GetSortModeForColumn, Column.GetId())
+			.OnSort(this, &STableTreeView::OnSortModeChanged)
+			//.ManualWidth(Column.GetInitialWidth())
+			.FillWidth(Column.GetInitialWidth())
+			//.FixedWidth(Column.IsFixedWidth() ? Column.GetInitialWidth() : TOptional<float>())
+			.HeaderContent()
 			[
-				SNew(STextBlock)
-				.Text(this, &STableTreeView::GetColumnHeaderText, Column.GetId())
+				SNew(SBox)
+				.ToolTip(STableTreeViewTooltip::GetColumnTooltip(Column))
+				.HAlign(Column.GetHorizontalAlignment())
+				.VAlign(VAlign_Center)
+				[
+					SNew(STextBlock)
+					.Text(this, &STableTreeView::GetColumnHeaderText, Column.GetId())
+				]
 			]
-		]
-		.MenuContent()
-		[
-			TreeViewHeaderRow_GenerateColumnMenu(Column)
-		];
+			.MenuContent()
+			[
+				TreeViewHeaderRow_GenerateColumnMenu(Column)
+			];
 
-	int32 ColumnIndex = 0;
-	const int32 NewColumnPosition = Table->GetColumnPositionIndex(ColumnId);
-	const int32 NumColumns = TreeViewHeaderRow->GetColumns().Num();
-	for (; ColumnIndex < NumColumns; ColumnIndex++)
-	{
-		const SHeaderRow::FColumn& CurrentColumn = TreeViewHeaderRow->GetColumns()[ColumnIndex];
-		const int32 CurrentColumnPosition = Table->GetColumnPositionIndex(CurrentColumn.ColumnId);
-		if (NewColumnPosition < CurrentColumnPosition)
+		int32 ColumnIndex = 0;
+		const int32 NewColumnPosition = Table->GetColumnPositionIndex(Column.GetId());
+		const int32 NumColumns = TreeViewHeaderRow->GetColumns().Num();
+		for (; ColumnIndex < NumColumns; ColumnIndex++)
 		{
-			break;
+			const SHeaderRow::FColumn& CurrentColumn = TreeViewHeaderRow->GetColumns()[ColumnIndex];
+			const int32 CurrentColumnPosition = Table->GetColumnPositionIndex(CurrentColumn.ColumnId);
+			if (NewColumnPosition < CurrentColumnPosition)
+			{
+				break;
+			}
 		}
-	}
 
-	TreeViewHeaderRow->InsertColumn(ColumnArgs, ColumnIndex);
+		TreeViewHeaderRow->InsertColumn(ColumnArgs, ColumnIndex);
+	}
 }
 
 ////////////////////////////////////////////////////////////////////////////////////////////////////
@@ -1977,9 +1998,18 @@ bool STableTreeView::CanHideColumn(const FName ColumnId) const
 void STableTreeView::HideColumn(const FName ColumnId)
 {
 	FTableColumn& Column = *Table->FindColumnChecked(ColumnId);
-	Column.Hide();
+	HideColumn(Column);
+}
+
+////////////////////////////////////////////////////////////////////////////////////////////////////
 
-	TreeViewHeaderRow->RemoveColumn(ColumnId);
+void STableTreeView::HideColumn(FTableColumn& Column)
+{
+	if (Column.IsVisible())
+	{
+		Column.Hide();
+		TreeViewHeaderRow->RemoveColumn(Column.GetId());
+	}
 }
 
 ////////////////////////////////////////////////////////////////////////////////////////////////////
@@ -2009,14 +2039,14 @@ bool STableTreeView::CanToggleColumnVisibility(const FName ColumnId) const
 
 void STableTreeView::ToggleColumnVisibility(const FName ColumnId)
 {
-	const FTableColumn& Column = *Table->FindColumnChecked(ColumnId);
+	FTableColumn& Column = *Table->FindColumnChecked(ColumnId);
 	if (Column.IsVisible())
 	{
-		HideColumn(ColumnId);
+		HideColumn(Column);
 	}
 	else
 	{
-		ShowColumn(ColumnId);
+		ShowColumn(Column);
 	}
 }
 
@@ -2044,11 +2074,10 @@ void STableTreeView::ContextMenu_ShowAllColumns_Execute()
 
 	for (const TSharedRef<FTableColumn>& ColumnRef : Table->GetColumns())
 	{
-		const FTableColumn& Column = *ColumnRef;
-
+		FTableColumn& Column = ColumnRef.Get();
 		if (!Column.IsVisible())
 		{
-			ShowColumn(Column.GetId());
+			ShowColumn(Column);
 		}
 	}
 }
@@ -2077,15 +2106,14 @@ void STableTreeView::ContextMenu_ResetColumns_Execute()
 
 	for (const TSharedRef<FTableColumn>& ColumnRef : Table->GetColumns())
 	{
-		const FTableColumn& Column = *ColumnRef;
-
+		FTableColumn& Column = ColumnRef.Get();
 		if (Column.ShouldBeVisible() && !Column.IsVisible())
 		{
-			ShowColumn(Column.GetId());
+			ShowColumn(Column);
 		}
 		else if (!Column.ShouldBeVisible() && Column.IsVisible())
 		{
-			HideColumn(Column.GetId());
+			HideColumn(Column);
 		}
 	}
 }
@@ -2172,6 +2200,9 @@ void STableTreeView::OnPostAsyncUpdate()
 			}
 		}
 
+		// Expand each group node on the first few depths (if it doesn't have too many children).
+		ExpandChildGroups(Root.Get(), 1000);
+
 		ClearInProgressAsyncOperations();
 		TreeView_Refresh();
 	}
@@ -2182,6 +2213,73 @@ void STableTreeView::OnPostAsyncUpdate()
 
 ////////////////////////////////////////////////////////////////////////////////////////////////////
 
+void STableTreeView::ExpandChildGroups(FBaseTreeNode* InRoot, int32 InMaxExpandedNodes)
+{
+	const int32 MaxDepthToExpand = 4;
+
+	TArray<int32> NumNodesPerDepth;
+	NumNodesPerDepth.AddDefaulted(MaxDepthToExpand + 1);
+	CountNumNodesPerDepthRec(InRoot, NumNodesPerDepth, 0, MaxDepthToExpand, InMaxExpandedNodes);
+
+	int32 MaxDepth = 0;
+	for (int32 Depth = 0; Depth <= MaxDepthToExpand; ++Depth)
+	{
+		if (Depth > 0)
+		{
+			NumNodesPerDepth[Depth] += NumNodesPerDepth[Depth - 1];
+		}
+		if (NumNodesPerDepth[Depth] > InMaxExpandedNodes)
+		{
+			break;
+		}
+		MaxDepth = Depth;
+	}
+
+	if (MaxDepth > 0)
+	{
+		ExpandChildGroupsRec(InRoot, 1, MaxDepth);
+	}
+}
+
+////////////////////////////////////////////////////////////////////////////////////////////////////
+
+void STableTreeView::CountNumNodesPerDepthRec(FBaseTreeNode* InRoot, TArray<int32>& InOutNumNodesPerDepth, int32 InDepth, int32 InMaxDepth, int32 InMaxNodes) const
+{
+	InOutNumNodesPerDepth[InDepth] += InRoot->GetChildren().Num();
+
+	if (InDepth < InMaxDepth && InOutNumNodesPerDepth[InDepth] < InMaxNodes)
+	{
+		for (const FBaseTreeNodePtr& Node : InRoot->GetChildren())
+		{
+			if (Node->IsGroup())
+			{
+				CountNumNodesPerDepthRec(Node.Get(), InOutNumNodesPerDepth, InDepth + 1, InMaxDepth, InMaxNodes);
+			}
+		}
+	}
+}
+
+////////////////////////////////////////////////////////////////////////////////////////////////////
+
+void STableTreeView::ExpandChildGroupsRec(FBaseTreeNode* InRoot, int32 InDepth, int32 InMaxDepth)
+{
+	for (const FBaseTreeNodePtr& Node : InRoot->GetChildren())
+	{
+		if (Node->IsGroup())
+		{
+			Node->SetExpansion(true);
+			TreeView->SetItemExpansion(StaticCastSharedPtr<FTableTreeNode>(Node), true);
+
+			if (InDepth < InMaxDepth)
+			{
+				ExpandChildGroupsRec(Node.Get(), InDepth + 1, InMaxDepth);
+			}
+		}
+	}
+}
+
+////////////////////////////////////////////////////////////////////////////////////////////////////
+
 FGraphEventRef STableTreeView::StartSortTreeNodesTask(FGraphEventRef Prerequisite)
 {
 	if (!CurrentSorter.IsValid())
@@ -2455,6 +2553,13 @@ void STableTreeView::OnAdvancedFiltersChangesCommited()
 
 ////////////////////////////////////////////////////////////////////////////////////////////////////
 
+TSharedPtr<SWidget> STableTreeView::ConstructToolbar()
+{
+	return nullptr;
+}
+
+////////////////////////////////////////////////////////////////////////////////////////////////////
+
 TSharedPtr<SWidget> STableTreeView::ConstructFooter()
 {
 	return nullptr;
diff --git a/Engine/Source/Developer/TraceInsights/Private/Insights/Table/Widgets/STableTreeView.h b/Engine/Source/Developer/TraceInsights/Private/Insights/Table/Widgets/STableTreeView.h
--- a/Engine/Source/Developer/TraceInsights/Private/Insights/Table/Widgets/STableTreeView.h
+++ b/Engine/Source/Developer/TraceInsights/Private/Insights/Table/Widgets/STableTreeView.h
@@ -135,6 +135,7 @@ public:
 protected:
 	void ConstructWidget(TSharedPtr<FTable> InTablePtr);
 	void InitCommandList();
+	virtual TSharedPtr<SWidget> ConstructToolbar();
 	virtual TSharedPtr<SWidget> ConstructFooter();
 
 	void UpdateTree();
@@ -320,10 +321,12 @@ protected:
 	// ShowColumn
 	bool CanShowColumn(const FName ColumnId) const;
 	void ShowColumn(const FName ColumnId);
+	void ShowColumn(FTableColumn& Column);
 
 	// HideColumn
 	bool CanHideColumn(const FName ColumnId) const;
 	void HideColumn(const FName ColumnId);
+	void HideColumn(FTableColumn& Column);
 
 	// ToggleColumnVisibility
 	bool IsColumnVisible(const FName ColumnId);
@@ -360,6 +363,10 @@ protected:
 
 	////////////////////////////////////////////////////////////////////////////////////////////////////
 
+	void ExpandChildGroups(FBaseTreeNode* InRoot, int32 InMaxExpandedNodes);
+	void CountNumNodesPerDepthRec(FBaseTreeNode* InRoot, TArray<int32>& InOutNumNodesPerDepth, int32 InDepth, int32 InMaxDepth, int InMaxNodes) const;
+	void ExpandChildGroupsRec(FBaseTreeNode* InRoot, int32 InDepth, int32 InMaxDepth);
+
 protected:
 	/** Table view model. */
 	TSharedPtr<Insights::FTable> Table;
