[Insights] Functionality to copy column value and tooltip to clipboard for STableTreeView.

diff --git a/Engine/Source/Developer/TraceInsights/Private/Insights/MemoryProfiler/ViewModels/MemAllocTable.cpp b/Engine/Source/Developer/TraceInsights/Private/Insights/MemoryProfiler/ViewModels/MemAllocTable.cpp
--- a/Engine/Source/Developer/TraceInsights/Private/Insights/MemoryProfiler/ViewModels/MemAllocTable.cpp
+++ b/Engine/Source/Developer/TraceInsights/Private/Insights/MemoryProfiler/ViewModels/MemAllocTable.cpp
@@ -576,6 +576,12 @@ void FMemAllocTable::AddDefaultColumns()
 							]
 						];
 			}
+
+			virtual FText FormatValueForTooltip(const FTableColumn& Column, const FBaseTreeNode& Node) const override
+			{
+				const FMemAllocNode& MemAllocNode = static_cast<const FMemAllocNode&>(Node);
+				return MemAllocNode.GetFullCallstack();
+			}
 		};
 
 		TSharedRef<ITableCellValueFormatter> Formatter = MakeShared<FunctionValueFormatter>();
diff --git a/Engine/Source/Developer/TraceInsights/Private/Insights/Table/Widgets/STableTreeView.cpp b/Engine/Source/Developer/TraceInsights/Private/Insights/Table/Widgets/STableTreeView.cpp
--- a/Engine/Source/Developer/TraceInsights/Private/Insights/Table/Widgets/STableTreeView.cpp
+++ b/Engine/Source/Developer/TraceInsights/Private/Insights/Table/Widgets/STableTreeView.cpp
@@ -61,10 +61,14 @@ public:
 		virtual void RegisterCommands() override
 	{
 		UI_COMMAND(Command_CopyToClipboard, "Copy To Clipboard", "Copies selection to clipboard", EUserInterfaceActionType::Button, FInputChord(EModifierKey::Control, EKeys::C));
+		UI_COMMAND(Command_CopyColumnToClipboard, "Copy Column Value To Clipboard", "Copies the value of hovered column to clipboard", EUserInterfaceActionType::Button, FInputChord(EModifierKey::Control | EModifierKey::Shift, EKeys::C));
+		UI_COMMAND(Command_CopyColumnTooltipToClipboard, "Copy Column Tooltip To Clipboard", "Copies the value of hovered column's tooltip to clipboard", EUserInterfaceActionType::Button, FInputChord(EModifierKey::Control | EModifierKey::Alt, EKeys::C));
 	}
 	PRAGMA_ENABLE_OPTIMIZATION
 
 	TSharedPtr<FUICommandInfo> Command_CopyToClipboard;
+	TSharedPtr<FUICommandInfo> Command_CopyColumnToClipboard;
+	TSharedPtr<FUICommandInfo> Command_CopyColumnTooltipToClipboard;
 };
 
 const FName STableTreeView::RootNodeName(TEXT("Root"));
@@ -341,6 +345,8 @@ void STableTreeView::InitCommandList()
 	FTableTreeViewCommands::Register();
 	CommandList = MakeShared<FUICommandList>();
 	CommandList->MapAction(FTableTreeViewCommands::Get().Command_CopyToClipboard, FExecuteAction::CreateSP(this, &STableTreeView::ContextMenu_CopySelectedToClipboard_Execute), FCanExecuteAction::CreateSP(this, &STableTreeView::ContextMenu_CopySelectedToClipboard_CanExecute));
+	CommandList->MapAction(FTableTreeViewCommands::Get().Command_CopyColumnToClipboard, FExecuteAction::CreateSP(this, &STableTreeView::ContextMenu_CopyColumnToClipboard_Execute), FCanExecuteAction::CreateSP(this, &STableTreeView::ContextMenu_CopyColumnToClipboard_CanExecute));
+	CommandList->MapAction(FTableTreeViewCommands::Get().Command_CopyColumnTooltipToClipboard, FExecuteAction::CreateSP(this, &STableTreeView::ContextMenu_CopyColumnTooltipToClipboard_Execute), FCanExecuteAction::CreateSP(this, &STableTreeView::ContextMenu_CopyColumnTooltipToClipboard_CanExecute));
 }
 
 ////////////////////////////////////////////////////////////////////////////////////////////////////
@@ -417,6 +423,24 @@ TSharedPtr<SWidget> STableTreeView::TreeView_GetMenuContent()
 			FSlateIcon(FEditorStyle::GetStyleSetName(), "Profiler.Misc.CopyToClipboard")
 		);
 
+		MenuBuilder.AddMenuEntry
+		(
+			FTableTreeViewCommands::Get().Command_CopyColumnToClipboard,
+			NAME_None,
+			TAttribute<FText>(),
+			TAttribute<FText>(),
+			FSlateIcon(FEditorStyle::GetStyleSetName(), "Profiler.Misc.CopyToClipboard")
+		);
+
+		MenuBuilder.AddMenuEntry
+		(
+			FTableTreeViewCommands::Get().Command_CopyColumnTooltipToClipboard,
+			NAME_None,
+			TAttribute<FText>(),
+			TAttribute<FText>(),
+			FSlateIcon(FEditorStyle::GetStyleSetName(), "Profiler.Misc.CopyToClipboard")
+		);
+
 		MenuBuilder.AddSubMenu
 		(
 			LOCTEXT("ContextMenu_Header_Misc_Sort", "Sort By"),
@@ -953,11 +977,10 @@ bool STableTreeView::TableRow_ShouldBeEnabled(FTableTreeNodePtr NodePtr) const
 
 void STableTreeView::TableRow_SetHoveredCell(TSharedPtr<FTable> InTablePtr, TSharedPtr<FTableColumn> InColumnPtr, FTableTreeNodePtr InNodePtr)
 {
-	HoveredColumnId = InColumnPtr ? InColumnPtr->GetId() : FName();
-
 	const bool bIsAnyMenusVisible = FSlateApplication::Get().AnyMenusVisible();
 	if (!HasMouseCapture() && !bIsAnyMenusVisible)
 	{
+		HoveredColumnId = InColumnPtr ? InColumnPtr->GetId() : FName();
 		HoveredNodePtr = InNodePtr;
 	}
 }
@@ -2616,6 +2639,52 @@ void STableTreeView::ContextMenu_CopySelectedToClipboard_Execute()
 
 ////////////////////////////////////////////////////////////////////////////////////////////////////
 
+bool STableTreeView::ContextMenu_CopyColumnToClipboard_CanExecute() const
+{
+	const TSharedPtr<FTableColumn> HoveredColumnPtr = Table->FindColumn(HoveredColumnId);
+
+	if (!HoveredColumnPtr.IsValid() || !HoveredNodePtr.IsValid())
+	{
+		return false;
+	}
+	
+	return true;
+}
+
+////////////////////////////////////////////////////////////////////////////////////////////////////
+
+void STableTreeView::ContextMenu_CopyColumnToClipboard_Execute()
+{
+	const TSharedPtr<FTableColumn> HoveredColumnPtr = Table->FindColumn(HoveredColumnId);
+	FString Text = HoveredColumnPtr->GetValueAsText(*HoveredNodePtr).ToString();
+	FPlatformApplicationMisc::ClipboardCopy(*Text);
+}
+
+////////////////////////////////////////////////////////////////////////////////////////////////////
+
+bool STableTreeView::ContextMenu_CopyColumnTooltipToClipboard_CanExecute() const
+{
+	const TSharedPtr<FTableColumn> HoveredColumnPtr = Table->FindColumn(HoveredColumnId);
+
+	if (!HoveredColumnPtr.IsValid() || !HoveredNodePtr.IsValid())
+	{
+		return false;
+	}
+
+	return true;
+}
+
+////////////////////////////////////////////////////////////////////////////////////////////////////
+
+void STableTreeView::ContextMenu_CopyColumnTooltipToClipboard_Execute()
+{
+	const TSharedPtr<FTableColumn> HoveredColumnPtr = Table->FindColumn(HoveredColumnId);
+	FString Text = HoveredColumnPtr->GetValueAsTooltipText(*HoveredNodePtr).ToString();
+	FPlatformApplicationMisc::ClipboardCopy(*Text);
+}
+
+////////////////////////////////////////////////////////////////////////////////////////////////////
+
 } // namespace Insights
 
 #undef LOCTEXT_NAMESPACE
diff --git a/Engine/Source/Developer/TraceInsights/Private/Insights/Table/Widgets/STableTreeView.h b/Engine/Source/Developer/TraceInsights/Private/Insights/Table/Widgets/STableTreeView.h
--- a/Engine/Source/Developer/TraceInsights/Private/Insights/Table/Widgets/STableTreeView.h
+++ b/Engine/Source/Developer/TraceInsights/Private/Insights/Table/Widgets/STableTreeView.h
@@ -161,6 +161,10 @@ protected:
 
 	bool ContextMenu_CopySelectedToClipboard_CanExecute() const;
 	void ContextMenu_CopySelectedToClipboard_Execute();
+	bool ContextMenu_CopyColumnToClipboard_CanExecute() const;
+	void ContextMenu_CopyColumnToClipboard_Execute();
+	bool ContextMenu_CopyColumnTooltipToClipboard_CanExecute() const;
+	void ContextMenu_CopyColumnTooltipToClipboard_Execute();
 
 	////////////////////////////////////////////////////////////////////////////////////////////////////
 	// Tree View - Columns' Header
