[Insights] Mem Allocs: Enabled sorting settings in the predefined breakdown views.

diff --git a/Engine/Source/Developer/TraceInsights/Private/Insights/MemoryProfiler/Widgets/SMemAllocTableTreeView.cpp b/Engine/Source/Developer/TraceInsights/Private/Insights/MemoryProfiler/Widgets/SMemAllocTableTreeView.cpp
--- a/Engine/Source/Developer/TraceInsights/Private/Insights/MemoryProfiler/Widgets/SMemAllocTableTreeView.cpp
+++ b/Engine/Source/Developer/TraceInsights/Private/Insights/MemoryProfiler/Widgets/SMemAllocTableTreeView.cpp
@@ -449,6 +449,10 @@ TSharedPtr<SWidget> SMemAllocTableTreeView::ConstructToolbar()
 
 FReply SMemAllocTableTreeView::OnDetailedViewClicked()
 {
+	ColumnBeingSorted = FTable::GetHierarchyColumnId();
+	ColumnSortMode = EColumnSortMode::Type::Ascending;
+	UpdateCurrentSortingByColumn();
+
 	PreChangeGroupings();
 
 	CurrentGroupings.Reset();
@@ -479,6 +483,10 @@ FReply SMemAllocTableTreeView::OnDetailedViewClicked()
 
 FReply SMemAllocTableTreeView::OnSizeViewClicked()
 {
+	ColumnBeingSorted = FMemAllocTableColumns::SizeColumnId;
+	ColumnSortMode = EColumnSortMode::Type::Descending;
+	UpdateCurrentSortingByColumn();
+
 	PreChangeGroupings();
 
 	CurrentGroupings.Reset();
@@ -519,6 +527,10 @@ FReply SMemAllocTableTreeView::OnSizeViewClicked()
 
 FReply SMemAllocTableTreeView::OnTagViewClicked()
 {
+	ColumnBeingSorted = FTable::GetHierarchyColumnId();
+	ColumnSortMode = EColumnSortMode::Type::Ascending;
+	UpdateCurrentSortingByColumn();
+
 	PreChangeGroupings();
 
 	CurrentGroupings.Reset();
@@ -560,6 +572,10 @@ FReply SMemAllocTableTreeView::OnTagViewClicked()
 
 FReply SMemAllocTableTreeView::OnCallstackViewClicked(bool bIsInverted)
 {
+	ColumnBeingSorted = FMemAllocTableColumns::SizeColumnId;
+	ColumnSortMode = EColumnSortMode::Type::Descending;
+	UpdateCurrentSortingByColumn();
+
 	PreChangeGroupings();
 
 	CurrentGroupings.Reset();
