[Insights] Memory: Split Mem Investigation into separate panel (from LLM Tags panel).

diff --git a/Engine/Source/Developer/TraceInsights/Private/Insights/InsightsStyle.cpp b/Engine/Source/Developer/TraceInsights/Private/Insights/InsightsStyle.cpp
--- a/Engine/Source/Developer/TraceInsights/Private/Insights/InsightsStyle.cpp
+++ b/Engine/Source/Developer/TraceInsights/Private/Insights/InsightsStyle.cpp
@@ -143,6 +143,12 @@ TSharedRef<FSlateStyleSet> FInsightsStyle::Create()
 	Style.Set("MemoryProfiler.Icon.Large", new IMAGE_BRUSH("Icons/Profiler/profiler_stats_40x", Icon32x32));
 	Style.Set("MemoryProfiler.Icon.Small", new IMAGE_BRUSH("Icons/Profiler/profiler_stats_40x", Icon16x16));
 
+	Style.Set("MemInvestigationView.Icon.Large", new IMAGE_BRUSH("Icons/Profiler/profiler_stats_40x", Icon32x32));
+	Style.Set("MemInvestigationView.Icon.Small", new IMAGE_BRUSH("Icons/Profiler/profiler_stats_40x", Icon16x16));
+
+	Style.Set("MemAllocTableTreeView.Icon.Large", new IMAGE_BRUSH("Icons/Profiler/profiler_stats_40x", Icon32x32));
+	Style.Set("MemAllocTableTreeView.Icon.Small", new IMAGE_BRUSH("Icons/Profiler/profiler_stats_40x", Icon16x16));
+
 	Style.Set("MemTagTreeView.Icon.Large", new IMAGE_BRUSH("Icons/Profiler/profiler_stats_40x", Icon32x32));
 	Style.Set("MemTagTreeView.Icon.Small", new IMAGE_BRUSH("Icons/Profiler/profiler_stats_40x", Icon16x16));
 
diff --git a/Engine/Source/Developer/TraceInsights/Private/Insights/MemoryProfiler/MemoryProfilerCommands.cpp b/Engine/Source/Developer/TraceInsights/Private/Insights/MemoryProfiler/MemoryProfilerCommands.cpp
--- a/Engine/Source/Developer/TraceInsights/Private/Insights/MemoryProfiler/MemoryProfilerCommands.cpp
+++ b/Engine/Source/Developer/TraceInsights/Private/Insights/MemoryProfiler/MemoryProfilerCommands.cpp
@@ -53,6 +53,7 @@ PRAGMA_DISABLE_OPTIMIZATION
 void FMemoryProfilerCommands::RegisterCommands()
 {
 	UI_COMMAND(ToggleTimingViewVisibility, "Timing", "Toggles the visibility of the main Timing view", EUserInterfaceActionType::ToggleButton, FInputChord(EModifierKey::Control, EKeys::T));
+	UI_COMMAND(ToggleMemInvestigationViewVisibility, "Investigation", "Toggles the visibility of the Memory Investigation (Alloc Queries) view", EUserInterfaceActionType::ToggleButton, FInputChord(EModifierKey::Control, EKeys::I));
 	UI_COMMAND(ToggleMemTagTreeViewVisibility, "LLM Tags", "Toggles the visibility of the LLM Tags tree view", EUserInterfaceActionType::ToggleButton, FInputChord(EModifierKey::Control, EKeys::M));
 }
 PRAGMA_ENABLE_OPTIMIZATION
@@ -95,6 +96,7 @@ PRAGMA_ENABLE_OPTIMIZATION
 	}
 
 IMPLEMENT_TOGGLE_COMMAND(ToggleTimingViewVisibility, IsTimingViewVisible, ShowHideTimingView)
+IMPLEMENT_TOGGLE_COMMAND(ToggleMemInvestigationViewVisibility, IsMemInvestigationViewVisible, ShowHideMemInvestigationView)
 IMPLEMENT_TOGGLE_COMMAND(ToggleMemTagTreeViewVisibility, IsMemTagTreeViewVisible, ShowHideMemTagTreeView)
 
 #undef IMPLEMENT_TOGGLE_COMMAND
diff --git a/Engine/Source/Developer/TraceInsights/Private/Insights/MemoryProfiler/MemoryProfilerCommands.h b/Engine/Source/Developer/TraceInsights/Private/Insights/MemoryProfiler/MemoryProfilerCommands.h
--- a/Engine/Source/Developer/TraceInsights/Private/Insights/MemoryProfiler/MemoryProfilerCommands.h
+++ b/Engine/Source/Developer/TraceInsights/Private/Insights/MemoryProfiler/MemoryProfilerCommands.h
@@ -32,16 +32,10 @@ public:
 	/** Toggles visibility for the Timing view. Global and custom command. */
 	TSharedPtr<FUICommandInfo> ToggleTimingViewVisibility;
 
-	/** Toggles visibility for the Event Aggregation tree view. Global and custom command. */
-	TSharedPtr<FUICommandInfo> ToggleEventAggregationTreeViewVisibility;
+	/** Toggles visibility for the Memory Investigation view. Global and custom command. */
+	TSharedPtr<FUICommandInfo> ToggleMemInvestigationViewVisibility;
 
-	/** Toggles visibility for the Object Type Aggregation tree view. Global and custom command. */
-	TSharedPtr<FUICommandInfo> ToggleObjectTypeAggregationTreeViewVisibility;
-
-	/** Toggles visibility for the Package Details tree view. Global and custom command. */
-	TSharedPtr<FUICommandInfo> TogglePackageDetailsTreeViewVisibility;
-
-	/** Toggles visibility for the Export Details tree view. Global and custom command. */
+	/** Toggles visibility for the LLM Tags tree view. Global and custom command. */
 	TSharedPtr<FUICommandInfo> ToggleMemTagTreeViewVisibility;
 };
 
@@ -95,6 +89,7 @@ protected:\
 	ECheckBoxState CmdName##_GetCheckState() const; /**< Handles FGetActionCheckState for CmdName. */
 
 	DECLARE_TOGGLE_COMMAND(ToggleTimingViewVisibility)
+	DECLARE_TOGGLE_COMMAND(ToggleMemInvestigationViewVisibility)
 	DECLARE_TOGGLE_COMMAND(ToggleMemTagTreeViewVisibility)
 #undef DECLARE_TOGGLE_COMMAND
 
diff --git a/Engine/Source/Developer/TraceInsights/Private/Insights/MemoryProfiler/MemoryProfilerManager.cpp b/Engine/Source/Developer/TraceInsights/Private/Insights/MemoryProfiler/MemoryProfilerManager.cpp
--- a/Engine/Source/Developer/TraceInsights/Private/Insights/MemoryProfiler/MemoryProfilerManager.cpp
+++ b/Engine/Source/Developer/TraceInsights/Private/Insights/MemoryProfiler/MemoryProfilerManager.cpp
@@ -58,6 +58,7 @@ FMemoryProfilerManager::FMemoryProfilerManager(TSharedRef<FUICommandList> InComm
 	, ActionManager(this)
 	, ProfilerWindow(nullptr)
 	, bIsTimingViewVisible(false)
+	, bIsMemInvestigationViewVisible(false)
 	, bIsMemTagTreeViewVisible(false)
 	, LogListingName(TEXT("MemoryInsights"))
 {
@@ -131,6 +132,7 @@ FMemoryProfilerManager::~FMemoryProfilerManager()
 void FMemoryProfilerManager::BindCommands()
 {
 	ActionManager.Map_ToggleTimingViewVisibility_Global();
+	ActionManager.Map_ToggleMemInvestigationViewVisibility_Global();
 	ActionManager.Map_ToggleMemTagTreeViewVisibility_Global();
 }
 
@@ -331,6 +333,19 @@ void FMemoryProfilerManager::ShowHideTimingView(const bool bIsVisible)
 
 ////////////////////////////////////////////////////////////////////////////////////////////////////
 
+void FMemoryProfilerManager::ShowHideMemInvestigationView(const bool bIsVisible)
+{
+	bIsMemInvestigationViewVisible = bIsVisible;
+
+	TSharedPtr<SMemoryProfilerWindow> Wnd = GetProfilerWindow();
+	if (Wnd.IsValid())
+	{
+		Wnd->ShowHideTab(FMemoryProfilerTabs::MemInvestigationViewID, bIsVisible);
+	}
+}
+
+////////////////////////////////////////////////////////////////////////////////////////////////////
+
 void FMemoryProfilerManager::ShowHideMemTagTreeView(const bool bIsVisible)
 {
 	bIsMemTagTreeViewVisible = bIsVisible;
diff --git a/Engine/Source/Developer/TraceInsights/Private/Insights/MemoryProfiler/MemoryProfilerManager.h b/Engine/Source/Developer/TraceInsights/Private/Insights/MemoryProfiler/MemoryProfilerManager.h
--- a/Engine/Source/Developer/TraceInsights/Private/Insights/MemoryProfiler/MemoryProfilerManager.h
+++ b/Engine/Source/Developer/TraceInsights/Private/Insights/MemoryProfiler/MemoryProfilerManager.h
@@ -87,12 +87,14 @@ public:
 	////////////////////////////////////////////////////////////////////////////////////////////////////
 	// Getters and setters used by Toggle Commands.
 
-	/** @return true, if the Timing view is visible */
 	const bool IsTimingViewVisible() const { return bIsTimingViewVisible; }
 	void SetTimingViewVisible(const bool bIsVisible) { bIsTimingViewVisible = bIsVisible; }
 	void ShowHideTimingView(const bool bIsVisible);
 
-	/** @return true, if the Export Details tree view is visible */
+	const bool IsMemInvestigationViewVisible() const { return bIsMemInvestigationViewVisible; }
+	void SetMemInvestigationViewVisible(const bool bIsVisible) { bIsMemInvestigationViewVisible = bIsVisible; }
+	void ShowHideMemInvestigationView(const bool bIsVisible);
+
 	const bool IsMemTagTreeViewVisible() const { return bIsMemTagTreeViewVisible; }
 	void SetMemTagTreeViewVisible(const bool bIsVisible) { bIsMemTagTreeViewVisible = bIsVisible; }
 	void ShowHideMemTagTreeView(const bool bIsVisible);
@@ -141,7 +143,10 @@ private:
 	/** If the Timing view is visible or hidden. */
 	bool bIsTimingViewVisible;
 
-	/** If the Categories tree view is visible or hidden. */
+	/** If the Memory Investigation (Alloc Queries) view is visible or hidden. */
+	bool bIsMemInvestigationViewVisible;
+
+	/** If the LLM Tags tree view is visible or hidden. */
 	bool bIsMemTagTreeViewVisible;
 
 	/** The name of the Memory Insights log listing. */
diff --git a/Engine/Source/Developer/TraceInsights/Private/Insights/MemoryProfiler/Widgets/SMemInvestigationView.cpp b/Engine/Source/Developer/TraceInsights/Private/Insights/MemoryProfiler/Widgets/SMemInvestigationView.cpp
new file mode 100644
--- /dev/null
+++ b/Engine/Source/Developer/TraceInsights/Private/Insights/MemoryProfiler/Widgets/SMemInvestigationView.cpp
@@ -0,0 +1,409 @@
+// Copyright Epic Games, Inc. All Rights Reserved.
+
+#include "SMemInvestigationView.h"
+
+#include "DesktopPlatformModule.h"
+#include "EditorStyleSet.h"
+#include "SlateOptMacros.h"
+#include "Widgets/Input/SButton.h"
+#include "Widgets/Input/SEditableTextBox.h"
+#include "Widgets/Images/SImage.h"
+
+// Insights
+#include "Insights/Common/Stopwatch.h"
+#include "Insights/Common/TimeUtils.h"
+#include "Insights/InsightsStyle.h"
+#include "Insights/MemoryProfiler/MemoryProfilerManager.h"
+#include "Insights/MemoryProfiler/ViewModels/MemorySharedState.h"
+#include "Insights/MemoryProfiler/Widgets/SMemAllocTableTreeView.h"
+#include "Insights/MemoryProfiler/Widgets/SMemoryProfilerWindow.h"
+#include "Insights/ViewModels/TimeRulerTrack.h"
+#include "Insights/Widgets/STimingView.h"
+
+#define LOCTEXT_NAMESPACE "SMemInvestigationView"
+
+////////////////////////////////////////////////////////////////////////////////////////////////////
+
+SMemInvestigationView::SMemInvestigationView()
+	: ProfilerWindow()
+{
+}
+
+////////////////////////////////////////////////////////////////////////////////////////////////////
+
+SMemInvestigationView::~SMemInvestigationView()
+{
+	// Remove ourselves from the Insights manager.
+	if (FInsightsManager::Get().IsValid())
+	{
+		FInsightsManager::Get()->GetSessionChangedEvent().RemoveAll(this);
+	}
+}
+
+////////////////////////////////////////////////////////////////////////////////////////////////////
+BEGIN_SLATE_FUNCTION_BUILD_OPTIMIZATION
+
+void SMemInvestigationView::Construct(const FArguments& InArgs, TSharedPtr<SMemoryProfilerWindow> InProfilerWindow)
+{
+	check(InProfilerWindow.IsValid());
+	ProfilerWindow = InProfilerWindow;
+
+	ChildSlot
+	[
+		SNew(SBorder)
+		.BorderImage(FEditorStyle::GetBrush("NoBorder"))
+		.Padding(8.0f)
+		[
+			ConstructInvestigationWidgetArea()
+		]
+	];
+
+	//BindCommands();
+
+	// Register ourselves with the Insights manager.
+	FInsightsManager::Get()->GetSessionChangedEvent().AddSP(this, &SMemInvestigationView::InsightsManager_OnSessionChanged);
+
+	// Update the Session (i.e. when analysis session was already started).
+	InsightsManager_OnSessionChanged();
+}
+
+////////////////////////////////////////////////////////////////////////////////////////////////////
+
+TSharedRef<SWidget> SMemInvestigationView::ConstructInvestigationWidgetArea()
+{
+	TSharedRef<SWidget> Widget =
+		SNew(SVerticalBox)
+
+		+ SVerticalBox::Slot()
+		.AutoHeight()
+		[
+			SNew(SHorizontalBox)
+
+			+ SHorizontalBox::Slot()
+			.AutoWidth()
+			.VAlign(VAlign_Center)
+			.Padding(0.0f, 0.0f, 4.0f, 0.0f)
+			[
+				SNew(STextBlock)
+				.Text(LOCTEXT("QueryRuleText", "Rule:"))
+			]
+
+			+ SHorizontalBox::Slot()
+			.FillWidth(1.0f)
+			.VAlign(VAlign_Center)
+			[
+				SAssignNew(QueryRuleComboBox, SComboBox<TSharedPtr<Insights::FMemoryRuleSpec>>)
+				.ToolTipText(this, &SMemInvestigationView::QueryRule_GetTooltipText)
+				.OptionsSource(GetAvailableQueryRules())
+				.OnSelectionChanged(this, &SMemInvestigationView::QueryRule_OnSelectionChanged)
+				.OnGenerateWidget(this, &SMemInvestigationView::QueryRule_OnGenerateWidget)
+				[
+					SNew(STextBlock)
+					.Text(this, &SMemInvestigationView::QueryRule_GetSelectedText)
+				]
+			]
+		]
+
+		+ SVerticalBox::Slot()
+		.AutoHeight()
+		[
+			ConstructTimeMarkerWidget(0)
+		]
+
+		+ SVerticalBox::Slot()
+		.AutoHeight()
+		[
+			ConstructTimeMarkerWidget(1)
+		]
+
+		+ SVerticalBox::Slot()
+		.AutoHeight()
+		[
+			ConstructTimeMarkerWidget(2)
+		]
+
+		+ SVerticalBox::Slot()
+		.AutoHeight()
+		[
+			ConstructTimeMarkerWidget(3)
+		]
+
+		+ SVerticalBox::Slot()
+		.AutoHeight()
+		.Padding(0.0f, 4.0f, 0.0f, 2.0f)
+		[
+			SNew(STextBlock)
+			.Text(this, &SMemInvestigationView::QueryRule_GetTooltipText)
+			.ColorAndOpacity(FLinearColor(0.3f, 0.3f, 0.3f, 1.0f))
+			.AutoWrapText(true)
+		]
+
+		+ SVerticalBox::Slot()
+		.AutoHeight()
+		.HAlign(HAlign_Right)
+		[
+			SNew(SButton)
+			.Text(LOCTEXT("RunQueryBtnText", "Run Query"))
+			.ToolTipText(LOCTEXT("RunQueryBtnToolTipText", "Run Memory Query.\nThe resulting list of allocations will be available in a tree view."))
+			.OnClicked(this, &SMemInvestigationView::RunQuery)
+		]
+	;
+
+	if (ProfilerWindow.IsValid())
+	{
+		FMemorySharedState& SharedState = ProfilerWindow->GetSharedState();
+		QueryRuleComboBox->SetSelectedItem(SharedState.GetCurrentMemoryRule());
+	}
+
+	return Widget;
+}
+
+////////////////////////////////////////////////////////////////////////////////////////////////////
+
+TSharedRef<SWidget> SMemInvestigationView::QueryRule_OnGenerateWidget(TSharedPtr<Insights::FMemoryRuleSpec> InRule)
+{
+	const FText QueryRuleText = FText::Format(LOCTEXT("QueryRuleComboBox_TextFmt", "{0} ({1})"), InRule->GetVerboseName(), InRule->GetShortName());
+
+	TSharedRef<SHorizontalBox> Widget = SNew(SHorizontalBox);
+
+	Widget->AddSlot()
+		.AutoWidth()
+		.HAlign(HAlign_Center)
+		.VAlign(VAlign_Center)
+		.Padding(2.0f, 2.0f, 0.0f, 2.0f)
+		[
+			SNew(SImage)
+			.Visibility_Lambda([Widget]() { return Widget->GetParentWidget()->IsHovered() ? EVisibility::Visible : EVisibility::Hidden; })
+			.Image(FEditorStyle::GetBrush("Profiler.Tooltip.HintIcon10"))
+			.ToolTipText(InRule->GetDescription())
+		];
+
+	Widget->AddSlot()
+		.AutoWidth()
+		[
+			SNew(STextBlock)
+			.Text(QueryRuleText)
+			.Margin(2.0f)
+		];
+
+	return Widget;
+}
+
+////////////////////////////////////////////////////////////////////////////////////////////////////
+
+TSharedRef<SWidget> SMemInvestigationView::ConstructTimeMarkerWidget(uint32 TimeMarkerIndex)
+{
+	if (!ProfilerWindow.IsValid())
+	{
+		return SNew(SBox);
+	}
+
+	FMemorySharedState& SharedState = ProfilerWindow->GetSharedState();
+	const uint32 NumTimeMarkers = ProfilerWindow->GetNumCustomTimeMarkers();
+	if (TimeMarkerIndex >= NumTimeMarkers)
+	{
+		return SNew(SBox);
+	}
+
+	const TSharedRef<Insights::FTimeMarker>& TimeMarker = ProfilerWindow->GetCustomTimeMarker(TimeMarkerIndex);
+
+	TSharedRef<SWidget> Widget = SNew(SHorizontalBox)
+		.Visibility_Lambda([TimeMarker]()
+			{
+				return TimeMarker->IsVisible() ? EVisibility::Visible : EVisibility::Collapsed;
+			})
+
+		+ SHorizontalBox::Slot()
+		.AutoWidth()
+		.VAlign(VAlign_Center)
+		.Padding(20.0f, 1.0f, 4.0f, 1.0f)
+		[
+			SNew(STextBlock)
+			.Text(FText::FromString(TimeMarker->GetName()))
+			.ColorAndOpacity(FSlateColor(TimeMarker->GetColor()))
+		]
+
+		+ SHorizontalBox::Slot()
+		.AutoWidth()
+		.VAlign(VAlign_Center)
+		.Padding(0.0f, 1.0f, 0.0f, 1.0f)
+		[
+			SNew(SEditableTextBox)
+			.MinDesiredWidth(110.0f)
+			.Text_Lambda([TimeMarker]()
+				{
+					return FText::FromString(FString::Printf(TEXT("%.9f"), TimeMarker->GetTime()));
+				})
+			.OnTextCommitted_Lambda([TimeMarker](const FText& InText, ETextCommit::Type InCommitType)
+				{
+					const double Time = FCString::Atod(*InText.ToString());
+					TimeMarker->SetTime(Time);
+				})
+		]
+
+		+ SHorizontalBox::Slot()
+		.AutoWidth()
+		.VAlign(VAlign_Center)
+		.Padding(4.0f, 1.0f, 0.0f, 1.0f)
+		[
+			SNew(STextBlock)
+			.Text_Lambda([TimeMarker]()
+				{
+					return FText::FromString(TimeUtils::FormatTime(TimeMarker->GetTime(), 0.1));
+				})
+		];
+
+	return Widget;
+}
+
+END_SLATE_FUNCTION_BUILD_OPTIMIZATION
+////////////////////////////////////////////////////////////////////////////////////////////////////
+
+void SMemInvestigationView::InsightsManager_OnSessionChanged()
+{
+	TSharedPtr<const TraceServices::IAnalysisSession> NewSession = FInsightsManager::Get()->GetSession();
+
+	if (NewSession != Session)
+	{
+		Session = NewSession;
+		Reset();
+	}
+}
+
+////////////////////////////////////////////////////////////////////////////////////////////////////
+
+void SMemInvestigationView::Reset()
+{
+	//...
+}
+
+////////////////////////////////////////////////////////////////////////////////////////////////////
+
+void SMemInvestigationView::Tick(const FGeometry& AllottedGeometry, const double InCurrentTime, const float InDeltaTime)
+{
+}
+
+////////////////////////////////////////////////////////////////////////////////////////////////////
+// Query Rules
+////////////////////////////////////////////////////////////////////////////////////////////////////
+
+const TArray<TSharedPtr<Insights::FMemoryRuleSpec>>* SMemInvestigationView::GetAvailableQueryRules()
+{
+	if (ProfilerWindow.IsValid())
+	{
+		FMemorySharedState& SharedState = ProfilerWindow->GetSharedState();
+		return &SharedState.GetMemoryRules();
+	}
+	return nullptr;
+}
+
+////////////////////////////////////////////////////////////////////////////////////////////////////
+
+void SMemInvestigationView::QueryRule_OnSelectionChanged(TSharedPtr<Insights::FMemoryRuleSpec> InRule, ESelectInfo::Type SelectInfo)
+{
+	if (SelectInfo != ESelectInfo::Direct)
+	{
+		if (ProfilerWindow.IsValid() && InRule)
+		{
+			FMemorySharedState& SharedState = ProfilerWindow->GetSharedState();
+			SharedState.SetCurrentMemoryRule(InRule);
+			ProfilerWindow->UpdateTimingViewMarkers();
+		}
+	}
+}
+
+////////////////////////////////////////////////////////////////////////////////////////////////////
+
+FText SMemInvestigationView::QueryRule_GetSelectedText() const
+{
+	if (ProfilerWindow.IsValid())
+	{
+		FMemorySharedState& SharedState = ProfilerWindow->GetSharedState();
+		TSharedPtr<Insights::FMemoryRuleSpec> Rule = SharedState.GetCurrentMemoryRule();
+		if (Rule)
+		{
+			return FText::Format(LOCTEXT("QueryRuleComboBox_TextFmt", "{0} ({1})"), Rule->GetVerboseName(), Rule->GetShortName());
+		}
+	}
+	return FText::GetEmpty();
+}
+
+////////////////////////////////////////////////////////////////////////////////////////////////////
+
+FText SMemInvestigationView::QueryRule_GetTooltipText() const
+{
+	if (ProfilerWindow.IsValid())
+	{
+		FMemorySharedState& SharedState = ProfilerWindow->GetSharedState();
+		if (SharedState.GetCurrentMemoryRule())
+		{
+			return SharedState.GetCurrentMemoryRule()->GetDescription();
+		}
+	}
+	return FText::GetEmpty();
+}
+
+////////////////////////////////////////////////////////////////////////////////////////////////////
+
+FReply SMemInvestigationView::RunQuery()
+{
+	if (!ProfilerWindow.IsValid())
+	{
+		UE_LOG(MemoryProfiler, Error, TEXT("[MemQuery] Invalid Profiler Window!"));
+		return FReply::Handled();
+	}
+
+	FMemorySharedState& SharedState = ProfilerWindow->GetSharedState();
+	TSharedPtr<Insights::FMemoryRuleSpec> Rule = SharedState.GetCurrentMemoryRule();
+	if (!Rule)
+	{
+		UE_LOG(MemoryProfiler, Error, TEXT("[MemQuery] Invalid Rule!"));
+		return FReply::Handled();
+	}
+
+	const uint32 NumTimeMarkers = ProfilerWindow->GetNumCustomTimeMarkers();
+	const uint32 RuleNumTimeMarkers = Rule->GetNumTimeMarkers();
+	if (RuleNumTimeMarkers > NumTimeMarkers)
+	{
+		UE_LOG(MemoryProfiler, Error, TEXT("[MemQuery] Only %d time markers available. Current rule (%s) requires %u time markers!"),
+			NumTimeMarkers, *Rule->GetShortName().ToString(), RuleNumTimeMarkers);
+		return FReply::Handled();
+	}
+
+#if !NO_LOGGING
+	TStringBuilder<1024> Builder;
+	if (RuleNumTimeMarkers > 0)
+	{
+		Builder.Append(TEXT(" ("));
+		for (uint32 TimeMarkerIndex = 0; TimeMarkerIndex < RuleNumTimeMarkers; ++TimeMarkerIndex)
+		{
+			if (TimeMarkerIndex != 0)
+			{
+				Builder.Append(TEXT(", "));
+			}
+			Builder.Append(TEXT('A') + TimeMarkerIndex);
+			const TSharedRef<Insights::FTimeMarker>& TimeMarker = ProfilerWindow->GetCustomTimeMarker(TimeMarkerIndex);
+			Builder.Appendf(TEXT("=%.9f"), TimeMarker->GetTime());
+		}
+		Builder.Append(TEXT(")"));
+	}
+	UE_LOG(MemoryProfiler, Log, TEXT("[MemQuery] Run Query %s%s..."), *Rule->GetShortName().ToString(), Builder.ToString());
+#endif
+
+	TSharedPtr<Insights::SMemAllocTableTreeView> MemAllocTableTreeView = ProfilerWindow->ShowMemAllocTableTreeViewTab();
+	if (MemAllocTableTreeView)
+	{
+		double TimeMarkers[4];
+		TimeMarkers[0] = (RuleNumTimeMarkers > 0) ? ProfilerWindow->GetCustomTimeMarker(0)->GetTime() : 0.0;
+		TimeMarkers[1] = (RuleNumTimeMarkers > 1) ? ProfilerWindow->GetCustomTimeMarker(1)->GetTime() : 0.0;
+		TimeMarkers[2] = (RuleNumTimeMarkers > 2) ? ProfilerWindow->GetCustomTimeMarker(2)->GetTime() : 0.0;
+		TimeMarkers[3] = (RuleNumTimeMarkers > 3) ? ProfilerWindow->GetCustomTimeMarker(3)->GetTime() : 0.0;
+		MemAllocTableTreeView->SetQueryParams(Rule, TimeMarkers[0], TimeMarkers[1], TimeMarkers[2], TimeMarkers[3]);
+	}
+
+	return FReply::Handled();
+}
+
+////////////////////////////////////////////////////////////////////////////////////////////////////
+
+#undef LOCTEXT_NAMESPACE
diff --git a/Engine/Source/Developer/TraceInsights/Private/Insights/MemoryProfiler/Widgets/SMemInvestigationView.h b/Engine/Source/Developer/TraceInsights/Private/Insights/MemoryProfiler/Widgets/SMemInvestigationView.h
new file mode 100644
--- /dev/null
+++ b/Engine/Source/Developer/TraceInsights/Private/Insights/MemoryProfiler/Widgets/SMemInvestigationView.h
@@ -0,0 +1,83 @@
+// Copyright Epic Games, Inc. All Rights Reserved.
+
+#pragma once
+
+#include "CoreMinimal.h"
+#include "SlateFwd.h"
+#include "Widgets/DeclarativeSyntaxSupport.h"
+#include "Widgets/Input/SComboBox.h"
+#include "Widgets/SCompoundWidget.h"
+#include "Widgets/SWidget.h"
+
+////////////////////////////////////////////////////////////////////////////////////////////////////
+
+class FMenuBuilder;
+class SMemoryProfilerWindow;
+
+namespace TraceServices
+{
+	class IAnalysisSession;
+}
+
+namespace Insights
+{
+	class FMemoryRuleSpec;
+}
+
+////////////////////////////////////////////////////////////////////////////////////////////////////
+/**
+ * A custom widget used to setup and run mem (allocations) queries.
+ */
+class SMemInvestigationView : public SCompoundWidget
+{
+public:
+	/** Default constructor. */
+	SMemInvestigationView();
+
+	/** Virtual destructor. */
+	virtual ~SMemInvestigationView();
+
+	SLATE_BEGIN_ARGS(SMemInvestigationView) {}
+	SLATE_END_ARGS()
+
+	/**
+	 * Construct this widget
+	 * @param InArgs - The declaration data for this widget
+	 */
+	void Construct(const FArguments& InArgs, TSharedPtr<SMemoryProfilerWindow> InProfilerWindow);
+
+	void Reset();
+
+private:
+	TSharedRef<SWidget> ConstructInvestigationWidgetArea();
+	TSharedRef<SWidget> ConstructTimeMarkerWidget(uint32 TimeMarkerIndex);
+
+	/** Called when the analysis session has changed. */
+	void InsightsManager_OnSessionChanged();
+
+	const TArray<TSharedPtr<Insights::FMemoryRuleSpec>>* GetAvailableQueryRules();
+	void QueryRule_OnSelectionChanged(TSharedPtr<Insights::FMemoryRuleSpec> InRule, ESelectInfo::Type SelectInfo);
+	TSharedRef<SWidget> QueryRule_OnGenerateWidget(TSharedPtr<Insights::FMemoryRuleSpec> InRule);
+	FText QueryRule_GetSelectedText() const;
+	FText QueryRule_GetTooltipText() const;
+	FReply RunQuery();
+
+	/**
+	 * Ticks this widget.  Override in derived classes, but always call the parent implementation.
+	 *
+	 * @param  AllottedGeometry The space allotted for this widget
+	 * @param  InCurrentTime  Current absolute real time
+	 * @param  InDeltaTime  Real time passed since last tick
+	 */
+	virtual void Tick(const FGeometry& AllottedGeometry, const double InCurrentTime, const float InDeltaTime) override;
+
+private:
+	TSharedPtr<SMemoryProfilerWindow> ProfilerWindow;
+
+	/** A weak pointer to the profiler session used to populate this widget. */
+	TSharedPtr<const TraceServices::IAnalysisSession>/*Weak*/ Session;
+
+	TSharedPtr<SComboBox<TSharedPtr<Insights::FMemoryRuleSpec>>> QueryRuleComboBox;
+};
+
+////////////////////////////////////////////////////////////////////////////////////////////////////
diff --git a/Engine/Source/Developer/TraceInsights/Private/Insights/MemoryProfiler/Widgets/SMemTagTreeView.cpp b/Engine/Source/Developer/TraceInsights/Private/Insights/MemoryProfiler/Widgets/SMemTagTreeView.cpp
--- a/Engine/Source/Developer/TraceInsights/Private/Insights/MemoryProfiler/Widgets/SMemTagTreeView.cpp
+++ b/Engine/Source/Developer/TraceInsights/Private/Insights/MemoryProfiler/Widgets/SMemTagTreeView.cpp
@@ -11,12 +11,8 @@
 #include "Widgets/Input/SCheckBox.h"
 #include "Widgets/Layout/SScrollBox.h"
 #include "Widgets/Input/SButton.h"
-#include "Widgets/Input/SEditableTextBox.h"
 #include "Widgets/Input/SSearchBox.h"
 #include "Widgets/Images/SImage.h"
-#include "Widgets/Layout/SExpandableArea.h"
-#include "Widgets/Layout/SGridPanel.h"
-#include "Widgets/Layout/SSeparator.h"
 #include "Widgets/SToolTip.h"
 #include "Widgets/Views/STableViewBase.h"
 
@@ -29,15 +25,11 @@
 #include "Insights/MemoryProfiler/ViewModels/MemoryTracker.h"
 #include "Insights/MemoryProfiler/ViewModels/MemorySharedState.h"
 #include "Insights/MemoryProfiler/ViewModels/MemTagTreeViewColumnFactory.h"
-#include "Insights/MemoryProfiler/Widgets/SMemAllocTableTreeView.h"
 #include "Insights/MemoryProfiler/Widgets/SMemTagTreeViewTooltip.h"
 #include "Insights/MemoryProfiler/Widgets/SMemTagTreeViewTableRow.h"
 #include "Insights/MemoryProfiler/Widgets/SMemoryProfilerWindow.h"
 #include "Insights/Table/ViewModels/Table.h"
 #include "Insights/Table/ViewModels/TableColumn.h"
-#include "Insights/ViewModels/TimeRulerTrack.h"
-#include "Insights/ViewModels/TimingGraphTrack.h"
-#include "Insights/Widgets/STimingView.h"
 
 #define LOCTEXT_NAMESPACE "SMemTagTreeView"
 
@@ -71,8 +63,8 @@ SMemTagTreeView::~SMemTagTreeView()
 }
 
 ////////////////////////////////////////////////////////////////////////////////////////////////////
-
 BEGIN_SLATE_FUNCTION_BUILD_OPTIMIZATION
+
 void SMemTagTreeView::Construct(const FArguments& InArgs, TSharedPtr<SMemoryProfilerWindow> InProfilerWindow)
 {
 	check(InProfilerWindow.IsValid());
@@ -95,15 +87,6 @@ void SMemTagTreeView::Construct(const FArguments& InArgs, TSharedPtr<SMemoryProf
 			[
 				SNew(SVerticalBox)
 
-				// Investigation
-				+ SVerticalBox::Slot()
-				.VAlign(VAlign_Center)
-				.Padding(2.0f)
-				.AutoHeight()
-				[
-					ConstructInvestigationWidgetArea()
-				]
-
 				// Search / Filtering
 				+ SVerticalBox::Slot()
 				.VAlign(VAlign_Center)
@@ -205,216 +188,6 @@ void SMemTagTreeView::Construct(const FArguments& InArgs, TSharedPtr<SMemoryProf
 
 ////////////////////////////////////////////////////////////////////////////////////////////////////
 
-TSharedRef<SWidget> SMemTagTreeView::ConstructInvestigationWidgetArea()
-{
-	TSharedRef<SWidget> Widget = SNew(SExpandableArea)
-
-	.InitiallyCollapsed(false)
-	.BodyBorderImage(FEditorStyle::GetBrush("ToolPanel.GroupBorder"))
-	.BodyBorderBackgroundColor(FLinearColor::White)
-	.Padding(0.0f)
-	.HeaderContent()
-	[
-		SNew(STextBlock)
-		.Text(LOCTEXT("InvestigationTitle", "Investigation"))
-		.Font(FEditorStyle::GetFontStyle("DetailsView.CategoryFontStyle"))
-		.ShadowOffset(FVector2D(1.0f, 1.0f))
-	]
-	.BodyContent()
-	[
-		SNew(SVerticalBox)
-
-		+ SVerticalBox::Slot()
-		.AutoHeight()
-		[
-			SNew(SHorizontalBox)
-
-			+ SHorizontalBox::Slot()
-			.AutoWidth()
-			.VAlign(VAlign_Center)
-			.Padding(0.0f, 0.0f, 4.0f, 0.0f)
-			[
-				SNew(STextBlock)
-				.Text(LOCTEXT("QueryRuleText", "Rule:"))
-			]
-
-			+ SHorizontalBox::Slot()
-			.FillWidth(1.0f)
-			.VAlign(VAlign_Center)
-			[
-				SAssignNew(QueryRuleComboBox, SComboBox<TSharedPtr<Insights::FMemoryRuleSpec>>)
-				.ToolTipText(this, &SMemTagTreeView::QueryRule_GetTooltipText)
-				.OptionsSource(GetAvailableQueryRules())
-				.OnSelectionChanged(this, &SMemTagTreeView::QueryRule_OnSelectionChanged)
-				.OnGenerateWidget(this, &SMemTagTreeView::QueryRule_OnGenerateWidget)
-				[
-					SNew(STextBlock)
-					.Text(this, &SMemTagTreeView::QueryRule_GetSelectedText)
-				]
-			]
-		]
-
-		+ SVerticalBox::Slot()
-		.AutoHeight()
-		[
-			ConstructTimeMarkerWidget(0)
-		]
-
-		+ SVerticalBox::Slot()
-		.AutoHeight()
-		[
-			ConstructTimeMarkerWidget(1)
-		]
-
-		+ SVerticalBox::Slot()
-		.AutoHeight()
-		[
-			ConstructTimeMarkerWidget(2)
-		]
-
-		+ SVerticalBox::Slot()
-		.AutoHeight()
-		[
-			ConstructTimeMarkerWidget(3)
-		]
-
-		+ SVerticalBox::Slot()
-		.AutoHeight()
-		.Padding(0.0f, 4.0f, 0.0f, 2.0f)
-		[
-			SNew(STextBlock)
-			.Text(this, &SMemTagTreeView::QueryRule_GetTooltipText)
-			.ColorAndOpacity(FLinearColor(0.3f, 0.3f, 0.3f, 1.0f))
-			.AutoWrapText(true)
-		]
-
-		+ SVerticalBox::Slot()
-		.AutoHeight()
-		.HAlign(HAlign_Right)
-		[
-			SNew(SButton)
-			.Text(LOCTEXT("RunQueryBtnText", "Run Query"))
-			.ToolTipText(LOCTEXT("RunQueryBtnToolTipText", "Run Memory Query.\nThe resulting list of allocations will be available in a tree view."))
-			.OnClicked(this, &SMemTagTreeView::RunQuery)
-		]
-
-		+ SVerticalBox::Slot()
-		.AutoHeight()
-		.Padding(0.0f, 6.0f, 0.0f, 0.0f)
-		[
-			SNew(SSeparator)
-		]
-	];
-
-	if (ProfilerWindow.IsValid())
-	{
-		FMemorySharedState& SharedState = ProfilerWindow->GetSharedState();
-		QueryRuleComboBox->SetSelectedItem(SharedState.GetCurrentMemoryRule());
-	}
-
-	return Widget;
-}
-
-////////////////////////////////////////////////////////////////////////////////////////////////////
-
-TSharedRef<SWidget> SMemTagTreeView::QueryRule_OnGenerateWidget(TSharedPtr<Insights::FMemoryRuleSpec> InRule)
-{
-	const FText QueryRuleText = FText::Format(LOCTEXT("QueryRuleComboBox_TextFmt", "{0} ({1})"), InRule->GetVerboseName(), InRule->GetShortName());
-
-	TSharedRef<SHorizontalBox> Widget = SNew(SHorizontalBox);
-
-	Widget->AddSlot()
-		.AutoWidth()
-		.HAlign(HAlign_Center)
-		.VAlign(VAlign_Center)
-		.Padding(2.0f, 2.0f, 0.0f, 2.0f)
-		[
-			SNew(SImage)
-			.Visibility_Lambda([Widget]() { return Widget->GetParentWidget()->IsHovered() ? EVisibility::Visible : EVisibility::Hidden; })
-			.Image(FEditorStyle::GetBrush("Profiler.Tooltip.HintIcon10"))
-			.ToolTipText(InRule->GetDescription())
-		];
-
-	Widget->AddSlot()
-		.AutoWidth()
-		[
-			SNew(STextBlock)
-			.Text(QueryRuleText)
-			.Margin(2.0f)
-		];
-
-	return Widget;
-}
-
-////////////////////////////////////////////////////////////////////////////////////////////////////
-
-TSharedRef<SWidget> SMemTagTreeView::ConstructTimeMarkerWidget(uint32 TimeMarkerIndex)
-{
-	if (!ProfilerWindow.IsValid())
-	{
-		return SNew(SBox);
-	}
-
-	FMemorySharedState& SharedState = ProfilerWindow->GetSharedState();
-	const uint32 NumTimeMarkers = ProfilerWindow->GetNumCustomTimeMarkers();
-	if (TimeMarkerIndex >= NumTimeMarkers)
-	{
-		return SNew(SBox);
-	}
-
-	const TSharedRef<Insights::FTimeMarker>& TimeMarker = ProfilerWindow->GetCustomTimeMarker(TimeMarkerIndex);
-
-	TSharedRef<SWidget> Widget = SNew(SHorizontalBox)
-		.Visibility_Lambda([TimeMarker]()
-			{
-				return TimeMarker->IsVisible() ? EVisibility::Visible : EVisibility::Collapsed;
-			})
-
-		+ SHorizontalBox::Slot()
-		.AutoWidth()
-		.VAlign(VAlign_Center)
-		.Padding(20.0f, 1.0f, 4.0f, 1.0f)
-		[
-			SNew(STextBlock)
-			.Text(FText::FromString(TimeMarker->GetName()))
-			.ColorAndOpacity(FSlateColor(TimeMarker->GetColor()))
-		]
-
-		+ SHorizontalBox::Slot()
-		.AutoWidth()
-		.VAlign(VAlign_Center)
-		.Padding(0.0f, 1.0f, 0.0f, 1.0f)
-		[
-			SNew(SEditableTextBox)
-			.MinDesiredWidth(110.0f)
-			.Text_Lambda([TimeMarker]()
-				{
-					return FText::FromString(FString::Printf(TEXT("%.9f"), TimeMarker->GetTime()));
-				})
-			.OnTextCommitted_Lambda([TimeMarker](const FText& InText, ETextCommit::Type InCommitType)
-				{
-					const double Time = FCString::Atod(*InText.ToString());
-					TimeMarker->SetTime(Time);
-				})
-		]
-
-		+ SHorizontalBox::Slot()
-		.AutoWidth()
-		.VAlign(VAlign_Center)
-		.Padding(4.0f, 1.0f, 0.0f, 1.0f)
-		[
-			SNew(STextBlock)
-			.Text_Lambda([TimeMarker]()
-				{
-					return FText::FromString(TimeUtils::FormatTime(TimeMarker->GetTime(), 0.1));
-				})
-		];
-
-	return Widget;
-}
-
-////////////////////////////////////////////////////////////////////////////////////////////////////
-
 TSharedRef<SWidget> SMemTagTreeView::ConstructTagsFilteringWidgetArea()
 {
 	TSharedRef<SWidget> Widget = SNew(SHorizontalBox)
@@ -2291,127 +2064,6 @@ void SMemTagTreeView::SelectMemTagNode(Insights::FMemoryTagId MemTagId)
 	}
 }
 
-////////////////////////////////////////////////////////////////////////////////////////////////////
-// Query Rules
-////////////////////////////////////////////////////////////////////////////////////////////////////
-
-const TArray<TSharedPtr<Insights::FMemoryRuleSpec>>* SMemTagTreeView::GetAvailableQueryRules()
-{
-	if (ProfilerWindow.IsValid())
-	{
-		FMemorySharedState& SharedState = ProfilerWindow->GetSharedState();
-		return &SharedState.GetMemoryRules();
-	}
-	return nullptr;
-}
-
-////////////////////////////////////////////////////////////////////////////////////////////////////
-
-void SMemTagTreeView::QueryRule_OnSelectionChanged(TSharedPtr<Insights::FMemoryRuleSpec> InRule, ESelectInfo::Type SelectInfo)
-{
-	if (SelectInfo != ESelectInfo::Direct)
-	{
-		if (ProfilerWindow.IsValid() && InRule)
-		{
-			FMemorySharedState& SharedState = ProfilerWindow->GetSharedState();
-			SharedState.SetCurrentMemoryRule(InRule);
-			ProfilerWindow->UpdateTimingViewMarkers();
-		}
-	}
-}
-
-////////////////////////////////////////////////////////////////////////////////////////////////////
-
-FText SMemTagTreeView::QueryRule_GetSelectedText() const
-{
-	if (ProfilerWindow.IsValid())
-	{
-		FMemorySharedState& SharedState = ProfilerWindow->GetSharedState();
-		TSharedPtr<Insights::FMemoryRuleSpec> Rule = SharedState.GetCurrentMemoryRule();
-		if (Rule)
-		{
-			return FText::Format(LOCTEXT("QueryRuleComboBox_TextFmt", "{0} ({1})"), Rule->GetVerboseName(), Rule->GetShortName());
-		}
-	}
-	return FText::GetEmpty();
-}
-
-////////////////////////////////////////////////////////////////////////////////////////////////////
-
-FText SMemTagTreeView::QueryRule_GetTooltipText() const
-{
-	if (ProfilerWindow.IsValid())
-	{
-		FMemorySharedState& SharedState = ProfilerWindow->GetSharedState();
-		if (SharedState.GetCurrentMemoryRule())
-		{
-			return SharedState.GetCurrentMemoryRule()->GetDescription();
-		}
-	}
-	return FText::GetEmpty();
-}
-
-////////////////////////////////////////////////////////////////////////////////////////////////////
-
-FReply SMemTagTreeView::RunQuery()
-{
-	if (!ProfilerWindow.IsValid())
-	{
-		UE_LOG(MemoryProfiler, Error, TEXT("[MemQuery] Invalid Profiler Window!"));
-		return FReply::Handled();
-	}
-
-	FMemorySharedState& SharedState = ProfilerWindow->GetSharedState();
-	TSharedPtr<Insights::FMemoryRuleSpec> Rule = SharedState.GetCurrentMemoryRule();
-	if (!Rule)
-	{
-		UE_LOG(MemoryProfiler, Error, TEXT("[MemQuery] Invalid Rule!"));
-		return FReply::Handled();
-	}
-
-	const uint32 NumTimeMarkers = ProfilerWindow->GetNumCustomTimeMarkers();
-	const uint32 RuleNumTimeMarkers = Rule->GetNumTimeMarkers();
-	if (RuleNumTimeMarkers > NumTimeMarkers)
-	{
-		UE_LOG(MemoryProfiler, Error, TEXT("[MemQuery] Only %d time markers available. Current rule (%s) requires %u time markers!"),
-			NumTimeMarkers, *Rule->GetShortName().ToString(), RuleNumTimeMarkers);
-		return FReply::Handled();
-	}
-
-#if !NO_LOGGING
-	TStringBuilder<1024> Builder;
-	if (RuleNumTimeMarkers > 0)
-	{
-		Builder.Append(TEXT(" ("));
-		for (uint32 TimeMarkerIndex = 0; TimeMarkerIndex < RuleNumTimeMarkers; ++TimeMarkerIndex)
-		{
-			if (TimeMarkerIndex != 0)
-			{
-				Builder.Append(TEXT(", "));
-			}
-			Builder.Append(TEXT('A') + TimeMarkerIndex);
-			const TSharedRef<Insights::FTimeMarker>& TimeMarker = ProfilerWindow->GetCustomTimeMarker(TimeMarkerIndex);
-			Builder.Appendf(TEXT("=%.9f"), TimeMarker->GetTime());
-		}
-		Builder.Append(TEXT(")"));
-	}
-	UE_LOG(MemoryProfiler, Log, TEXT("[MemQuery] Run Query %s%s..."), *Rule->GetShortName().ToString(), Builder.ToString());
-#endif
-
-	TSharedPtr<Insights::SMemAllocTableTreeView> MemAllocTableTreeView = ProfilerWindow->ShowMemAllocTableTreeViewTab();
-	if (MemAllocTableTreeView)
-	{
-		double TimeMarkers[4];
-		TimeMarkers[0] = (RuleNumTimeMarkers > 0) ? ProfilerWindow->GetCustomTimeMarker(0)->GetTime() : 0.0;
-		TimeMarkers[1] = (RuleNumTimeMarkers > 1) ? ProfilerWindow->GetCustomTimeMarker(1)->GetTime() : 0.0;
-		TimeMarkers[2] = (RuleNumTimeMarkers > 2) ? ProfilerWindow->GetCustomTimeMarker(2)->GetTime() : 0.0;
-		TimeMarkers[3] = (RuleNumTimeMarkers > 3) ? ProfilerWindow->GetCustomTimeMarker(3)->GetTime() : 0.0;
-		MemAllocTableTreeView->SetQueryParams(Rule, TimeMarkers[0], TimeMarkers[1], TimeMarkers[2], TimeMarkers[3]);
-	}
-
-	return FReply::Handled();
-}
-
 ////////////////////////////////////////////////////////////////////////////////////////////////////
 // Trackers
 ////////////////////////////////////////////////////////////////////////////////////////////////////
diff --git a/Engine/Source/Developer/TraceInsights/Private/Insights/MemoryProfiler/Widgets/SMemTagTreeView.h b/Engine/Source/Developer/TraceInsights/Private/Insights/MemoryProfiler/Widgets/SMemTagTreeView.h
--- a/Engine/Source/Developer/TraceInsights/Private/Insights/MemoryProfiler/Widgets/SMemTagTreeView.h
+++ b/Engine/Source/Developer/TraceInsights/Private/Insights/MemoryProfiler/Widgets/SMemTagTreeView.h
@@ -32,7 +32,6 @@ namespace TraceServices
 
 namespace Insights
 {
-	class FMemoryRuleSpec;
 	class FTable;
 	class FTableColumn;
 	class ITableCellValueSorter;
@@ -250,15 +249,6 @@ private:
 
 	////////////////////////////////////////////////////////////////////////////////////////////////////
 
-	const TArray<TSharedPtr<Insights::FMemoryRuleSpec>>* GetAvailableQueryRules();
-	void QueryRule_OnSelectionChanged(TSharedPtr<Insights::FMemoryRuleSpec> InRule, ESelectInfo::Type SelectInfo);
-	TSharedRef<SWidget> QueryRule_OnGenerateWidget(TSharedPtr<Insights::FMemoryRuleSpec> InRule);
-	FText QueryRule_GetSelectedText() const;
-	FText QueryRule_GetTooltipText() const;
-	FReply RunQuery();
-
-	////////////////////////////////////////////////////////////////////////////////////////////////////
-
 	const TArray<TSharedPtr<Insights::FMemoryTracker>>* GetAvailableTrackers();
 	void Tracker_OnSelectionChanged(TSharedPtr<Insights::FMemoryTracker> InTracker, ESelectInfo::Type SelectInfo);
 	TSharedRef<SWidget> Tracker_OnGenerateWidget(TSharedPtr<Insights::FMemoryTracker> InTracker);
@@ -382,11 +372,6 @@ private:
 	/** If true, the expanded nodes have been saved before applying a text filter. */
 	bool bExpansionSaved;
 
-	//////////////////////////////////////////////////
-	// Queries
-
-	TSharedPtr<SComboBox<TSharedPtr<Insights::FMemoryRuleSpec>>> QueryRuleComboBox;
-
 	//////////////////////////////////////////////////
 	// Search box and filters
 
diff --git a/Engine/Source/Developer/TraceInsights/Private/Insights/MemoryProfiler/Widgets/SMemoryProfilerToolbar.cpp b/Engine/Source/Developer/TraceInsights/Private/Insights/MemoryProfiler/Widgets/SMemoryProfilerToolbar.cpp
--- a/Engine/Source/Developer/TraceInsights/Private/Insights/MemoryProfiler/Widgets/SMemoryProfilerToolbar.cpp
+++ b/Engine/Source/Developer/TraceInsights/Private/Insights/MemoryProfiler/Widgets/SMemoryProfilerToolbar.cpp
@@ -46,6 +46,7 @@ void SMemoryProfilerToolbar::Construct(const FArguments& InArgs)
 			ToolbarBuilder.BeginSection("View");
 			{
 				ToolbarBuilder.AddToolBarButton(FMemoryProfilerCommands::Get().ToggleTimingViewVisibility);
+				ToolbarBuilder.AddToolBarButton(FMemoryProfilerCommands::Get().ToggleMemInvestigationViewVisibility);
 				ToolbarBuilder.AddToolBarButton(FMemoryProfilerCommands::Get().ToggleMemTagTreeViewVisibility);
 			}
 			ToolbarBuilder.EndSection();
diff --git a/Engine/Source/Developer/TraceInsights/Private/Insights/MemoryProfiler/Widgets/SMemoryProfilerWindow.cpp b/Engine/Source/Developer/TraceInsights/Private/Insights/MemoryProfiler/Widgets/SMemoryProfilerWindow.cpp
--- a/Engine/Source/Developer/TraceInsights/Private/Insights/MemoryProfiler/Widgets/SMemoryProfilerWindow.cpp
+++ b/Engine/Source/Developer/TraceInsights/Private/Insights/MemoryProfiler/Widgets/SMemoryProfilerWindow.cpp
@@ -31,6 +31,7 @@
 #include "Insights/MemoryProfiler/MemoryProfilerManager.h"
 #include "Insights/MemoryProfiler/ViewModels/MemorySharedState.h"
 #include "Insights/MemoryProfiler/Widgets/SMemAllocTableTreeView.h"
+#include "Insights/MemoryProfiler/Widgets/SMemInvestigationView.h"
 #include "Insights/MemoryProfiler/Widgets/SMemoryProfilerToolbar.h"
 #include "Insights/MemoryProfiler/Widgets/SMemTagTreeView.h"
 #include "Insights/TraceInsightsModule.h"
@@ -47,6 +48,7 @@
 
 const FName FMemoryProfilerTabs::ToolbarID(TEXT("Toolbar"));
 const FName FMemoryProfilerTabs::TimingViewID(TEXT("TimingView"));
+const FName FMemoryProfilerTabs::MemInvestigationViewID(TEXT("MemInvestigation"));
 const FName FMemoryProfilerTabs::MemTagTreeViewID(TEXT("LowLevelMemTags"));
 const FName FMemoryProfilerTabs::MemAllocTableTreeViewID(TEXT("MemAllocTableTreeView"));
 
@@ -55,6 +57,7 @@ const FName FMemoryProfilerTabs::MemAllocTableTreeViewID(TEXT("MemAllocTableTree
 SMemoryProfilerWindow::SMemoryProfilerWindow()
 	: TimingView()
 	, SharedState(MakeShared<FMemorySharedState>())
+	, MemInvestigationView()
 	, MemTagTreeView()
 	, TabManager()
 	, ActiveTimerHandle()
@@ -72,6 +75,12 @@ SMemoryProfilerWindow::~SMemoryProfilerWindow()
 		check(MemTagTreeView == nullptr);
 	}
 
+	if (MemInvestigationView)
+	{
+		HideTab(FMemoryProfilerTabs::MemInvestigationViewID);
+		check(MemInvestigationView == nullptr);
+	}
+
 	if (TimingView)
 	{
 		HideTab(FMemoryProfilerTabs::TimingViewID);
@@ -98,11 +107,17 @@ void SMemoryProfilerWindow::Reset()
 		ResetTimingViewMarkers();
 	}
 
+	if (MemInvestigationView)
+	{
+		MemInvestigationView->Reset();
+	}
+
 	if (MemTagTreeView)
 	{
 		MemTagTreeView->Reset();
 	}
 
+	UpdateMemInvestigationView();
 	UpdateTableTreeViews();
 }
 
@@ -170,6 +185,16 @@ void SMemoryProfilerWindow::UpdateTimingViewMarkers()
 
 ////////////////////////////////////////////////////////////////////////////////////////////////////
 
+void SMemoryProfilerWindow::UpdateMemInvestigationView()
+{
+	if (MemInvestigationView)
+	{
+		//MemInvestigationView->Update();
+	}
+}
+
+////////////////////////////////////////////////////////////////////////////////////////////////////
+
 void SMemoryProfilerWindow::UpdateTableTreeViews()
 {
 	UpdateMemTagTreeView();
@@ -280,6 +305,34 @@ void SMemoryProfilerWindow::OnTimingViewTabClosed(TSharedRef<SDockTab> TabBeingC
 
 ////////////////////////////////////////////////////////////////////////////////////////////////////
 
+TSharedRef<SDockTab> SMemoryProfilerWindow::SpawnTab_MemInvestigationView(const FSpawnTabArgs& Args)
+{
+	FMemoryProfilerManager::Get()->SetMemInvestigationViewVisible(true);
+
+	const TSharedRef<SDockTab> DockTab = SNew(SDockTab)
+		.ShouldAutosize(false)
+		.TabRole(ETabRole::PanelTab)
+		[
+			SAssignNew(MemInvestigationView, SMemInvestigationView, SharedThis(this))
+		];
+
+	UpdateMemInvestigationView();
+
+	DockTab->SetOnTabClosed(SDockTab::FOnTabClosedCallback::CreateRaw(this, &SMemoryProfilerWindow::OnMemInvestigationViewTabClosed));
+
+	return DockTab;
+}
+
+////////////////////////////////////////////////////////////////////////////////////////////////////
+
+void SMemoryProfilerWindow::OnMemInvestigationViewTabClosed(TSharedRef<SDockTab> TabBeingClosed)
+{
+	FMemoryProfilerManager::Get()->SetMemInvestigationViewVisible(false);
+	MemInvestigationView = nullptr;
+}
+
+////////////////////////////////////////////////////////////////////////////////////////////////////
+
 TSharedRef<SDockTab> SMemoryProfilerWindow::SpawnTab_MemTagTreeView(const FSpawnTabArgs& Args)
 {
 	FMemoryProfilerManager::Get()->SetMemTagTreeViewVisible(true);
@@ -392,6 +445,11 @@ void SMemoryProfilerWindow::Construct(const FArguments& InArgs, const TSharedRef
 		.SetIcon(FSlateIcon(FInsightsStyle::GetStyleSetName(), "TimingView.Icon.Small"))
 		.SetGroup(AppMenuGroup);
 
+	TabManager->RegisterTabSpawner(FMemoryProfilerTabs::MemInvestigationViewID, FOnSpawnTab::CreateRaw(this, &SMemoryProfilerWindow::SpawnTab_MemInvestigationView))
+		.SetDisplayName(LOCTEXT("MemoryProfiler.MemInvestigationViewTabTitle", "Investigation"))
+		.SetIcon(FSlateIcon(FInsightsStyle::GetStyleSetName(), "MemInvestigationView.Icon.Small"))
+		.SetGroup(AppMenuGroup);
+
 	TabManager->RegisterTabSpawner(FMemoryProfilerTabs::MemTagTreeViewID, FOnSpawnTab::CreateRaw(this, &SMemoryProfilerWindow::SpawnTab_MemTagTreeView))
 		.SetDisplayName(LOCTEXT("MemoryProfiler.MemTagTreeViewTabTitle", "LLM Tags"))
 		.SetIcon(FSlateIcon(FInsightsStyle::GetStyleSetName(), "MemTagTreeView.Icon.Small"))
@@ -404,7 +462,7 @@ void SMemoryProfilerWindow::Construct(const FArguments& InArgs, const TSharedRef
 		FText MemAllocTableTreeViewTabDisplayName = FText::Format(LOCTEXT("MemoryProfiler.MemAllocTableTreeViewTabTitle", "Allocs Table {0}"), FText::AsNumber(TabIndex));
 		TabManager->RegisterTabSpawner(MemAllocTableTreeViewTabId, FOnSpawnTab::CreateRaw(this, &SMemoryProfilerWindow::SpawnTab_MemAllocTableTreeView, TabIndex))
 			.SetDisplayName(MemAllocTableTreeViewTabDisplayName)
-			.SetIcon(FSlateIcon(FInsightsStyle::GetStyleSetName(), "MemTagTreeView.Icon.Small")) // TODO
+			.SetIcon(FSlateIcon(FInsightsStyle::GetStyleSetName(), "MemAllocTableTreeView.Icon.Small"))
 			.SetGroup(AppMenuGroup);
 	}
 
@@ -439,8 +497,9 @@ void SMemoryProfilerWindow::Construct(const FArguments& InArgs, const TSharedRef
 				(
 					FTabManager::NewStack()
 					->SetSizeCoefficient(0.23f)
+					->AddTab(FMemoryProfilerTabs::MemInvestigationViewID, ETabState::OpenedTab)
 					->AddTab(FMemoryProfilerTabs::MemTagTreeViewID, ETabState::OpenedTab)
-					->SetForegroundTab(FMemoryProfilerTabs::MemTagTreeViewID)
+					->SetForegroundTab(FMemoryProfilerTabs::MemInvestigationViewID)
 				)
 			)
 		);
diff --git a/Engine/Source/Developer/TraceInsights/Private/Insights/MemoryProfiler/Widgets/SMemoryProfilerWindow.h b/Engine/Source/Developer/TraceInsights/Private/Insights/MemoryProfiler/Widgets/SMemoryProfilerWindow.h
--- a/Engine/Source/Developer/TraceInsights/Private/Insights/MemoryProfiler/Widgets/SMemoryProfilerWindow.h
+++ b/Engine/Source/Developer/TraceInsights/Private/Insights/MemoryProfiler/Widgets/SMemoryProfilerWindow.h
@@ -21,8 +21,9 @@ class FActiveTimerHandle;
 class FMenuBuilder;
 
 class FMemorySharedState;
-class STimingView;
+class SMemInvestigationView;
 class SMemTagTreeView;
+class STimingView;
 
 namespace Insights
 {
@@ -38,6 +39,7 @@ struct FMemoryProfilerTabs
 	// Tab identifiers
 	static const FName ToolbarID;
 	static const FName TimingViewID;
+	static const FName MemInvestigationViewID;
 	static const FName MemTagTreeViewID;
 	static const FName MemAllocTableTreeViewID; // base id
 };
@@ -61,6 +63,7 @@ public:
 	void ResetTimingViewMarkers();
 	void UpdateTimingViewMarkers();
 	void UpdateTableTreeViews();
+	void UpdateMemInvestigationView();
 	void UpdateMemTagTreeView();
 	void UpdateMemAllocTableTreeView(TSharedPtr<Insights::SMemAllocTableTreeView> MemAllocTableTreeView);
 
@@ -76,6 +79,7 @@ public:
 	TSharedPtr<FTabManager> GetTabManager() const { return TabManager; }
 
 	TSharedPtr<STimingView> GetTimingView() const { return TimingView; }
+	TSharedPtr<SMemInvestigationView> GetMemInvestigationView() const { return MemInvestigationView; }
 	TSharedPtr<SMemTagTreeView> GetMemTagTreeView() const { return MemTagTreeView; }
 
 	uint32 GetNumCustomTimeMarkers() const { return (uint32)CustomTimeMarkers.Num(); }
@@ -92,6 +96,9 @@ private:
 	TSharedRef<SDockTab> SpawnTab_TimingView(const FSpawnTabArgs& Args);
 	void OnTimingViewTabClosed(TSharedRef<SDockTab> TabBeingClosed);
 
+	TSharedRef<SDockTab> SpawnTab_MemInvestigationView(const FSpawnTabArgs& Args);
+	void OnMemInvestigationViewTabClosed(TSharedRef<SDockTab> TabBeingClosed);
+
 	TSharedRef<SDockTab> SpawnTab_MemTagTreeView(const FSpawnTabArgs& Args);
 	void OnMemTagTreeViewTabClosed(TSharedRef<SDockTab> TabBeingClosed);
 
@@ -175,6 +182,9 @@ private:
 	TArray<TSharedRef<Insights::FTimeMarker>> CustomTimeMarkers;
 	TSharedPtr<FMemorySharedState> SharedState;
 
+	/** The Memory Investigation (Allocation Queries) view widget */
+	TSharedPtr<SMemInvestigationView> MemInvestigationView;
+
 	/** The "LLM Tags" tree view widget */
 	TSharedPtr<SMemTagTreeView> MemTagTreeView;
 
