[Insights] Mem Alloc: Removed address prefix from resolved callstack frames (when grouping by callstack).

diff --git a/Engine/Source/Developer/TraceInsights/Private/Insights/MemoryProfiler/ViewModels/MemAllocGroupingByCallstack.cpp b/Engine/Source/Developer/TraceInsights/Private/Insights/MemoryProfiler/ViewModels/MemAllocGroupingByCallstack.cpp
--- a/Engine/Source/Developer/TraceInsights/Private/Insights/MemoryProfiler/ViewModels/MemAllocGroupingByCallstack.cpp
+++ b/Engine/Source/Developer/TraceInsights/Private/Insights/MemoryProfiler/ViewModels/MemAllocGroupingByCallstack.cpp
@@ -110,8 +110,8 @@ void FMemAllocGroupingByCallstack::GroupNodes(const TArray<FTableTreeNodePtr>& N
 							GroupName = FName(*FString::Printf(TEXT("0x%X"), Frame->Addr));
 							break;
 						case TraceServices::QueryResult::QR_OK:
-							//GroupName = FName(Frame->Symbol->Name);
-							GroupName = FName(*FString::Printf(TEXT("0x%X %s"), Frame->Addr, Frame->Symbol->Name));
+							GroupName = FName(Frame->Symbol->Name);
+							//GroupName = FName(*FString::Printf(TEXT("0x%X %s"), Frame->Addr, Frame->Symbol->Name));
 							break;
 						}
 
