[Insights] Increased version to 0.15.

diff --git a/Engine/Source/Developer/TraceInsights/Public/Insights/Version.h b/Engine/Source/Developer/TraceInsights/Public/Insights/Version.h
--- a/Engine/Source/Developer/TraceInsights/Public/Insights/Version.h
+++ b/Engine/Source/Developer/TraceInsights/Public/Insights/Version.h
@@ -5,10 +5,10 @@
 // See https://semver.org/
 
 #define UNREAL_INSIGHTS_VERSION_MAJOR 0
-#define UNREAL_INSIGHTS_VERSION_MINOR 14
+#define UNREAL_INSIGHTS_VERSION_MINOR 15
 #define UNREAL_INSIGHTS_VERSION_PATCH 0
 
-#define UNREAL_INSIGHTS_VERSION_STRING "0.14"
+#define UNREAL_INSIGHTS_VERSION_STRING "0.15"
 
 /** Extra identifier string (like: "alpha", "beta.1", "beta.2", "test", "featureX", etc.). */
 //#define UNREAL_INSIGHTS_VERSION_ID_STRING "beta.1"
