[Insights] Target window option in Memory Investigation View

diff --git a/Engine/Source/Developer/TraceInsights/Private/Insights/MemoryProfiler/ViewModels/MemorySharedState.cpp b/Engine/Source/Developer/TraceInsights/Private/Insights/MemoryProfiler/ViewModels/MemorySharedState.cpp
--- a/Engine/Source/Developer/TraceInsights/Private/Insights/MemoryProfiler/ViewModels/MemorySharedState.cpp
+++ b/Engine/Source/Developer/TraceInsights/Private/Insights/MemoryProfiler/ViewModels/MemorySharedState.cpp
@@ -27,6 +27,8 @@
 
 #define LOCTEXT_NAMESPACE "MemorySharedState"
 
+const FName Insights::FQueryTargetWindowSpec::NewWindow = TEXT("New Window");
+
 ////////////////////////////////////////////////////////////////////////////////////////////////////
 // FMemorySharedState
 ////////////////////////////////////////////////////////////////////////////////////////////////////
@@ -44,6 +46,9 @@ FMemorySharedState::FMemorySharedState()
 	, CurrentMemoryRule(nullptr)
 {
 	InitMemoryRules();
+
+	CurrentQueryTarget = MakeShared<Insights::FQueryTargetWindowSpec>(Insights::FQueryTargetWindowSpec::NewWindow, LOCTEXT("NewWindow", "New Window"));
+	QueryTargetSpecs.Add(CurrentQueryTarget);
 }
 
 ////////////////////////////////////////////////////////////////////////////////////////////////////
@@ -1035,4 +1040,19 @@ void FMemorySharedState::OnMemoryRuleChanged()
 
 ////////////////////////////////////////////////////////////////////////////////////////////////////
 
+void FMemorySharedState::AddQueryTarget(TSharedPtr<Insights::FQueryTargetWindowSpec> InPtr)
+{
+	QueryTargetSpecs.Add(InPtr);
+}
+
+////////////////////////////////////////////////////////////////////////////////////////////////////
+
+
+void FMemorySharedState::RemoveQueryTarget(TSharedPtr<Insights::FQueryTargetWindowSpec> InPtr)
+{
+	QueryTargetSpecs.Remove(InPtr);
+}
+
+////////////////////////////////////////////////////////////////////////////////////////////////////
+
 #undef LOCTEXT_NAMESPACE
diff --git a/Engine/Source/Developer/TraceInsights/Private/Insights/MemoryProfiler/ViewModels/MemorySharedState.h b/Engine/Source/Developer/TraceInsights/Private/Insights/MemoryProfiler/ViewModels/MemorySharedState.h
--- a/Engine/Source/Developer/TraceInsights/Private/Insights/MemoryProfiler/ViewModels/MemorySharedState.h
+++ b/Engine/Source/Developer/TraceInsights/Private/Insights/MemoryProfiler/ViewModels/MemorySharedState.h
@@ -51,6 +51,25 @@ private:
 	FText Description;     // ex.: "Allocations allocated and freed between time A and time B (A <= a <= f <= B)."
 };
 
+class FQueryTargetWindowSpec
+{
+public:
+	FQueryTargetWindowSpec(const FName& InName, const FText& InText)
+		: Text(InText)
+		, Name(InName)
+	{}
+
+	FText GetText() const { return Text; }
+	FName GetName() const { return Name; }
+
+public:
+	static const FName NewWindow;
+
+private:
+	FText Text;
+	FName Name;
+};
+
 } // namespace Insights
 
 ////////////////////////////////////////////////////////////////////////////////////////////////////
@@ -111,6 +130,13 @@ public:
 	TSharedPtr<Insights::FMemoryRuleSpec> GetCurrentMemoryRule() const { return CurrentMemoryRule; }
 	void SetCurrentMemoryRule(TSharedPtr<Insights::FMemoryRuleSpec> InRule) { CurrentMemoryRule = InRule; OnMemoryRuleChanged(); }
 
+	const TArray<TSharedPtr<Insights::FQueryTargetWindowSpec>>& GetQueryTargets() const { return QueryTargetSpecs; }
+
+	TSharedPtr<Insights::FQueryTargetWindowSpec> GetCurrentQueryTarget() const { return CurrentQueryTarget; }
+	void SetCurrentQueryTarget(TSharedPtr<Insights::FQueryTargetWindowSpec> InTarget) { CurrentQueryTarget = InTarget; }
+	void AddQueryTarget(TSharedPtr<Insights::FQueryTargetWindowSpec> InPtr);
+	void RemoveQueryTarget(TSharedPtr<Insights::FQueryTargetWindowSpec> InPtr);
+
 private:
 	void SyncTrackers();
 	void OnTrackerChanged();
@@ -142,6 +168,9 @@ private:
 
 	TArray<TSharedPtr<Insights::FMemoryRuleSpec>> MemoryRules;
 	TSharedPtr<Insights::FMemoryRuleSpec> CurrentMemoryRule;
+
+	TSharedPtr<Insights::FQueryTargetWindowSpec> CurrentQueryTarget;
+	TArray<TSharedPtr<Insights::FQueryTargetWindowSpec>> QueryTargetSpecs;
 };
 
 ////////////////////////////////////////////////////////////////////////////////////////////////////
diff --git a/Engine/Source/Developer/TraceInsights/Private/Insights/MemoryProfiler/Widgets/SMemInvestigationView.cpp b/Engine/Source/Developer/TraceInsights/Private/Insights/MemoryProfiler/Widgets/SMemInvestigationView.cpp
--- a/Engine/Source/Developer/TraceInsights/Private/Insights/MemoryProfiler/Widgets/SMemInvestigationView.cpp
+++ b/Engine/Source/Developer/TraceInsights/Private/Insights/MemoryProfiler/Widgets/SMemInvestigationView.cpp
@@ -138,6 +138,36 @@ TSharedRef<SWidget> SMemInvestigationView::ConstructInvestigationWidgetArea()
 			.AutoWrapText(true)
 		]
 
+		+ SVerticalBox::Slot()
+		.AutoHeight()
+		[
+			SNew(SHorizontalBox)
+
+			+ SHorizontalBox::Slot()
+			.AutoWidth()
+			.VAlign(VAlign_Center)
+			.Padding(0.0f, 0.0f, 4.0f, 0.0f)
+			[
+				SNew(STextBlock)
+				.Text(LOCTEXT("QueryTargetWindow", "Target Window"))
+			]
+
+			+ SHorizontalBox::Slot()
+			.FillWidth(1.0f)
+			.VAlign(VAlign_Center)
+			[
+				SAssignNew(QueryTargetComboBox, SComboBox<TSharedPtr<Insights::FQueryTargetWindowSpec>>)
+				.ToolTipText(LOCTEXT("QueryTargetWindowTooltip", "Select an existing or new window where the query results will be displayed"))
+				.OptionsSource(GetAvailableQueryTargets())
+				.OnSelectionChanged(this, &SMemInvestigationView::QueryTarget_OnSelectionChanged)
+				.OnGenerateWidget(this, &SMemInvestigationView::QueryTarget_OnGenerateWidget)
+				[
+					SNew(STextBlock)
+					.Text(this, &SMemInvestigationView::QueryTarget_GetSelectedText)
+				]
+			]
+		]
+
 		+ SVerticalBox::Slot()
 		.AutoHeight()
 		.HAlign(HAlign_Right)
@@ -345,6 +375,80 @@ FText SMemInvestigationView::QueryRule_GetTooltipText() const
 
 ////////////////////////////////////////////////////////////////////////////////////////////////////
 
+const TArray<TSharedPtr<Insights::FQueryTargetWindowSpec>>* SMemInvestigationView::GetAvailableQueryTargets()
+{
+	if (ProfilerWindow.IsValid())
+	{
+		FMemorySharedState& SharedState = ProfilerWindow->GetSharedState();
+		return &SharedState.GetQueryTargets();
+	}
+	return nullptr;
+}
+
+////////////////////////////////////////////////////////////////////////////////////////////////////
+
+void SMemInvestigationView::QueryTarget_OnSelectionChanged(TSharedPtr<Insights::FQueryTargetWindowSpec> InTarget, ESelectInfo::Type SelectInfo)
+{
+	if (SelectInfo != ESelectInfo::Type::Direct)
+	{
+		if (ProfilerWindow.IsValid() && InTarget)
+		{
+			FMemorySharedState& SharedState = ProfilerWindow->GetSharedState();
+			SharedState.SetCurrentQueryTarget(InTarget);
+		}
+	}
+	else
+	{
+		QueryTargetComboBox->SetSelectedItem(InTarget);
+	}
+}
+
+////////////////////////////////////////////////////////////////////////////////////////////////////
+
+TSharedRef<SWidget> SMemInvestigationView::QueryTarget_OnGenerateWidget(TSharedPtr<Insights::FQueryTargetWindowSpec> InTarget)
+{
+	TSharedRef<SHorizontalBox> Widget = SNew(SHorizontalBox);
+
+	Widget->AddSlot()
+		.AutoWidth()
+		.HAlign(HAlign_Center)
+		.VAlign(VAlign_Center)
+		.Padding(2.0f, 2.0f, 0.0f, 2.0f)
+		[
+			SNew(SImage)
+			.Image(FEditorStyle::GetBrush("Profiler.Tooltip.HintIcon10"))
+			.Visibility_Lambda([Widget]() { return Widget->GetParentWidget()->IsHovered() ? EVisibility::Visible : EVisibility::Hidden; })
+		];
+
+	Widget->AddSlot()
+		.AutoWidth()
+		[
+			SNew(STextBlock)
+			.Text(InTarget->GetText())
+			.Margin(2.0f)
+		];
+
+	return Widget;
+}
+
+////////////////////////////////////////////////////////////////////////////////////////////////////
+
+FText SMemInvestigationView::QueryTarget_GetSelectedText() const
+{
+	if (ProfilerWindow.IsValid())
+	{
+		FMemorySharedState& SharedState = ProfilerWindow->GetSharedState();
+		TSharedPtr<Insights::FQueryTargetWindowSpec> Target = SharedState.GetCurrentQueryTarget();
+		if (Target)
+		{
+			return Target->GetText();
+		}
+	}
+	return FText::GetEmpty();
+}
+
+////////////////////////////////////////////////////////////////////////////////////////////////////
+
 FReply SMemInvestigationView::RunQuery()
 {
 	if (!ProfilerWindow.IsValid())
diff --git a/Engine/Source/Developer/TraceInsights/Private/Insights/MemoryProfiler/Widgets/SMemInvestigationView.h b/Engine/Source/Developer/TraceInsights/Private/Insights/MemoryProfiler/Widgets/SMemInvestigationView.h
--- a/Engine/Source/Developer/TraceInsights/Private/Insights/MemoryProfiler/Widgets/SMemInvestigationView.h
+++ b/Engine/Source/Developer/TraceInsights/Private/Insights/MemoryProfiler/Widgets/SMemInvestigationView.h
@@ -22,6 +22,7 @@ namespace TraceServices
 namespace Insights
 {
 	class FMemoryRuleSpec;
+	class FQueryTargetWindowSpec;
 }
 
 ////////////////////////////////////////////////////////////////////////////////////////////////////
@@ -48,6 +49,8 @@ public:
 
 	void Reset();
 
+	void QueryTarget_OnSelectionChanged(TSharedPtr<Insights::FQueryTargetWindowSpec> InRule, ESelectInfo::Type SelectInfo);
+
 private:
 	TSharedRef<SWidget> ConstructInvestigationWidgetArea();
 	TSharedRef<SWidget> ConstructTimeMarkerWidget(uint32 TimeMarkerIndex);
@@ -55,6 +58,9 @@ private:
 	/** Called when the analysis session has changed. */
 	void InsightsManager_OnSessionChanged();
 
+	const TArray<TSharedPtr<Insights::FQueryTargetWindowSpec>>* GetAvailableQueryTargets();
+	TSharedRef<SWidget> QueryTarget_OnGenerateWidget(TSharedPtr<Insights::FQueryTargetWindowSpec> InRule);
+	FText QueryTarget_GetSelectedText() const;
 	const TArray<TSharedPtr<Insights::FMemoryRuleSpec>>* GetAvailableQueryRules();
 	void QueryRule_OnSelectionChanged(TSharedPtr<Insights::FMemoryRuleSpec> InRule, ESelectInfo::Type SelectInfo);
 	TSharedRef<SWidget> QueryRule_OnGenerateWidget(TSharedPtr<Insights::FMemoryRuleSpec> InRule);
@@ -78,6 +84,8 @@ private:
 	TSharedPtr<const TraceServices::IAnalysisSession>/*Weak*/ Session;
 
 	TSharedPtr<SComboBox<TSharedPtr<Insights::FMemoryRuleSpec>>> QueryRuleComboBox;
+
+	TSharedPtr<SComboBox<TSharedPtr<Insights::FQueryTargetWindowSpec>>> QueryTargetComboBox;
 };
 
 ////////////////////////////////////////////////////////////////////////////////////////////////////
diff --git a/Engine/Source/Developer/TraceInsights/Private/Insights/MemoryProfiler/Widgets/SMemoryProfilerWindow.cpp b/Engine/Source/Developer/TraceInsights/Private/Insights/MemoryProfiler/Widgets/SMemoryProfilerWindow.cpp
--- a/Engine/Source/Developer/TraceInsights/Private/Insights/MemoryProfiler/Widgets/SMemoryProfilerWindow.cpp
+++ b/Engine/Source/Developer/TraceInsights/Private/Insights/MemoryProfiler/Widgets/SMemoryProfilerWindow.cpp
@@ -393,7 +393,30 @@ void SMemoryProfilerWindow::OnMemAllocTableTreeViewTabClosed(TSharedRef<SDockTab
 {
 	TSharedRef<Insights::SMemAllocTableTreeView> MemAllocTableTreeView = StaticCastSharedRef<Insights::SMemAllocTableTreeView>(TabBeingClosed->GetContent());
 
-	//FMemoryProfilerManager::Get()->SetMemAllocTableTreeViewVisible(MemAllocTableTreeView->GetTabIndex(), false);
+	FName ClosingTabId = FMemoryProfilerTabs::MemAllocTableTreeViewID;
+	ClosingTabId.SetNumber(MemAllocTableTreeView->GetTabIndex());
+
+	TSharedPtr<Insights::FQueryTargetWindowSpec> TargetToDelete;
+	const TArray<TSharedPtr<Insights::FQueryTargetWindowSpec>>& Targets = this->GetSharedState().GetQueryTargets();
+	for (int32 Index = 0; Index < Targets.Num(); ++Index)
+	{
+		if (Targets[Index]->GetName() == ClosingTabId)
+		{
+			TargetToDelete = Targets[Index];
+			break;
+		}
+	}
+
+	if (TargetToDelete.IsValid())
+	{
+		this->GetSharedState().RemoveQueryTarget(TargetToDelete);
+	}
+
+	TSharedPtr<Insights::FQueryTargetWindowSpec> NewSelection = this->GetSharedState().GetQueryTargets()[0];
+	this->GetSharedState().SetCurrentQueryTarget(NewSelection);
+	MemInvestigationView->QueryTarget_OnSelectionChanged(NewSelection, ESelectInfo::Type::Direct);
+
+	TabManager->UnregisterTabSpawner(ClosingTabId);
 
 	MemAllocTableTreeView->OnClose();
 	MemAllocTableTreeViews.Remove(MemAllocTableTreeView);
@@ -403,17 +426,37 @@ void SMemoryProfilerWindow::OnMemAllocTableTreeViewTabClosed(TSharedRef<SDockTab
 
 TSharedPtr<Insights::SMemAllocTableTreeView> SMemoryProfilerWindow::ShowMemAllocTableTreeViewTab()
 {
-	LastMemAllocTableTreeViewIndex = (LastMemAllocTableTreeViewIndex + 1) % MaxMemAllocTableTreeViews;
-	FName TabId = FMemoryProfilerTabs::MemAllocTableTreeViewID;
-	TabId.SetNumber(LastMemAllocTableTreeViewIndex);
+	if (SharedState->GetCurrentQueryTarget()->GetName() == Insights::FQueryTargetWindowSpec::NewWindow)
+	{
+		LastMemAllocTableTreeViewIndex = ++LastMemAllocTableTreeViewIndex;
+		FName TabId = FMemoryProfilerTabs::MemAllocTableTreeViewID;
+		TabId.SetNumber(LastMemAllocTableTreeViewIndex);
 
+		FText MemAllocTableTreeViewTabDisplayName = FText::Format(LOCTEXT("MemoryProfiler.MemAllocTableTreeViewTabTitle", "Allocs Table {0}"), FText::AsNumber(LastMemAllocTableTreeViewIndex));
+		TabManager->RegisterTabSpawner(TabId, FOnSpawnTab::CreateRaw(this, &SMemoryProfilerWindow::SpawnTab_MemAllocTableTreeView, LastMemAllocTableTreeViewIndex))
+			.SetDisplayName(MemAllocTableTreeViewTabDisplayName)
+			.SetIcon(FSlateIcon(FInsightsStyle::GetStyleSetName(), "MemAllocTableTreeView.Icon.Small"))
+			.SetGroup(AppMenuGroup.ToSharedRef());
+
+
+		TSharedPtr<Insights::FQueryTargetWindowSpec> NewTarget = MakeShared<Insights::FQueryTargetWindowSpec>(TabId, MemAllocTableTreeViewTabDisplayName);
+		SharedState->AddQueryTarget(NewTarget);
+		SharedState->SetCurrentQueryTarget(NewTarget);
+		MemInvestigationView->QueryTarget_OnSelectionChanged(NewTarget, ESelectInfo::Type::Direct);
+	}
+
+	FName TabId = SharedState->GetCurrentQueryTarget()->GetName();
 	if (TabManager->HasTabSpawner(TabId))
 	{
 		TSharedPtr<SDockTab> Tab = TabManager->TryInvokeTab(TabId);
 		if (Tab)
 		{
 			TSharedRef<Insights::SMemAllocTableTreeView> MemAllocTableTreeView = StaticCastSharedRef<Insights::SMemAllocTableTreeView>(Tab->GetContent());
-			ensure(MemAllocTableTreeView->GetTabIndex() == LastMemAllocTableTreeViewIndex);
+
+			if (SharedState->GetCurrentQueryTarget()->GetName() == Insights::FQueryTargetWindowSpec::NewWindow)
+			{
+				Tab->SetOnTabClosed(SDockTab::FOnTabClosedCallback::CreateRaw(this, &SMemoryProfilerWindow::OnMemAllocTableTreeViewTabClosed));
+			}
 			return MemAllocTableTreeView;
 		}
 	}
@@ -433,38 +476,27 @@ void SMemoryProfilerWindow::Construct(const FArguments& InArgs, const TSharedRef
 	};
 	TabManager->SetOnPersistLayout(FTabManager::FOnPersistLayout::CreateLambda(PersistLayout));
 
-	TSharedRef<FWorkspaceItem> AppMenuGroup = TabManager->AddLocalWorkspaceMenuCategory(LOCTEXT("MemoryProfilerMenuGroupName", "Memory Insights"));
+	AppMenuGroup = TabManager->AddLocalWorkspaceMenuCategory(LOCTEXT("MemoryProfilerMenuGroupName", "Memory Insights"));
 
 	TabManager->RegisterTabSpawner(FMemoryProfilerTabs::ToolbarID, FOnSpawnTab::CreateRaw(this, &SMemoryProfilerWindow::SpawnTab_Toolbar))
 		.SetDisplayName(LOCTEXT("DeviceToolbarTabTitle", "Toolbar"))
 		.SetIcon(FSlateIcon(FInsightsStyle::GetStyleSetName(), "Toolbar.Icon.Small"))
-		.SetGroup(AppMenuGroup);
+		.SetGroup(AppMenuGroup.ToSharedRef());
 
 	TabManager->RegisterTabSpawner(FMemoryProfilerTabs::TimingViewID, FOnSpawnTab::CreateRaw(this, &SMemoryProfilerWindow::SpawnTab_TimingView))
 		.SetDisplayName(LOCTEXT("MemoryProfiler.TimingViewTabTitle", "Timing View"))
 		.SetIcon(FSlateIcon(FInsightsStyle::GetStyleSetName(), "TimingView.Icon.Small"))
-		.SetGroup(AppMenuGroup);
+		.SetGroup(AppMenuGroup.ToSharedRef());
 
 	TabManager->RegisterTabSpawner(FMemoryProfilerTabs::MemInvestigationViewID, FOnSpawnTab::CreateRaw(this, &SMemoryProfilerWindow::SpawnTab_MemInvestigationView))
 		.SetDisplayName(LOCTEXT("MemoryProfiler.MemInvestigationViewTabTitle", "Investigation"))
 		.SetIcon(FSlateIcon(FInsightsStyle::GetStyleSetName(), "MemInvestigationView.Icon.Small"))
-		.SetGroup(AppMenuGroup);
+		.SetGroup(AppMenuGroup.ToSharedRef());
 
 	TabManager->RegisterTabSpawner(FMemoryProfilerTabs::MemTagTreeViewID, FOnSpawnTab::CreateRaw(this, &SMemoryProfilerWindow::SpawnTab_MemTagTreeView))
 		.SetDisplayName(LOCTEXT("MemoryProfiler.MemTagTreeViewTabTitle", "LLM Tags"))
 		.SetIcon(FSlateIcon(FInsightsStyle::GetStyleSetName(), "MemTagTreeView.Icon.Small"))
-		.SetGroup(AppMenuGroup);
-
-	FName MemAllocTableTreeViewTabId = FMemoryProfilerTabs::MemAllocTableTreeViewID;
-	for (int32 TabIndex = 0; TabIndex < MaxMemAllocTableTreeViews; ++TabIndex)
-	{
-		MemAllocTableTreeViewTabId.SetNumber(TabIndex);
-		FText MemAllocTableTreeViewTabDisplayName = FText::Format(LOCTEXT("MemoryProfiler.MemAllocTableTreeViewTabTitle", "Allocs Table {0}"), FText::AsNumber(TabIndex));
-		TabManager->RegisterTabSpawner(MemAllocTableTreeViewTabId, FOnSpawnTab::CreateRaw(this, &SMemoryProfilerWindow::SpawnTab_MemAllocTableTreeView, TabIndex))
-			.SetDisplayName(MemAllocTableTreeViewTabDisplayName)
-			.SetIcon(FSlateIcon(FInsightsStyle::GetStyleSetName(), "MemAllocTableTreeView.Icon.Small"))
-			.SetGroup(AppMenuGroup);
-	}
+		.SetGroup(AppMenuGroup.ToSharedRef());
 
 	TSharedPtr<FMemoryProfilerManager> MemoryProfilerManager = FMemoryProfilerManager::Get();
 	ensure(MemoryProfilerManager.IsValid());
diff --git a/Engine/Source/Developer/TraceInsights/Private/Insights/MemoryProfiler/Widgets/SMemoryProfilerWindow.h b/Engine/Source/Developer/TraceInsights/Private/Insights/MemoryProfiler/Widgets/SMemoryProfilerWindow.h
--- a/Engine/Source/Developer/TraceInsights/Private/Insights/MemoryProfiler/Widgets/SMemoryProfilerWindow.h
+++ b/Engine/Source/Developer/TraceInsights/Private/Insights/MemoryProfiler/Widgets/SMemoryProfilerWindow.h
@@ -19,6 +19,7 @@
 
 class FActiveTimerHandle;
 class FMenuBuilder;
+class FWorkspaceItem;
 
 class FMemorySharedState;
 class SMemInvestigationView;
@@ -197,6 +198,8 @@ private:
 	/** Holds the tab manager that manages the front-end's tabs. */
 	TSharedPtr<FTabManager> TabManager;
 
+	TSharedPtr<FWorkspaceItem> AppMenuGroup;
+
 	/** The handle to the active update duration tick */
 	TWeakPtr<FActiveTimerHandle> ActiveTimerHandle;
 
