[Insights] Add information about the current query in the Memory Allocation Tree View. #rb none

diff --git a/Engine/Source/Developer/TraceInsights/Private/Insights/MemoryProfiler/ViewModels/MemorySharedState.cpp b/Engine/Source/Developer/TraceInsights/Private/Insights/MemoryProfiler/ViewModels/MemorySharedState.cpp
--- a/Engine/Source/Developer/TraceInsights/Private/Insights/MemoryProfiler/ViewModels/MemorySharedState.cpp
+++ b/Engine/Source/Developer/TraceInsights/Private/Insights/MemoryProfiler/ViewModels/MemorySharedState.cpp
@@ -946,67 +946,67 @@ void FMemorySharedState::InitMemoryRules()
 		ERule::aAf, 1,
 		LOCTEXT("MemRule_aAf_Short", "*A*"),
 		LOCTEXT("MemRule_aAf_Verbose", "Active Allocs"),
-		LOCTEXT("MemRule_aAf_Desc", "Identifies active allocations at time A.\n(a ≤ A ≤ f)")));
+		LOCTEXT("MemRule_aAf_Desc_Fmt", "Identifies active allocations at time {0}.\n(a ≤ {0} ≤ f)")));
 
 	MemoryRules.Add(MakeShared<Insights::FMemoryRuleSpec>(
 		ERule::afA, 1,
 		LOCTEXT("MemRule_afA_Short", "**A"),
 		LOCTEXT("MemRule_afA_Verbose", "Before"),
-		LOCTEXT("MemRule_afA_Desc", "Identifies allocations allocated and freed before time A.\n(a ≤ f ≤ A)")));
+		LOCTEXT("MemRule_afA_Desc_Fmt", "Identifies allocations allocated and freed before time {0}.\n(a ≤ f ≤ {0})")));
 
 	MemoryRules.Add(MakeShared<Insights::FMemoryRuleSpec>(
 		ERule::Aaf, 1,
 		LOCTEXT("MemRule_Aaf_Short", "A**"),
 		LOCTEXT("MemRule_Aaf_Verbose", "After"),
-		LOCTEXT("MemRule_Aaf_Desc", "Identifies allocations allocated and freed after time A.\n(A ≤ a ≤ f)")));
+		LOCTEXT("MemRule_Aaf_Desc_Fmt", "Identifies allocations allocated and freed after time {0}.\n({0} ≤ a ≤ f)")));
 
 	MemoryRules.Add(MakeShared<Insights::FMemoryRuleSpec>(
 		ERule::aAfB, 2,
 		LOCTEXT("MemRule_aAfB_Short", "*A*B"),
 		LOCTEXT("MemRule_aAfB_Verbose", "Decline"),
-		LOCTEXT("MemRule_aAfB_Desc", "Identifies allocations allocated before time A and freed between time A and time B.\n(a ≤ A ≤ f ≤ B)")));
+		LOCTEXT("MemRule_aAfB_Desc_Fmt", "Identifies allocations allocated before time {0} and freed between time {0} and time {1}.\n(a ≤ {0} ≤ f ≤ {1})")));
 
 	MemoryRules.Add(MakeShared<Insights::FMemoryRuleSpec>(
 		ERule::AaBf, 2,
 		LOCTEXT("MemRule_AaBf_Short", "A*B*"),
 		LOCTEXT("MemRule_AaBf_Verbose", "Growth"),
-		LOCTEXT("MemRule_AaBf_Desc", "Identifies allocations allocated between time A and time B and freed after time B.\n(A ≤ a ≤ B ≤ f)")));
+		LOCTEXT("MemRule_AaBf_Desc_Fmt", "Identifies allocations allocated between time {0} and time {1} and freed after time {1}.\n({0} ≤ a ≤ {1} ≤ f)")));
 
 	MemoryRules.Add(MakeShared<Insights::FMemoryRuleSpec>(
 		ERule::AafB, 2,
 		LOCTEXT("MemRule_AafB_Short", "A**B"),
 		LOCTEXT("MemRule_AafB_Verbose", "Short Living Allocs"),
-		LOCTEXT("MemRule_AafB_Desc", "Identifies allocations allocated and freed between time A and time B.\n(A ≤ a ≤ f ≤ B)")));
+		LOCTEXT("MemRule_AafB_Desc_Fmt", "Identifies allocations allocated and freed between time {0} and time {1}.\n({0} ≤ a ≤ f ≤ {1})")));
 
 	MemoryRules.Add(MakeShared<Insights::FMemoryRuleSpec>(
 		ERule::aABf, 2,
 		LOCTEXT("MemRule_aABf_Short", "*A B*"),
 		LOCTEXT("MemRule_aABf_Verbose", "Long Living Allocs"),
-		LOCTEXT("MemRule_aABf_Desc", "Identifies allocations allocated before time A and freed after time B.\n(a ≤ A ≤ B ≤ f)")));
+		LOCTEXT("MemRule_aABf_Desc_Fmt", "Identifies allocations allocated before time {0} and freed after time {1}.\n(a ≤ {0} ≤ {1} ≤ f)")));
 
 	MemoryRules.Add(MakeShared<Insights::FMemoryRuleSpec>(
 		ERule::AaBCf, 3,
 		LOCTEXT("MemRule_AaBCf_Short", "A*B C*"),
 		LOCTEXT("MemRule_AaBCf_Verbose", "Memory Leaks"),
-		LOCTEXT("MemRule_AaBCf_Desc", "Identifies allocations allocated between time A and time B and freed after time C.\n(A ≤ a ≤ B ≤ C ≤ f)")));
+		LOCTEXT("MemRule_AaBCf_Desc_Fmt", "Identifies allocations allocated between time {0} and time {1} and freed after time {2}.\n({0} ≤ a ≤ {1} ≤ {2} ≤ f)")));
 
 	MemoryRules.Add(MakeShared<Insights::FMemoryRuleSpec>(
 		ERule::AaBfC, 3,
 		LOCTEXT("MemRule_AaBfC_Short", "A*B*C"),
 		LOCTEXT("MemRule_AaBfC_Verbose", "Limited Lifetime"),
-		LOCTEXT("MemRule_AaBfC_Desc", "Identifies allocations allocated between time A and time B and freed between time B and time C.\n(A ≤ a ≤ B ≤ f ≤ C)")));
+		LOCTEXT("MemRule_AaBfC_Desc_Fmt", "Identifies allocations allocated between time {0} and time {1} and freed between time {1} and time {2}.\n({0} ≤ a ≤ {1} ≤ f ≤ {2})")));
 
 	MemoryRules.Add(MakeShared<Insights::FMemoryRuleSpec>(
 		ERule::aABfC, 3,
 		LOCTEXT("MemRule_aABfC_Short", "*A B*C"),
 		LOCTEXT("MemRule_aABfC_Verbose", "Decline of Long Living Allocs"),
-		LOCTEXT("MemRule_aABfC_Desc", "Identifies allocations allocated before time A and freed between time B and time C.\n(a ≤ A ≤ B ≤ f ≤ C)")));
+		LOCTEXT("MemRule_aABfC_Desc_Fmt", "Identifies allocations allocated before time {0} and freed between time {1} and time {2}.\n(a ≤ {0} ≤ {1} ≤ f ≤ {2})")));
 
 	MemoryRules.Add(MakeShared<Insights::FMemoryRuleSpec>(
 		ERule::AaBCfD, 4,
 		LOCTEXT("MemRule_AaBCfD_Short", "A*B C*D"),
 		LOCTEXT("MemRule_AaBCfD_Verbose", "Specific Lifetime"),
-		LOCTEXT("MemRule_AaBCfD_Desc", "Identifies allocations allocated between time A and time B and freed between time C and time D.\n(A ≤ a ≤ B ≤ C ≤ f ≤ D)")));
+		LOCTEXT("MemRule_AaBCfD_Desc_Fmt", "Identifies allocations allocated between time {0} and time {1} and freed between time {2} and time {3}.\n({0} ≤ a ≤ {1} ≤ {2} ≤ f ≤ {3})")));
 
 	//TODO
 	//MemoryRules.Add(MakeShared<Insights::FMemoryRuleSpec>(
diff --git a/Engine/Source/Developer/TraceInsights/Private/Insights/MemoryProfiler/Widgets/SMemAllocTableTreeView.cpp b/Engine/Source/Developer/TraceInsights/Private/Insights/MemoryProfiler/Widgets/SMemAllocTableTreeView.cpp
--- a/Engine/Source/Developer/TraceInsights/Private/Insights/MemoryProfiler/Widgets/SMemAllocTableTreeView.cpp
+++ b/Engine/Source/Developer/TraceInsights/Private/Insights/MemoryProfiler/Widgets/SMemAllocTableTreeView.cpp
@@ -185,6 +185,8 @@ void SMemAllocTableTreeView::OnQueryInvalidated()
 		Allocs.Reset(10 * 1024 * 1024);
 	}
 
+	UpdateQueryInfo();
+
 	StartQuery();
 
 	RebuildTree(true);
@@ -648,6 +650,15 @@ TSharedPtr<SWidget> SMemAllocTableTreeView::ConstructFooter()
 	return
 		SNew(SHorizontalBox)
 
+		+ SHorizontalBox::Slot()
+		.HAlign(HAlign_Left)
+		.Padding(2.0f)
+		[
+			SNew(STextBlock)
+			.Text(this, &SMemAllocTableTreeView::GetQueryInfo)
+			.ToolTipText(this, &SMemAllocTableTreeView::GetQueryInfoTooltip)
+		]
+
 		+ SHorizontalBox::Slot()
 		.HAlign(HAlign_Right)
 		.Padding(2.0f)
@@ -671,6 +682,20 @@ FText SMemAllocTableTreeView::GetSymbolResolutionStatus() const
 
 ////////////////////////////////////////////////////////////////////////////////////////////////////
 
+FText SMemAllocTableTreeView::GetQueryInfo() const
+{
+	return QueryInfo;
+}
+
+////////////////////////////////////////////////////////////////////////////////////////////////////
+
+FText SMemAllocTableTreeView::GetQueryInfoTooltip() const
+{
+	return QueryInfoTooltip;
+}
+
+////////////////////////////////////////////////////////////////////////////////////////////////////
+
 void SMemAllocTableTreeView::InternalCreateGroupings()
 {
 	STableTreeView::InternalCreateGroupings();
@@ -700,6 +725,54 @@ void SMemAllocTableTreeView::InternalCreateGroupings()
 
 ////////////////////////////////////////////////////////////////////////////////////////////////////
 
+void SMemAllocTableTreeView::UpdateQueryInfo()
+{
+	if (Rule.IsValid())
+	{
+		FString RuleDesc = Rule->GetDescription().ToString();
+		FText TimeMarkersText;
+		int NumMarkers = Rule->GetNumTimeMarkers();
+
+		switch (NumMarkers)
+		{
+		case 1:
+		{
+			TimeMarkersText = FText::Format(LOCTEXT("TimeMarkersFmt", "A={0}"), FText::AsNumber(TimeMarkers[0]));
+			QueryInfoTooltip = FText::Format(Rule->GetDescription(), FText::AsNumber(TimeMarkers[0]));
+			break;
+		}
+		case 2:
+		{
+			TimeMarkersText = FText::Format(LOCTEXT("TimeMarkersFmt", "A={0} B={1}"), FText::AsNumber(TimeMarkers[0]), FText::AsNumber(TimeMarkers[1]));
+			QueryInfoTooltip = FText::Format(Rule->GetDescription(), FText::AsNumber(TimeMarkers[0]), FText::AsNumber(TimeMarkers[1]));
+			break;
+		}
+		case 3:
+		{
+			TimeMarkersText = FText::Format(LOCTEXT("TimeMarkersFmt", "A={0} B={1} C={2}"), FText::AsNumber(TimeMarkers[0]), FText::AsNumber(TimeMarkers[1]), FText::AsNumber(TimeMarkers[2]));
+			QueryInfoTooltip = FText::Format(Rule->GetDescription(), FText::AsNumber(TimeMarkers[0]), FText::AsNumber(TimeMarkers[1]), FText::AsNumber(TimeMarkers[2]));
+			break;
+		}
+		case 4:
+		{
+			TimeMarkersText = FText::Format(LOCTEXT("TimeMarkersFmt", "A={0} B={1} C={2} D={3}"), FText::AsNumber(TimeMarkers[0]), FText::AsNumber(TimeMarkers[1]), FText::AsNumber(TimeMarkers[2]), FText::AsNumber(TimeMarkers[3]));
+			QueryInfoTooltip = FText::Format(Rule->GetDescription(), FText::AsNumber(TimeMarkers[0]), FText::AsNumber(TimeMarkers[1]), FText::AsNumber(TimeMarkers[2]), FText::AsNumber(TimeMarkers[3]));
+			break;
+		}
+		default:
+		{
+			// Unhandled value
+			check(false);
+		}
+		}
+
+		QueryInfo = Rule->GetVerboseName();
+		QueryInfo = FText::Format(LOCTEXT("QueryInfoFmt", "Showing {0} for markers {1}"), QueryInfo, TimeMarkersText);
+	}
+}
+
+////////////////////////////////////////////////////////////////////////////////////////////////////
+
 } // namespace Insights
 
 #undef LOCTEXT_NAMESPACE
diff --git a/Engine/Source/Developer/TraceInsights/Private/Insights/MemoryProfiler/Widgets/SMemAllocTableTreeView.h b/Engine/Source/Developer/TraceInsights/Private/Insights/MemoryProfiler/Widgets/SMemAllocTableTreeView.h
--- a/Engine/Source/Developer/TraceInsights/Private/Insights/MemoryProfiler/Widgets/SMemAllocTableTreeView.h
+++ b/Engine/Source/Developer/TraceInsights/Private/Insights/MemoryProfiler/Widgets/SMemAllocTableTreeView.h
@@ -106,14 +106,19 @@ private:
 	FReply OnCallstackViewClicked(bool bIsInverted);
 
 	FText GetSymbolResolutionStatus() const;
+	FText GetQueryInfo() const;
+	FText GetQueryInfoTooltip() const;
 
 	void ApplyColumnConfig(const TArrayView<FColumnConfig>& Preset);
+	void UpdateQueryInfo();
 
 private:
 	int32 TabIndex = -1;
 	TSharedPtr<FMemoryRuleSpec> Rule = nullptr;
 	double TimeMarkers[4];
 	TraceServices::IAllocationsProvider::FQueryHandle Query = 0;
+	FText QueryInfo;
+	FText QueryInfoTooltip;
 
 	FStopwatch QueryStopwatch;
 };
diff --git a/Engine/Source/Developer/TraceInsights/Private/Insights/MemoryProfiler/Widgets/SMemInvestigationView.cpp b/Engine/Source/Developer/TraceInsights/Private/Insights/MemoryProfiler/Widgets/SMemInvestigationView.cpp
--- a/Engine/Source/Developer/TraceInsights/Private/Insights/MemoryProfiler/Widgets/SMemInvestigationView.cpp
+++ b/Engine/Source/Developer/TraceInsights/Private/Insights/MemoryProfiler/Widgets/SMemInvestigationView.cpp
@@ -367,7 +367,34 @@ FText SMemInvestigationView::QueryRule_GetTooltipText() const
 		FMemorySharedState& SharedState = ProfilerWindow->GetSharedState();
 		if (SharedState.GetCurrentMemoryRule())
 		{
-			return SharedState.GetCurrentMemoryRule()->GetDescription();
+			switch (SharedState.GetCurrentMemoryRule()->GetNumTimeMarkers())
+			{
+			case 1:
+			{
+				return FText::Format(SharedState.GetCurrentMemoryRule()->GetDescription(), FText::FromString(TEXT("A")));
+				break;
+			}
+			case 2:
+			{
+				return FText::Format(SharedState.GetCurrentMemoryRule()->GetDescription(), FText::FromString(TEXT("A")), FText::FromString(TEXT("B")));
+				break;
+			}
+			case 3:
+			{
+				return FText::Format(SharedState.GetCurrentMemoryRule()->GetDescription(), FText::FromString(TEXT("A")), FText::FromString(TEXT("B")), FText::FromString(TEXT("C")));
+				break;
+			}
+			case 4:
+			{
+				return FText::Format(SharedState.GetCurrentMemoryRule()->GetDescription(), FText::FromString(TEXT("A")), FText::FromString(TEXT("B")), FText::FromString(TEXT("C")), FText::FromString(TEXT("D")));
+				break;
+			}
+			default:
+			{
+				// Unhandled value
+				check(false);
+			}
+			}
 		}
 	}
 	return FText::GetEmpty();
