[Insights] String filter support in FilterConfigurator
Add callstack and full callstack filters in MemAlloc TreeView

diff --git a/Engine/Source/Developer/TraceInsights/Private/Insights/MemoryProfiler/Widgets/SMemAllocTableTreeView.cpp b/Engine/Source/Developer/TraceInsights/Private/Insights/MemoryProfiler/Widgets/SMemAllocTableTreeView.cpp
--- a/Engine/Source/Developer/TraceInsights/Private/Insights/MemoryProfiler/Widgets/SMemAllocTableTreeView.cpp
+++ b/Engine/Source/Developer/TraceInsights/Private/Insights/MemoryProfiler/Widgets/SMemAllocTableTreeView.cpp
@@ -18,6 +18,7 @@
 #include "Insights/Table/ViewModels/TableCellValueFormatter.h"
 #include "Insights/Table/ViewModels/TableColumn.h"
 #include "Insights/TimingProfilerCommon.h"
+#include "Insights/ViewModels/FilterConfigurator.h"
 
 #include <limits>
 
@@ -28,6 +29,8 @@ using namespace TraceServices;
 namespace Insights
 {
 
+const int SMemAllocTableTreeView::FullCallStackIndex = 0x0000FFFFF;
+
 ////////////////////////////////////////////////////////////////////////////////////////////////////
 
 SMemAllocTableTreeView::SMemAllocTableTreeView()
@@ -773,6 +776,26 @@ void SMemAllocTableTreeView::UpdateQueryInfo()
 
 ////////////////////////////////////////////////////////////////////////////////////////////////////
 
+bool SMemAllocTableTreeView::ApplyCustomAdvancedFilters(const FTableTreeNodePtr& NodePtr)
+{
+	FMemAllocNodePtr MemNodePtr = StaticCastSharedPtr<FMemAllocNode>(NodePtr);
+	Context.SetFilterData<FString>(FullCallStackIndex, MemNodePtr->GetFullCallstack().ToString());
+
+	return true;
+}
+
+////////////////////////////////////////////////////////////////////////////////////////////////////
+
+void SMemAllocTableTreeView::AddCustomAdvancedFilters()
+{
+	TSharedPtr<TArray<TSharedPtr<struct FFilter>>>& AvailableFilters = FilterConfigurator->GetAvailableFilters();
+
+	AvailableFilters->Add(MakeShared<FFilter>(FullCallStackIndex, LOCTEXT("FullCallstack", "Full Callstack"), LOCTEXT("FullCallstack", "Search in all the callstack frames"), EFilterDataType::String, FFilterService::Get()->GetStringOperators()));
+	Context.AddFilterData<FString>(FullCallStackIndex, FString());
+}
+
+////////////////////////////////////////////////////////////////////////////////////////////////////
+
 } // namespace Insights
 
 #undef LOCTEXT_NAMESPACE
diff --git a/Engine/Source/Developer/TraceInsights/Private/Insights/MemoryProfiler/Widgets/SMemAllocTableTreeView.h b/Engine/Source/Developer/TraceInsights/Private/Insights/MemoryProfiler/Widgets/SMemAllocTableTreeView.h
--- a/Engine/Source/Developer/TraceInsights/Private/Insights/MemoryProfiler/Widgets/SMemAllocTableTreeView.h
+++ b/Engine/Source/Developer/TraceInsights/Private/Insights/MemoryProfiler/Widgets/SMemAllocTableTreeView.h
@@ -111,8 +111,11 @@ private:
 
 	void ApplyColumnConfig(const TArrayView<FColumnConfig>& Preset);
 	void UpdateQueryInfo();
+	bool virtual ApplyCustomAdvancedFilters(const FTableTreeNodePtr& NodePtr) override;
+	virtual void AddCustomAdvancedFilters() override;
 
 private:
+	const static int FullCallStackIndex;
 	int32 TabIndex = -1;
 	TSharedPtr<FMemoryRuleSpec> Rule = nullptr;
 	double TimeMarkers[4];
diff --git a/Engine/Source/Developer/TraceInsights/Private/Insights/Table/Widgets/STableTreeView.cpp b/Engine/Source/Developer/TraceInsights/Private/Insights/Table/Widgets/STableTreeView.cpp
--- a/Engine/Source/Developer/TraceInsights/Private/Insights/Table/Widgets/STableTreeView.cpp
+++ b/Engine/Source/Developer/TraceInsights/Private/Insights/Table/Widgets/STableTreeView.cpp
@@ -2507,9 +2507,22 @@ FReply STableTreeView::OnAdvancedFiltersClicked()
 					Context.AddFilterData<double>(Column->GetIndex(), 0.0);
 					break;
 				}
+				case ETableCellDataType::CString:
+				case ETableCellDataType::Text:
+				case ETableCellDataType::Custom:
+				{
+					if (!Column->IsHierarchy())
+					{
+						AvailableFilters->Add(MakeShared<FFilter>(Column->GetIndex(), Column->GetTitleName(), Column->GetDescription(), EFilterDataType::String, FFilterService::Get()->GetStringOperators()));
+						Context.AddFilterData<FString>(Column->GetIndex(), FString());
+					}
+					break;
+				}
 			}
 		}
 
+		AddCustomAdvancedFilters();
+
 		CurrentAsyncOpFilterConfigurator = new FFilterConfigurator(*FilterConfigurator);
 		OnFilterChangesCommitedHandle = FilterConfigurator->GetOnChangesCommitedEvent().AddSP(this, &STableTreeView::OnAdvancedFiltersChangesCommited);
 	}
@@ -2547,10 +2560,20 @@ bool STableTreeView::ApplyAdvancedFilters(const FTableTreeNodePtr& NodePtr)
 				Context.SetFilterData<double>(Column->GetIndex(), Column->GetValue(*NodePtr)->AsDouble());
 				break;
 			}
+			case ETableCellDataType::CString:
+			case ETableCellDataType::Text:
+			case ETableCellDataType::Custom:
+			{
+				if (!Column->IsHierarchy())
+				{
+					Context.SetFilterData<FString>(Column->GetIndex(), Column->GetValue(*NodePtr)->AsString());
+				}
+				break;
+			}
 		}
 	}
 
-	return CurrentAsyncOpFilterConfigurator->ApplyFilters(Context);
+	return ApplyCustomAdvancedFilters(NodePtr) && CurrentAsyncOpFilterConfigurator->ApplyFilters(Context);
 }
 
 ////////////////////////////////////////////////////////////////////////////////////////////////////
diff --git a/Engine/Source/Developer/TraceInsights/Private/Insights/Table/Widgets/STableTreeView.h b/Engine/Source/Developer/TraceInsights/Private/Insights/Table/Widgets/STableTreeView.h
--- a/Engine/Source/Developer/TraceInsights/Private/Insights/Table/Widgets/STableTreeView.h
+++ b/Engine/Source/Developer/TraceInsights/Private/Insights/Table/Widgets/STableTreeView.h
@@ -364,6 +364,8 @@ protected:
 	FReply OnAdvancedFiltersClicked();
 	void OnAdvancedFiltersChangesCommited();
 	bool ApplyAdvancedFilters(const FTableTreeNodePtr& NodePtr);
+	bool virtual ApplyCustomAdvancedFilters(const FTableTreeNodePtr& NodePtr) { return true; };
+	virtual void AddCustomAdvancedFilters() {}
 
 	////////////////////////////////////////////////////////////////////////////////////////////////////
 
diff --git a/Engine/Source/Developer/TraceInsights/Private/Insights/ViewModels/FilterConfiguratorNode.cpp b/Engine/Source/Developer/TraceInsights/Private/Insights/ViewModels/FilterConfiguratorNode.cpp
--- a/Engine/Source/Developer/TraceInsights/Private/Insights/ViewModels/FilterConfiguratorNode.cpp
+++ b/Engine/Source/Developer/TraceInsights/Private/Insights/ViewModels/FilterConfiguratorNode.cpp
@@ -23,6 +23,8 @@ FFilterConfiguratorNode::FFilterConfiguratorNode(const FName InName, bool bInIsG
 	{
 		SelectedFilterGroupOperator = GetFilterGroupOperators()[0];
 	}
+
+	AvailableFilterOperators = MakeShared<TArray<TSharedPtr<IFilterOperator>>>();
 }
 
 ////////////////////////////////////////////////////////////////////////////////////////////////////
@@ -40,6 +42,7 @@ FFilterConfiguratorNode& FFilterConfiguratorNode::operator=(const FFilterConfigu
 	GetChildrenMutable().Empty();
 
 	AvailableFilters = Other.AvailableFilters;
+	AvailableFilterOperators = Other.AvailableFilterOperators;
 	SelectedFilter = Other.SelectedFilter;
 	SelectedFilterOperator = Other.SelectedFilterOperator;
 	SelectedFilterGroupOperator = Other.SelectedFilterGroupOperator;
@@ -124,6 +127,13 @@ void FFilterConfiguratorNode::SetSelectedFilter(TSharedPtr<struct FFilter> InSel
 	if (SelectedFilter.IsValid() && SelectedFilter->GetSupportedOperators()->Num() > 0)
 	{
 		SetSelectedFilterOperator(SelectedFilter->GetSupportedOperators()->GetData()[0]);
+
+		AvailableFilterOperators->Empty();
+		Insights::FFilter::SupportedOperatorsArrayPtr AvailableOperators = InSelectedFilter->SupportedOperators;
+		for (auto& FilterOperator : *AvailableOperators)
+		{
+			AvailableFilterOperators->Add(FilterOperator);
+		}
 	}
 }
 
@@ -154,6 +164,11 @@ void FFilterConfiguratorNode::ProcessFilter()
 			FilterValue.Set<int64>(FCString::Atoi64(*TextBoxValue));
 			break;
 		}
+		case EFilterDataType::String:
+		{
+			FilterValue.Set<FString>(TextBoxValue);
+			break;
+		}
 		}
 	}
 }
@@ -172,7 +187,7 @@ bool FFilterConfiguratorNode::ApplyFilters(const FFilterContext& Context) const
 			auto& ChildrenArr = GetChildren();
 			for (int Index = 0; Index < ChildrenArr.Num() && Ret; ++Index)
 			{
-				Ret &= ((FFilterConfiguratorNode*)&*ChildrenArr[Index])->ApplyFilters(Context);
+				Ret &= ((FFilterConfiguratorNode*)ChildrenArr[Index].Get())->ApplyFilters(Context);
 			}
 			break;
 		}
@@ -186,7 +201,7 @@ bool FFilterConfiguratorNode::ApplyFilters(const FFilterContext& Context) const
 
 			for (int Index = 0; Index < ChildrenArr.Num() && !Ret; ++Index)
 			{
-				Ret |= ((FFilterConfiguratorNode*)&*ChildrenArr[Index])->ApplyFilters(Context);
+				Ret |= ((FFilterConfiguratorNode*)ChildrenArr[Index].Get())->ApplyFilters(Context);
 			}
 			break;
 		}
@@ -198,7 +213,7 @@ bool FFilterConfiguratorNode::ApplyFilters(const FFilterContext& Context) const
 		{
 		case EFilterDataType::Double:
 		{
-			FFilterOperator<double>* Operator = (FFilterOperator<double>*)&*SelectedFilterOperator;
+			FFilterOperator<double>* Operator = (FFilterOperator<double>*) SelectedFilterOperator.Get();
 			double Value = 0.0;
 			Context.GetFilterData<double>(SelectedFilter->Key, Value);
 
@@ -207,13 +222,22 @@ bool FFilterConfiguratorNode::ApplyFilters(const FFilterContext& Context) const
 		}
 		case EFilterDataType::Int64:
 		{
-			FFilterOperator<int64>* Operator = (FFilterOperator<int64>*) & *SelectedFilterOperator;
+			FFilterOperator<int64>* Operator = (FFilterOperator<int64>*) SelectedFilterOperator.Get();
 			int64 Value = 0;
 			Context.GetFilterData<int64>(SelectedFilter->Key, Value);
 
 			Ret = Operator->Func(Value, FilterValue.Get<int64>());
 			break;
 		}
+		case EFilterDataType::String:
+		{
+			FFilterOperator<FString>* Operator = (FFilterOperator<FString>*) SelectedFilterOperator.Get();
+			FString Value;
+			Context.GetFilterData<FString>(SelectedFilter->Key, Value);
+
+			Ret = Operator->Func(Value, FilterValue.Get<FString>());
+			break;
+		}
 		default:
 			break;
 		}
diff --git a/Engine/Source/Developer/TraceInsights/Private/Insights/ViewModels/FilterConfiguratorNode.h b/Engine/Source/Developer/TraceInsights/Private/Insights/ViewModels/FilterConfiguratorNode.h
--- a/Engine/Source/Developer/TraceInsights/Private/Insights/ViewModels/FilterConfiguratorNode.h
+++ b/Engine/Source/Developer/TraceInsights/Private/Insights/ViewModels/FilterConfiguratorNode.h
@@ -72,6 +72,8 @@ public:
 	void SetSelectedFilterGroupOperator(TSharedPtr<struct FFilterGroupOperator> InSelectedFilterGroupOperator) { SelectedFilterGroupOperator = InSelectedFilterGroupOperator; }
 	TSharedPtr<struct FFilterGroupOperator> GetSelectedFilterGroupOperator() const { return SelectedFilterGroupOperator; }
 
+	TSharedPtr<TArray<TSharedPtr<IFilterOperator>>> GetAvailableFilterOperators() const {	return AvailableFilterOperators;	}
+
 	void DeleteChildNode(FFilterConfiguratorNodePtr InNode);
 
 	const FString& GetTextBoxValue() { return TextBoxValue; }
@@ -94,6 +96,8 @@ private:
 
 	TSharedPtr<struct FFilterGroupOperator> SelectedFilterGroupOperator;
 
+	TSharedPtr<TArray<TSharedPtr<IFilterOperator>>> AvailableFilterOperators;
+
 	FString TextBoxValue;
 
 	// The textbox value, converted to its data type
diff --git a/Engine/Source/Developer/TraceInsights/Private/Insights/ViewModels/Filters.h b/Engine/Source/Developer/TraceInsights/Private/Insights/ViewModels/Filters.h
--- a/Engine/Source/Developer/TraceInsights/Private/Insights/ViewModels/Filters.h
+++ b/Engine/Source/Developer/TraceInsights/Private/Insights/ViewModels/Filters.h
@@ -21,7 +21,8 @@ namespace Insights
 enum class EFilterDataType : uint32
 {
 	Int64,
-	Double
+	Double,
+	String,
 };
 
 ////////////////////////////////////////////////////////////////////////////////////////////////////
@@ -33,6 +34,7 @@ enum class EFilterOperator : uint8
 	Lte = 2, // Less than or equal to
 	Gt = 3, // Greater than
 	Gte = 4, // Greater than or equal to 
+	Contains = 5,
 };
 
 ////////////////////////////////////////////////////////////////////////////////////////////////////
@@ -137,6 +139,10 @@ public:
 		IntegerOperators->Add(StaticCastSharedRef<IFilterOperator>(MakeShared<FFilterOperator<int64>>(EFilterOperator::Gt, TEXT(">"), std::greater<>())));
 		IntegerOperators->Add(StaticCastSharedRef<IFilterOperator>(MakeShared<FFilterOperator<int64>>(EFilterOperator::Gte, TEXT(">="), std::greater_equal<>())));
 
+		StringOperators = MakeShared<TArray<TSharedPtr<IFilterOperator>>>();
+		StringOperators->Add(StaticCastSharedRef<IFilterOperator>(MakeShared<FFilterOperator<FString>>(EFilterOperator::Eq, TEXT("Equals"), [](const FString& lhs, const FString& rhs) { return lhs.Equals(rhs); })));
+		StringOperators->Add(StaticCastSharedRef<IFilterOperator>(MakeShared<FFilterOperator<FString>>(EFilterOperator::Contains, TEXT("Contains"), [](const FString& lhs, const FString& rhs) { return lhs.Contains(rhs); })));
+
 		FilterGroupOperators.Add(MakeShared<FFilterGroupOperator>(EFilterGroupOperator::And, LOCTEXT("AllOf", "All Of"), LOCTEXT("AllOfDesc", "All of the children must be true for the group to return true. Equivalent to an AND operation.")));
 		FilterGroupOperators.Add(MakeShared<FFilterGroupOperator>(EFilterGroupOperator::Or, LOCTEXT("AnyOf", "Any Of"), LOCTEXT("AnyOfDesc", "Any of the children must be true for the group to return true. Equivalent to an OR operation.")));
 	}
@@ -145,12 +151,14 @@ public:
 
 	TSharedPtr<TArray<TSharedPtr<IFilterOperator>>> GetDoubleOperators() { return DoubleOperators; }
 	TSharedPtr<TArray<TSharedPtr<IFilterOperator>>> GetIntegerOperators() { return IntegerOperators; }
+	TSharedPtr<TArray<TSharedPtr<IFilterOperator>>> GetStringOperators() { return StringOperators; }
 
 private:
 	TArray<TSharedPtr<FFilterGroupOperator>> FilterGroupOperators;
 
 	TSharedPtr<TArray<TSharedPtr<IFilterOperator>>> DoubleOperators;
 	TSharedPtr<TArray<TSharedPtr<IFilterOperator>>> IntegerOperators;
+	TSharedPtr<TArray<TSharedPtr<IFilterOperator>>> StringOperators;
 };
 
 ////////////////////////////////////////////////////////////////////////////////////////////////////
@@ -171,6 +179,7 @@ public:
 
 	TSharedPtr<TArray<TSharedPtr<IFilterOperator>>> GetDoubleOperators() { return FilterStorage.GetDoubleOperators(); }
 	TSharedPtr<TArray<TSharedPtr<IFilterOperator>>> GetIntegerOperators() { return FilterStorage.GetIntegerOperators(); }
+	TSharedPtr<TArray<TSharedPtr<IFilterOperator>>> GetStringOperators() { return FilterStorage.GetStringOperators(); }
 
 	TSharedPtr<SWidget> CreateFilterConfiguratorWidget(TSharedPtr<class FFilterConfigurator> FilterConfiguratorViewModel);
 
@@ -188,7 +197,7 @@ private:
 class FFilterContext
 {
 public:
-	typedef TVariant<double, int64> ContextData;
+	typedef TVariant<double, int64, FString> ContextData;
 
 	template<typename T>
 	void AddFilterData(int32 Key, const T& InData)
diff --git a/Engine/Source/Developer/TraceInsights/Private/Insights/Widgets/SFilterConfiguratorRow.cpp b/Engine/Source/Developer/TraceInsights/Private/Insights/Widgets/SFilterConfiguratorRow.cpp
--- a/Engine/Source/Developer/TraceInsights/Private/Insights/Widgets/SFilterConfiguratorRow.cpp
+++ b/Engine/Source/Developer/TraceInsights/Private/Insights/Widgets/SFilterConfiguratorRow.cpp
@@ -288,7 +288,7 @@ FText SFilterConfiguratorRow::AvailableFilters_GetSelectionText() const
 
 const TArray<TSharedPtr<IFilterOperator>>* SFilterConfiguratorRow::GetAvailableFilterOperators()
 {
-	return &*FilterConfiguratorNodePtr->GetSelectedFilter()->SupportedOperators;
+	return FilterConfiguratorNodePtr->GetAvailableFilterOperators().Get();
 }
 
 ///////////////////////////////////////////////////////////////////////////////////////////////////
