[Insights] Changes to TableTreeView filtering:
- When using the search box filter, if a group is matched, all children of the group will be made visible
- When using advanced filters, only leaf nodes will be filtered, and only groups with visible children will be visible
- Because of this, the option to use the search box filter and the advanced filters at the same time will now be disabled

diff --git a/Engine/Source/Developer/TraceInsights/Private/Insights/Table/Widgets/STableTreeView.cpp b/Engine/Source/Developer/TraceInsights/Private/Insights/Table/Widgets/STableTreeView.cpp
--- a/Engine/Source/Developer/TraceInsights/Private/Insights/Table/Widgets/STableTreeView.cpp
+++ b/Engine/Source/Developer/TraceInsights/Private/Insights/Table/Widgets/STableTreeView.cpp
@@ -170,7 +170,7 @@ void STableTreeView::ConstructWidget(TSharedPtr<FTable> InTablePtr)
 					.HintText(LOCTEXT("SearchBox_Hint", "Search"))
 					.OnTextChanged(this, &STableTreeView::SearchBox_OnTextChanged)
 					.IsEnabled(this, &STableTreeView::SearchBox_IsEnabled)
-					.ToolTipText(LOCTEXT("SearchBox_ToolTip", "Type here to search the tree hierarchy by item or group name"))
+					.ToolTipText(this, &STableTreeView::SearchBox_GetTooltipText)
 				]
 
 				+ SHorizontalBox::Slot()
@@ -180,8 +180,9 @@ void STableTreeView::ConstructWidget(TSharedPtr<FTable> InTablePtr)
 				[
 					SNew(SButton)
 					.Text(LOCTEXT("AdvancedFiltersBtn_Text", "Advanced Filters"))
-					.ToolTipText(LOCTEXT("AdvancedFiltersBtn_ToolTip", "Opens the filter configurator window."))
+					.ToolTipText(this, &STableTreeView::AdvancedFilters_GetTooltipText)
 					.OnClicked(this, &STableTreeView::OnAdvancedFiltersClicked)
+					.IsEnabled(this, &STableTreeView::AdvancedFilters_ShouldBeEnabled)
 				]
 			]
 
@@ -786,7 +787,14 @@ void STableTreeView::ApplyFiltering()
 	Stopwatch.Start();
 
 	// Apply filter to all groups and its children.
-	ApplyFilteringForNode(Root);
+	if (FilterConfigurator_HasFilters())
+	{
+		ApplyAdvancedFiltersForNode(Root);
+	}
+	else
+	{
+		ApplyHierarchicalFilterForNode(Root);
+	}
 
 	FilteredGroupNodes.Reset();
 	const TArray<FBaseTreeNodePtr>& RootChildren = Root->GetFilteredChildren();
@@ -851,12 +859,19 @@ void STableTreeView::ApplyFiltering()
 
 ////////////////////////////////////////////////////////////////////////////////////////////////////
 
-bool STableTreeView::ApplyFilteringForNode(FTableTreeNodePtr NodePtr)
+bool STableTreeView::ApplyHierarchicalFilterForNode(FTableTreeNodePtr NodePtr)
 {
 	bool bIsNodeVisible = Filters->PassesAllFilters(NodePtr);
 
 	if (NodePtr->IsGroup())
 	{
+		// If a group node passes the filter, all child nodes will be shown
+		if (bIsNodeVisible)
+		{
+			MakeSubtreeVisible(NodePtr);
+			return true;
+		}
+
 		NodePtr->ClearFilteredChildren();
 
 		const TArray<FBaseTreeNodePtr>& GroupChildren = NodePtr->GetChildren();
@@ -870,7 +885,7 @@ bool STableTreeView::ApplyFilteringForNode(FTableTreeNodePtr NodePtr)
 			}
 			// Add a child.
 			const FTableTreeNodePtr& ChildNodePtr = StaticCastSharedPtr<FTableTreeNode>(GroupChildren[Cx]);
-			if (ApplyFilteringForNode(ChildNodePtr))
+			if (ApplyHierarchicalFilterForNode(ChildNodePtr))
 			{
 				NodePtr->AddFilteredChild(ChildNodePtr);
 				NumVisibleChildren++;
@@ -893,7 +908,6 @@ bool STableTreeView::ApplyFilteringForNode(FTableTreeNodePtr NodePtr)
 	}
 	else
 	{
-		bIsNodeVisible &= ApplyAdvancedFilters(NodePtr);
 		NodePtr->SetIsFiltered(!bIsNodeVisible);
 		return bIsNodeVisible;
 	}
@@ -901,6 +915,82 @@ bool STableTreeView::ApplyFilteringForNode(FTableTreeNodePtr NodePtr)
 
 ////////////////////////////////////////////////////////////////////////////////////////////////////
 
+bool STableTreeView::ApplyAdvancedFiltersForNode(FTableTreeNodePtr NodePtr)
+{
+	if (NodePtr->IsGroup())
+	{
+		NodePtr->ClearFilteredChildren();
+
+		const TArray<FBaseTreeNodePtr>& GroupChildren = NodePtr->GetChildren();
+		const int32 NumChildren = GroupChildren.Num();
+		int32 NumVisibleChildren = 0;
+		for (int32 Cx = 0; Cx < NumChildren; ++Cx)
+		{
+			if (bCancelCurrentAsyncOp)
+			{
+				break;
+			}
+			// Add a child.
+			const FTableTreeNodePtr& ChildNodePtr = StaticCastSharedPtr<FTableTreeNode>(GroupChildren[Cx]);
+			if (ApplyAdvancedFiltersForNode(ChildNodePtr))
+			{
+				NodePtr->AddFilteredChild(ChildNodePtr);
+				NumVisibleChildren++;
+			}
+		}
+
+		const bool bIsGroupNodeVisible = NumVisibleChildren > 0;
+
+		if (bIsGroupNodeVisible)
+		{
+			// Add a group.
+			NodePtr->SetExpansion(true);
+		}
+		else
+		{
+			NodePtr->SetExpansion(false);
+		}
+
+		return bIsGroupNodeVisible;
+	}
+	else
+	{
+		bool bIsNodeVisible = ApplyAdvancedFilters(NodePtr);
+		NodePtr->SetIsFiltered(!bIsNodeVisible);
+		return bIsNodeVisible;
+	}
+}
+
+////////////////////////////////////////////////////////////////////////////////////////////////////
+
+void STableTreeView::MakeSubtreeVisible(FTableTreeNodePtr NodePtr)
+{
+	if (NodePtr->IsGroup())
+	{
+		NodePtr->ClearFilteredChildren();
+
+		const TArray<FBaseTreeNodePtr>& GroupChildren = NodePtr->GetChildren();
+		const int32 NumChildren = GroupChildren.Num();
+		int32 NumVisibleChildren = 0;
+		for (int32 Cx = 0; Cx < NumChildren; ++Cx)
+		{
+			if (bCancelCurrentAsyncOp)
+			{
+				break;
+			}
+
+			const FTableTreeNodePtr& ChildNodePtr = StaticCastSharedPtr<FTableTreeNode>(GroupChildren[Cx]);
+			MakeSubtreeVisible(ChildNodePtr);
+			NodePtr->AddFilteredChild(ChildNodePtr);
+			NodePtr->SetExpansion(true);
+		}
+	}
+
+	NodePtr->SetIsFiltered(false);
+}
+
+////////////////////////////////////////////////////////////////////////////////////////////////////
+
 void STableTreeView::HandleItemToStringArray(const FTableTreeNodePtr& FTableTreeNodePtr, TArray<FString>& OutSearchStrings)
 {
 	OutSearchStrings.Add(FTableTreeNodePtr->GetName().GetPlainNameString());
@@ -1054,7 +1144,19 @@ void STableTreeView::SearchBox_OnTextChanged(const FText& InFilterText)
 
 bool STableTreeView::SearchBox_IsEnabled() const
 {
-	return TableTreeNodes.Num() > 0;
+	return TableTreeNodes.Num() > 0 && !FilterConfigurator_HasFilters();
+}
+
+////////////////////////////////////////////////////////////////////////////////////////////////////
+
+FText STableTreeView::SearchBox_GetTooltipText() const
+{
+	if (SearchBox_IsEnabled())
+	{
+		return LOCTEXT("SearchBox_ToolTip", "Type here to search the tree hierarchy by item or group name");
+	}
+
+	return LOCTEXT("SearchBox_Disabled_ToolTip", "Searching the tree hierarchy is disabled when advanced filters are set or when there are no entries in the table");
 }
 
 ////////////////////////////////////////////////////////////////////////////////////////////////////
@@ -2599,6 +2701,32 @@ void STableTreeView::OnAdvancedFiltersChangesCommited()
 
 ////////////////////////////////////////////////////////////////////////////////////////////////////
 
+bool STableTreeView::AdvancedFilters_ShouldBeEnabled() const
+{
+	return TextFilter->GetRawFilterText().IsEmpty();
+}
+
+////////////////////////////////////////////////////////////////////////////////////////////////////
+
+FText STableTreeView::AdvancedFilters_GetTooltipText() const
+{
+	if (AdvancedFilters_ShouldBeEnabled())
+	{
+		return LOCTEXT("AdvancedFiltersBtn_ToolTip", "Opens the filter configurator window.");
+	}
+	
+	return LOCTEXT("AdvancedFiltersBtn_Disabled_ToolTip", "Advanced filters cannot be added when filters are already applied using the search box.");
+}
+
+////////////////////////////////////////////////////////////////////////////////////////////////////
+
+bool STableTreeView::FilterConfigurator_HasFilters() const
+{
+	return FilterConfigurator.IsValid() && FilterConfigurator->GetRootNode()->GetChildren().Num() > 0;
+}
+
+////////////////////////////////////////////////////////////////////////////////////////////////////
+
 TSharedPtr<SWidget> STableTreeView::ConstructToolbar()
 {
 	return nullptr;
diff --git a/Engine/Source/Developer/TraceInsights/Private/Insights/Table/Widgets/STableTreeView.h b/Engine/Source/Developer/TraceInsights/Private/Insights/Table/Widgets/STableTreeView.h
--- a/Engine/Source/Developer/TraceInsights/Private/Insights/Table/Widgets/STableTreeView.h
+++ b/Engine/Source/Developer/TraceInsights/Private/Insights/Table/Widgets/STableTreeView.h
@@ -213,11 +213,17 @@ protected:
 	/** Populates the group and stat tree with items based on the current data. */
 	void ApplyFiltering();
 
-	bool ApplyFilteringForNode(FTableTreeNodePtr NodePtr);
+	bool ApplyAdvancedFiltersForNode(FTableTreeNodePtr NodePtr);
+
+	bool ApplyHierarchicalFilterForNode(FTableTreeNodePtr NodePtr);
+
+	void MakeSubtreeVisible(FTableTreeNodePtr NodePtr);
 
 	bool SearchBox_IsEnabled() const;
 	void SearchBox_OnTextChanged(const FText& InFilterText);
 
+	FText SearchBox_GetTooltipText() const;
+
 	////////////////////////////////////////////////////////////////////////////////////////////////////
 	// Grouping
 
@@ -362,6 +368,9 @@ protected:
 	void CancelCurrentAsyncOp();
 
 	FReply OnAdvancedFiltersClicked();
+	bool AdvancedFilters_ShouldBeEnabled() const;
+	FText AdvancedFilters_GetTooltipText() const;
+	bool FilterConfigurator_HasFilters() const;
 	void OnAdvancedFiltersChangesCommited();
 	bool ApplyAdvancedFilters(const FTableTreeNodePtr& NodePtr);
 	bool virtual ApplyCustomAdvancedFilters(const FTableTreeNodePtr& NodePtr) { return true; };
