Fix static analysis error

diff --git a/Engine/Source/Developer/TraceInsights/Private/Insights/MemoryProfiler/Widgets/SMemoryProfilerWindow.cpp b/Engine/Source/Developer/TraceInsights/Private/Insights/MemoryProfiler/Widgets/SMemoryProfilerWindow.cpp
--- a/Engine/Source/Developer/TraceInsights/Private/Insights/MemoryProfiler/Widgets/SMemoryProfilerWindow.cpp
+++ b/Engine/Source/Developer/TraceInsights/Private/Insights/MemoryProfiler/Widgets/SMemoryProfilerWindow.cpp
@@ -428,7 +428,7 @@ TSharedPtr<Insights::SMemAllocTableTreeView> SMemoryProfilerWindow::ShowMemAlloc
 {
 	if (SharedState->GetCurrentQueryTarget()->GetName() == Insights::FQueryTargetWindowSpec::NewWindow)
 	{
-		LastMemAllocTableTreeViewIndex = ++LastMemAllocTableTreeViewIndex;
+		++LastMemAllocTableTreeViewIndex;
 		FName TabId = FMemoryProfilerTabs::MemAllocTableTreeViewID;
 		TabId.SetNumber(LastMemAllocTableTreeViewIndex);
 
