[Insights] Fix crash when using the filter configurator in a tabletree view set to non async workflow.#rb none

diff --git a/Engine/Source/Developer/TraceInsights/Private/Insights/Table/Widgets/STableTreeView.cpp b/Engine/Source/Developer/TraceInsights/Private/Insights/Table/Widgets/STableTreeView.cpp
--- a/Engine/Source/Developer/TraceInsights/Private/Insights/Table/Widgets/STableTreeView.cpp
+++ b/Engine/Source/Developer/TraceInsights/Private/Insights/Table/Widgets/STableTreeView.cpp
@@ -2638,12 +2638,14 @@ FReply STableTreeView::OnAdvancedFiltersClicked()
 
 bool STableTreeView::ApplyAdvancedFilters(const FTableTreeNodePtr& NodePtr)
 {
-	if (CurrentAsyncOpFilterConfigurator == nullptr)
+	FFilterConfigurator* FilterConfiguratorToUse = bRunInAsyncMode ? CurrentAsyncOpFilterConfigurator : FilterConfigurator.Get();
+
+	if (FilterConfiguratorToUse == nullptr)
 	{
 		return true;
 	}
 
-	if (CurrentAsyncOpFilterConfigurator->GetRootNode()->GetChildren().Num() == 0)
+	if (FilterConfiguratorToUse->GetRootNode()->GetChildren().Num() == 0)
 	{
 		return true;
 	}
@@ -2675,27 +2677,28 @@ bool STableTreeView::ApplyAdvancedFilters(const FTableTreeNodePtr& NodePtr)
 		}
 	}
 
-	return ApplyCustomAdvancedFilters(NodePtr) && CurrentAsyncOpFilterConfigurator->ApplyFilters(Context);
+	return ApplyCustomAdvancedFilters(NodePtr) && FilterConfiguratorToUse->ApplyFilters(Context);
 }
 
 ////////////////////////////////////////////////////////////////////////////////////////////////////
 
 void STableTreeView::OnAdvancedFiltersChangesCommited()
 {
-	// If settings have not changed, do nothing.
-	if (*FilterConfigurator == *CurrentAsyncOpFilterConfigurator)
-	{
-		return;
-	}
-
-	if (!bIsUpdateRunning)
+	if (bRunInAsyncMode)
 	{
-		OnPreAsyncUpdate();
-		InProgressAsyncOperationEvent = StartApplyFiltersTask();
+		if (!bIsUpdateRunning)
+		{
+			OnPreAsyncUpdate();
+			InProgressAsyncOperationEvent = StartApplyFiltersTask();
+		}
+		else
+		{
+			CancelCurrentAsyncOp();
+		}
 	}
 	else
 	{
-		CancelCurrentAsyncOp();
+		ApplyFiltering();
 	}
 }
 
diff --git a/Engine/Source/Developer/TraceInsights/Private/Insights/ViewModels/FilterConfigurator.cpp b/Engine/Source/Developer/TraceInsights/Private/Insights/ViewModels/FilterConfigurator.cpp
--- a/Engine/Source/Developer/TraceInsights/Private/Insights/ViewModels/FilterConfigurator.cpp
+++ b/Engine/Source/Developer/TraceInsights/Private/Insights/ViewModels/FilterConfigurator.cpp
@@ -34,15 +34,20 @@ FFilterConfigurator::FFilterConfigurator(const FFilterConfigurator& Other)
 
 FFilterConfigurator& FFilterConfigurator::operator=(const FFilterConfigurator& Other)
 {
+	if (*this == Other)
+	{
+		return *this;
+	}
+
 	RootNode = MakeShared<FFilterConfiguratorNode>(*Other.RootNode);
 	RootNode->SetGroupPtrForChildren();
 	AvailableFilters = Other.AvailableFilters;
 	OnDestroyedEvent = Other.OnDestroyedEvent;
 
-	OnChangesCommitedEvent.Broadcast();
-
 	RootNode->ProcessFilter();
 
+	OnChangesCommitedEvent.Broadcast();
+
 	return *this;
 }
 
@@ -51,7 +56,11 @@ FFilterConfigurator& FFilterConfigurator::operator=(const FFilterConfigurator& O
 bool FFilterConfigurator::operator==(const FFilterConfigurator& Other)
 {
 	bool bIsEqual = AvailableFilters.Get() == Other.AvailableFilters.Get();
-	bIsEqual &= *RootNode == *Other.RootNode;
+	bIsEqual &= RootNode.IsValid() == Other.RootNode.IsValid();
+	if (RootNode.IsValid() && Other.RootNode.IsValid())
+	{
+		bIsEqual &= *RootNode == *Other.RootNode;
+	}
 
 	return bIsEqual;
 }
