[Insights] Fix MemAlloc TreeView Tabs spawning as "Unregistered Tab" when the Memory Insights window is reopened.
Fix memory leak of MemoryProfilerWindow caused by circular dependency.

diff --git a/Engine/Source/Developer/TraceInsights/Private/Insights/MemoryProfiler/MemoryProfilerManager.cpp b/Engine/Source/Developer/TraceInsights/Private/Insights/MemoryProfiler/MemoryProfilerManager.cpp
--- a/Engine/Source/Developer/TraceInsights/Private/Insights/MemoryProfiler/MemoryProfilerManager.cpp
+++ b/Engine/Source/Developer/TraceInsights/Private/Insights/MemoryProfiler/MemoryProfilerManager.cpp
@@ -195,6 +195,14 @@ bool FMemoryProfilerManager::CanSpawnTab(const FSpawnTabArgs& Args) const
 
 void FMemoryProfilerManager::OnTabClosed(TSharedRef<SDockTab> TabBeingClosed)
 {
+	{
+		TSharedPtr<SMemoryProfilerWindow> ProfilerWindowSharedPtr = ProfilerWindow.Pin();
+		if (ProfilerWindowSharedPtr.IsValid())
+		{
+			ProfilerWindowSharedPtr->CloseMemAllocTableTreeTabs();
+		}
+	}
+
 	RemoveProfilerWindow();
 
 	// Disable TabClosed delegate.
diff --git a/Engine/Source/Developer/TraceInsights/Private/Insights/MemoryProfiler/Widgets/SMemInvestigationView.cpp b/Engine/Source/Developer/TraceInsights/Private/Insights/MemoryProfiler/Widgets/SMemInvestigationView.cpp
--- a/Engine/Source/Developer/TraceInsights/Private/Insights/MemoryProfiler/Widgets/SMemInvestigationView.cpp
+++ b/Engine/Source/Developer/TraceInsights/Private/Insights/MemoryProfiler/Widgets/SMemInvestigationView.cpp
@@ -25,7 +25,7 @@
 ////////////////////////////////////////////////////////////////////////////////////////////////////
 
 SMemInvestigationView::SMemInvestigationView()
-	: ProfilerWindow()
+	: ProfilerWindowWeakPtr()
 {
 }
 
@@ -46,7 +46,7 @@ BEGIN_SLATE_FUNCTION_BUILD_OPTIMIZATION
 void SMemInvestigationView::Construct(const FArguments& InArgs, TSharedPtr<SMemoryProfilerWindow> InProfilerWindow)
 {
 	check(InProfilerWindow.IsValid());
-	ProfilerWindow = InProfilerWindow;
+	ProfilerWindowWeakPtr = InProfilerWindow;
 
 	ChildSlot
 	[
@@ -181,6 +181,8 @@ TSharedRef<SWidget> SMemInvestigationView::ConstructInvestigationWidgetArea()
 		]
 	;
 
+	TSharedPtr<SMemoryProfilerWindow> ProfilerWindow = ProfilerWindowWeakPtr.Pin();
+
 	if (ProfilerWindow.IsValid())
 	{
 		FMemorySharedState& SharedState = ProfilerWindow->GetSharedState();
@@ -225,6 +227,8 @@ TSharedRef<SWidget> SMemInvestigationView::QueryRule_OnGenerateWidget(TSharedPtr
 
 TSharedRef<SWidget> SMemInvestigationView::ConstructTimeMarkerWidget(uint32 TimeMarkerIndex)
 {
+	TSharedPtr<SMemoryProfilerWindow> ProfilerWindow = ProfilerWindowWeakPtr.Pin();
+
 	if (!ProfilerWindow.IsValid())
 	{
 		return SNew(SBox);
@@ -321,6 +325,8 @@ void SMemInvestigationView::Tick(const FGeometry& AllottedGeometry, const double
 
 const TArray<TSharedPtr<Insights::FMemoryRuleSpec>>* SMemInvestigationView::GetAvailableQueryRules()
 {
+	TSharedPtr<SMemoryProfilerWindow> ProfilerWindow = ProfilerWindowWeakPtr.Pin();
+
 	if (ProfilerWindow.IsValid())
 	{
 		FMemorySharedState& SharedState = ProfilerWindow->GetSharedState();
@@ -335,6 +341,8 @@ void SMemInvestigationView::QueryRule_OnSelectionChanged(TSharedPtr<Insights::FM
 {
 	if (SelectInfo != ESelectInfo::Direct)
 	{
+		TSharedPtr<SMemoryProfilerWindow> ProfilerWindow = ProfilerWindowWeakPtr.Pin();
+
 		if (ProfilerWindow.IsValid() && InRule)
 		{
 			FMemorySharedState& SharedState = ProfilerWindow->GetSharedState();
@@ -348,6 +356,8 @@ void SMemInvestigationView::QueryRule_OnSelectionChanged(TSharedPtr<Insights::FM
 
 FText SMemInvestigationView::QueryRule_GetSelectedText() const
 {
+	TSharedPtr<SMemoryProfilerWindow> ProfilerWindow = ProfilerWindowWeakPtr.Pin();
+
 	if (ProfilerWindow.IsValid())
 	{
 		FMemorySharedState& SharedState = ProfilerWindow->GetSharedState();
@@ -364,6 +374,8 @@ FText SMemInvestigationView::QueryRule_GetSelectedText() const
 
 FText SMemInvestigationView::QueryRule_GetTooltipText() const
 {
+	TSharedPtr<SMemoryProfilerWindow> ProfilerWindow = ProfilerWindowWeakPtr.Pin();
+
 	if (ProfilerWindow.IsValid())
 	{
 		FMemorySharedState& SharedState = ProfilerWindow->GetSharedState();
@@ -380,6 +392,8 @@ FText SMemInvestigationView::QueryRule_GetTooltipText() const
 
 const TArray<TSharedPtr<Insights::FQueryTargetWindowSpec>>* SMemInvestigationView::GetAvailableQueryTargets()
 {
+	TSharedPtr<SMemoryProfilerWindow> ProfilerWindow = ProfilerWindowWeakPtr.Pin();
+
 	if (ProfilerWindow.IsValid())
 	{
 		FMemorySharedState& SharedState = ProfilerWindow->GetSharedState();
@@ -394,6 +408,8 @@ void SMemInvestigationView::QueryTarget_OnSelectionChanged(TSharedPtr<Insights::
 {
 	if (SelectInfo != ESelectInfo::Type::Direct)
 	{
+		TSharedPtr<SMemoryProfilerWindow> ProfilerWindow = ProfilerWindowWeakPtr.Pin();
+
 		if (ProfilerWindow.IsValid() && InTarget)
 		{
 			FMemorySharedState& SharedState = ProfilerWindow->GetSharedState();
@@ -438,6 +454,8 @@ TSharedRef<SWidget> SMemInvestigationView::QueryTarget_OnGenerateWidget(TSharedP
 
 FText SMemInvestigationView::QueryTarget_GetSelectedText() const
 {
+	TSharedPtr<SMemoryProfilerWindow> ProfilerWindow = ProfilerWindowWeakPtr.Pin();
+
 	if (ProfilerWindow.IsValid())
 	{
 		FMemorySharedState& SharedState = ProfilerWindow->GetSharedState();
@@ -454,6 +472,8 @@ FText SMemInvestigationView::QueryTarget_GetSelectedText() const
 
 FReply SMemInvestigationView::RunQuery()
 {
+	TSharedPtr<SMemoryProfilerWindow> ProfilerWindow = ProfilerWindowWeakPtr.Pin();
+
 	if (!ProfilerWindow.IsValid())
 	{
 		UE_LOG(MemoryProfiler, Error, TEXT("[MemQuery] Invalid Profiler Window!"));
diff --git a/Engine/Source/Developer/TraceInsights/Private/Insights/MemoryProfiler/Widgets/SMemInvestigationView.h b/Engine/Source/Developer/TraceInsights/Private/Insights/MemoryProfiler/Widgets/SMemInvestigationView.h
--- a/Engine/Source/Developer/TraceInsights/Private/Insights/MemoryProfiler/Widgets/SMemInvestigationView.h
+++ b/Engine/Source/Developer/TraceInsights/Private/Insights/MemoryProfiler/Widgets/SMemInvestigationView.h
@@ -78,7 +78,7 @@ private:
 	virtual void Tick(const FGeometry& AllottedGeometry, const double InCurrentTime, const float InDeltaTime) override;
 
 private:
-	TSharedPtr<SMemoryProfilerWindow> ProfilerWindow;
+	TWeakPtr<SMemoryProfilerWindow> ProfilerWindowWeakPtr;
 
 	/** A weak pointer to the profiler session used to populate this widget. */
 	TSharedPtr<const TraceServices::IAnalysisSession>/*Weak*/ Session;
diff --git a/Engine/Source/Developer/TraceInsights/Private/Insights/MemoryProfiler/Widgets/SMemTagTreeView.cpp b/Engine/Source/Developer/TraceInsights/Private/Insights/MemoryProfiler/Widgets/SMemTagTreeView.cpp
--- a/Engine/Source/Developer/TraceInsights/Private/Insights/MemoryProfiler/Widgets/SMemTagTreeView.cpp
+++ b/Engine/Source/Developer/TraceInsights/Private/Insights/MemoryProfiler/Widgets/SMemTagTreeView.cpp
@@ -36,7 +36,7 @@
 ////////////////////////////////////////////////////////////////////////////////////////////////////
 
 SMemTagTreeView::SMemTagTreeView()
-	: ProfilerWindow()
+	: ProfilerWindowWeakPtr()
 	, Table(MakeShared<Insights::FTable>())
 	, bExpansionSaved(false)
 	, bFilterOutZeroCountMemTags(false)
@@ -68,7 +68,7 @@ BEGIN_SLATE_FUNCTION_BUILD_OPTIMIZATION
 void SMemTagTreeView::Construct(const FArguments& InArgs, TSharedPtr<SMemoryProfilerWindow> InProfilerWindow)
 {
 	check(InProfilerWindow.IsValid());
-	ProfilerWindow = InProfilerWindow;
+	ProfilerWindowWeakPtr = InProfilerWindow;
 
 	SAssignNew(ExternalScrollbar, SScrollBar)
 	.AlwaysShowScrollbar(true);
@@ -966,6 +966,8 @@ void SMemTagTreeView::ApplyFiltering()
 {
 	FilteredGroupNodes.Reset();
 
+	TSharedPtr<SMemoryProfilerWindow> ProfilerWindow = ProfilerWindowWeakPtr.Pin();
+
 	uint64 TrackerFilterMask = uint64(-1);
 	if (ProfilerWindow.IsValid())
 	{
@@ -1134,6 +1136,8 @@ void SMemTagTreeView::TreeView_OnMouseButtonDoubleClick(FMemTagNodePtr MemTagNod
 	}
 	else
 	{
+		TSharedPtr<SMemoryProfilerWindow> ProfilerWindow = ProfilerWindowWeakPtr.Pin();
+
 		if (ProfilerWindow.IsValid())
 		{
 			FMemorySharedState& SharedState = ProfilerWindow->GetSharedState();
@@ -1882,6 +1886,8 @@ void SMemTagTreeView::RebuildTree(bool bResync)
 	const uint32 PreviousNodeCount = MemTagNodes.Num();
 
 	SyncStopwatch.Start();
+	TSharedPtr<SMemoryProfilerWindow> ProfilerWindow = ProfilerWindowWeakPtr.Pin();
+
 	if (ProfilerWindow.IsValid())
 	{
 		FMemorySharedState& SharedState = ProfilerWindow->GetSharedState();
@@ -2070,6 +2076,8 @@ void SMemTagTreeView::SelectMemTagNode(Insights::FMemoryTagId MemTagId)
 
 const TArray<TSharedPtr<Insights::FMemoryTracker>>* SMemTagTreeView::GetAvailableTrackers()
 {
+	TSharedPtr<SMemoryProfilerWindow> ProfilerWindow = ProfilerWindowWeakPtr.Pin();
+
 	if (ProfilerWindow.IsValid())
 	{
 		FMemorySharedState& SharedState = ProfilerWindow->GetSharedState();
@@ -2084,6 +2092,8 @@ void SMemTagTreeView::Tracker_OnSelectionChanged(TSharedPtr<Insights::FMemoryTra
 {
 	if (SelectInfo != ESelectInfo::Direct)
 	{
+		TSharedPtr<SMemoryProfilerWindow> ProfilerWindow = ProfilerWindowWeakPtr.Pin();
+
 		if (ProfilerWindow.IsValid())
 		{
 			FMemorySharedState& SharedState = ProfilerWindow->GetSharedState();
@@ -2119,6 +2129,8 @@ void SMemTagTreeView::Tracker_OnCheckStateChanged(ECheckBoxState CheckType, TSha
 {
 	if (CheckType == ECheckBoxState::Checked)
 	{
+		TSharedPtr<SMemoryProfilerWindow> ProfilerWindow = ProfilerWindowWeakPtr.Pin();
+
 		if (ProfilerWindow.IsValid())
 		{
 			FMemorySharedState& SharedState = ProfilerWindow->GetSharedState();
@@ -2137,6 +2149,8 @@ void SMemTagTreeView::Tracker_OnCheckStateChanged(ECheckBoxState CheckType, TSha
 ECheckBoxState SMemTagTreeView::Tracker_IsChecked(TSharedPtr<Insights::FMemoryTracker> InTracker) const
 {
 	TSharedPtr<Insights::FMemoryTracker> CurrentTracker = nullptr;
+	TSharedPtr<SMemoryProfilerWindow> ProfilerWindow = ProfilerWindowWeakPtr.Pin();
+
 	if (ProfilerWindow.IsValid())
 	{
 		FMemorySharedState& SharedState = ProfilerWindow->GetSharedState();
@@ -2149,6 +2163,8 @@ ECheckBoxState SMemTagTreeView::Tracker_IsChecked(TSharedPtr<Insights::FMemoryTr
 
 FText SMemTagTreeView::Tracker_GetSelectedText() const
 {
+	TSharedPtr<SMemoryProfilerWindow> ProfilerWindow = ProfilerWindowWeakPtr.Pin();
+
 	if (ProfilerWindow.IsValid())
 	{
 		FMemorySharedState& SharedState = ProfilerWindow->GetSharedState();
@@ -2173,6 +2189,8 @@ FText SMemTagTreeView::Tracker_GetTooltipText() const
 
 FReply SMemTagTreeView::LoadReportXML_OnClicked()
 {
+	TSharedPtr<SMemoryProfilerWindow> ProfilerWindow = ProfilerWindowWeakPtr.Pin();
+
 	if (ProfilerWindow.IsValid())
 	{
 		TArray<FString> OutFiles;
@@ -2231,6 +2249,8 @@ FReply SMemTagTreeView::HideAllTracks_OnClicked()
 
 FReply SMemTagTreeView::AllTracksSmallHeight_OnClicked()
 {
+	TSharedPtr<SMemoryProfilerWindow> ProfilerWindow = ProfilerWindowWeakPtr.Pin();
+
 	if (ProfilerWindow.IsValid())
 	{
 		FMemorySharedState& SharedState = ProfilerWindow->GetSharedState();
@@ -2244,6 +2264,8 @@ FReply SMemTagTreeView::AllTracksSmallHeight_OnClicked()
 
 FReply SMemTagTreeView::AllTracksMediumHeight_OnClicked()
 {
+	TSharedPtr<SMemoryProfilerWindow> ProfilerWindow = ProfilerWindowWeakPtr.Pin();
+
 	if (ProfilerWindow.IsValid())
 	{
 		FMemorySharedState& SharedState = ProfilerWindow->GetSharedState();
@@ -2257,6 +2279,8 @@ FReply SMemTagTreeView::AllTracksMediumHeight_OnClicked()
 
 FReply SMemTagTreeView::AllTracksLargeHeight_OnClicked()
 {
+	TSharedPtr<SMemoryProfilerWindow> ProfilerWindow = ProfilerWindowWeakPtr.Pin();
+
 	if (ProfilerWindow.IsValid())
 	{
 		FMemorySharedState& SharedState = ProfilerWindow->GetSharedState();
@@ -2272,6 +2296,8 @@ FReply SMemTagTreeView::AllTracksLargeHeight_OnClicked()
 
 bool SMemTagTreeView::CanCreateGraphTracksForSelectedMemTags() const
 {
+	TSharedPtr<SMemoryProfilerWindow> ProfilerWindow = ProfilerWindowWeakPtr.Pin();
+
 	return ProfilerWindow.IsValid() && TreeView->GetNumItemsSelected() > 0;
 }
 
@@ -2279,6 +2305,8 @@ bool SMemTagTreeView::CanCreateGraphTracksForSelectedMemTags() const
 
 void SMemTagTreeView::CreateGraphTracksForSelectedMemTags()
 {
+	TSharedPtr<SMemoryProfilerWindow> ProfilerWindow = ProfilerWindowWeakPtr.Pin();
+
 	if (ProfilerWindow.IsValid())
 	{
 		FMemorySharedState& SharedState = ProfilerWindow->GetSharedState();
@@ -2310,6 +2338,8 @@ void SMemTagTreeView::CreateGraphTracksForSelectedMemTags()
 
 bool SMemTagTreeView::CanCreateGraphTracksForFilteredMemTags() const
 {
+	TSharedPtr<SMemoryProfilerWindow> ProfilerWindow = ProfilerWindowWeakPtr.Pin();
+
 	if (!ProfilerWindow.IsValid())
 	{
 		return false;
@@ -2326,6 +2356,8 @@ bool SMemTagTreeView::CanCreateGraphTracksForFilteredMemTags() const
 
 void SMemTagTreeView::CreateGraphTracksForFilteredMemTags()
 {
+	TSharedPtr<SMemoryProfilerWindow> ProfilerWindow = ProfilerWindowWeakPtr.Pin();
+
 	if (ProfilerWindow.IsValid())
 	{
 		FMemorySharedState& SharedState = ProfilerWindow->GetSharedState();
@@ -2348,6 +2380,8 @@ void SMemTagTreeView::CreateGraphTracksForFilteredMemTags()
 
 bool SMemTagTreeView::CanCreateAllGraphTracks() const
 {
+	TSharedPtr<SMemoryProfilerWindow> ProfilerWindow = ProfilerWindowWeakPtr.Pin();
+
 	return ProfilerWindow.IsValid();
 }
 
@@ -2355,6 +2389,8 @@ bool SMemTagTreeView::CanCreateAllGraphTracks() const
 
 void SMemTagTreeView::CreateAllGraphTracks()
 {
+	TSharedPtr<SMemoryProfilerWindow> ProfilerWindow = ProfilerWindowWeakPtr.Pin();
+
 	if (ProfilerWindow.IsValid())
 	{
 		FMemorySharedState& SharedState = ProfilerWindow->GetSharedState();
@@ -2372,6 +2408,8 @@ void SMemTagTreeView::CreateAllGraphTracks()
 
 bool SMemTagTreeView::CanRemoveGraphTracksForSelectedMemTags() const
 {
+	TSharedPtr<SMemoryProfilerWindow> ProfilerWindow = ProfilerWindowWeakPtr.Pin();
+
 	return ProfilerWindow.IsValid() && TreeView->GetNumItemsSelected() > 0;
 }
 
@@ -2379,6 +2417,8 @@ bool SMemTagTreeView::CanRemoveGraphTracksForSelectedMemTags() const
 
 void SMemTagTreeView::RemoveGraphTracksForSelectedMemTags()
 {
+	TSharedPtr<SMemoryProfilerWindow> ProfilerWindow = ProfilerWindowWeakPtr.Pin();
+
 	if (ProfilerWindow.IsValid())
 	{
 		FMemorySharedState& SharedState = ProfilerWindow->GetSharedState();
@@ -2410,6 +2450,8 @@ void SMemTagTreeView::RemoveGraphTracksForSelectedMemTags()
 
 bool SMemTagTreeView::CanRemoveGraphTracksForUnusedMemTags() const
 {
+	TSharedPtr<SMemoryProfilerWindow> ProfilerWindow = ProfilerWindowWeakPtr.Pin();
+
 	return ProfilerWindow.IsValid();
 }
 
@@ -2417,6 +2459,8 @@ bool SMemTagTreeView::CanRemoveGraphTracksForUnusedMemTags() const
 
 void SMemTagTreeView::RemoveGraphTracksForUnusedMemTags()
 {
+	TSharedPtr<SMemoryProfilerWindow> ProfilerWindow = ProfilerWindowWeakPtr.Pin();
+
 	if (ProfilerWindow.IsValid())
 	{
 		FMemorySharedState& SharedState = ProfilerWindow->GetSharedState();
@@ -2430,6 +2474,8 @@ void SMemTagTreeView::RemoveGraphTracksForUnusedMemTags()
 
 bool SMemTagTreeView::CanRemoveAllGraphTracks() const
 {
+	TSharedPtr<SMemoryProfilerWindow> ProfilerWindow = ProfilerWindowWeakPtr.Pin();
+
 	return ProfilerWindow.IsValid();
 }
 
@@ -2437,6 +2483,8 @@ bool SMemTagTreeView::CanRemoveAllGraphTracks() const
 
 void SMemTagTreeView::RemoveAllGraphTracks()
 {
+	TSharedPtr<SMemoryProfilerWindow> ProfilerWindow = ProfilerWindowWeakPtr.Pin();
+
 	if (ProfilerWindow.IsValid())
 	{
 		FMemorySharedState& SharedState = ProfilerWindow->GetSharedState();
@@ -2450,6 +2498,8 @@ void SMemTagTreeView::RemoveAllGraphTracks()
 
 bool SMemTagTreeView::CanGenerateColorForSelectedMemTags() const
 {
+	TSharedPtr<SMemoryProfilerWindow> ProfilerWindow = ProfilerWindowWeakPtr.Pin();
+
 	return ProfilerWindow.IsValid() && TreeView->GetNumItemsSelected() > 0;
 }
 
@@ -2457,6 +2507,8 @@ bool SMemTagTreeView::CanGenerateColorForSelectedMemTags() const
 
 void SMemTagTreeView::GenerateColorForSelectedMemTags()
 {
+	TSharedPtr<SMemoryProfilerWindow> ProfilerWindow = ProfilerWindowWeakPtr.Pin();
+
 	if (ProfilerWindow.IsValid())
 	{
 		const TArray<FMemTagNodePtr> SelectedNodes = TreeView->GetSelectedItems();
@@ -2503,6 +2555,8 @@ void SMemTagTreeView::SetColorToNode(const FMemTagNodePtr& MemTagNode, FLinearCo
 
 	const Insights::FMemoryTagId MemTagId = MemTagNode->GetMemTagId();
 
+	TSharedPtr<SMemoryProfilerWindow> ProfilerWindow = ProfilerWindowWeakPtr.Pin();
+
 	FMemorySharedState& SharedState = ProfilerWindow->GetSharedState();
 
 	TSharedPtr<FMemoryGraphTrack> MainGraphTrack = SharedState.GetMainGraphTrack();
@@ -2533,6 +2587,8 @@ void SMemTagTreeView::SetColorToNode(const FMemTagNodePtr& MemTagNode, FLinearCo
 
 bool SMemTagTreeView::CanEditColorForSelectedMemTags() const
 {
+	TSharedPtr<SMemoryProfilerWindow> ProfilerWindow = ProfilerWindowWeakPtr.Pin();
+
 	return ProfilerWindow.IsValid() && TreeView->GetNumItemsSelected() > 0;
 }
 
@@ -2540,6 +2596,8 @@ bool SMemTagTreeView::CanEditColorForSelectedMemTags() const
 
 void SMemTagTreeView::EditColorForSelectedMemTags()
 {
+	TSharedPtr<SMemoryProfilerWindow> ProfilerWindow = ProfilerWindowWeakPtr.Pin();
+
 	if (ProfilerWindow.IsValid())
 	{
 		EditableColorValue = FLinearColor(0.5f, 0.5f, 0.5f, 1.0f);
diff --git a/Engine/Source/Developer/TraceInsights/Private/Insights/MemoryProfiler/Widgets/SMemTagTreeView.h b/Engine/Source/Developer/TraceInsights/Private/Insights/MemoryProfiler/Widgets/SMemTagTreeView.h
--- a/Engine/Source/Developer/TraceInsights/Private/Insights/MemoryProfiler/Widgets/SMemTagTreeView.h
+++ b/Engine/Source/Developer/TraceInsights/Private/Insights/MemoryProfiler/Widgets/SMemTagTreeView.h
@@ -316,7 +316,7 @@ private:
 	virtual void Tick(const FGeometry& AllottedGeometry, const double InCurrentTime, const float InDeltaTime) override;
 
 private:
-	TSharedPtr<SMemoryProfilerWindow> ProfilerWindow;
+	TWeakPtr<SMemoryProfilerWindow> ProfilerWindowWeakPtr;
 
 	/** Table view model. */
 	TSharedPtr<Insights::FTable> Table;
diff --git a/Engine/Source/Developer/TraceInsights/Private/Insights/MemoryProfiler/Widgets/SMemoryProfilerWindow.cpp b/Engine/Source/Developer/TraceInsights/Private/Insights/MemoryProfiler/Widgets/SMemoryProfilerWindow.cpp
--- a/Engine/Source/Developer/TraceInsights/Private/Insights/MemoryProfiler/Widgets/SMemoryProfilerWindow.cpp
+++ b/Engine/Source/Developer/TraceInsights/Private/Insights/MemoryProfiler/Widgets/SMemoryProfilerWindow.cpp
@@ -412,9 +412,15 @@ void SMemoryProfilerWindow::OnMemAllocTableTreeViewTabClosed(TSharedRef<SDockTab
 		this->GetSharedState().RemoveQueryTarget(TargetToDelete);
 	}
 
-	TSharedPtr<Insights::FQueryTargetWindowSpec> NewSelection = this->GetSharedState().GetQueryTargets()[0];
-	this->GetSharedState().SetCurrentQueryTarget(NewSelection);
-	MemInvestigationView->QueryTarget_OnSelectionChanged(NewSelection, ESelectInfo::Type::Direct);
+	if (Targets.Num() > 0)
+	{
+		TSharedPtr<Insights::FQueryTargetWindowSpec> NewSelection = Targets[0];
+		this->GetSharedState().SetCurrentQueryTarget(NewSelection);
+		if (MemInvestigationView.IsValid())
+		{
+			MemInvestigationView->QueryTarget_OnSelectionChanged(NewSelection, ESelectInfo::Type::Direct);
+		}
+	}
 
 	TabManager->UnregisterTabSpawner(ClosingTabId);
 
@@ -424,6 +430,24 @@ void SMemoryProfilerWindow::OnMemAllocTableTreeViewTabClosed(TSharedRef<SDockTab
 
 ////////////////////////////////////////////////////////////////////////////////////////////////////
 
+void SMemoryProfilerWindow::CloseMemAllocTableTreeTabs()
+{
+	const TArray<TSharedPtr<Insights::FQueryTargetWindowSpec>>& Targets = this->GetSharedState().GetQueryTargets();
+	while(Targets.Num() > 0)
+	{
+		if (Targets[0]->GetName() != Insights::FQueryTargetWindowSpec::NewWindow)
+		{
+			HideTab(Targets[0]->GetName());
+		}
+		else
+		{
+			this->GetSharedState().RemoveQueryTarget(Targets[0]);
+		}
+	}
+}
+
+////////////////////////////////////////////////////////////////////////////////////////////////////
+
 TSharedPtr<Insights::SMemAllocTableTreeView> SMemoryProfilerWindow::ShowMemAllocTableTreeViewTab()
 {
 	if (SharedState->GetCurrentQueryTarget()->GetName() == Insights::FQueryTargetWindowSpec::NewWindow)
diff --git a/Engine/Source/Developer/TraceInsights/Private/Insights/MemoryProfiler/Widgets/SMemoryProfilerWindow.h b/Engine/Source/Developer/TraceInsights/Private/Insights/MemoryProfiler/Widgets/SMemoryProfilerWindow.h
--- a/Engine/Source/Developer/TraceInsights/Private/Insights/MemoryProfiler/Widgets/SMemoryProfilerWindow.h
+++ b/Engine/Source/Developer/TraceInsights/Private/Insights/MemoryProfiler/Widgets/SMemoryProfilerWindow.h
@@ -90,6 +90,8 @@ public:
 	FMemorySharedState& GetSharedState() { return *SharedState; }
 	const FMemorySharedState& GetSharedState() const { return *SharedState; }
 
+	void CloseMemAllocTableTreeTabs();
+
 private:
 	TSharedRef<SDockTab> SpawnTab_Toolbar(const FSpawnTabArgs& Args);
 	void OnToolbarTabClosed(TSharedRef<SDockTab> TabBeingClosed);
