[Insights] Add Expand/ExpandCriticalPath/Collapse actions for group nodes in TableTreeView.

diff --git a/Engine/Source/Developer/TraceInsights/Private/Insights/Table/Widgets/STableTreeView.cpp b/Engine/Source/Developer/TraceInsights/Private/Insights/Table/Widgets/STableTreeView.cpp
--- a/Engine/Source/Developer/TraceInsights/Private/Insights/Table/Widgets/STableTreeView.cpp
+++ b/Engine/Source/Developer/TraceInsights/Private/Insights/Table/Widgets/STableTreeView.cpp
@@ -58,17 +58,23 @@ public:
 
 	// UI_COMMAND takes long for the compiler to optimize
 	PRAGMA_DISABLE_OPTIMIZATION
-		virtual void RegisterCommands() override
+	virtual void RegisterCommands() override
 	{
-		UI_COMMAND(Command_CopyToClipboard, "Copy To Clipboard", "Copies selection to clipboard", EUserInterfaceActionType::Button, FInputChord(EModifierKey::Control, EKeys::C));
-		UI_COMMAND(Command_CopyColumnToClipboard, "Copy Column Value To Clipboard", "Copies the value of hovered column to clipboard", EUserInterfaceActionType::Button, FInputChord(EModifierKey::Control | EModifierKey::Shift, EKeys::C));
-		UI_COMMAND(Command_CopyColumnTooltipToClipboard, "Copy Column Tooltip To Clipboard", "Copies the value of hovered column's tooltip to clipboard", EUserInterfaceActionType::Button, FInputChord(EModifierKey::Control | EModifierKey::Alt, EKeys::C));
+		UI_COMMAND(Command_CopyToClipboard, "Copy To Clipboard", "Copies selection to clipboard.", EUserInterfaceActionType::Button, FInputChord(EModifierKey::Control, EKeys::C));
+		UI_COMMAND(Command_CopyColumnToClipboard, "Copy Column Value To Clipboard", "Copies the value of hovered column to clipboard.", EUserInterfaceActionType::Button, FInputChord(EModifierKey::Control | EModifierKey::Shift, EKeys::C));
+		UI_COMMAND(Command_CopyColumnTooltipToClipboard, "Copy Column Tooltip To Clipboard", "Copies the value of hovered column's tooltip to clipboard.", EUserInterfaceActionType::Button, FInputChord(EModifierKey::Control | EModifierKey::Alt, EKeys::C));
+		UI_COMMAND(Command_ExpandSubtree, "Expand Subtree", "Expand the subtree that starts from the selected group node.", EUserInterfaceActionType::Button, FInputChord(EModifierKey::None, EKeys::E));
+		UI_COMMAND(Command_ExpandCriticalPath, "Expand Critical Path", "Expand the first group child node recursively until a leaf nodes in reached.", EUserInterfaceActionType::Button, FInputChord(EModifierKey::None, EKeys::R));
+		UI_COMMAND(Command_CollapseSubtree, "Collapse Subtree", "Collapse the subtree that starts from the selected group node.", EUserInterfaceActionType::Button, FInputChord(EModifierKey::None, EKeys::C));
 	}
 	PRAGMA_ENABLE_OPTIMIZATION
 
 	TSharedPtr<FUICommandInfo> Command_CopyToClipboard;
 	TSharedPtr<FUICommandInfo> Command_CopyColumnToClipboard;
 	TSharedPtr<FUICommandInfo> Command_CopyColumnTooltipToClipboard;
+	TSharedPtr<FUICommandInfo> Command_ExpandSubtree;
+	TSharedPtr<FUICommandInfo> Command_ExpandCriticalPath;
+	TSharedPtr<FUICommandInfo> Command_CollapseSubtree;
 };
 
 const FName STableTreeView::RootNodeName(TEXT("Root"));
@@ -348,6 +354,9 @@ void STableTreeView::InitCommandList()
 	CommandList->MapAction(FTableTreeViewCommands::Get().Command_CopyToClipboard, FExecuteAction::CreateSP(this, &STableTreeView::ContextMenu_CopySelectedToClipboard_Execute), FCanExecuteAction::CreateSP(this, &STableTreeView::ContextMenu_CopySelectedToClipboard_CanExecute));
 	CommandList->MapAction(FTableTreeViewCommands::Get().Command_CopyColumnToClipboard, FExecuteAction::CreateSP(this, &STableTreeView::ContextMenu_CopyColumnToClipboard_Execute), FCanExecuteAction::CreateSP(this, &STableTreeView::ContextMenu_CopyColumnToClipboard_CanExecute));
 	CommandList->MapAction(FTableTreeViewCommands::Get().Command_CopyColumnTooltipToClipboard, FExecuteAction::CreateSP(this, &STableTreeView::ContextMenu_CopyColumnTooltipToClipboard_Execute), FCanExecuteAction::CreateSP(this, &STableTreeView::ContextMenu_CopyColumnTooltipToClipboard_CanExecute));
+	CommandList->MapAction(FTableTreeViewCommands::Get().Command_ExpandSubtree, FExecuteAction::CreateSP(this, &STableTreeView::ContextMenu_ExpandSubtree_Execute), FCanExecuteAction::CreateSP(this, &STableTreeView::ContextMenu_ExpandSubtree_CanExecute));
+	CommandList->MapAction(FTableTreeViewCommands::Get().Command_ExpandCriticalPath, FExecuteAction::CreateSP(this, &STableTreeView::ContextMenu_ExpandCriticalPath_Execute), FCanExecuteAction::CreateSP(this, &STableTreeView::ContextMenu_ExpandCriticalPath_CanExecute));
+	CommandList->MapAction(FTableTreeViewCommands::Get().Command_CollapseSubtree, FExecuteAction::CreateSP(this, &STableTreeView::ContextMenu_CollapseSubtree_Execute), FCanExecuteAction::CreateSP(this, &STableTreeView::ContextMenu_CollapseSubtree_CanExecute));
 }
 
 ////////////////////////////////////////////////////////////////////////////////////////////////////
@@ -413,6 +422,38 @@ TSharedPtr<SWidget> STableTreeView::TreeView_GetMenuContent()
 	}
 	MenuBuilder.EndSection();
 
+	MenuBuilder.BeginSection("Node", LOCTEXT("ContextMenu_Header_Node", "Node"));
+	{
+		MenuBuilder.AddMenuEntry
+		(
+			FTableTreeViewCommands::Get().Command_ExpandSubtree,
+			NAME_None,
+			TAttribute<FText>(),
+			TAttribute<FText>(),
+			FSlateIcon(FEditorStyle::GetStyleSetName(), "Profiler.EventGraph.ExpandSelection")
+		);
+
+		MenuBuilder.AddMenuEntry
+		(
+			FTableTreeViewCommands::Get().Command_ExpandCriticalPath,
+			NAME_None,
+			TAttribute<FText>(),
+			TAttribute<FText>(),
+			FSlateIcon(FEditorStyle::GetStyleSetName(), "Profiler.EventGraph.ExpandSelection")
+		);
+
+		MenuBuilder.AddMenuEntry
+		(
+			FTableTreeViewCommands::Get().Command_CollapseSubtree,
+			NAME_None,
+			TAttribute<FText>(),
+			TAttribute<FText>(),
+			FSlateIcon(FEditorStyle::GetStyleSetName(), "Profiler.EventGraph.CollapseSelection")
+		);
+	}
+
+	MenuBuilder.EndSection();
+
 	MenuBuilder.BeginSection("Misc", LOCTEXT("ContextMenu_Header_Misc", "Miscellaneous"));
 	{
 		MenuBuilder.AddMenuEntry
@@ -2326,7 +2367,7 @@ void STableTreeView::OnPostAsyncUpdate()
 		}
 
 		// Expand each group node on the first few depths (if it doesn't have too many children).
-		ExpandChildGroups(Root.Get(), 1000);
+		SetExpandValueForChildGroups(Root.Get(), 1000, 4, true);
 
 		ClearInProgressAsyncOperations();
 		TreeView_Refresh();
@@ -2338,16 +2379,14 @@ void STableTreeView::OnPostAsyncUpdate()
 
 ////////////////////////////////////////////////////////////////////////////////////////////////////
 
-void STableTreeView::ExpandChildGroups(FBaseTreeNode* InRoot, int32 InMaxExpandedNodes)
+void STableTreeView::SetExpandValueForChildGroups(FBaseTreeNode* InRoot, int32 InMaxExpandedNodes, int32 InMaxDepthToExpand, bool InValue)
 {
-	const int32 MaxDepthToExpand = 4;
-
 	TArray<int32> NumNodesPerDepth;
-	NumNodesPerDepth.AddDefaulted(MaxDepthToExpand + 1);
-	CountNumNodesPerDepthRec(InRoot, NumNodesPerDepth, 0, MaxDepthToExpand, InMaxExpandedNodes);
+	NumNodesPerDepth.AddDefaulted(InMaxDepthToExpand + 1);
+	CountNumNodesPerDepthRec(InRoot, NumNodesPerDepth, 0, InMaxDepthToExpand, InMaxExpandedNodes);
 
 	int32 MaxDepth = 0;
-	for (int32 Depth = 0; Depth <= MaxDepthToExpand; ++Depth)
+	for (int32 Depth = 0; Depth <= InMaxDepthToExpand; ++Depth)
 	{
 		if (Depth > 0)
 		{
@@ -2362,7 +2401,7 @@ void STableTreeView::ExpandChildGroups(FBaseTreeNode* InRoot, int32 InMaxExpande
 
 	if (MaxDepth > 0)
 	{
-		ExpandChildGroupsRec(InRoot, 1, MaxDepth);
+		SetExpandValueForChildGroupsRec(InRoot, 1, MaxDepth, InValue);
 	}
 }
 
@@ -2386,18 +2425,18 @@ void STableTreeView::CountNumNodesPerDepthRec(FBaseTreeNode* InRoot, TArray<int3
 
 ////////////////////////////////////////////////////////////////////////////////////////////////////
 
-void STableTreeView::ExpandChildGroupsRec(FBaseTreeNode* InRoot, int32 InDepth, int32 InMaxDepth)
+void STableTreeView::SetExpandValueForChildGroupsRec(FBaseTreeNode* InRoot, int32 InDepth, int32 InMaxDepth, bool InValue)
 {
 	for (const FBaseTreeNodePtr& Node : InRoot->GetChildren())
 	{
 		if (Node->IsGroup())
 		{
-			Node->SetExpansion(true);
-			TreeView->SetItemExpansion(StaticCastSharedPtr<FTableTreeNode>(Node), true);
+			Node->SetExpansion(InValue);
+			TreeView->SetItemExpansion(StaticCastSharedPtr<FTableTreeNode>(Node), InValue);
 
 			if (InDepth < InMaxDepth)
 			{
-				ExpandChildGroupsRec(Node.Get(), InDepth + 1, InMaxDepth);
+				SetExpandValueForChildGroupsRec(Node.Get(), InDepth + 1, InMaxDepth, InValue);
 			}
 		}
 	}
@@ -2839,6 +2878,120 @@ void STableTreeView::ContextMenu_CopyColumnTooltipToClipboard_Execute()
 
 ////////////////////////////////////////////////////////////////////////////////////////////////////
 
+bool STableTreeView::ContextMenu_ExpandSubtree_CanExecute() const
+{
+	const TArray<FTableTreeNodePtr> SelectedNodes = TreeView->GetSelectedItems();
+	for (const FTableTreeNodePtr& Node : SelectedNodes)
+	{
+		if (Node->IsGroup() && Node->GetFilteredChildren().Num() > 0 && !TreeView->IsItemExpanded(Node))
+		{
+			return true;
+		}
+	}
+
+	return false;
+}
+
+////////////////////////////////////////////////////////////////////////////////////////////////////
+
+void STableTreeView::ContextMenu_ExpandSubtree_Execute()
+{
+	const TArray<FTableTreeNodePtr> SelectedNodes = TreeView->GetSelectedItems();
+	for (const FTableTreeNodePtr& Node : SelectedNodes)
+	{
+		if (Node->IsGroup() && !TreeView->IsItemExpanded(Node))
+		{
+			Node->SetExpansion(true);
+			TreeView->SetItemExpansion(Node, true);
+			SetExpandValueForChildGroups((FBaseTreeNode*)Node.Get(), MAX_NUMBER_OF_NODES_TO_EXPAND, MAX_DEPTH_TO_EXPAND, true);
+		}
+	}
+
+	TreeView->RequestTreeRefresh();
+}
+
+////////////////////////////////////////////////////////////////////////////////////////////////////
+
+bool STableTreeView::ContextMenu_CollapseSubtree_CanExecute() const
+{
+	const TArray<FTableTreeNodePtr> SelectedNodes = TreeView->GetSelectedItems();
+	for (const FTableTreeNodePtr& Node : SelectedNodes)
+	{
+		if (Node->IsGroup() && Node->GetFilteredChildren().Num() > 0 && TreeView->IsItemExpanded(Node))
+		{
+			return true;
+		}
+	}
+
+	return false;
+}
+
+////////////////////////////////////////////////////////////////////////////////////////////////////
+
+void STableTreeView::ContextMenu_CollapseSubtree_Execute()
+{
+	const TArray<FTableTreeNodePtr> SelectedNodes = TreeView->GetSelectedItems();
+	for (const FTableTreeNodePtr& Node : SelectedNodes)
+	{
+		if (Node->IsGroup() && TreeView->IsItemExpanded(Node))
+		{
+			Node->SetExpansion(false);
+			TreeView->SetItemExpansion(Node, false);
+			SetExpandValueForChildGroups((FBaseTreeNode*)Node.Get(), MAX_NUMBER_OF_NODES_TO_EXPAND, MAX_DEPTH_TO_EXPAND, false);
+		}
+	}
+
+	TreeView->RequestTreeRefresh();
+}
+
+////////////////////////////////////////////////////////////////////////////////////////////////////
+
+bool STableTreeView::ContextMenu_ExpandCriticalPath_CanExecute() const
+{
+	const TArray<FTableTreeNodePtr> SelectedNodes = TreeView->GetSelectedItems();
+	for (const FTableTreeNodePtr& Node : SelectedNodes)
+	{
+		if (Node->IsGroup() && Node->GetFilteredChildren().Num() > 0 && !TreeView->IsItemExpanded(Node))
+		{
+			return true;
+		}
+	}
+
+	return false;
+}
+
+////////////////////////////////////////////////////////////////////////////////////////////////////
+
+void STableTreeView::ContextMenu_ExpandCriticalPath_Execute()
+{
+	const TArray<FTableTreeNodePtr> SelectedNodes = TreeView->GetSelectedItems();
+	for (const FTableTreeNodePtr& Node : SelectedNodes)
+	{
+		FTableTreeNodePtr CurrentNode = Node;
+		while (CurrentNode.IsValid() && CurrentNode->IsGroup())
+		{
+			if (!TreeView->IsItemExpanded(CurrentNode))
+			{
+				CurrentNode->SetExpansion(true);
+				TreeView->SetItemExpansion(CurrentNode, true);
+			}
+
+			if (CurrentNode->GetFilteredChildren().Num() > 0)
+			{
+				CurrentNode = StaticCastSharedPtr<FTableTreeNode>(CurrentNode->GetFilteredChildren()[0]);
+			}
+			else
+			{
+				break;
+			}
+		}
+	}
+
+	TreeView->RequestTreeRefresh();
+}
+
+////////////////////////////////////////////////////////////////////////////////////////////////////
+
 } // namespace Insights
 
 #undef LOCTEXT_NAMESPACE
diff --git a/Engine/Source/Developer/TraceInsights/Private/Insights/Table/Widgets/STableTreeView.h b/Engine/Source/Developer/TraceInsights/Private/Insights/Table/Widgets/STableTreeView.h
--- a/Engine/Source/Developer/TraceInsights/Private/Insights/Table/Widgets/STableTreeView.h
+++ b/Engine/Source/Developer/TraceInsights/Private/Insights/Table/Widgets/STableTreeView.h
@@ -165,6 +165,12 @@ protected:
 	void ContextMenu_CopyColumnToClipboard_Execute();
 	bool ContextMenu_CopyColumnTooltipToClipboard_CanExecute() const;
 	void ContextMenu_CopyColumnTooltipToClipboard_Execute();
+	bool ContextMenu_ExpandSubtree_CanExecute() const;
+	void ContextMenu_ExpandSubtree_Execute();
+	bool ContextMenu_ExpandCriticalPath_CanExecute() const;
+	void ContextMenu_ExpandCriticalPath_Execute();
+	bool ContextMenu_CollapseSubtree_CanExecute() const;
+	void ContextMenu_CollapseSubtree_Execute();
 
 	////////////////////////////////////////////////////////////////////////////////////////////////////
 	// Tree View - Columns' Header
@@ -378,9 +384,9 @@ protected:
 
 	////////////////////////////////////////////////////////////////////////////////////////////////////
 
-	void ExpandChildGroups(FBaseTreeNode* InRoot, int32 InMaxExpandedNodes);
+	void SetExpandValueForChildGroups(FBaseTreeNode* InRoot, int32 InMaxExpandedNodes, int32 MaxDepthToExpand, bool InValue);
 	void CountNumNodesPerDepthRec(FBaseTreeNode* InRoot, TArray<int32>& InOutNumNodesPerDepth, int32 InDepth, int32 InMaxDepth, int InMaxNodes) const;
-	void ExpandChildGroupsRec(FBaseTreeNode* InRoot, int32 InDepth, int32 InMaxDepth);
+	void SetExpandValueForChildGroupsRec(FBaseTreeNode* InRoot, int32 InDepth, int32 InMaxDepth, bool InValue);
 
 protected:
 	/** Table view model. */
@@ -500,6 +506,9 @@ protected:
 	double StatsStartTime;
 	double StatsEndTime;
 
+	static constexpr int32 MAX_NUMBER_OF_NODES_TO_EXPAND = 1000 * 1000;
+	static constexpr int32 MAX_DEPTH_TO_EXPAND = 100;
+
 	//////////////////////////////////////////////////
 	// Logging
 
