task traces: added task priority info (as well as other info soldered into ENamedThreads::Type)

diff --git a/Engine/Source/Developer/TraceInsights/Private/Insights/ViewModels/ThreadTimingTrack.cpp b/Engine/Source/Developer/TraceInsights/Private/Insights/ViewModels/ThreadTimingTrack.cpp
--- a/Engine/Source/Developer/TraceInsights/Private/Insights/ViewModels/ThreadTimingTrack.cpp
+++ b/Engine/Source/Developer/TraceInsights/Private/Insights/ViewModels/ThreadTimingTrack.cpp
@@ -9,6 +9,7 @@
 #include "Serialization/MemoryReader.h"
 #include "Styling/SlateBrush.h"
 #include "TraceServices/Model/TasksProfiler.h"
+#include "Async/TaskGraphInterfaces.h"
 
 // Insights
 #include "Insights/Common/PaintUtils.h"
@@ -1012,20 +1013,43 @@ void FThreadTimingTrack::InitTooltip(FTooltipDrawState& InOutTooltip, const ITim
 
 		auto AddTaskInfo = [&InOutTooltip, this](const TraceServices::FTaskInfo& Task)
 		{
-			InOutTooltip.AddTextLine(TEXT("-------- TaskInfo --------"), FLinearColor::Green);
-			InOutTooltip.AddNameValueTextLine(TEXT("TaskId:"), FString::Printf(TEXT("%d"), Task.Id));
-			InOutTooltip.AddNameValueTextLine(TEXT("Task Debug Name:"), FString::Printf(TEXT("%s"), Task.DebugName));
-			InOutTooltip.AddNameValueTextLine(TEXT("Tracked:"), FString::Printf(TEXT("%d"), Task.bTracked));
-			InOutTooltip.AddNameValueTextLine(TEXT("Created:"), FString::Printf(TEXT("%f"), Task.CreatedTimestamp));
-			InOutTooltip.AddNameValueTextLine(TEXT("Launched:"), FString::Printf(TEXT("%f (+%s)"), Task.LaunchedTimestamp, *TimeUtils::FormatTimeAuto(Task.LaunchedTimestamp - Task.CreatedTimestamp)));
-			InOutTooltip.AddNameValueTextLine(TEXT("Started:"), FString::Printf(TEXT("%f (+%s) on thread \"%s\""), Task.StartedTimestamp, *TimeUtils::FormatTimeAuto(Task.StartedTimestamp - Task.LaunchedTimestamp), *SharedState.GetCpuTrack(Task.StartedThreadId)->GetName()));
+			InOutTooltip.AddTextLine(FString::Printf(TEXT("-------- Task %d%s --------"), Task.Id, Task.bTracked ? TEXT("") : TEXT(" (not tracked)")), FLinearColor::Green);
+
+			if (Task.DebugName != nullptr)
+			{
+				InOutTooltip.AddTextLine(FString::Printf(TEXT("%s"), Task.DebugName), FLinearColor::Green);
+			}
+
+			ENamedThreads::Type ThreadInfo = (ENamedThreads::Type)Task.ThreadToExecuteOn;
+			const TCHAR* NamedThreadsStr[] = {TEXT("Stats"), TEXT("RHI"), TEXT("Audio"), TEXT("Game"), TEXT("Rendering") };
+			ENamedThreads::Type ThreadIndex = ENamedThreads::GetThreadIndex(ThreadInfo);
+
+			if (ThreadIndex == ENamedThreads::AnyThread)
+			{
+				const TCHAR* TaskPri = ENamedThreads::GetTaskPriority(ThreadInfo) == ENamedThreads::NormalTaskPriority ? TEXT("Normal") : TEXT("High");
+
+				int32 ThreadPriIndex = ENamedThreads::GetThreadPriorityIndex(ThreadInfo);
+				const TCHAR* ThreadPriStrs[] = { TEXT("Normal"), TEXT("High"), TEXT("Low") };
+				const TCHAR* ThreadPri = ThreadPriStrs[ThreadPriIndex];
+
+				InOutTooltip.AddTextLine(FString::Printf(TEXT("%s Pri task on %s Pri worker"), TaskPri, ThreadPri), FLinearColor::Green);
+			}
+			else
+			{
+				const TCHAR* QueueStr = ENamedThreads::GetQueueIndex(ThreadInfo) == ENamedThreads::MainQueue ? TEXT("Main") : TEXT("Local");
+				InOutTooltip.AddTextLine(FString::Printf(TEXT("%s (%s queue)"), NamedThreadsStr[ThreadIndex], QueueStr), FLinearColor::Green);
+			}
+			
+			InOutTooltip.AddNameValueTextLine(TEXT("Created:"), FString::Printf(TEXT("%f on %s"), Task.CreatedTimestamp, *SharedState.GetCpuTrack(Task.CreatedThreadId)->GetName()));
+			InOutTooltip.AddNameValueTextLine(TEXT("Launched:"), FString::Printf(TEXT("%f (+%s) on %s"), Task.LaunchedTimestamp, *TimeUtils::FormatTimeAuto(Task.LaunchedTimestamp - Task.CreatedTimestamp), *SharedState.GetCpuTrack(Task.LaunchedThreadId)->GetName()));
+			InOutTooltip.AddNameValueTextLine(TEXT("Started:"), FString::Printf(TEXT("%f (+%s)"), Task.StartedTimestamp, *TimeUtils::FormatTimeAuto(Task.StartedTimestamp - Task.LaunchedTimestamp)));
 			if (Task.FinishedTimestamp != TraceServices::FTaskInfo::InvalidTimestamp)
 			{
 				InOutTooltip.AddNameValueTextLine(TEXT("Finished:"), FString::Printf(TEXT("%f (+%s)"), Task.FinishedTimestamp, *TimeUtils::FormatTimeAuto(Task.FinishedTimestamp - Task.StartedTimestamp)));
 
 				if (Task.CompletedTimestamp != TraceServices::FTaskInfo::InvalidTimestamp)
 				{
-					InOutTooltip.AddNameValueTextLine(TEXT("Completed:"), FString::Printf(TEXT("%f (+%s)"), Task.FinishedTimestamp, *TimeUtils::FormatTimeAuto(Task.CompletedTimestamp - Task.FinishedTimestamp)));
+					InOutTooltip.AddNameValueTextLine(TEXT("Completed:"), FString::Printf(TEXT("%f (+%s) on %s"), Task.FinishedTimestamp, *TimeUtils::FormatTimeAuto(Task.CompletedTimestamp - Task.FinishedTimestamp), *SharedState.GetCpuTrack(Task.CompletedThreadId)->GetName()));
 				}
 			}
 		};
diff --git a/Engine/Source/Developer/TraceServices/Private/Model/TasksProfiler.cpp b/Engine/Source/Developer/TraceServices/Private/Model/TasksProfiler.cpp
--- a/Engine/Source/Developer/TraceServices/Private/Model/TasksProfiler.cpp
+++ b/Engine/Source/Developer/TraceServices/Private/Model/TasksProfiler.cpp
@@ -7,6 +7,7 @@
 #include "Misc/Timespan.h"
 #include "Algo/BinarySearch.h"
 #include "Common/Utils.h"
+#include "Async/TaskGraphInterfaces.h"
 
 namespace TraceServices
 {
@@ -131,6 +132,7 @@ namespace TraceServices
 		{
 			Task->Id = TaskId;
 			Task->CreatedTimestamp = Timestamp;
+			Task->CreatedThreadId = ThreadId;
 		}
 
 		Task->DebugName = DebugName;
@@ -296,7 +298,7 @@ namespace TraceServices
 
 	bool FTasksProvider::IsNamedThread(int32 Thread)
 	{
-		return Thread != ENamedThreads::AnyThread;
+		return ENamedThreads::GetThreadIndex((ENamedThreads::Type)Thread) != ENamedThreads::AnyThread;
 	}
 
 	bool FTasksProvider::TryRegisterEvent(const TCHAR* EventName, TaskTrace::FId TaskId, double FTaskInfo::* TimestampPtr, double TimestampValue, uint32 FTaskInfo::* ThreadIdPtr/* = nullptr*/, uint32 ThreadIdValue/* = 0*/)
diff --git a/Engine/Source/Developer/TraceServices/Private/Model/TasksProfilerPrivate.h b/Engine/Source/Developer/TraceServices/Private/Model/TasksProfilerPrivate.h
--- a/Engine/Source/Developer/TraceServices/Private/Model/TasksProfilerPrivate.h
+++ b/Engine/Source/Developer/TraceServices/Private/Model/TasksProfilerPrivate.h
@@ -45,7 +45,6 @@ namespace TraceServices
 		// translates `TaskId` to the index into `Tasks` array, 
 		int64 GetTaskIndex(TaskTrace::FId TaskId) const; // `FistTaskId` must be initialised before the first use
 
-		FTaskInfo& GetTask(TaskTrace::FId TaskId);
 		FTaskInfo* TryGetTask(TaskTrace::FId TaskId);
 		// default-constructs a task if required
 		FTaskInfo* TryGetOrCreateTask(TaskTrace::FId TaskId);
diff --git a/Engine/Source/Developer/TraceServices/Public/TraceServices/Model/TasksProfiler.h b/Engine/Source/Developer/TraceServices/Public/TraceServices/Model/TasksProfiler.h
--- a/Engine/Source/Developer/TraceServices/Public/TraceServices/Model/TasksProfiler.h
+++ b/Engine/Source/Developer/TraceServices/Public/TraceServices/Model/TasksProfiler.h
@@ -18,7 +18,9 @@ namespace TraceServices
 		// this members are filled during task launch
 		const TCHAR* DebugName; // user provided or automatically generated (in format `full_path(lineno)`)
 		bool bTracked; // if the task is awaitable, not awaitable tasks will never be "completed", they are done when their execution is finished
-		int32 ThreadToExecuteOn; // ENamedThreads::Type - which thread type the task will be executed on, one of named threads or a worker thread
+		int32 ThreadToExecuteOn; // ENamedThreads::Type - combines the info about thread to execute on (one of named threads or a worker
+		// thread - ENamedThreads::AnyThread), task priority, thread priority (in case of a worker thread), and queue index (in case of a named
+		// thread - main queue or local queue). See ENamedThreads to split the info into separate fields
 
 		static constexpr double InvalidTimestamp = 0.0;
 
diff --git a/Engine/Source/Runtime/Core/Public/Async/TaskGraphInterfaces.h b/Engine/Source/Runtime/Core/Public/Async/TaskGraphInterfaces.h
--- a/Engine/Source/Runtime/Core/Public/Async/TaskGraphInterfaces.h
+++ b/Engine/Source/Runtime/Core/Public/Async/TaskGraphInterfaces.h
@@ -919,7 +919,7 @@ public:
 	{
 		bool bWakeUpWorker = true;
 		ConditionalQueueTask(CurrentThreadIfKnown, bWakeUpWorker);
-		TaskTrace::Launched(GetTraceId(), nullptr, Subsequents.IsValid(), ENamedThreads::GetThreadIndex(((TTask*)&TaskStorage)->GetDesiredThread()));
+		TaskTrace::Launched(GetTraceId(), nullptr, Subsequents.IsValid(), ((TTask*)&TaskStorage)->GetDesiredThread());
 	}
 
 	FGraphEventRef GetCompletionEvent()
@@ -1063,7 +1063,7 @@ private:
 	 **/
 	FGraphEventRef Setup(const FGraphEventArray* Prerequisites = NULL, ENamedThreads::Type CurrentThreadIfKnown = ENamedThreads::AnyThread)
 	{
-		TaskTrace::Launched(GetTraceId(), nullptr, Subsequents.IsValid(), ENamedThreads::GetThreadIndex(((TTask*)&TaskStorage)->GetDesiredThread()));
+		TaskTrace::Launched(GetTraceId(), nullptr, Subsequents.IsValid(), ((TTask*)&TaskStorage)->GetDesiredThread());
 
 		FGraphEventRef ReturnedEventRef = Subsequents; // very important so that this doesn't get destroyed before we return
 		SetupPrereqs(Prerequisites, CurrentThreadIfKnown, true);
