[Insights] Auto expand for searchbox filter in TableTreeView

diff --git a/Engine/Source/Developer/TraceInsights/Private/Insights/Table/Widgets/STableTreeView.cpp b/Engine/Source/Developer/TraceInsights/Private/Insights/Table/Widgets/STableTreeView.cpp
--- a/Engine/Source/Developer/TraceInsights/Private/Insights/Table/Widgets/STableTreeView.cpp
+++ b/Engine/Source/Developer/TraceInsights/Private/Insights/Table/Widgets/STableTreeView.cpp
@@ -834,7 +834,8 @@ void STableTreeView::ApplyFiltering()
 	}
 	else
 	{
-		ApplyHierarchicalFilterForNode(Root);
+		const bool bFilterIsEmpty = bRunInAsyncMode ? CurrentAsyncOpTextFilter->GetRawFilterText().IsEmpty() : TextFilter->GetRawFilterText().IsEmpty();
+		ApplyHierarchicalFilterForNode(Root, bFilterIsEmpty);
 	}
 
 	FilteredGroupNodes.Reset();
@@ -900,16 +901,16 @@ void STableTreeView::ApplyFiltering()
 
 ////////////////////////////////////////////////////////////////////////////////////////////////////
 
-bool STableTreeView::ApplyHierarchicalFilterForNode(FTableTreeNodePtr NodePtr)
+bool STableTreeView::ApplyHierarchicalFilterForNode(FTableTreeNodePtr NodePtr, bool bFilterIsEmpty)
 {
-	bool bIsNodeVisible = Filters->PassesAllFilters(NodePtr);
+	bool bIsNodeVisible = bFilterIsEmpty || Filters->PassesAllFilters(NodePtr);
 
 	if (NodePtr->IsGroup())
 	{
 		// If a group node passes the filter, all child nodes will be shown
 		if (bIsNodeVisible)
 		{
-			MakeSubtreeVisible(NodePtr);
+			MakeSubtreeVisible(NodePtr, bFilterIsEmpty);
 			return true;
 		}
 
@@ -926,7 +927,7 @@ bool STableTreeView::ApplyHierarchicalFilterForNode(FTableTreeNodePtr NodePtr)
 			}
 			// Add a child.
 			const FTableTreeNodePtr& ChildNodePtr = StaticCastSharedPtr<FTableTreeNode>(GroupChildren[Cx]);
-			if (ApplyHierarchicalFilterForNode(ChildNodePtr))
+			if (ApplyHierarchicalFilterForNode(ChildNodePtr, bFilterIsEmpty))
 			{
 				NodePtr->AddFilteredChild(ChildNodePtr);
 				NumVisibleChildren++;
@@ -934,17 +935,17 @@ bool STableTreeView::ApplyHierarchicalFilterForNode(FTableTreeNodePtr NodePtr)
 		}
 
 		const bool bIsGroupNodeVisible = bIsNodeVisible || NumVisibleChildren > 0;
-
-		if (bIsGroupNodeVisible)
-		{
-			// Add a group.
-			NodePtr->SetExpansion(true);
-		}
-		else
+		if (!bFilterIsEmpty && bIsGroupNodeVisible )
 		{
-			NodePtr->SetExpansion(false);
+			if (bRunInAsyncMode)
+			{
+				NodesToExpand.Add(NodePtr);
+			}
+			else
+			{
+				TreeView->SetItemExpansion(NodePtr, true);
+			}
 		}
-
 		return bIsGroupNodeVisible;
 	}
 	else
@@ -1004,8 +1005,9 @@ bool STableTreeView::ApplyAdvancedFiltersForNode(FTableTreeNodePtr NodePtr)
 
 ////////////////////////////////////////////////////////////////////////////////////////////////////
 
-void STableTreeView::MakeSubtreeVisible(FTableTreeNodePtr NodePtr)
+bool STableTreeView::MakeSubtreeVisible(FTableTreeNodePtr NodePtr, bool bFilterIsEmpty)
 {
+	bool bPassesNonEmptyFilter = !bFilterIsEmpty && Filters->PassesAllFilters(NodePtr);
 	if (NodePtr->IsGroup())
 	{
 		NodePtr->ClearFilteredChildren();
@@ -1013,6 +1015,7 @@ void STableTreeView::MakeSubtreeVisible(FTableTreeNodePtr NodePtr)
 		const TArray<FBaseTreeNodePtr>& GroupChildren = NodePtr->GetChildren();
 		const int32 NumChildren = GroupChildren.Num();
 		int32 NumVisibleChildren = 0;
+		bool bShouldExpand = false;
 		for (int32 Cx = 0; Cx < NumChildren; ++Cx)
 		{
 			if (bCancelCurrentAsyncOp)
@@ -1021,13 +1024,28 @@ void STableTreeView::MakeSubtreeVisible(FTableTreeNodePtr NodePtr)
 			}
 
 			const FTableTreeNodePtr& ChildNodePtr = StaticCastSharedPtr<FTableTreeNode>(GroupChildren[Cx]);
-			MakeSubtreeVisible(ChildNodePtr);
+			bShouldExpand |= MakeSubtreeVisible(ChildNodePtr, bFilterIsEmpty);
 			NodePtr->AddFilteredChild(ChildNodePtr);
 			NodePtr->SetExpansion(true);
+
+			if (bShouldExpand)
+			{
+				if (bRunInAsyncMode)
+				{
+					NodesToExpand.Add(NodePtr);
+				}
+				else
+				{
+					TreeView->SetItemExpansion(NodePtr, true);
+				}
+			}
 		}
+
+		return bShouldExpand || bPassesNonEmptyFilter;
 	}
 
 	NodePtr->SetIsFiltered(false);
+	return bPassesNonEmptyFilter;
 }
 
 ////////////////////////////////////////////////////////////////////////////////////////////////////
@@ -2366,6 +2384,12 @@ void STableTreeView::OnPostAsyncUpdate()
 			}
 		}
 
+		for (FTableTreeNodePtr& Node : NodesToExpand)
+		{
+			TreeView->SetItemExpansion(Node, true);
+		}
+		NodesToExpand.Empty();
+
 		// Expand each group node on the first few depths (if it doesn't have too many children).
 		SetExpandValueForChildGroups(Root.Get(), 1000, 4, true);
 
diff --git a/Engine/Source/Developer/TraceInsights/Private/Insights/Table/Widgets/STableTreeView.h b/Engine/Source/Developer/TraceInsights/Private/Insights/Table/Widgets/STableTreeView.h
--- a/Engine/Source/Developer/TraceInsights/Private/Insights/Table/Widgets/STableTreeView.h
+++ b/Engine/Source/Developer/TraceInsights/Private/Insights/Table/Widgets/STableTreeView.h
@@ -221,9 +221,10 @@ protected:
 
 	bool ApplyAdvancedFiltersForNode(FTableTreeNodePtr NodePtr);
 
-	bool ApplyHierarchicalFilterForNode(FTableTreeNodePtr NodePtr);
+	bool ApplyHierarchicalFilterForNode(FTableTreeNodePtr NodePtr, bool bFilterIsEmpty);
 
-	void MakeSubtreeVisible(FTableTreeNodePtr NodePtr);
+	/** Set all the nodes belonging to a subtree as visible. Returns true if the caller node should be expanded. */
+	bool MakeSubtreeVisible(FTableTreeNodePtr NodePtr, bool bFilterIsEmpty);
 
 	bool SearchBox_IsEnabled() const;
 	void SearchBox_OnTextChanged(const FText& InFilterText);
@@ -497,6 +498,7 @@ protected:
 	FFilterConfigurator* CurrentAsyncOpFilterConfigurator = nullptr;
 	std::atomic<bool> bCancelCurrentAsyncOp { false };
 	FGraphEventRef DispatchEvent;
+	TArray<FTableTreeNodePtr> NodesToExpand;
 
 	//////////////////////////////////////////////////
 	TSharedPtr<FFilterConfigurator> FilterConfigurator;
