[Insights]
  - Tree: Separated the name suffix for group nodes from display name. Also changed color of this name suffix to be darker then the group name (in all tree views).
  - Tree: Improved memory usage for tree views. Properties specific to group nodes are only allocated for group nodes and not for leaf nodes.
  - Tree: Added possibility to set a custom tooltip text to group nodes.

diff --git a/Engine/Source/Developer/TraceInsights/Private/Insights/MemoryProfiler/Widgets/SMemTagTreeViewTableCell.cpp b/Engine/Source/Developer/TraceInsights/Private/Insights/MemoryProfiler/Widgets/SMemTagTreeViewTableCell.cpp
--- a/Engine/Source/Developer/TraceInsights/Private/Insights/MemoryProfiler/Widgets/SMemTagTreeViewTableCell.cpp
+++ b/Engine/Source/Developer/TraceInsights/Private/Insights/MemoryProfiler/Widgets/SMemTagTreeViewTableCell.cpp
@@ -106,7 +106,7 @@ TSharedRef<SWidget> SMemTagTreeViewTableCell::GenerateWidgetForNameColumn(const
 		+ SHorizontalBox::Slot()
 		.AutoWidth()
 		.VAlign(VAlign_Center)
-		.HAlign(ColumnPtr->GetHorizontalAlignment())
+		.HAlign(HAlign_Left)
 		.Padding(FMargin(2.0f, 0.0f))
 		[
 			SNew(STextBlock)
@@ -116,6 +116,21 @@ TSharedRef<SWidget> SMemTagTreeViewTableCell::GenerateWidgetForNameColumn(const
 			.ColorAndOpacity(this, &SMemTagTreeViewTableCell::GetColorAndOpacity)
 			.ShadowColorAndOpacity(this, &SMemTagTreeViewTableCell::GetShadowColorAndOpacity)
 		]
+
+		// Name Suffix
+		+ SHorizontalBox::Slot()
+		.AutoWidth()
+		.VAlign(VAlign_Center)
+		.HAlign(HAlign_Left)
+		.Padding(FMargin(2.0f, 0.0f))
+		[
+			SNew(STextBlock)
+			.Visibility(this, &SMemTagTreeViewTableCell::HasExtraDisplayName)
+			.Text(this, &SMemTagTreeViewTableCell::GetExtraDisplayName)
+			.TextStyle(FEditorStyle::Get(), TEXT("Profiler.Tooltip"))
+			.ColorAndOpacity(this, &SMemTagTreeViewTableCell::GetExtraColorAndOpacity)
+			.ShadowColorAndOpacity(this, &SMemTagTreeViewTableCell::GetShadowColorAndOpacity)
+		]
 	;
 }
 
diff --git a/Engine/Source/Developer/TraceInsights/Private/Insights/MemoryProfiler/Widgets/SMemTagTreeViewTableCell.h b/Engine/Source/Developer/TraceInsights/Private/Insights/MemoryProfiler/Widgets/SMemTagTreeViewTableCell.h
--- a/Engine/Source/Developer/TraceInsights/Private/Insights/MemoryProfiler/Widgets/SMemTagTreeViewTableCell.h
+++ b/Engine/Source/Developer/TraceInsights/Private/Insights/MemoryProfiler/Widgets/SMemTagTreeViewTableCell.h
@@ -112,6 +112,16 @@ protected:
 		return MemTagNodePtr->GetDisplayName();
 	}
 
+	FText GetExtraDisplayName() const
+	{
+		return MemTagNodePtr->GetExtraDisplayName();
+	}
+
+	EVisibility HasExtraDisplayName() const
+	{
+		return MemTagNodePtr->HasExtraDisplayName() ? EVisibility::Visible : EVisibility::Collapsed;
+	}
+
 	FText GetValueAsText() const;
 
 	FSlateColor GetColorAndOpacity() const
@@ -123,6 +133,15 @@ protected:
 		return TextColor;
 	}
 
+	FSlateColor GetExtraColorAndOpacity() const
+	{
+		const FLinearColor TextColor =
+			MemTagNodePtr->IsFiltered() ?
+				FLinearColor(0.5f, 0.5f, 0.5f, 0.5f) :
+				FLinearColor(0.5f, 0.5f, 0.5f, 1.0f);
+		return TextColor;
+	}
+
 	FSlateColor GetStatsColorAndOpacity() const
 	{
 		const FLinearColor TextColor =
diff --git a/Engine/Source/Developer/TraceInsights/Private/Insights/NetworkingProfiler/Widgets/SNetStatsTableCell.cpp b/Engine/Source/Developer/TraceInsights/Private/Insights/NetworkingProfiler/Widgets/SNetStatsTableCell.cpp
--- a/Engine/Source/Developer/TraceInsights/Private/Insights/NetworkingProfiler/Widgets/SNetStatsTableCell.cpp
+++ b/Engine/Source/Developer/TraceInsights/Private/Insights/NetworkingProfiler/Widgets/SNetStatsTableCell.cpp
@@ -86,7 +86,7 @@ TSharedRef<SWidget> SNetStatsTableCell::GenerateWidgetForNameColumn(const FArgum
 		+ SHorizontalBox::Slot()
 		.AutoWidth()
 		.VAlign(VAlign_Center)
-		.HAlign(ColumnPtr->GetHorizontalAlignment())
+		.HAlign(HAlign_Left)
 		.Padding(FMargin(2.0f, 0.0f))
 		[
 			SNew(STextBlock)
@@ -96,6 +96,21 @@ TSharedRef<SWidget> SNetStatsTableCell::GenerateWidgetForNameColumn(const FArgum
 			.ColorAndOpacity(this, &SNetStatsTableCell::GetColorAndOpacity)
 			.ShadowColorAndOpacity(this, &SNetStatsTableCell::GetShadowColorAndOpacity)
 		]
+
+		// Name Suffix
+		+ SHorizontalBox::Slot()
+		.AutoWidth()
+		.VAlign(VAlign_Center)
+		.HAlign(HAlign_Left)
+		.Padding(FMargin(2.0f, 0.0f))
+		[
+			SNew(STextBlock)
+			.Visibility(this, &SNetStatsTableCell::HasExtraDisplayName)
+			.Text(this, &SNetStatsTableCell::GetExtraDisplayName)
+			.TextStyle(FEditorStyle::Get(), TEXT("Profiler.Tooltip"))
+			.ColorAndOpacity(this, &SNetStatsTableCell::GetExtraColorAndOpacity)
+			.ShadowColorAndOpacity(this, &SNetStatsTableCell::GetShadowColorAndOpacity)
+		]
 	;
 }
 
diff --git a/Engine/Source/Developer/TraceInsights/Private/Insights/NetworkingProfiler/Widgets/SNetStatsTableCell.h b/Engine/Source/Developer/TraceInsights/Private/Insights/NetworkingProfiler/Widgets/SNetStatsTableCell.h
--- a/Engine/Source/Developer/TraceInsights/Private/Insights/NetworkingProfiler/Widgets/SNetStatsTableCell.h
+++ b/Engine/Source/Developer/TraceInsights/Private/Insights/NetworkingProfiler/Widgets/SNetStatsTableCell.h
@@ -107,6 +107,16 @@ protected:
 		return NetEventNodePtr->GetDisplayName();
 	}
 
+	FText GetExtraDisplayName() const
+	{
+		return NetEventNodePtr->GetExtraDisplayName();
+	}
+
+	EVisibility HasExtraDisplayName() const
+	{
+		return NetEventNodePtr->HasExtraDisplayName() ? EVisibility::Visible : EVisibility::Collapsed;
+	}
+
 	FText GetValueAsText() const;
 
 	FSlateColor GetColorAndOpacity() const
@@ -118,6 +128,15 @@ protected:
 		return TextColor;
 	}
 
+	FSlateColor GetExtraColorAndOpacity() const
+	{
+		const FLinearColor TextColor =
+			NetEventNodePtr->IsFiltered() ?
+				FLinearColor(0.5f, 0.5f, 0.5f, 0.5f) :
+				FLinearColor(0.5f, 0.5f, 0.5f, 1.0f);
+		return TextColor;
+	}
+
 	FSlateColor GetStatsColorAndOpacity() const
 	{
 		const FLinearColor TextColor =
diff --git a/Engine/Source/Developer/TraceInsights/Private/Insights/Table/ViewModels/BaseTreeNode.cpp b/Engine/Source/Developer/TraceInsights/Private/Insights/Table/ViewModels/BaseTreeNode.cpp
--- a/Engine/Source/Developer/TraceInsights/Private/Insights/Table/ViewModels/BaseTreeNode.cpp
+++ b/Engine/Source/Developer/TraceInsights/Private/Insights/Table/ViewModels/BaseTreeNode.cpp
@@ -10,46 +10,56 @@
 namespace Insights
 {
 
+FBaseTreeNode::FGroupNodeData FBaseTreeNode::DefaultGroupData;
+
 ////////////////////////////////////////////////////////////////////////////////////////////////////
 
 const FText FBaseTreeNode::GetDisplayName() const
 {
-	FText Text = FText::GetEmpty();
+	return FText::FromName(GetName());
+}
+
+////////////////////////////////////////////////////////////////////////////////////////////////////
 
+const FText FBaseTreeNode::GetExtraDisplayName() const
+{
 	if (IsGroup())
 	{
-		const int32 NumChildren = Children.Num();
-		const int32 NumFilteredChildren = FilteredChildren.Num();
+		const int32 NumChildren = GroupData->Children.Num();
+		const int32 NumFilteredChildren = GroupData->FilteredChildren.Num();
 
 		if (NumFilteredChildren == NumChildren)
 		{
-			Text = FText::Format(LOCTEXT("TreeNodeGroupTextFmt1", "{0} ({1})"), FText::FromName(GetName()), FText::AsNumber(NumChildren));
+			return FText::Format(LOCTEXT("TreeNodeGroup_ExtraText_Fmt1", "({0})"), FText::AsNumber(NumChildren));
 		}
 		else
 		{
-			Text = FText::Format(LOCTEXT("TreeNodeGroupTextFmt2", "{0} ({1} / {2})"), FText::FromName(GetName()), FText::AsNumber(NumFilteredChildren), FText::AsNumber(NumChildren));
+			return FText::Format(LOCTEXT("TreeNodeGroup_ExtraText_Fmt2", "({0} / {1})"), FText::AsNumber(NumFilteredChildren), FText::AsNumber(NumChildren));
 		}
 	}
-	else
-	{
-		Text = FText::FromName(GetName());
-	}
 
-	return Text;
+	return FText::GetEmpty();
+}
+
+////////////////////////////////////////////////////////////////////////////////////////////////////
+
+bool FBaseTreeNode::HasExtraDisplayName() const
+{
+	return IsGroup();
 }
 
 ////////////////////////////////////////////////////////////////////////////////////////////////////
 
 void FBaseTreeNode::SortChildrenAscending(const ITableCellValueSorter& Sorter)
 {
-	Sorter.Sort(Children, ESortMode::Ascending);
+	Sorter.Sort(GroupData->Children, ESortMode::Ascending);
 }
 
 ////////////////////////////////////////////////////////////////////////////////////////////////////
 
 void FBaseTreeNode::SortChildrenDescending(const ITableCellValueSorter& Sorter)
 {
-	Sorter.Sort(Children, ESortMode::Descending);
+	Sorter.Sort(GroupData->Children, ESortMode::Descending);
 }
 
 ////////////////////////////////////////////////////////////////////////////////////////////////////
diff --git a/Engine/Source/Developer/TraceInsights/Private/Insights/Table/ViewModels/BaseTreeNode.h b/Engine/Source/Developer/TraceInsights/Private/Insights/Table/ViewModels/BaseTreeNode.h
--- a/Engine/Source/Developer/TraceInsights/Private/Insights/Table/ViewModels/BaseTreeNode.h
+++ b/Engine/Source/Developer/TraceInsights/Private/Insights/Table/ViewModels/BaseTreeNode.h
@@ -29,18 +29,38 @@ typedef TWeakPtr<class FBaseTreeNode> FBaseTreeNodeWeak;
 
 class FBaseTreeNode : public TSharedFromThis<FBaseTreeNode>
 {
+protected:
+	struct FGroupNodeData
+	{
+		/** Children of this node. */
+		TArray<FBaseTreeNodePtr> Children;
+
+		/** Filtered children of this node. */
+		TArray<FBaseTreeNodePtr> FilteredChildren;
+
+		/** Descriptive text for this node to display in a tooltip. */
+		FText Tooltip;
+
+		/** Whether this group node should be expanded or not. */
+		bool bIsExpanded = false;
+	};
+
 public:
 	/** Initialization constructor for the node. */
 	FBaseTreeNode(const FName InName, bool bInIsGroup)
 		: DefaultSortOrder(0)
 		, Name(InName)
-		, bIsGroup(bInIsGroup)
-		, bIsExpanded(false)
+		, GroupData(bInIsGroup ? new FGroupNodeData() : &DefaultGroupData)
 	{
 	}
 
 	virtual ~FBaseTreeNode()
 	{
+		if (GroupData != &DefaultGroupData)
+		{
+			delete GroupData;
+			GroupData = nullptr;
+		}
 	}
 
 	virtual const FName& GetTypeName() const = 0;
@@ -57,16 +77,33 @@ public:
 	}
 
 	/**
-	 * @return a name of this node to display in tree view. Ex.: group nodes may include additional info, like the number of visible/total children.
+	 * @return a name of this node to display in tree view.
 	 */
 	virtual const FText GetDisplayName() const;
 
+	/**
+	 * @return a name suffix for this node to display in tree view. Ex.: group nodes may include additional info, like the number of visible / total children.
+	 */
+	virtual const FText GetExtraDisplayName() const;
+
+	virtual bool HasExtraDisplayName() const;
+
+	/**
+	 * @return a descriptive text for this node to display in a tooltip.
+	 */
+	virtual const FText GetTooltip() const { return IsGroup() ? GroupData->Tooltip : FText::GetEmpty(); }
+
+	/**
+	 * Sets a descriptive text for this node to display in a tooltip.
+	 */
+	virtual void SetTooltip(const FText& InTooltip) { if (IsGroup()) { GroupData->Tooltip = InTooltip; } }
+
 	/**
 	 * @return true, if this node is a group node.
 	 */
 	bool IsGroup() const
 	{
-		return bIsGroup;
+		return GroupData != &DefaultGroupData;
 	}
 
 	/**
@@ -74,7 +111,7 @@ public:
 	 */
 	FORCEINLINE_DEBUGGABLE const TArray<FBaseTreeNodePtr>& GetChildren() const
 	{
-		return Children;
+		return GroupData->Children;
 	}
 
 	/**
@@ -82,7 +119,7 @@ public:
 	 */
 	FORCEINLINE_DEBUGGABLE const TArray<FBaseTreeNodePtr>& GetFilteredChildren() const
 	{
-		return FilteredChildren;
+		return GroupData->FilteredChildren;
 	}
 
 	/**
@@ -113,33 +150,33 @@ public:
 	FORCEINLINE_DEBUGGABLE void AddChildAndSetGroupPtr(const FBaseTreeNodePtr& ChildPtr)
 	{
 		ChildPtr->GroupPtr = AsShared();
-		Children.Add(ChildPtr);
+		GroupData->Children.Add(ChildPtr);
 	}
 
 	/** Clears children. */
 	void ClearChildren(int32 NewSize = 0)
 	{
-		for (FBaseTreeNodePtr& NodePtr : Children)
+		for (FBaseTreeNodePtr& NodePtr : GroupData->Children)
 		{
 			NodePtr->GroupPtr = nullptr;
 		}
-		Children.Reset(NewSize);
+		GroupData->Children.Reset(NewSize);
 	}
 
 	/** Adds specified child to the filtered children. */
 	FORCEINLINE_DEBUGGABLE void AddFilteredChild(const FBaseTreeNodePtr& ChildPtr)
 	{
-		FilteredChildren.Add(ChildPtr);
+		GroupData->FilteredChildren.Add(ChildPtr);
 	}
 
 	/** Clears filtered children. */
 	void ClearFilteredChildren()
 	{
-		FilteredChildren.Reset();
+		GroupData->FilteredChildren.Reset();
 	}
 
-	bool IsExpanded() const { return bIsExpanded; }
-	void SetExpansion(bool bOnOff) { bIsExpanded = bOnOff; }
+	bool IsExpanded() const { return GroupData->bIsExpanded; }
+	void SetExpansion(bool bOnOff) { GroupData->bIsExpanded = bOnOff; }
 
 protected:
 	/**
@@ -147,7 +184,7 @@ protected:
 	 */
 	FORCEINLINE_DEBUGGABLE TArray<FBaseTreeNodePtr>& GetChildrenMutable()
 	{
-		return Children;
+		return GroupData->Children;
 	}
 
 	void SetGroupPtr(FBaseTreeNodePtr InGroupPtr)
@@ -162,20 +199,14 @@ private:
 	/** The name of this node. */
 	const FName Name;
 
-	/** Children of this node. */
-	TArray<FBaseTreeNodePtr> Children;
-
-	/** Filtered children of this node. */
-	TArray<FBaseTreeNodePtr> FilteredChildren;
-
 	/** A weak pointer to the group/parent of this node. */
 	FBaseTreeNodeWeak GroupPtr;
 
-	/** If this node is a group node or not. */
-	bool bIsGroup;
+	/** The struct containing properties of a group node. It is allocated only for group nodes. */
+	FGroupNodeData* GroupData;
 
-	/** Whether this group node should be expanded or not. */
-	bool bIsExpanded;
+	/** The only group data for "not a group" nodes. */
+	static FGroupNodeData DefaultGroupData;
 };
 
 ////////////////////////////////////////////////////////////////////////////////////////////////////
diff --git a/Engine/Source/Developer/TraceInsights/Private/Insights/Table/Widgets/STableTreeViewCell.cpp b/Engine/Source/Developer/TraceInsights/Private/Insights/Table/Widgets/STableTreeViewCell.cpp
--- a/Engine/Source/Developer/TraceInsights/Private/Insights/Table/Widgets/STableTreeViewCell.cpp
+++ b/Engine/Source/Developer/TraceInsights/Private/Insights/Table/Widgets/STableTreeViewCell.cpp
@@ -89,7 +89,7 @@ TSharedRef<SWidget> STableTreeViewCell::GenerateWidgetForNameColumn(const FArgum
 		+ SHorizontalBox::Slot()
 		.AutoWidth()
 		.VAlign(VAlign_Center)
-		.HAlign(ColumnPtr->GetHorizontalAlignment())
+		.HAlign(HAlign_Left)
 		.Padding(FMargin(2.0f, 0.0f))
 		[
 			SNew(STextBlock)
@@ -99,6 +99,21 @@ TSharedRef<SWidget> STableTreeViewCell::GenerateWidgetForNameColumn(const FArgum
 			.ColorAndOpacity(this, &STableTreeViewCell::GetColorAndOpacity)
 			.ShadowColorAndOpacity(this, &STableTreeViewCell::GetShadowColorAndOpacity)
 		]
+
+		// Name Suffix
+		+ SHorizontalBox::Slot()
+		.AutoWidth()
+		.VAlign(VAlign_Center)
+		.HAlign(HAlign_Left)
+		.Padding(FMargin(2.0f, 0.0f))
+		[
+			SNew(STextBlock)
+			.Visibility(this, &STableTreeViewCell::HasExtraDisplayName)
+			.Text(this, &STableTreeViewCell::GetExtraDisplayName)
+			.TextStyle(FEditorStyle::Get(), TEXT("Profiler.Tooltip"))
+			.ColorAndOpacity(this, &STableTreeViewCell::GetExtraColorAndOpacity)
+			.ShadowColorAndOpacity(this, &STableTreeViewCell::GetShadowColorAndOpacity)
+		]
 	;
 }
 
diff --git a/Engine/Source/Developer/TraceInsights/Private/Insights/Table/Widgets/STableTreeViewCell.h b/Engine/Source/Developer/TraceInsights/Private/Insights/Table/Widgets/STableTreeViewCell.h
--- a/Engine/Source/Developer/TraceInsights/Private/Insights/Table/Widgets/STableTreeViewCell.h
+++ b/Engine/Source/Developer/TraceInsights/Private/Insights/Table/Widgets/STableTreeViewCell.h
@@ -108,6 +108,16 @@ protected:
 		return TableTreeNodePtr->GetDisplayName();
 	}
 
+	FText GetExtraDisplayName() const
+	{
+		return TableTreeNodePtr->GetExtraDisplayName();
+	}
+
+	EVisibility HasExtraDisplayName() const
+	{
+		return TableTreeNodePtr->HasExtraDisplayName() ? EVisibility::Visible : EVisibility::Collapsed;
+	}
+
 	FText GetValueAsText() const;
 
 	FSlateColor GetColorAndOpacity() const
@@ -119,6 +129,15 @@ protected:
 		return TextColor;
 	}
 
+	FSlateColor GetExtraColorAndOpacity() const
+	{
+		const FLinearColor TextColor =
+			TableTreeNodePtr->IsFiltered() ?
+				FLinearColor(0.5f, 0.5f, 0.5f, 0.5f) :
+				FLinearColor(0.5f, 0.5f, 0.5f, 1.0f);
+		return TextColor;
+	}
+
 	FSlateColor GetStatsColorAndOpacity() const
 	{
 		const FLinearColor TextColor =
diff --git a/Engine/Source/Developer/TraceInsights/Private/Insights/Table/Widgets/STableTreeViewTooltip.cpp b/Engine/Source/Developer/TraceInsights/Private/Insights/Table/Widgets/STableTreeViewTooltip.cpp
--- a/Engine/Source/Developer/TraceInsights/Private/Insights/Table/Widgets/STableTreeViewTooltip.cpp
+++ b/Engine/Source/Developer/TraceInsights/Private/Insights/Table/Widgets/STableTreeViewTooltip.cpp
@@ -92,6 +92,9 @@ TSharedPtr<SToolTip> STableTreeViewTooltip::GetRowTooltip(const TSharedPtr<FTabl
 	TSharedPtr<SGridPanel> GridPanel;
 	TSharedPtr<SHorizontalBox> HBox;
 
+	const FText NodeTooltip = TreeNodePtr->GetTooltip();
+	const EVisibility NodeTooltipVisibility = NodeTooltip.IsEmpty() ? EVisibility::Collapsed : EVisibility::Visible;
+
 	TSharedPtr<SToolTip> TableCellTooltip =
 		SNew(SToolTip)
 		[
@@ -167,6 +170,28 @@ TSharedPtr<SToolTip> STableTreeViewTooltip::GetRowTooltip(const TSharedPtr<FTabl
 					]
 				]
 
+				+ SVerticalBox::Slot()
+				.AutoHeight()
+				.Padding(2.0f)
+				[
+					SNew(SSeparator)
+					.Visibility(NodeTooltipVisibility)
+					.Orientation(Orient_Horizontal)
+				]
+
+				+ SVerticalBox::Slot()
+				.AutoHeight()
+				.Padding(2.0f)
+				[
+					SNew(STextBlock)
+					.Visibility(NodeTooltipVisibility)
+					.WrapTextAt(512.0f)
+					.WrappingPolicy(ETextWrappingPolicy::AllowPerCharacterWrapping)
+					.Text(NodeTooltip)
+					.TextStyle(FEditorStyle::Get(), TEXT("Profiler.Tooltip"))
+					.ColorAndOpacity(FLinearColor(0.5f, 0.5f, 0.5f, 1.0f))
+				]
+
 				+ SVerticalBox::Slot()
 				.AutoHeight()
 				.Padding(2.0f)
diff --git a/Engine/Source/Developer/TraceInsights/Private/Insights/Widgets/SStatsTableCell.cpp b/Engine/Source/Developer/TraceInsights/Private/Insights/Widgets/SStatsTableCell.cpp
--- a/Engine/Source/Developer/TraceInsights/Private/Insights/Widgets/SStatsTableCell.cpp
+++ b/Engine/Source/Developer/TraceInsights/Private/Insights/Widgets/SStatsTableCell.cpp
@@ -106,7 +106,7 @@ TSharedRef<SWidget> SStatsTableCell::GenerateWidgetForNameColumn(const FArgument
 		+ SHorizontalBox::Slot()
 		.AutoWidth()
 		.VAlign(VAlign_Center)
-		.HAlign(ColumnPtr->GetHorizontalAlignment())
+		.HAlign(HAlign_Left)
 		.Padding(FMargin(2.0f, 0.0f))
 		[
 			SNew(STextBlock)
@@ -116,6 +116,21 @@ TSharedRef<SWidget> SStatsTableCell::GenerateWidgetForNameColumn(const FArgument
 			.ColorAndOpacity(this, &SStatsTableCell::GetColorAndOpacity)
 			.ShadowColorAndOpacity(this, &SStatsTableCell::GetShadowColorAndOpacity)
 		]
+
+		// Name Suffix
+		+ SHorizontalBox::Slot()
+		.AutoWidth()
+		.VAlign(VAlign_Center)
+		.HAlign(HAlign_Left)
+		.Padding(FMargin(2.0f, 0.0f))
+		[
+			SNew(STextBlock)
+			.Visibility(this, &SStatsTableCell::HasExtraDisplayName)
+			.Text(this, &SStatsTableCell::GetExtraDisplayName)
+			.TextStyle(FEditorStyle::Get(), TEXT("Profiler.Tooltip"))
+			.ColorAndOpacity(this, &SStatsTableCell::GetExtraColorAndOpacity)
+			.ShadowColorAndOpacity(this, &SStatsTableCell::GetShadowColorAndOpacity)
+		]
 	;
 }
 
diff --git a/Engine/Source/Developer/TraceInsights/Private/Insights/Widgets/SStatsTableCell.h b/Engine/Source/Developer/TraceInsights/Private/Insights/Widgets/SStatsTableCell.h
--- a/Engine/Source/Developer/TraceInsights/Private/Insights/Widgets/SStatsTableCell.h
+++ b/Engine/Source/Developer/TraceInsights/Private/Insights/Widgets/SStatsTableCell.h
@@ -112,6 +112,16 @@ protected:
 		return StatsNodePtr->GetDisplayName();
 	}
 
+	FText GetExtraDisplayName() const
+	{
+		return StatsNodePtr->GetExtraDisplayName();
+	}
+
+	EVisibility HasExtraDisplayName() const
+	{
+		return StatsNodePtr->HasExtraDisplayName() ? EVisibility::Visible : EVisibility::Collapsed;
+	}
+
 	FText GetValueAsText() const;
 
 	FSlateColor GetColorAndOpacity() const
@@ -123,6 +133,15 @@ protected:
 		return TextColor;
 	}
 
+	FSlateColor GetExtraColorAndOpacity() const
+	{
+		const FLinearColor TextColor =
+			StatsNodePtr->IsFiltered() ?
+				FLinearColor(0.5f, 0.5f, 0.5f, 0.5f) :
+				FLinearColor(0.5f, 0.5f, 0.5f, 1.0f);
+		return TextColor;
+	}
+
 	FSlateColor GetStatsColorAndOpacity() const
 	{
 		const FLinearColor TextColor =
diff --git a/Engine/Source/Developer/TraceInsights/Private/Insights/Widgets/STimerTableCell.cpp b/Engine/Source/Developer/TraceInsights/Private/Insights/Widgets/STimerTableCell.cpp
--- a/Engine/Source/Developer/TraceInsights/Private/Insights/Widgets/STimerTableCell.cpp
+++ b/Engine/Source/Developer/TraceInsights/Private/Insights/Widgets/STimerTableCell.cpp
@@ -120,7 +120,7 @@ TSharedRef<SWidget> STimerTableCell::GenerateWidgetForNameColumn(const FArgument
 		+ SHorizontalBox::Slot()
 		.AutoWidth()
 		.VAlign(VAlign_Center)
-		.HAlign(ColumnPtr->GetHorizontalAlignment())
+		.HAlign(HAlign_Left)
 		.Padding(FMargin(2.0f, 0.0f))
 		[
 			SNew(STextBlock)
@@ -130,6 +130,21 @@ TSharedRef<SWidget> STimerTableCell::GenerateWidgetForNameColumn(const FArgument
 			.ColorAndOpacity(this, &STimerTableCell::GetColorAndOpacity)
 			.ShadowColorAndOpacity(this, &STimerTableCell::GetShadowColorAndOpacity)
 		]
+
+		// Name Suffix
+		+ SHorizontalBox::Slot()
+		.AutoWidth()
+		.VAlign(VAlign_Center)
+		.HAlign(HAlign_Left)
+		.Padding(FMargin(2.0f, 0.0f))
+		[
+			SNew(STextBlock)
+			.Visibility(this, &STimerTableCell::HasExtraDisplayName)
+			.Text(this, &STimerTableCell::GetExtraDisplayName)
+			.TextStyle(FEditorStyle::Get(), TEXT("Profiler.Tooltip"))
+			.ColorAndOpacity(this, &STimerTableCell::GetExtraColorAndOpacity)
+			.ShadowColorAndOpacity(this, &STimerTableCell::GetShadowColorAndOpacity)
+		]
 	;
 }
 
diff --git a/Engine/Source/Developer/TraceInsights/Private/Insights/Widgets/STimerTableCell.h b/Engine/Source/Developer/TraceInsights/Private/Insights/Widgets/STimerTableCell.h
--- a/Engine/Source/Developer/TraceInsights/Private/Insights/Widgets/STimerTableCell.h
+++ b/Engine/Source/Developer/TraceInsights/Private/Insights/Widgets/STimerTableCell.h
@@ -117,6 +117,16 @@ protected:
 		return TimerNodePtr->GetDisplayName();
 	}
 
+	FText GetExtraDisplayName() const
+	{
+		return TimerNodePtr->GetExtraDisplayName();
+	}
+
+	EVisibility HasExtraDisplayName() const
+	{
+		return TimerNodePtr->HasExtraDisplayName() ? EVisibility::Visible : EVisibility::Collapsed;
+	}
+
 	FText GetValueAsText() const;
 
 	FSlateColor GetColorAndOpacity() const
@@ -128,6 +138,15 @@ protected:
 		return TextColor;
 	}
 
+	FSlateColor GetExtraColorAndOpacity() const
+	{
+		const FLinearColor TextColor =
+			TimerNodePtr->IsFiltered() ?
+				FLinearColor(0.5f, 0.5f, 0.5f, 0.5f) :
+				FLinearColor(0.5f, 0.5f, 0.5f, 1.0f);
+		return TextColor;
+	}
+
 	FSlateColor GetStatsColorAndOpacity() const
 	{
 		const FLinearColor TextColor =
