[Insights]
  - Memory: Added "Callstack Grouping by Function" toggle button for callstack grouping to merge nodes by function.
  - Memory: Added callstack frame info in tooltip of group nodes (when grouping by callstack) showing address, function, source file / line number and the status of resolving.

diff --git a/Engine/Source/Developer/TraceInsights/Private/Insights/MemoryProfiler/ViewModels/MemAllocGroupingByCallstack.cpp b/Engine/Source/Developer/TraceInsights/Private/Insights/MemoryProfiler/ViewModels/MemAllocGroupingByCallstack.cpp
--- a/Engine/Source/Developer/TraceInsights/Private/Insights/MemoryProfiler/ViewModels/MemAllocGroupingByCallstack.cpp
+++ b/Engine/Source/Developer/TraceInsights/Private/Insights/MemoryProfiler/ViewModels/MemAllocGroupingByCallstack.cpp
@@ -19,7 +19,7 @@ INSIGHTS_IMPLEMENT_RTTI(FMemAllocGroupingByCallstack)
 
 ////////////////////////////////////////////////////////////////////////////////////////////////////
 
-FMemAllocGroupingByCallstack::FMemAllocGroupingByCallstack(bool bInIsInverted)
+FMemAllocGroupingByCallstack::FMemAllocGroupingByCallstack(bool bInIsInverted, bool bInIsGroupingByFunction)
 	: FTreeNodeGrouping(
 		bInIsInverted ? LOCTEXT("Grouping_ByCallstack2_ShortName", "Inverted Callstack")
 					  : LOCTEXT("Grouping_ByCallstack1_ShortName", "Callstack"),
@@ -29,6 +29,7 @@ FMemAllocGroupingByCallstack::FMemAllocGroupingByCallstack(bool bInIsInverted)
 		TEXT("Profiler.FiltersAndPresets.GroupNameIcon"),
 		nullptr)
 	, bIsInverted(bInIsInverted)
+	, bIsGroupingByFunction(bInIsGroupingByFunction)
 {
 }
 
@@ -49,6 +50,7 @@ void FMemAllocGroupingByCallstack::GroupNodes(const TArray<FTableTreeNodePtr>& N
 		FTableTreeNode* Node = nullptr;
 		const TraceServices::FStackFrame* Frame = nullptr;
 		TMap<uint64, FCallstackGroup*> GroupMap; // Callstack Frame Address -> FCallstackGroup*
+		TMap<FName, FCallstackGroup*> GroupMapByName; // Group Name --> FCallstackGroup*
 	};
 	TArray<FCallstackGroup*> CallstackGroups;
 
@@ -56,9 +58,8 @@ void FMemAllocGroupingByCallstack::GroupNodes(const TArray<FTableTreeNodePtr>& N
 	Root->Node = &ParentGroup;
 	CallstackGroups.Add(Root);
 
-	FTableTreeNodePtr UnsetGroupPtr = nullptr;
-
-	static FName NotAvailableName(TEXT("N/A"));
+	FTableTreeNode* UnsetGroupPtr = nullptr;
+	TMap<const TraceServices::FCallstack*, FCallstackGroup*> GroupMapByCallstack; // Callstack* -> FCallstackGroup*
 
 	for (FTableTreeNodePtr NodePtr : Nodes)
 	{
@@ -88,40 +89,46 @@ void FMemAllocGroupingByCallstack::GroupNodes(const TArray<FTableTreeNodePtr>& N
 					const TraceServices::FStackFrame* Frame = Callstack->Frame(bIsInverted ? NumFrames - FrameDepth - 1 : FrameDepth);
 					check(Frame != nullptr);
 
-					FCallstackGroup** GroupPtrPtr = GroupPtr->GroupMap.Find(Frame->Addr);
-					if (!GroupPtrPtr)
+					if (bIsGroupingByFunction)
 					{
-						FCallstackGroup* NewGroupPtr = new FCallstackGroup();
-						CallstackGroups.Add(NewGroupPtr);
-						GroupPtr->GroupMap.Add(Frame->Addr, NewGroupPtr);
-
-						FName GroupName = NotAvailableName;
-						const TraceServices::ESymbolQueryResult Result = Frame->Symbol->GetResult();
-						switch (Result)
+						const FName GroupName = GetGroupName(Frame);
+						FCallstackGroup** GroupPtrPtr = GroupPtr->GroupMapByName.Find(GroupName);
+						if (!GroupPtrPtr)
 						{
-						case TraceServices::ESymbolQueryResult::Pending:
-							//GroupName = PendingName;
-							GroupName = FName(*FString::Printf(TEXT("0x%X (...)"), Frame->Addr));
-							break;
-						case TraceServices::ESymbolQueryResult::OK:
-							GroupName = FName(Frame->Symbol->Name);
-							//GroupName = FName(*FString::Printf(TEXT("0x%X %s"), Frame->Addr, Frame->Symbol->Name));
-							break;
-						default:
-							GroupName = FName(*FString::Printf(TEXT("0x%X"), Frame->Addr));
-							break;
+							FCallstackGroup* NewGroupPtr = new FCallstackGroup();
+							CallstackGroups.Add(NewGroupPtr);
+
+							NewGroupPtr->Node = CreateGroup(GroupName, InParentTable, *(GroupPtr->Node));
+							NewGroupPtr->Node->SetTooltip(GetGroupTooltip(Frame));
+							
+							GroupPtr->GroupMapByName.Add(GroupName, NewGroupPtr);
+							GroupPtr = NewGroupPtr;
+						}
+						else
+						{
+							GroupPtr = *GroupPtrPtr;
 						}
-
-						FTableTreeNodePtr NewGroupNodePtr = MakeShared<FTableTreeNode>(GroupName, InParentTable);
-						NewGroupPtr->Node = NewGroupNodePtr.Get();
-						NewGroupNodePtr->SetExpansion(false);
-						GroupPtr->Node->AddChildAndSetGroupPtr(NewGroupNodePtr);
-
-						GroupPtr = NewGroupPtr;
 					}
 					else
 					{
-						GroupPtr = *GroupPtrPtr;
+						FCallstackGroup** GroupPtrPtr = GroupPtr->GroupMap.Find(Frame->Addr);
+						if (!GroupPtrPtr)
+						{
+							const FName GroupName = GetGroupName(Frame);
+
+							FCallstackGroup* NewGroupPtr = new FCallstackGroup();
+							CallstackGroups.Add(NewGroupPtr);
+
+							NewGroupPtr->Node = CreateGroup(GroupName, InParentTable, *(GroupPtr->Node));
+							NewGroupPtr->Node->SetTooltip(GetGroupTooltip(Frame));
+
+							GroupPtr->GroupMap.Add(Frame->Addr, NewGroupPtr);
+							GroupPtr = NewGroupPtr;
+						}
+						else
+						{
+							GroupPtr = *GroupPtrPtr;
+						}
 					}
 				}
 			}
@@ -135,9 +142,7 @@ void FMemAllocGroupingByCallstack::GroupNodes(const TArray<FTableTreeNodePtr>& N
 		{
 			if (!UnsetGroupPtr)
 			{
-				UnsetGroupPtr = MakeShared<FTableTreeNode>(NotAvailableName, InParentTable);
-				UnsetGroupPtr->SetExpansion(false);
-				ParentGroup.AddChildAndSetGroupPtr(UnsetGroupPtr);
+				UnsetGroupPtr = CreateUnsetGroup(InParentTable, ParentGroup);
 			}
 			UnsetGroupPtr->AddChildAndSetGroupPtr(NodePtr);
 		}
@@ -151,6 +156,66 @@ void FMemAllocGroupingByCallstack::GroupNodes(const TArray<FTableTreeNodePtr>& N
 
 ////////////////////////////////////////////////////////////////////////////////////////////////////
 
+FName FMemAllocGroupingByCallstack::GetGroupName(const TraceServices::FStackFrame* Frame) const
+{
+	const TraceServices::ESymbolQueryResult Result = Frame->Symbol->GetResult();
+	if (Result == TraceServices::ESymbolQueryResult::OK)
+	{
+		return FName(Frame->Symbol->Name);
+	}
+	else if (Result == TraceServices::ESymbolQueryResult::Pending)
+	{
+		return FName(*FString::Printf(TEXT("0x%X [...]"), Frame->Addr));
+	}
+	else
+	{
+		return FName(*FString::Printf(TEXT("0x%X"), Frame->Addr));
+	}
+}
+
+////////////////////////////////////////////////////////////////////////////////////////////////////
+
+FText FMemAllocGroupingByCallstack::GetGroupTooltip(const TraceServices::FStackFrame* Frame) const
+{
+	const TraceServices::ESymbolQueryResult Result = Frame->Symbol->GetResult();
+	if (Result == TraceServices::ESymbolQueryResult::OK)
+	{
+		return FText::Format(LOCTEXT("CallstackFrameTooltipFmt1", "Callstack Frame\n\t{0}\n\t{1}\n\t{2}"),
+			FText::FromString(FString::Printf(TEXT("0x%X"), Frame->Addr)),
+			FText::FromString(FString(Frame->Symbol->Name)),
+			FText::FromString(FString(Frame->Symbol->FileAndLine)));
+	}
+	else
+	{
+		return FText::Format(LOCTEXT("CallstackFrameTooltipFmt2", "Callstack Frame\n\t{0}\n\t{1}"),
+			FText::FromString(FString::Printf(TEXT("0x%X"), Frame->Addr)),
+			FText::FromString(FString(TraceServices::QueryResultToString(Result))));
+	}
+}
+
+////////////////////////////////////////////////////////////////////////////////////////////////////
+
+FTableTreeNode* FMemAllocGroupingByCallstack::CreateGroup(const FName GroupName, TWeakPtr<FTable> ParentTable, FTableTreeNode& Parent) const
+{
+	FTableTreeNodePtr NodePtr = MakeShared<FTableTreeNode>(GroupName, ParentTable);
+	NodePtr->SetExpansion(false);
+	Parent.AddChildAndSetGroupPtr(NodePtr);
+	return NodePtr.Get();
+}
+
+////////////////////////////////////////////////////////////////////////////////////////////////////
+
+FTableTreeNode* FMemAllocGroupingByCallstack::CreateUnsetGroup(TWeakPtr<FTable> ParentTable, FTableTreeNode& Parent) const
+{
+	static FName NotAvailableName(TEXT("N/A"));
+	FTableTreeNodePtr NodePtr = MakeShared<FTableTreeNode>(NotAvailableName, ParentTable);
+	NodePtr->SetExpansion(false);
+	Parent.AddChildAndSetGroupPtr(NodePtr);
+	return NodePtr.Get();
+}
+
+////////////////////////////////////////////////////////////////////////////////////////////////////
+
 } // namespace Insights
 
 #undef LOCTEXT_NAMESPACE
diff --git a/Engine/Source/Developer/TraceInsights/Private/Insights/MemoryProfiler/ViewModels/MemAllocGroupingByCallstack.h b/Engine/Source/Developer/TraceInsights/Private/Insights/MemoryProfiler/ViewModels/MemAllocGroupingByCallstack.h
--- a/Engine/Source/Developer/TraceInsights/Private/Insights/MemoryProfiler/ViewModels/MemAllocGroupingByCallstack.h
+++ b/Engine/Source/Developer/TraceInsights/Private/Insights/MemoryProfiler/ViewModels/MemAllocGroupingByCallstack.h
@@ -8,6 +8,11 @@
 #include "Insights/Table/ViewModels/BaseTreeNode.h"
 #include "Insights/Table/ViewModels/TreeNodeGrouping.h"
 
+namespace TraceServices
+{
+	struct FStackFrame;
+}
+
 namespace Insights
 {
 
@@ -18,15 +23,26 @@ class FMemAllocGroupingByCallstack : public FTreeNodeGrouping
 	INSIGHTS_DECLARE_RTTI(FMemAllocGroupingByCallstack, FTreeNodeGrouping)
 
 public:
-	FMemAllocGroupingByCallstack(bool bInIsInverted);
+	FMemAllocGroupingByCallstack(bool bInIsInverted, bool bInIsGroupingByFunction);
 	virtual ~FMemAllocGroupingByCallstack();
 
 	virtual void GroupNodes(const TArray<FTableTreeNodePtr>& Nodes, FTableTreeNode& ParentGroup, TWeakPtr<FTable> InParentTable, std::atomic<bool>& bCancelGrouping) const override;
 
 	bool IsInverted() const { return bIsInverted; }
 
+	bool IsGroupingByFunction() const { return bIsGroupingByFunction; }
+	void SetGroupingByFunction(bool bOnOff) { bIsGroupingByFunction = bOnOff; }
+
+private:
+	FName GetGroupName(const TraceServices::FStackFrame* Frame) const;
+	FText GetGroupTooltip(const TraceServices::FStackFrame* Frame) const;
+
+	FTableTreeNode* CreateGroup(const FName GroupName, TWeakPtr<FTable> ParentTable, FTableTreeNode& Parent) const;
+	FTableTreeNode* CreateUnsetGroup(TWeakPtr<FTable> ParentTable, FTableTreeNode& Parent) const;
+
 private:
 	bool bIsInverted;
+	bool bIsGroupingByFunction;
 };
 
 ////////////////////////////////////////////////////////////////////////////////////////////////////
diff --git a/Engine/Source/Developer/TraceInsights/Private/Insights/MemoryProfiler/Widgets/SMemAllocTableTreeView.cpp b/Engine/Source/Developer/TraceInsights/Private/Insights/MemoryProfiler/Widgets/SMemAllocTableTreeView.cpp
--- a/Engine/Source/Developer/TraceInsights/Private/Insights/MemoryProfiler/Widgets/SMemAllocTableTreeView.cpp
+++ b/Engine/Source/Developer/TraceInsights/Private/Insights/MemoryProfiler/Widgets/SMemAllocTableTreeView.cpp
@@ -447,6 +447,23 @@ TSharedPtr<SWidget> SMemAllocTableTreeView::ConstructToolbar()
 			.ToolTipText(LOCTEXT("ByInvertedCallstackBtn_Tooltip", "Inverted Callstack Breakdown View\nConfigure the tree view to show a breakdown of allocations by inverted callstack."))
 			.OnClicked(this, &SMemAllocTableTreeView::OnCallstackViewClicked, true)
 		]
+
+		+ SHorizontalBox::Slot()
+		.AutoWidth()
+		.Padding(4.0f, 0.0f, 0.0f, 0.0f)
+		[
+			SNew(SCheckBox)
+			.Style(FCoreStyle::Get(), "ToggleButtonCheckbox")
+			.HAlign(HAlign_Center)
+			.Padding(2.0f)
+			.OnCheckStateChanged(this, &SMemAllocTableTreeView::CallstackGroupingByFunction_OnCheckStateChanged)
+			.IsChecked(this, &SMemAllocTableTreeView::CallstackGroupingByFunction_IsChecked)
+			.ToolTipText(LOCTEXT("CallstackGroupingByFunction_Tooltip", "Callstack Grouping by Function Name\nIf enabled, the callstack grouping will create a single group node per function name\n(ex. when two callstack frames have same function but different file line numbers).\nOtherwise it will create separate group nodes for each unique callstack frame."))
+			[
+				SNew(STextBlock)
+				.Text(LOCTEXT("CallstackGroupingByFunction_Text", " Fn "))
+			]
+		]
 	;
 }
 
@@ -728,8 +745,8 @@ void SMemAllocTableTreeView::InternalCreateGroupings()
 		AvailableGroupings.Insert(TagGrouping, Index++);
 	}
 
-	AvailableGroupings.Insert(MakeShared<FMemAllocGroupingByCallstack>(false), Index++);
-	AvailableGroupings.Insert(MakeShared<FMemAllocGroupingByCallstack>(true), Index++);
+	AvailableGroupings.Insert(MakeShared<FMemAllocGroupingByCallstack>(false, bIsCallstackGroupingByFunction), Index++);
+	AvailableGroupings.Insert(MakeShared<FMemAllocGroupingByCallstack>(true, bIsCallstackGroupingByFunction), Index++);
 }
 
 ////////////////////////////////////////////////////////////////////////////////////////////////////
@@ -787,6 +804,29 @@ void SMemAllocTableTreeView::AddCustomAdvancedFilters()
 
 ////////////////////////////////////////////////////////////////////////////////////////////////////
 
+void SMemAllocTableTreeView::CallstackGroupingByFunction_OnCheckStateChanged(ECheckBoxState NewRadioState)
+{
+	PreChangeGroupings();
+	bIsCallstackGroupingByFunction = (NewRadioState == ECheckBoxState::Checked);
+	for (TSharedPtr<FTreeNodeGrouping>& Grouping : AvailableGroupings)
+	{
+		if (Grouping->Is<FMemAllocGroupingByCallstack>())
+		{
+			Grouping->As<FMemAllocGroupingByCallstack>().SetGroupingByFunction(bIsCallstackGroupingByFunction);
+		}
+	}
+	PostChangeGroupings();
+}
+
+////////////////////////////////////////////////////////////////////////////////////////////////////
+
+ECheckBoxState SMemAllocTableTreeView::CallstackGroupingByFunction_IsChecked() const
+{
+	return bIsCallstackGroupingByFunction ? ECheckBoxState::Checked : ECheckBoxState::Unchecked;
+}
+
+////////////////////////////////////////////////////////////////////////////////////////////////////
+
 } // namespace Insights
 
 #undef LOCTEXT_NAMESPACE
diff --git a/Engine/Source/Developer/TraceInsights/Private/Insights/MemoryProfiler/Widgets/SMemAllocTableTreeView.h b/Engine/Source/Developer/TraceInsights/Private/Insights/MemoryProfiler/Widgets/SMemAllocTableTreeView.h
--- a/Engine/Source/Developer/TraceInsights/Private/Insights/MemoryProfiler/Widgets/SMemAllocTableTreeView.h
+++ b/Engine/Source/Developer/TraceInsights/Private/Insights/MemoryProfiler/Widgets/SMemAllocTableTreeView.h
@@ -114,6 +114,9 @@ private:
 	bool virtual ApplyCustomAdvancedFilters(const FTableTreeNodePtr& NodePtr) override;
 	virtual void AddCustomAdvancedFilters() override;
 
+	void CallstackGroupingByFunction_OnCheckStateChanged(ECheckBoxState NewRadioState);
+	ECheckBoxState CallstackGroupingByFunction_IsChecked() const;
+
 private:
 	const static int FullCallStackIndex;
 	int32 TabIndex = -1;
@@ -122,8 +125,8 @@ private:
 	TraceServices::IAllocationsProvider::FQueryHandle Query = 0;
 	FText QueryInfo;
 	FText QueryInfoTooltip;
-
 	FStopwatch QueryStopwatch;
+	bool bIsCallstackGroupingByFunction = true;
 };
 
 ////////////////////////////////////////////////////////////////////////////////////////////////////
