[Insights] Fix MemAllocTreeView spawning as an unregisterd tab on application restart. This fix covers a use case missed by the previous fix

diff --git a/Engine/Source/Developer/TraceInsights/Private/Insights/MemoryProfiler/MemoryProfilerManager.cpp b/Engine/Source/Developer/TraceInsights/Private/Insights/MemoryProfiler/MemoryProfilerManager.cpp
--- a/Engine/Source/Developer/TraceInsights/Private/Insights/MemoryProfiler/MemoryProfilerManager.cpp
+++ b/Engine/Source/Developer/TraceInsights/Private/Insights/MemoryProfiler/MemoryProfilerManager.cpp
@@ -367,4 +367,16 @@ void FMemoryProfilerManager::ShowHideMemTagTreeView(const bool bIsVisible)
 
 ////////////////////////////////////////////////////////////////////////////////////////////////////
 
+void FMemoryProfilerManager::OnWindowClosedEvent()
+{
+	// Need to close MemAlloc window to prevent it being saved in the layout and spawning as an "Unregister Tab" on the next application start. 
+	TSharedPtr<SMemoryProfilerWindow> Wnd = GetProfilerWindow();
+	if (Wnd.IsValid())
+	{
+		Wnd->CloseMemAllocTableTreeTabs();
+	}
+}
+
+////////////////////////////////////////////////////////////////////////////////////////////////////
+
 #undef LOCTEXT_NAMESPACE
diff --git a/Engine/Source/Developer/TraceInsights/Private/Insights/MemoryProfiler/MemoryProfilerManager.h b/Engine/Source/Developer/TraceInsights/Private/Insights/MemoryProfiler/MemoryProfilerManager.h
--- a/Engine/Source/Developer/TraceInsights/Private/Insights/MemoryProfiler/MemoryProfilerManager.h
+++ b/Engine/Source/Developer/TraceInsights/Private/Insights/MemoryProfiler/MemoryProfilerManager.h
@@ -51,6 +51,7 @@ public:
 	virtual void Shutdown() override;
 	virtual void RegisterMajorTabs(IUnrealInsightsModule& InsightsModule) override;
 	virtual void UnregisterMajorTabs() override;
+	virtual void OnWindowClosedEvent() override;
 
 	//////////////////////////////////////////////////
 
diff --git a/Engine/Source/Developer/TraceInsights/Private/Insights/MemoryProfiler/Widgets/SMemoryProfilerWindow.cpp b/Engine/Source/Developer/TraceInsights/Private/Insights/MemoryProfiler/Widgets/SMemoryProfilerWindow.cpp
--- a/Engine/Source/Developer/TraceInsights/Private/Insights/MemoryProfiler/Widgets/SMemoryProfilerWindow.cpp
+++ b/Engine/Source/Developer/TraceInsights/Private/Insights/MemoryProfiler/Widgets/SMemoryProfilerWindow.cpp
@@ -435,13 +435,12 @@ void SMemoryProfilerWindow::CloseMemAllocTableTreeTabs()
 	const TArray<TSharedPtr<Insights::FQueryTargetWindowSpec>>& Targets = this->GetSharedState().GetQueryTargets();
 	while(Targets.Num() > 0)
 	{
-		if (Targets[0]->GetName() != Insights::FQueryTargetWindowSpec::NewWindow)
-		{
-			HideTab(Targets[0]->GetName());
-		}
-		else
+		FName Name = Targets[0]->GetName();
+		this->GetSharedState().RemoveQueryTarget(Targets[0]);
+
+		if (Name != Insights::FQueryTargetWindowSpec::NewWindow)
 		{
-			this->GetSharedState().RemoveQueryTarget(Targets[0]);
+			HideTab(Name);
 		}
 	}
 }
diff --git a/Engine/Source/Developer/TraceInsights/Private/Insights/TraceInsightsModule.cpp b/Engine/Source/Developer/TraceInsights/Private/Insights/TraceInsightsModule.cpp
--- a/Engine/Source/Developer/TraceInsights/Private/Insights/TraceInsightsModule.cpp
+++ b/Engine/Source/Developer/TraceInsights/Private/Insights/TraceInsightsModule.cpp
@@ -420,6 +420,11 @@ void FTraceInsightsModule::StartAnalysisForTraceFile(const TCHAR* InTraceFile, b
 
 void FTraceInsightsModule::OnWindowClosedEvent(const TSharedRef<SWindow>&)
 {
+	for (TSharedRef<IInsightsComponent>& Component : Components)
+	{
+		Component->OnWindowClosedEvent();
+	}
+
 	FGlobalTabmanager::Get()->SaveAllVisualState();
 }
 
diff --git a/Engine/Source/Developer/TraceInsights/Public/Insights/IUnrealInsightsModule.h b/Engine/Source/Developer/TraceInsights/Public/Insights/IUnrealInsightsModule.h
--- a/Engine/Source/Developer/TraceInsights/Public/Insights/IUnrealInsightsModule.h
+++ b/Engine/Source/Developer/TraceInsights/Public/Insights/IUnrealInsightsModule.h
@@ -304,6 +304,9 @@ public:
 
 	/* Requests this component to unregister its major tabs. */
 	virtual void UnregisterMajorTabs() = 0;
+
+	/* Called by the TraceInsights module when it receives the OnWindowClosedEvent. Can be used to close any panels that should not persist in the layout. */
+	virtual void OnWindowClosedEvent() {};
 };
 
 ////////////////////////////////////////////////////////////////////////////////////////////////////
