[Insights] TableTreeView expand/collapse minor fixes

diff --git a/Engine/Source/Developer/TraceInsights/Private/Insights/Table/Widgets/STableTreeView.cpp b/Engine/Source/Developer/TraceInsights/Private/Insights/Table/Widgets/STableTreeView.cpp
--- a/Engine/Source/Developer/TraceInsights/Private/Insights/Table/Widgets/STableTreeView.cpp
+++ b/Engine/Source/Developer/TraceInsights/Private/Insights/Table/Widgets/STableTreeView.cpp
@@ -2860,21 +2860,28 @@ bool STableTreeView::ContextMenu_CopyColumnToClipboard_CanExecute() const
 {
 	const TSharedPtr<FTableColumn> HoveredColumnPtr = Table->FindColumn(HoveredColumnId);
 
-	if (!HoveredColumnPtr.IsValid() || !HoveredNodePtr.IsValid())
+	if (HoveredColumnPtr.IsValid() && TreeView->GetNumItemsSelected() == 1)
 	{
-		return false;
+		return true;
 	}
-	
-	return true;
+
+	return false;
 }
 
 ////////////////////////////////////////////////////////////////////////////////////////////////////
 
 void STableTreeView::ContextMenu_CopyColumnToClipboard_Execute()
 {
-	const TSharedPtr<FTableColumn> HoveredColumnPtr = Table->FindColumn(HoveredColumnId);
-	FString Text = HoveredColumnPtr->GetValueAsText(*HoveredNodePtr).ToString();
-	FPlatformApplicationMisc::ClipboardCopy(*Text);
+	if (TreeView->GetNumItemsSelected() > 0)
+	{
+		FTableTreeNodePtr SelectedNode = TreeView->GetSelectedItems()[0];
+		const TSharedPtr<FTableColumn> HoveredColumnPtr = Table->FindColumn(HoveredColumnId);
+		if (HoveredColumnPtr.IsValid())
+		{
+			FString Text = HoveredColumnPtr->GetValueAsText(*SelectedNode).ToString();
+			FPlatformApplicationMisc::ClipboardCopy(*Text);
+		}
+	}
 }
 
 ////////////////////////////////////////////////////////////////////////////////////////////////////
@@ -2883,37 +2890,35 @@ bool STableTreeView::ContextMenu_CopyColumnTooltipToClipboard_CanExecute() const
 {
 	const TSharedPtr<FTableColumn> HoveredColumnPtr = Table->FindColumn(HoveredColumnId);
 
-	if (!HoveredColumnPtr.IsValid() || !HoveredNodePtr.IsValid())
+	if (HoveredColumnPtr.IsValid() && TreeView->GetNumItemsSelected() == 1)
 	{
-		return false;
+		return true;
 	}
 
-	return true;
+	return false;
 }
 
 ////////////////////////////////////////////////////////////////////////////////////////////////////
 
 void STableTreeView::ContextMenu_CopyColumnTooltipToClipboard_Execute()
 {
-	const TSharedPtr<FTableColumn> HoveredColumnPtr = Table->FindColumn(HoveredColumnId);
-	FString Text = HoveredColumnPtr->GetValueAsTooltipText(*HoveredNodePtr).ToString();
-	FPlatformApplicationMisc::ClipboardCopy(*Text);
+	if (TreeView->GetNumItemsSelected() > 0)
+	{
+		FTableTreeNodePtr SelectedNode = TreeView->GetSelectedItems()[0];
+		const TSharedPtr<FTableColumn> HoveredColumnPtr = Table->FindColumn(HoveredColumnId);
+		if (HoveredColumnPtr.IsValid())
+		{
+			FString Text = HoveredColumnPtr->GetValueAsTooltipText(*SelectedNode).ToString();
+			FPlatformApplicationMisc::ClipboardCopy(*Text);
+		}
+	}
 }
 
 ////////////////////////////////////////////////////////////////////////////////////////////////////
 
 bool STableTreeView::ContextMenu_ExpandSubtree_CanExecute() const
 {
-	const TArray<FTableTreeNodePtr> SelectedNodes = TreeView->GetSelectedItems();
-	for (const FTableTreeNodePtr& Node : SelectedNodes)
-	{
-		if (Node->IsGroup() && Node->GetFilteredChildren().Num() > 0 && !TreeView->IsItemExpanded(Node))
-		{
-			return true;
-		}
-	}
-
-	return false;
+	return true;
 }
 
 ////////////////////////////////////////////////////////////////////////////////////////////////////
@@ -2923,7 +2928,7 @@ void STableTreeView::ContextMenu_ExpandSubtree_Execute()
 	const TArray<FTableTreeNodePtr> SelectedNodes = TreeView->GetSelectedItems();
 	for (const FTableTreeNodePtr& Node : SelectedNodes)
 	{
-		if (Node->IsGroup() && !TreeView->IsItemExpanded(Node))
+		if (Node->IsGroup())
 		{
 			Node->SetExpansion(true);
 			TreeView->SetItemExpansion(Node, true);
@@ -2972,16 +2977,7 @@ void STableTreeView::ContextMenu_CollapseSubtree_Execute()
 
 bool STableTreeView::ContextMenu_ExpandCriticalPath_CanExecute() const
 {
-	const TArray<FTableTreeNodePtr> SelectedNodes = TreeView->GetSelectedItems();
-	for (const FTableTreeNodePtr& Node : SelectedNodes)
-	{
-		if (Node->IsGroup() && Node->GetFilteredChildren().Num() > 0 && !TreeView->IsItemExpanded(Node))
-		{
-			return true;
-		}
-	}
-
-	return false;
+	return true;
 }
 
 ////////////////////////////////////////////////////////////////////////////////////////////////////
@@ -2992,8 +2988,9 @@ void STableTreeView::ContextMenu_ExpandCriticalPath_Execute()
 	for (const FTableTreeNodePtr& Node : SelectedNodes)
 	{
 		FTableTreeNodePtr CurrentNode = Node;
-		while (CurrentNode.IsValid() && CurrentNode->IsGroup())
+		while (CurrentNode->IsGroup())
 		{
+			check(CurrentNode.IsValid());
 			if (!TreeView->IsItemExpanded(CurrentNode))
 			{
 				CurrentNode->SetExpansion(true);
