[Insights]
  - Allocs Table: When "Fn" (Callstack Grouping by [REDACTED_PERSON]) toggle is enabled, callstack grouping will merge also the callstack frames of recusive calls into a single tree node.
  - Allocs Table: Improved speed of Callstack grouping by 4X (tested for a large set of 11.3M allocs with 124K callstack frames --> time to breakdown by callstack was reduced from 55.6s to 12.5s).
  - Allocs Table: Updated tooltip for "Fn" (Callstack Grouping by [REDACTED_PERSON]) option.

diff --git a/Engine/Source/Developer/TraceInsights/Private/Insights/MemoryProfiler/ViewModels/MemAllocGroupingByCallstack.cpp b/Engine/Source/Developer/TraceInsights/Private/Insights/MemoryProfiler/ViewModels/MemAllocGroupingByCallstack.cpp
--- a/Engine/Source/Developer/TraceInsights/Private/Insights/MemoryProfiler/ViewModels/MemAllocGroupingByCallstack.cpp
+++ b/Engine/Source/Developer/TraceInsights/Private/Insights/MemoryProfiler/ViewModels/MemAllocGroupingByCallstack.cpp
@@ -45,13 +45,6 @@ void FMemAllocGroupingByCallstack::GroupNodes(const TArray<FTableTreeNodePtr>& N
 {
 	ParentGroup.ClearChildren();
 
-	struct FCallstackGroup
-	{
-		FTableTreeNode* Node = nullptr;
-		const TraceServices::FStackFrame* Frame = nullptr;
-		TMap<uint64, FCallstackGroup*> GroupMap; // Callstack Frame Address -> FCallstackGroup*
-		TMap<FName, FCallstackGroup*> GroupMapByName; // Group Name --> FCallstackGroup*
-	};
 	TArray<FCallstackGroup*> CallstackGroups;
 
 	FCallstackGroup* Root = new FCallstackGroup();
@@ -81,7 +74,13 @@ void FMemAllocGroupingByCallstack::GroupNodes(const TArray<FTableTreeNodePtr>& N
 		if (Alloc)
 		{
 			const TraceServices::FCallstack* Callstack = Alloc->GetCallstack();
-			if (Callstack)
+
+			FCallstackGroup** FoundGroupPtrPtr = GroupMapByCallstack.Find(Callstack);
+			if (FoundGroupPtrPtr)
+			{
+				GroupPtr = *FoundGroupPtrPtr;
+			}
+			else if (Callstack)
 			{
 				const int32 NumFrames = static_cast<int32>(Callstack->Num());
 				for (int32 FrameDepth = NumFrames - 1; FrameDepth >= 0; --FrameDepth)
@@ -92,17 +91,19 @@ void FMemAllocGroupingByCallstack::GroupNodes(const TArray<FTableTreeNodePtr>& N
 					if (bIsGroupingByFunction)
 					{
 						const FName GroupName = GetGroupName(Frame);
+
+						// Merge with parent group, if it has the same name (i.e. same function).
+						if (GroupPtr->Parent != nullptr && GroupPtr->Name == GroupName)
+						{
+							GroupPtr = GroupPtr->Parent;
+						}
+
+						// Merge groups by name.
 						FCallstackGroup** GroupPtrPtr = GroupPtr->GroupMapByName.Find(GroupName);
 						if (!GroupPtrPtr)
 						{
-							FCallstackGroup* NewGroupPtr = new FCallstackGroup();
-							CallstackGroups.Add(NewGroupPtr);
-
-							NewGroupPtr->Node = CreateGroup(GroupName, InParentTable, *(GroupPtr->Node));
-							NewGroupPtr->Node->SetTooltip(GetGroupTooltip(Frame));
-							
-							GroupPtr->GroupMapByName.Add(GroupName, NewGroupPtr);
-							GroupPtr = NewGroupPtr;
+							GroupPtr = CreateGroup(CallstackGroups, GroupPtr, GroupName, InParentTable, Frame);
+							GroupPtr->Parent->GroupMapByName.Add(GroupName, GroupPtr);
 						}
 						else
 						{
@@ -111,19 +112,13 @@ void FMemAllocGroupingByCallstack::GroupNodes(const TArray<FTableTreeNodePtr>& N
 					}
 					else
 					{
+						// Merge groups by unique callstack frame.
 						FCallstackGroup** GroupPtrPtr = GroupPtr->GroupMap.Find(Frame->Addr);
 						if (!GroupPtrPtr)
 						{
 							const FName GroupName = GetGroupName(Frame);
-
-							FCallstackGroup* NewGroupPtr = new FCallstackGroup();
-							CallstackGroups.Add(NewGroupPtr);
-
-							NewGroupPtr->Node = CreateGroup(GroupName, InParentTable, *(GroupPtr->Node));
-							NewGroupPtr->Node->SetTooltip(GetGroupTooltip(Frame));
-
-							GroupPtr->GroupMap.Add(Frame->Addr, NewGroupPtr);
-							GroupPtr = NewGroupPtr;
+							GroupPtr = CreateGroup(CallstackGroups, GroupPtr, GroupName, InParentTable, Frame);
+							GroupPtr->Parent->GroupMap.Add(Frame->Addr, GroupPtr);
 						}
 						else
 						{
@@ -131,6 +126,8 @@ void FMemAllocGroupingByCallstack::GroupNodes(const TArray<FTableTreeNodePtr>& N
 						}
 					}
 				}
+
+				GroupMapByCallstack.Add(Callstack, GroupPtr);
 			}
 		}
 
@@ -195,12 +192,27 @@ FText FMemAllocGroupingByCallstack::GetGroupTooltip(const TraceServices::FStackF
 
 ////////////////////////////////////////////////////////////////////////////////////////////////////
 
-FTableTreeNode* FMemAllocGroupingByCallstack::CreateGroup(const FName GroupName, TWeakPtr<FTable> ParentTable, FTableTreeNode& Parent) const
+FMemAllocGroupingByCallstack::FCallstackGroup* FMemAllocGroupingByCallstack::CreateGroup(
+	TArray<FCallstackGroup*>& InOutAllCallstackGroup,
+	FCallstackGroup* InParentGroup,
+	const FName InGroupName,
+	TWeakPtr<FTable> InParentTable,
+	const TraceServices::FStackFrame* InFrame) const
 {
-	FTableTreeNodePtr NodePtr = MakeShared<FTableTreeNode>(GroupName, ParentTable);
+	FCallstackGroup* NewGroupPtr = new FCallstackGroup();
+	NewGroupPtr->Parent = InParentGroup;
+	NewGroupPtr->Name = InGroupName;
+
+	InOutAllCallstackGroup.Add(NewGroupPtr);
+
+	FTableTreeNodePtr NodePtr = MakeShared<FTableTreeNode>(InGroupName, InParentTable);
 	NodePtr->SetExpansion(false);
-	Parent.AddChildAndSetGroupPtr(NodePtr);
-	return NodePtr.Get();
+	InParentGroup->Node->AddChildAndSetGroupPtr(NodePtr);
+
+	NewGroupPtr->Node = NodePtr.Get();
+	NewGroupPtr->Node->SetTooltip(GetGroupTooltip(InFrame));
+
+	return NewGroupPtr;
 }
 
 ////////////////////////////////////////////////////////////////////////////////////////////////////
diff --git a/Engine/Source/Developer/TraceInsights/Private/Insights/MemoryProfiler/ViewModels/MemAllocGroupingByCallstack.h b/Engine/Source/Developer/TraceInsights/Private/Insights/MemoryProfiler/ViewModels/MemAllocGroupingByCallstack.h
--- a/Engine/Source/Developer/TraceInsights/Private/Insights/MemoryProfiler/ViewModels/MemAllocGroupingByCallstack.h
+++ b/Engine/Source/Developer/TraceInsights/Private/Insights/MemoryProfiler/ViewModels/MemAllocGroupingByCallstack.h
@@ -22,6 +22,17 @@ class FMemAllocGroupingByCallstack : public FTreeNodeGrouping
 {
 	INSIGHTS_DECLARE_RTTI(FMemAllocGroupingByCallstack, FTreeNodeGrouping)
 
+private:
+	struct FCallstackGroup
+	{
+		FCallstackGroup* Parent = nullptr;
+		FName Name;
+		FTableTreeNode* Node = nullptr;
+		const TraceServices::FStackFrame* Frame = nullptr;
+		TMap<uint64, FCallstackGroup*> GroupMap; // Callstack Frame Address -> FCallstackGroup*
+		TMap<FName, FCallstackGroup*> GroupMapByName; // Group Name --> FCallstackGroup*
+	};
+
 public:
 	FMemAllocGroupingByCallstack(bool bInIsInverted, bool bInIsGroupingByFunction);
 	virtual ~FMemAllocGroupingByCallstack();
@@ -37,7 +48,8 @@ private:
 	FName GetGroupName(const TraceServices::FStackFrame* Frame) const;
 	FText GetGroupTooltip(const TraceServices::FStackFrame* Frame) const;
 
-	FTableTreeNode* CreateGroup(const FName GroupName, TWeakPtr<FTable> ParentTable, FTableTreeNode& Parent) const;
+	FCallstackGroup* CreateGroup(TArray<FCallstackGroup*>& InOutAllCallstackGroup, FCallstackGroup* InParentGroup, const FName InGroupName, TWeakPtr<FTable> InParentTable, const TraceServices::FStackFrame* InFrame) const;
+
 	FTableTreeNode* CreateUnsetGroup(TWeakPtr<FTable> ParentTable, FTableTreeNode& Parent) const;
 
 private:
diff --git a/Engine/Source/Developer/TraceInsights/Private/Insights/MemoryProfiler/ViewModels/MemAllocTable.cpp b/Engine/Source/Developer/TraceInsights/Private/Insights/MemoryProfiler/ViewModels/MemAllocTable.cpp
--- a/Engine/Source/Developer/TraceInsights/Private/Insights/MemoryProfiler/ViewModels/MemAllocTable.cpp
+++ b/Engine/Source/Developer/TraceInsights/Private/Insights/MemoryProfiler/ViewModels/MemAllocTable.cpp
@@ -484,7 +484,7 @@ void FMemAllocTable::AddDefaultColumns()
 		Column.SetIndex(ColumnIndex++);
 
 		Column.SetShortName(LOCTEXT("FunctionColumnName", "Function"));
-		Column.SetTitleName(LOCTEXT("FunctionColumnTitle", "Function"));
+		Column.SetTitleName(LOCTEXT("FunctionColumnTitle", "Function (Callstack)"));
 		Column.SetDescription(LOCTEXT("FunctionColumnDesc", "Resolved top function from the callstack of allocation"));
 
 		Column.SetFlags(ETableColumnFlags::ShouldBeVisible | ETableColumnFlags::CanBeHidden | ETableColumnFlags::CanBeFiltered);
diff --git a/Engine/Source/Developer/TraceInsights/Private/Insights/MemoryProfiler/Widgets/SMemAllocTableTreeView.cpp b/Engine/Source/Developer/TraceInsights/Private/Insights/MemoryProfiler/Widgets/SMemAllocTableTreeView.cpp
--- a/Engine/Source/Developer/TraceInsights/Private/Insights/MemoryProfiler/Widgets/SMemAllocTableTreeView.cpp
+++ b/Engine/Source/Developer/TraceInsights/Private/Insights/MemoryProfiler/Widgets/SMemAllocTableTreeView.cpp
@@ -2,10 +2,12 @@
 
 #include "SMemAllocTableTreeView.h"
 
+#include "EditorStyleSet.h"
 #include "TraceServices/AnalysisService.h"
 #include "TraceServices/Model/AllocationsProvider.h"
 #include "TraceServices/Model/Callstack.h"
 #include "TraceServices/Model/Modules.h"
+#include "Widgets/SToolTip.h"
 
 // Insights
 #include "Insights/MemoryProfiler/MemoryProfilerManager.h"
@@ -455,10 +457,59 @@ TSharedPtr<SWidget> SMemAllocTableTreeView::ConstructToolbar()
 			SNew(SCheckBox)
 			.Style(FCoreStyle::Get(), "ToggleButtonCheckbox")
 			.HAlign(HAlign_Center)
-			.Padding(2.0f)
+			.Padding(FMargin(2.0f, 4.0f, 2.0f, 0.0f))
 			.OnCheckStateChanged(this, &SMemAllocTableTreeView::CallstackGroupingByFunction_OnCheckStateChanged)
 			.IsChecked(this, &SMemAllocTableTreeView::CallstackGroupingByFunction_IsChecked)
-			.ToolTipText(LOCTEXT("CallstackGroupingByFunction_Tooltip", "Callstack Grouping by Function Name\nIf enabled, the callstack grouping will create a single group node per function name\n(ex. when two callstack frames have same function but different file line numbers).\nOtherwise it will create separate group nodes for each unique callstack frame."))
+			.ToolTip(
+				SNew(SToolTip)
+				[
+					SNew(SVerticalBox)
+
+					+ SVerticalBox::Slot()
+					.AutoHeight()
+					.Padding(2.0f)
+					[
+						SNew(STextBlock)
+						.Text(LOCTEXT("CallstackGroupingByFunction_Tooltip_Title", "Callstack Grouping by Function Name"))
+						.TextStyle(FEditorStyle::Get(), TEXT("Profiler.TooltipBold"))
+					]
+
+					+ SVerticalBox::Slot()
+					.AutoHeight()
+					.Padding(2.0f, 8.0f, 2.0f, 2.0f)
+					[
+						SNew(STextBlock)
+						.Text(LOCTEXT("CallstackGroupingByFunction_Tooltip_Content", "If enabled, the callstack grouping will create a single group node per function name.\nExample 1: When two callstack frames are located in same function, but at different line numbers; \nExample 2: When a function is called recursively.\nOtherwise it will create separate group nodes for each unique callstack frame."))
+						.TextStyle(FEditorStyle::Get(), TEXT("Profiler.Tooltip"))
+					]
+
+					+ SVerticalBox::Slot()
+					.AutoHeight()
+					.Padding(2.0f, 8.0f, 2.0f, 2.0f)
+					[
+						SNew(SHorizontalBox)
+
+						+ SHorizontalBox::Slot()
+						.AutoWidth()
+						.VAlign(VAlign_Top)
+						.Padding(0.0f)
+						[
+							SNew(STextBlock)
+							.Text(LOCTEXT("CallstackGroupingByFunction_Warning", "Warning:"))
+							.TextStyle(FEditorStyle::Get(), TEXT("Profiler.Tooltip"))
+							.ColorAndOpacity(FLinearColor(1.0f, 0.6f, 0.3f, 1.0f))
+						]
+
+						+ SHorizontalBox::Slot()
+						.AutoWidth()
+						.Padding(4.0f, 0.0f, 0.0f, 0.0f)
+						[
+							SNew(STextBlock)
+							.Text(LOCTEXT("CallstackGroupingByFunction_Warning_Content", "When this option is enabled, the tree nodes that have merged multiple callstack frames\nwill show in their tooltips the source file name and the line number of an arbitrary\ncallstack frame from ones merged by respective tree node."))
+							.TextStyle(FEditorStyle::Get(), TEXT("Profiler.Tooltip"))
+						]
+					]
+				])
 			[
 				SNew(STextBlock)
 				.Text(LOCTEXT("CallstackGroupingByFunction_Text", " Fn "))
