[Insights] Fixed freeze (infinite loop) when callers/callees are updating following a selection in Timing view.

diff --git a/Engine/Source/Developer/TraceInsights/Private/Insights/ViewModels/TimerNode.cpp b/Engine/Source/Developer/TraceInsights/Private/Insights/ViewModels/TimerNode.cpp
--- a/Engine/Source/Developer/TraceInsights/Private/Insights/ViewModels/TimerNode.cpp
+++ b/Engine/Source/Developer/TraceInsights/Private/Insights/ViewModels/TimerNode.cpp
@@ -14,7 +14,7 @@ const FName FTimerNode::CpuGroup(TEXT("CPU"));
 ////////////////////////////////////////////////////////////////////////////////////////////////////
 
 FTimerNode::FTimerNode(uint32 InTimerId, const TCHAR* InName, ETimerNodeType InType)
-	: FBaseTreeNode(FName(InName), InType == ETimerNodeType::Group)
+	: FBaseTreeNode(FName(InName), true) // Timer nodes are groups (for callers/callees use case)
 	, TimerId(InTimerId)
 	, MetaGroupName(InType == ETimerNodeType::CpuScope ? CpuGroup : InType == ETimerNodeType::GpuScope ? GpuGroup : NAME_None)
 	, Type(InType)
