[Insights] TableTreeView searchbox is no longer disabled when the tree is empty

diff --git a/Engine/Source/Developer/TraceInsights/Private/Insights/Table/Widgets/STableTreeView.cpp b/Engine/Source/Developer/TraceInsights/Private/Insights/Table/Widgets/STableTreeView.cpp
--- a/Engine/Source/Developer/TraceInsights/Private/Insights/Table/Widgets/STableTreeView.cpp
+++ b/Engine/Source/Developer/TraceInsights/Private/Insights/Table/Widgets/STableTreeView.cpp
@@ -1203,7 +1203,7 @@ void STableTreeView::SearchBox_OnTextChanged(const FText& InFilterText)
 
 bool STableTreeView::SearchBox_IsEnabled() const
 {
-	return TableTreeNodes.Num() > 0 && !FilterConfigurator_HasFilters();
+	return !FilterConfigurator_HasFilters();
 }
 
 ////////////////////////////////////////////////////////////////////////////////////////////////////
@@ -1212,10 +1212,10 @@ FText STableTreeView::SearchBox_GetTooltipText() const
 {
 	if (SearchBox_IsEnabled())
 	{
-		return LOCTEXT("SearchBox_ToolTip", "Type here to search the tree hierarchy by item or group name");
+		return LOCTEXT("SearchBox_ToolTip", "Type here to search the tree hierarchy by item or group name.");
 	}
 
-	return LOCTEXT("SearchBox_Disabled_ToolTip", "Searching the tree hierarchy is disabled when advanced filters are set or when there are no entries in the table");
+	return LOCTEXT("SearchBox_Disabled_ToolTip", "Searching the tree hierarchy is disabled when advanced filters are set.");
 }
 
 ////////////////////////////////////////////////////////////////////////////////////////////////////
