[Insights] Fixed static analysis warnings.

diff --git a/Engine/Source/Developer/TraceInsights/Private/Insights/MemoryProfiler/ViewModels/MemAllocGroupingByCallstack.cpp b/Engine/Source/Developer/TraceInsights/Private/Insights/MemoryProfiler/ViewModels/MemAllocGroupingByCallstack.cpp
--- a/Engine/Source/Developer/TraceInsights/Private/Insights/MemoryProfiler/ViewModels/MemAllocGroupingByCallstack.cpp
+++ b/Engine/Source/Developer/TraceInsights/Private/Insights/MemoryProfiler/ViewModels/MemAllocGroupingByCallstack.cpp
@@ -103,6 +103,7 @@ void FMemAllocGroupingByCallstack::GroupNodes(const TArray<FTableTreeNodePtr>& N
 						if (!GroupPtrPtr)
 						{
 							GroupPtr = CreateGroup(CallstackGroups, GroupPtr, GroupName, InParentTable, Frame);
+							check(GroupPtr->Parent != nullptr);
 							GroupPtr->Parent->GroupMapByName.Add(GroupName, GroupPtr);
 						}
 						else
@@ -118,6 +119,7 @@ void FMemAllocGroupingByCallstack::GroupNodes(const TArray<FTableTreeNodePtr>& N
 						{
 							const FName GroupName = GetGroupName(Frame);
 							GroupPtr = CreateGroup(CallstackGroups, GroupPtr, GroupName, InParentTable, Frame);
+							check(GroupPtr->Parent != nullptr);
 							GroupPtr->Parent->GroupMap.Add(Frame->Addr, GroupPtr);
 						}
 						else
diff --git a/Engine/Source/Programs/UnrealInsights/Private/Windows/WindowsUnrealInsightsMain.cpp b/Engine/Source/Programs/UnrealInsights/Private/Windows/WindowsUnrealInsightsMain.cpp
--- a/Engine/Source/Programs/UnrealInsights/Private/Windows/WindowsUnrealInsightsMain.cpp
+++ b/Engine/Source/Programs/UnrealInsights/Private/Windows/WindowsUnrealInsightsMain.cpp
@@ -16,7 +16,7 @@
  * @param nShowCmd Specifies how the window is to be shown.
  * @return Application's exit value.
  */
-int32 WINAPI WinMain(HINSTANCE hInInstance, HINSTANCE hPrevInstance, char* lpCmdLine, int32 nShowCmd)
+int WINAPI WinMain(_In_ HINSTANCE hInInstance, _In_opt_ HINSTANCE hPrevInstance, _In_ LPSTR lpCmdLine, _In_ int nShowCmd)
 {
 	hInstance = hInInstance;
 
