[Insights] Fixed NonUnity compile errors.

diff --git a/Engine/Plugins/Animation/GameplayInsights/Source/GameplayInsights/Private/SVariantValueView.cpp b/Engine/Plugins/Animation/GameplayInsights/Source/GameplayInsights/Private/SVariantValueView.cpp
--- a/Engine/Plugins/Animation/GameplayInsights/Source/GameplayInsights/Private/SVariantValueView.cpp
+++ b/Engine/Plugins/Animation/GameplayInsights/Source/GameplayInsights/Private/SVariantValueView.cpp
@@ -18,6 +18,7 @@
 #include "VariantTreeNode.h"
 #include "Misc/TextFilterExpressionEvaluator.h"
 #include "Framework/Commands/GenericCommands.h"
+#include "Framework/Commands/UICommandList.h"
 #include "HAL/PlatformApplicationMisc.h"
 #include "Misc/StringBuilder.h"
 
diff --git a/Engine/Source/Developer/TraceInsights/Private/Insights/MemoryProfiler/Widgets/SMemAllocTableTreeView.cpp b/Engine/Source/Developer/TraceInsights/Private/Insights/MemoryProfiler/Widgets/SMemAllocTableTreeView.cpp
--- a/Engine/Source/Developer/TraceInsights/Private/Insights/MemoryProfiler/Widgets/SMemAllocTableTreeView.cpp
+++ b/Engine/Source/Developer/TraceInsights/Private/Insights/MemoryProfiler/Widgets/SMemAllocTableTreeView.cpp
@@ -3,10 +3,13 @@
 #include "SMemAllocTableTreeView.h"
 
 #include "EditorStyleSet.h"
+#include "SlateOptMacros.h"
+#include "Styling/CoreStyle.h"
 #include "TraceServices/AnalysisService.h"
 #include "TraceServices/Model/AllocationsProvider.h"
 #include "TraceServices/Model/Callstack.h"
 #include "TraceServices/Model/Modules.h"
+#include "Widgets/Input/SCheckBox.h"
 #include "Widgets/SToolTip.h"
 
 // Insights
@@ -395,6 +398,7 @@ FText SMemAllocTableTreeView::GetCurrentOperationName() const
 
 ////////////////////////////////////////////////////////////////////////////////////////////////////
 
+BEGIN_SLATE_FUNCTION_BUILD_OPTIMIZATION
 TSharedPtr<SWidget> SMemAllocTableTreeView::ConstructToolbar()
 {
 	return
@@ -517,6 +521,7 @@ TSharedPtr<SWidget> SMemAllocTableTreeView::ConstructToolbar()
 		]
 	;
 }
+END_SLATE_FUNCTION_BUILD_OPTIMIZATION
 
 ////////////////////////////////////////////////////////////////////////////////////////////////////
 
@@ -716,6 +721,7 @@ void SMemAllocTableTreeView::ApplyColumnConfig(const TArrayView<FColumnConfig>&
 
 ////////////////////////////////////////////////////////////////////////////////////////////////////
 
+BEGIN_SLATE_FUNCTION_BUILD_OPTIMIZATION
 TSharedPtr<SWidget> SMemAllocTableTreeView::ConstructFooter()
 {
 	return
@@ -738,6 +744,7 @@ TSharedPtr<SWidget> SMemAllocTableTreeView::ConstructFooter()
 			.Text(this, &SMemAllocTableTreeView::GetSymbolResolutionStatus)
 		];
 }
+END_SLATE_FUNCTION_BUILD_OPTIMIZATION
 
 ////////////////////////////////////////////////////////////////////////////////////////////////////
 
