[Insights] Minor code style fixes.

diff --git a/Engine/Source/Developer/TraceInsights/Private/Insights/MemoryProfiler/Widgets/SMemAllocTableTreeView.h b/Engine/Source/Developer/TraceInsights/Private/Insights/MemoryProfiler/Widgets/SMemAllocTableTreeView.h
--- a/Engine/Source/Developer/TraceInsights/Private/Insights/MemoryProfiler/Widgets/SMemAllocTableTreeView.h
+++ b/Engine/Source/Developer/TraceInsights/Private/Insights/MemoryProfiler/Widgets/SMemAllocTableTreeView.h
@@ -3,7 +3,6 @@
 #pragma once
 
 #include "CoreMinimal.h"
-
 #include "TraceServices/Model/AllocationsProvider.h"
 
 // Insights
@@ -44,7 +43,7 @@ public:
 	 */
 	void Construct(const FArguments& InArgs, TSharedPtr<FMemAllocTable> InTablePtr);
 
-	virtual TSharedPtr<SWidget> ConstructToolbar() override;;
+	virtual TSharedPtr<SWidget> ConstructToolbar() override;
 	virtual TSharedPtr<SWidget> ConstructFooter() override;
 
 	TSharedPtr<FMemAllocTable> GetMemAllocTable() { return StaticCastSharedPtr<FMemAllocTable>(GetTable()); }
