[Insights]
  - Fixed LLM Tags not allowing to change the current tracker, if previously the combobox was clicked when empty.
  - Removed usage of "ToolPanel.GroupBorder" (and the respective SBorder) from all panels.

diff --git a/Engine/Source/Developer/TraceInsights/Private/Insights/MemoryProfiler/Widgets/SMemTagTreeView.cpp b/Engine/Source/Developer/TraceInsights/Private/Insights/MemoryProfiler/Widgets/SMemTagTreeView.cpp
--- a/Engine/Source/Developer/TraceInsights/Private/Insights/MemoryProfiler/Widgets/SMemTagTreeView.cpp
+++ b/Engine/Source/Developer/TraceInsights/Private/Insights/MemoryProfiler/Widgets/SMemTagTreeView.cpp
@@ -80,39 +80,35 @@ void SMemTagTreeView::Construct(const FArguments& InArgs, TSharedPtr<SMemoryProf
 		+ SVerticalBox::Slot()
 		.VAlign(VAlign_Center)
 		.AutoHeight()
+		.Padding(2.0f, 2.0f, 2.0f, 2.0f)
 		[
-			SNew(SBorder)
-			.BorderImage(FEditorStyle::GetBrush("ToolPanel.GroupBorder"))
+			SNew(SVerticalBox)
+
+			// Search and Filtering
+			+ SVerticalBox::Slot()
+			.VAlign(VAlign_Center)
 			.Padding(2.0f)
+			.AutoHeight()
 			[
-				SNew(SVerticalBox)
-
-				// Search / Filtering
-				+ SVerticalBox::Slot()
-				.VAlign(VAlign_Center)
-				.Padding(2.0f)
-				.AutoHeight()
-				[
-					ConstructTagsFilteringWidgetArea()
-				]
+				ConstructTagsFilteringWidgetArea()
+			]
 
-				// Grouping
-				+ SVerticalBox::Slot()
-				.VAlign(VAlign_Center)
-				.Padding(2.0f)
-				.AutoHeight()
-				[
-					ConstructTagsGroupingWidgetArea()
-				]
+			// Tracker and Grouping
+			+ SVerticalBox::Slot()
+			.VAlign(VAlign_Center)
+			.Padding(2.0f)
+			.AutoHeight()
+			[
+				ConstructTagsGroupingWidgetArea()
+			]
 
-				// Tracks
-				+ SVerticalBox::Slot()
-				.VAlign(VAlign_Center)
-				.Padding(2.0f)
-				.AutoHeight()
-				[
-					ConstructTracksMiniToolbar()
-				]
+			// Tracks Mini Toolbar
+			+ SVerticalBox::Slot()
+			.VAlign(VAlign_Center)
+			.Padding(2.0f)
+			.AutoHeight()
+			[
+				ConstructTracksMiniToolbar()
 			]
 		]
 
@@ -132,26 +128,21 @@ void SMemTagTreeView::Construct(const FArguments& InArgs, TSharedPtr<SMemoryProf
 
 				+ SScrollBox::Slot()
 				[
-					SNew(SBorder)
-					.BorderImage(FEditorStyle::GetBrush("ToolPanel.GroupBorder"))
-					.Padding(0.0f)
-					[
-						SAssignNew(TreeView, STreeView<FMemTagNodePtr>)
-						.ExternalScrollbar(ExternalScrollbar)
-						.SelectionMode(ESelectionMode::Multi)
-						.TreeItemsSource(&FilteredGroupNodes)
-						.OnGetChildren(this, &SMemTagTreeView::TreeView_OnGetChildren)
-						.OnGenerateRow(this, &SMemTagTreeView::TreeView_OnGenerateRow)
-						.OnSelectionChanged(this, &SMemTagTreeView::TreeView_OnSelectionChanged)
-						.OnMouseButtonDoubleClick(this, &SMemTagTreeView::TreeView_OnMouseButtonDoubleClick)
-						.OnContextMenuOpening(FOnContextMenuOpening::CreateSP(this, &SMemTagTreeView::TreeView_GetMenuContent))
-						.ItemHeight(12.0f)
-						.HeaderRow
-						(
-							SAssignNew(TreeViewHeaderRow, SHeaderRow)
-							.Visibility(EVisibility::Visible)
-						)
-					]
+					SAssignNew(TreeView, STreeView<FMemTagNodePtr>)
+					.ExternalScrollbar(ExternalScrollbar)
+					.SelectionMode(ESelectionMode::Multi)
+					.TreeItemsSource(&FilteredGroupNodes)
+					.OnGetChildren(this, &SMemTagTreeView::TreeView_OnGetChildren)
+					.OnGenerateRow(this, &SMemTagTreeView::TreeView_OnGenerateRow)
+					.OnSelectionChanged(this, &SMemTagTreeView::TreeView_OnSelectionChanged)
+					.OnMouseButtonDoubleClick(this, &SMemTagTreeView::TreeView_OnMouseButtonDoubleClick)
+					.OnContextMenuOpening(FOnContextMenuOpening::CreateSP(this, &SMemTagTreeView::TreeView_GetMenuContent))
+					.ItemHeight(12.0f)
+					.HeaderRow
+					(
+						SAssignNew(TreeViewHeaderRow, SHeaderRow)
+						.Visibility(EVisibility::Visible)
+					)
 				]
 			]
 
@@ -268,6 +259,8 @@ TSharedRef<SWidget> SMemTagTreeView::ConstructTagsGroupingWidgetArea()
 		SAssignNew(TrackerComboBox, SComboBox<TSharedPtr<Insights::FMemoryTracker>>)
 		.ToolTipText(this, &SMemTagTreeView::Tracker_GetTooltipText)
 		.OptionsSource(GetAvailableTrackers())
+		.IsEnabled(this, &SMemTagTreeView::Tracker_IsEnabled)
+		.OnComboBoxOpening(this, &SMemTagTreeView::Tracker_OnComboBoxOpening)
 		.OnSelectionChanged(this, &SMemTagTreeView::Tracker_OnSelectionChanged)
 		.OnGenerateWidget(this, &SMemTagTreeView::Tracker_OnGenerateWidget)
 		[
@@ -302,6 +295,14 @@ TSharedRef<SWidget> SMemTagTreeView::ConstructTagsGroupingWidgetArea()
 		]
 	];
 
+	TSharedPtr<SMemoryProfilerWindow> ProfilerWindow = ProfilerWindowWeakPtr.Pin();
+	if (ProfilerWindow.IsValid())
+	{
+		FMemorySharedState& SharedState = ProfilerWindow->GetSharedState();
+		TrackerComboBox->SetSelectedItem(SharedState.GetCurrentTracker());
+		TrackerComboBox->RefreshOptions();
+	}
+
 	return Widget;
 }
 
@@ -2088,6 +2089,36 @@ const TArray<TSharedPtr<Insights::FMemoryTracker>>* SMemTagTreeView::GetAvailabl
 
 ////////////////////////////////////////////////////////////////////////////////////////////////////
 
+bool SMemTagTreeView::Tracker_IsEnabled() const
+{
+	TSharedPtr<SMemoryProfilerWindow> ProfilerWindow = ProfilerWindowWeakPtr.Pin();
+
+	if (ProfilerWindow.IsValid())
+	{
+		FMemorySharedState& SharedState = ProfilerWindow->GetSharedState();
+		return SharedState.GetTrackers().Num() > 0;
+	}
+	return false;
+}
+
+////////////////////////////////////////////////////////////////////////////////////////////////////
+
+void SMemTagTreeView::Tracker_OnComboBoxOpening()
+{
+	TSharedPtr<SMemoryProfilerWindow> ProfilerWindow = ProfilerWindowWeakPtr.Pin();
+
+	if (ProfilerWindow.IsValid())
+	{
+		FMemorySharedState& SharedState = ProfilerWindow->GetSharedState();
+		if (SharedState.GetCurrentTracker() != TrackerComboBox->GetSelectedItem())
+		{
+			TrackerComboBox->SetSelectedItem(SharedState.GetCurrentTracker());
+		}
+	}
+}
+
+////////////////////////////////////////////////////////////////////////////////////////////////////
+
 void SMemTagTreeView::Tracker_OnSelectionChanged(TSharedPtr<Insights::FMemoryTracker> InTracker, ESelectInfo::Type SelectInfo)
 {
 	if (SelectInfo != ESelectInfo::Direct)
@@ -2111,16 +2142,8 @@ void SMemTagTreeView::Tracker_OnSelectionChanged(TSharedPtr<Insights::FMemoryTra
 TSharedRef<SWidget> SMemTagTreeView::Tracker_OnGenerateWidget(TSharedPtr<Insights::FMemoryTracker> InTracker)
 {
 	const FText TrackerText = FText::Format(LOCTEXT("TrackerComboBox_TextFmt", "{0} Tracker (id {1})"), FText::FromString(InTracker->GetName()), FText::AsNumber(InTracker->GetId()));
-	return SNew(SCheckBox)
-		.Style(FEditorStyle::Get(), "Toolbar.RadioButton")
-		.OnCheckStateChanged(this, &SMemTagTreeView::Tracker_OnCheckStateChanged, InTracker)
-		.IsChecked(this, &SMemTagTreeView::Tracker_IsChecked, InTracker)
-		.Content()
-		[
-			SNew(STextBlock)
-			.Text(TrackerText)
-			.Margin(2.0f)
-		];
+	return SNew(STextBlock)
+		.Text(TrackerText);
 }
 
 ////////////////////////////////////////////////////////////////////////////////////////////////////
diff --git a/Engine/Source/Developer/TraceInsights/Private/Insights/MemoryProfiler/Widgets/SMemTagTreeView.h b/Engine/Source/Developer/TraceInsights/Private/Insights/MemoryProfiler/Widgets/SMemTagTreeView.h
--- a/Engine/Source/Developer/TraceInsights/Private/Insights/MemoryProfiler/Widgets/SMemTagTreeView.h
+++ b/Engine/Source/Developer/TraceInsights/Private/Insights/MemoryProfiler/Widgets/SMemTagTreeView.h
@@ -250,6 +250,8 @@ private:
 	////////////////////////////////////////////////////////////////////////////////////////////////////
 
 	const TArray<TSharedPtr<Insights::FMemoryTracker>>* GetAvailableTrackers();
+	bool Tracker_IsEnabled() const;
+	void Tracker_OnComboBoxOpening();
 	void Tracker_OnSelectionChanged(TSharedPtr<Insights::FMemoryTracker> InTracker, ESelectInfo::Type SelectInfo);
 	TSharedRef<SWidget> Tracker_OnGenerateWidget(TSharedPtr<Insights::FMemoryTracker> InTracker);
 	void Tracker_OnCheckStateChanged(ECheckBoxState CheckType, TSharedPtr<Insights::FMemoryTracker> InTracker);
diff --git a/Engine/Source/Developer/TraceInsights/Private/Insights/NetworkingProfiler/Widgets/SNetStatsView.cpp b/Engine/Source/Developer/TraceInsights/Private/Insights/NetworkingProfiler/Widgets/SNetStatsView.cpp
--- a/Engine/Source/Developer/TraceInsights/Private/Insights/NetworkingProfiler/Widgets/SNetStatsView.cpp
+++ b/Engine/Source/Developer/TraceInsights/Private/Insights/NetworkingProfiler/Widgets/SNetStatsView.cpp
@@ -82,117 +82,113 @@ void SNetStatsView::Construct(const FArguments& InArgs, TSharedPtr<SNetworkingPr
 		+ SVerticalBox::Slot()
 		.VAlign(VAlign_Center)
 		.AutoHeight()
+		.Padding(2.0f, 2.0f, 2.0f, 2.0f)
 		[
-			SNew(SBorder)
-			.BorderImage(FEditorStyle::GetBrush("ToolPanel.GroupBorder"))
+			SNew(SVerticalBox)
+
+			+ SVerticalBox::Slot()
+			.VAlign(VAlign_Center)
 			.Padding(2.0f)
+			.AutoHeight()
 			[
-				SNew(SVerticalBox)
+				SNew(SHorizontalBox)
 
-				+ SVerticalBox::Slot()
+				// Search box
+				+ SHorizontalBox::Slot()
 				.VAlign(VAlign_Center)
 				.Padding(2.0f)
-				.AutoHeight()
+				.FillWidth(1.0f)
 				[
-					SNew(SHorizontalBox)
-
-					// Search box
-					+ SHorizontalBox::Slot()
-					.VAlign(VAlign_Center)
-					.Padding(2.0f)
-					.FillWidth(1.0f)
-					[
-						SAssignNew(SearchBox, SSearchBox)
-						.HintText(LOCTEXT("SearchBoxHint", "Search net events or groups"))
-						.OnTextChanged(this, &SNetStatsView::SearchBox_OnTextChanged)
-						.IsEnabled(this, &SNetStatsView::SearchBox_IsEnabled)
-						.ToolTipText(LOCTEXT("FilterSearchHint", "Type here to search net events or groups"))
-					]
-
-					// Filter out net event types with zero instance count
-					+ SHorizontalBox::Slot()
-					.VAlign(VAlign_Center)
-					.Padding(2.0f)
-					.AutoWidth()
-					[
-						SNew(SCheckBox)
-						.Style(FCoreStyle::Get(), "ToggleButtonCheckbox")
-						.HAlign(HAlign_Center)
-						.Padding(2.0f)
-						.OnCheckStateChanged(this, &SNetStatsView::FilterOutZeroCountEvents_OnCheckStateChanged)
-						.IsChecked(this, &SNetStatsView::FilterOutZeroCountEvents_IsChecked)
-						.ToolTipText(LOCTEXT("FilterOutZeroCountEvents_Tooltip", "Filter out the net event types having zero total instance count (aggregated stats)."))
-						[
-							SNew(STextBlock)
-							.Text(LOCTEXT("FilterOutZeroCountEvents_Button", " !0 "))
-						]
-					]
+					SAssignNew(SearchBox, SSearchBox)
+					.HintText(LOCTEXT("SearchBoxHint", "Search net events or groups"))
+					.OnTextChanged(this, &SNetStatsView::SearchBox_OnTextChanged)
+					.IsEnabled(this, &SNetStatsView::SearchBox_IsEnabled)
+					.ToolTipText(LOCTEXT("FilterSearchHint", "Type here to search net events or groups"))
 				]
 
-				// Group by
-				+ SVerticalBox::Slot()
+				// Filter out net event types with zero instance count
+				+ SHorizontalBox::Slot()
 				.VAlign(VAlign_Center)
 				.Padding(2.0f)
-				.AutoHeight()
+				.AutoWidth()
 				[
-					SNew(SHorizontalBox)
-
-					+ SHorizontalBox::Slot()
-					.FillWidth(1.0f)
-					.VAlign(VAlign_Center)
+					SNew(SCheckBox)
+					.Style(FCoreStyle::Get(), "ToggleButtonCheckbox")
+					.HAlign(HAlign_Center)
+					.Padding(2.0f)
+					.OnCheckStateChanged(this, &SNetStatsView::FilterOutZeroCountEvents_OnCheckStateChanged)
+					.IsChecked(this, &SNetStatsView::FilterOutZeroCountEvents_IsChecked)
+					.ToolTipText(LOCTEXT("FilterOutZeroCountEvents_Tooltip", "Filter out the net event types having zero total instance count (aggregated stats)."))
 					[
 						SNew(STextBlock)
-						.Text(LOCTEXT("GroupByText", "Group by"))
-					]
-
-					+ SHorizontalBox::Slot()
-					.FillWidth(2.0f)
-					.VAlign(VAlign_Center)
-					[
-						SAssignNew(GroupByComboBox, SComboBox<TSharedPtr<ENetEventGroupingMode>>)
-						.ToolTipText(this, &SNetStatsView::GroupBy_GetSelectedTooltipText)
-						.OptionsSource(&GroupByOptionsSource)
-						.OnSelectionChanged(this, &SNetStatsView::GroupBy_OnSelectionChanged)
-						.OnGenerateWidget(this, &SNetStatsView::GroupBy_OnGenerateWidget)
-						[
-							SNew(STextBlock)
-							.Text(this, &SNetStatsView::GroupBy_GetSelectedText)
-						]
+						.Text(LOCTEXT("FilterOutZeroCountEvents_Button", " !0 "))
 					]
 				]
+			]
 
-				// TODO: Check boxes for: NetEvent, ...
-				/*
-				+ SVerticalBox::Slot()
+			// Group by
+			+ SVerticalBox::Slot()
+			.VAlign(VAlign_Center)
+			.Padding(2.0f)
+			.AutoHeight()
+			[
+				SNew(SHorizontalBox)
+
+				+ SHorizontalBox::Slot()
+				.FillWidth(1.0f)
 				.VAlign(VAlign_Center)
-				.Padding(2.0f)
-				.AutoHeight()
 				[
-					SNew(SHorizontalBox)
+					SNew(STextBlock)
+					.Text(LOCTEXT("GroupByText", "Group by"))
+				]
 
-					+ SHorizontalBox::Slot()
-					.Padding(FMargin(0.0f,0.0f,1.0f,0.0f))
-					.FillWidth(1.0f)
+				+ SHorizontalBox::Slot()
+				.FillWidth(2.0f)
+				.VAlign(VAlign_Center)
+				[
+					SAssignNew(GroupByComboBox, SComboBox<TSharedPtr<ENetEventGroupingMode>>)
+					.ToolTipText(this, &SNetStatsView::GroupBy_GetSelectedTooltipText)
+					.OptionsSource(&GroupByOptionsSource)
+					.OnSelectionChanged(this, &SNetStatsView::GroupBy_OnSelectionChanged)
+					.OnGenerateWidget(this, &SNetStatsView::GroupBy_OnGenerateWidget)
 					[
-						GetToggleButtonForNetEventType(ENetEventNodeType::NetEvent)
+						SNew(STextBlock)
+						.Text(this, &SNetStatsView::GroupBy_GetSelectedText)
 					]
+				]
+			]
 
-					//+ SHorizontalBox::Slot()
-					//.Padding(FMargin(1.0f,0.0f,1.0f,0.0f))
-					//.FillWidth(1.0f)
-					//[
-					//	GetToggleButtonForNetEventType(ENetEventNodeType::NetEvent1)
-					//]
-
-					//+ SHorizontalBox::Slot()
-					.//Padding(FMargin(1.0f,0.0f,1.0f,0.0f))
-					//.FillWidth(1.0f)
-					//[
-					//	GetToggleButtonForNetEventType(ENetEventNodeType::NetEvent2)
-					//]
+			// TODO: Check boxes for: NetEvent, ...
+			/*
+			+ SVerticalBox::Slot()
+			.VAlign(VAlign_Center)
+			.Padding(2.0f)
+			.AutoHeight()
+			[
+				SNew(SHorizontalBox)
+
+				+ SHorizontalBox::Slot()
+				.Padding(FMargin(0.0f,0.0f,1.0f,0.0f))
+				.FillWidth(1.0f)
+				[
+					GetToggleButtonForNetEventType(ENetEventNodeType::NetEvent)
 				]
-				*/
+
+				//+ SHorizontalBox::Slot()
+				//.Padding(FMargin(1.0f,0.0f,1.0f,0.0f))
+				//.FillWidth(1.0f)
+				//[
+				//	GetToggleButtonForNetEventType(ENetEventNodeType::NetEvent1)
+				//]
+
+				//+ SHorizontalBox::Slot()
+				.//Padding(FMargin(1.0f,0.0f,1.0f,0.0f))
+				//.FillWidth(1.0f)
+				//[
+				//	GetToggleButtonForNetEventType(ENetEventNodeType::NetEvent2)
+				//]
 			]
+			*/
 		]
 
 		// Tree view
@@ -211,26 +207,21 @@ void SNetStatsView::Construct(const FArguments& InArgs, TSharedPtr<SNetworkingPr
 
 				+ SScrollBox::Slot()
 				[
-					SNew(SBorder)
-					.BorderImage(FEditorStyle::GetBrush("ToolPanel.GroupBorder"))
-					.Padding(0.0f)
-					[
-						SAssignNew(TreeView, STreeView<FNetEventNodePtr>)
-						.ExternalScrollbar(ExternalScrollbar)
-						.SelectionMode(ESelectionMode::Multi)
-						.TreeItemsSource(&FilteredGroupNodes)
-						.OnGetChildren(this, &SNetStatsView::TreeView_OnGetChildren)
-						.OnGenerateRow(this, &SNetStatsView::TreeView_OnGenerateRow)
-						.OnSelectionChanged(this, &SNetStatsView::TreeView_OnSelectionChanged)
-						.OnMouseButtonDoubleClick(this, &SNetStatsView::TreeView_OnMouseButtonDoubleClick)
-						.OnContextMenuOpening(FOnContextMenuOpening::CreateSP(this, &SNetStatsView::TreeView_GetMenuContent))
-						.ItemHeight(12.0f)
-						.HeaderRow
-						(
-							SAssignNew(TreeViewHeaderRow, SHeaderRow)
-							.Visibility(EVisibility::Visible)
-						)
-					]
+					SAssignNew(TreeView, STreeView<FNetEventNodePtr>)
+					.ExternalScrollbar(ExternalScrollbar)
+					.SelectionMode(ESelectionMode::Multi)
+					.TreeItemsSource(&FilteredGroupNodes)
+					.OnGetChildren(this, &SNetStatsView::TreeView_OnGetChildren)
+					.OnGenerateRow(this, &SNetStatsView::TreeView_OnGenerateRow)
+					.OnSelectionChanged(this, &SNetStatsView::TreeView_OnSelectionChanged)
+					.OnMouseButtonDoubleClick(this, &SNetStatsView::TreeView_OnMouseButtonDoubleClick)
+					.OnContextMenuOpening(FOnContextMenuOpening::CreateSP(this, &SNetStatsView::TreeView_GetMenuContent))
+					.ItemHeight(12.0f)
+					.HeaderRow
+					(
+						SAssignNew(TreeViewHeaderRow, SHeaderRow)
+						.Visibility(EVisibility::Visible)
+					)
 				]
 			]
 
diff --git a/Engine/Source/Developer/TraceInsights/Private/Insights/Table/Widgets/STableTreeView.cpp b/Engine/Source/Developer/TraceInsights/Private/Insights/Table/Widgets/STableTreeView.cpp
--- a/Engine/Source/Developer/TraceInsights/Private/Insights/Table/Widgets/STableTreeView.cpp
+++ b/Engine/Source/Developer/TraceInsights/Private/Insights/Table/Widgets/STableTreeView.cpp
@@ -153,76 +153,72 @@ void STableTreeView::ConstructWidget(TSharedPtr<FTable> InTablePtr)
 	+ SVerticalBox::Slot()
 	.VAlign(VAlign_Center)
 	.AutoHeight()
+	.Padding(2.0f, 2.0f, 2.0f, 2.0f)
 	[
-		SNew(SBorder)
-		.BorderImage(FEditorStyle::GetBrush("ToolPanel.GroupBorder"))
+		SNew(SVerticalBox)
+
+		// Search box
+		+ SVerticalBox::Slot()
+		.VAlign(VAlign_Center)
 		.Padding(2.0f)
+		.AutoHeight()
 		[
-			SNew(SVerticalBox)
+			SNew(SHorizontalBox)
 
-			// Search box
-			+ SVerticalBox::Slot()
+			+ SHorizontalBox::Slot()
+			.FillWidth(1.0)
 			.VAlign(VAlign_Center)
-			.Padding(2.0f)
-			.AutoHeight()
 			[
-				SNew(SHorizontalBox)
-
-				+ SHorizontalBox::Slot()
-				.FillWidth(1.0)
-				.VAlign(VAlign_Center)
-				[
-					SAssignNew(SearchBox, SSearchBox)
-					.HintText(LOCTEXT("SearchBox_Hint", "Search"))
-					.OnTextChanged(this, &STableTreeView::SearchBox_OnTextChanged)
-					.IsEnabled(this, &STableTreeView::SearchBox_IsEnabled)
-					.ToolTipText(this, &STableTreeView::SearchBox_GetTooltipText)
-				]
-
-				+ SHorizontalBox::Slot()
-				.AutoWidth()
-				.Padding(4.0f, 0.0f, 0.0f, 0.0f)
-				.VAlign(VAlign_Center)
-				[
-					SNew(SButton)
-					.Text(LOCTEXT("AdvancedFiltersBtn_Text", "Advanced Filters"))
-					.ToolTipText(this, &STableTreeView::AdvancedFilters_GetTooltipText)
-					.OnClicked(this, &STableTreeView::OnAdvancedFiltersClicked)
-					.IsEnabled(this, &STableTreeView::AdvancedFilters_ShouldBeEnabled)
-				]
+				SAssignNew(SearchBox, SSearchBox)
+				.HintText(LOCTEXT("SearchBox_Hint", "Search"))
+				.OnTextChanged(this, &STableTreeView::SearchBox_OnTextChanged)
+				.IsEnabled(this, &STableTreeView::SearchBox_IsEnabled)
+				.ToolTipText(this, &STableTreeView::SearchBox_GetTooltipText)
 			]
 
-			// Group by
-			+ SVerticalBox::Slot()
+			+ SHorizontalBox::Slot()
+			.AutoWidth()
+			.Padding(4.0f, 0.0f, 0.0f, 0.0f)
 			.VAlign(VAlign_Center)
-			.Padding(2.0f)
-			.AutoHeight()
 			[
-				SAssignNew(ToolbarBox, SHorizontalBox)
+				SNew(SButton)
+				.Text(LOCTEXT("AdvancedFiltersBtn_Text", "Advanced Filters"))
+				.ToolTipText(this, &STableTreeView::AdvancedFilters_GetTooltipText)
+				.OnClicked(this, &STableTreeView::OnAdvancedFiltersClicked)
+				.IsEnabled(this, &STableTreeView::AdvancedFilters_ShouldBeEnabled)
+			]
+		]
 
-				+ SHorizontalBox::Slot()
-				.AutoWidth()
-				.VAlign(VAlign_Center)
-				[
-					SNew(STextBlock)
-					.Text(LOCTEXT("GroupByText", "Hierarchy:"))
-					.Margin(FMargin(0.0f, 0.0f, 4.0f, 0.0f))
-				]
+		// Group by
+		+ SVerticalBox::Slot()
+		.VAlign(VAlign_Center)
+		.Padding(2.0f)
+		.AutoHeight()
+		[
+			SAssignNew(ToolbarBox, SHorizontalBox)
 
-				+ SHorizontalBox::Slot()
-				.FillWidth(1.0f)
-				.VAlign(VAlign_Center)
-				[
-					SAssignNew(GroupingBreadcrumbTrail, SBreadcrumbTrail<TSharedPtr<FTreeNodeGrouping>>)
-					.ButtonContentPadding(FMargin(1.0f, 1.0f))
-					//.DelimiterImage(FEditorStyle::GetBrush("SlateFileDialogs.PathDelimiter"))
-					//.TextStyle(FEditorStyle::Get(), "Tutorials.Browser.PathText")
-					//.ShowLeadingDelimiter(true)
-					//.PersistentBreadcrumbs(true)
-					.InvertTextColorOnHover(true)
-					.OnCrumbClicked(this, &STableTreeView::OnGroupingCrumbClicked)
-					.GetCrumbMenuContent(this, &STableTreeView::GetGroupingCrumbMenuContent)
-				]
+			+ SHorizontalBox::Slot()
+			.AutoWidth()
+			.VAlign(VAlign_Center)
+			[
+				SNew(STextBlock)
+				.Text(LOCTEXT("GroupByText", "Hierarchy:"))
+				.Margin(FMargin(0.0f, 0.0f, 4.0f, 0.0f))
+			]
+
+			+ SHorizontalBox::Slot()
+			.FillWidth(1.0f)
+			.VAlign(VAlign_Center)
+			[
+				SAssignNew(GroupingBreadcrumbTrail, SBreadcrumbTrail<TSharedPtr<FTreeNodeGrouping>>)
+				.ButtonContentPadding(FMargin(1.0f, 1.0f))
+				//.DelimiterImage(FEditorStyle::GetBrush("SlateFileDialogs.PathDelimiter"))
+				//.TextStyle(FEditorStyle::Get(), "Tutorials.Browser.PathText")
+				//.ShowLeadingDelimiter(true)
+				//.PersistentBreadcrumbs(true)
+				.InvertTextColorOnHover(true)
+				.OnCrumbClicked(this, &STableTreeView::OnGroupingCrumbClicked)
+				.GetCrumbMenuContent(this, &STableTreeView::GetGroupingCrumbMenuContent)
 			]
 		]
 	]
@@ -249,26 +245,21 @@ void STableTreeView::ConstructWidget(TSharedPtr<FTable> InTablePtr)
 				.HAlign(HAlign_Fill)
 				.VAlign(VAlign_Fill)
 				[
-					SNew(SBorder)
-					.BorderImage(FEditorStyle::GetBrush("ToolPanel.GroupBorder"))
-					.Padding(0.0f)
-					[
-						SAssignNew(TreeView, STreeView<FTableTreeNodePtr>)
-						//.ExternalScrollbar(ExternalScrollbar)
-						.SelectionMode(ESelectionMode::Multi)
-						.TreeItemsSource(&FilteredGroupNodes)
-						.OnGetChildren(this, &STableTreeView::TreeView_OnGetChildren)
-						.OnGenerateRow(this, &STableTreeView::TreeView_OnGenerateRow)
-						.OnSelectionChanged(this, &STableTreeView::TreeView_OnSelectionChanged)
-						.OnMouseButtonDoubleClick(this, &STableTreeView::TreeView_OnMouseButtonDoubleClick)
-						.OnContextMenuOpening(FOnContextMenuOpening::CreateSP(this, &STableTreeView::TreeView_GetMenuContent))
-						.ItemHeight(12.0f)
-						.HeaderRow
-						(
-							SAssignNew(TreeViewHeaderRow, SHeaderRow)
-							.Visibility(EVisibility::Visible)
-						)
-					]
+					SAssignNew(TreeView, STreeView<FTableTreeNodePtr>)
+					//.ExternalScrollbar(ExternalScrollbar)
+					.SelectionMode(ESelectionMode::Multi)
+					.TreeItemsSource(&FilteredGroupNodes)
+					.OnGetChildren(this, &STableTreeView::TreeView_OnGetChildren)
+					.OnGenerateRow(this, &STableTreeView::TreeView_OnGenerateRow)
+					.OnSelectionChanged(this, &STableTreeView::TreeView_OnSelectionChanged)
+					.OnMouseButtonDoubleClick(this, &STableTreeView::TreeView_OnMouseButtonDoubleClick)
+					.OnContextMenuOpening(FOnContextMenuOpening::CreateSP(this, &STableTreeView::TreeView_GetMenuContent))
+					.ItemHeight(12.0f)
+					.HeaderRow
+					(
+						SAssignNew(TreeViewHeaderRow, SHeaderRow)
+						.Visibility(EVisibility::Visible)
+					)
 				]
 
 				+ SOverlay::Slot()
diff --git a/Engine/Source/Developer/TraceInsights/Private/Insights/Widgets/SLogView.cpp b/Engine/Source/Developer/TraceInsights/Private/Insights/Widgets/SLogView.cpp
--- a/Engine/Source/Developer/TraceInsights/Private/Insights/Widgets/SLogView.cpp
+++ b/Engine/Source/Developer/TraceInsights/Private/Insights/Widgets/SLogView.cpp
@@ -626,53 +626,48 @@ void SLogView::Construct(const FArguments& InArgs)
 					+ SScrollBox::Slot()
 					.VAlign(VAlign_Fill)
 					[
-						SNew(SBorder)
-						.BorderImage(FEditorStyle::GetBrush("ToolPanel.GroupBorder"))
-						.Padding(0.0f)
-						[
-							SAssignNew(ListView, SListView<TSharedPtr<FLogMessage>>)
-							.ExternalScrollbar(ExternalScrollbar)
-							.ItemHeight(20.0f)
-							.SelectionMode(ESelectionMode::Single)
-							.OnMouseButtonClick(this, &SLogView::OnMouseButtonClick)
-							.OnSelectionChanged(this, &SLogView::OnSelectionChanged)
-							.ListItemsSource(&Messages)
-							.OnGenerateRow(this, &SLogView::OnGenerateRow)
-							.ConsumeMouseWheel(EConsumeMouseWheel::Always)
-							.OnContextMenuOpening(FOnContextMenuOpening::CreateSP(this, &SLogView::ListView_GetContextMenu))
-							.HeaderRow
-							(
-								SNew(SHeaderRow)
-
-								+ SHeaderRow::Column(LogViewColumns::IdColumnName)
-								.ManualWidth(60.0f)
-								.DefaultLabel(LOCTEXT("IdColumn", "Index"))
-
-								+ SHeaderRow::Column(LogViewColumns::TimeColumnName)
-								.ManualWidth(94.0f)
-								.DefaultLabel(LOCTEXT("TimeColumn", "Time"))
-
-								+ SHeaderRow::Column(LogViewColumns::VerbosityColumnName)
-								.ManualWidth(80.0f)
-								.DefaultLabel(LOCTEXT("VerbosityColumn", "Verbosity"))
-
-								+ SHeaderRow::Column(LogViewColumns::CategoryColumnName)
-								.ManualWidth(120.0f)
-								.DefaultLabel(LOCTEXT("CategoryColumn", "Category"))
-
-								+ SHeaderRow::Column(LogViewColumns::MessageColumnName)
-								.ManualWidth(880.0f)
-								.DefaultLabel(LOCTEXT("MessageColumn", "Message"))
-
-								+ SHeaderRow::Column(LogViewColumns::FileColumnName)
-								.ManualWidth(600.0f)
-								.DefaultLabel(LOCTEXT("FileColumn", "File"))
-
-								+ SHeaderRow::Column(LogViewColumns::LineColumnName)
-								.ManualWidth(60.0f)
-								.DefaultLabel(LOCTEXT("LineColumn", "Line"))
-							)
-						]
+						SAssignNew(ListView, SListView<TSharedPtr<FLogMessage>>)
+						.ExternalScrollbar(ExternalScrollbar)
+						.ItemHeight(20.0f)
+						.SelectionMode(ESelectionMode::Single)
+						.OnMouseButtonClick(this, &SLogView::OnMouseButtonClick)
+						.OnSelectionChanged(this, &SLogView::OnSelectionChanged)
+						.ListItemsSource(&Messages)
+						.OnGenerateRow(this, &SLogView::OnGenerateRow)
+						.ConsumeMouseWheel(EConsumeMouseWheel::Always)
+						.OnContextMenuOpening(FOnContextMenuOpening::CreateSP(this, &SLogView::ListView_GetContextMenu))
+						.HeaderRow
+						(
+							SNew(SHeaderRow)
+
+							+ SHeaderRow::Column(LogViewColumns::IdColumnName)
+							.ManualWidth(60.0f)
+							.DefaultLabel(LOCTEXT("IdColumn", "Index"))
+
+							+ SHeaderRow::Column(LogViewColumns::TimeColumnName)
+							.ManualWidth(94.0f)
+							.DefaultLabel(LOCTEXT("TimeColumn", "Time"))
+
+							+ SHeaderRow::Column(LogViewColumns::VerbosityColumnName)
+							.ManualWidth(80.0f)
+							.DefaultLabel(LOCTEXT("VerbosityColumn", "Verbosity"))
+
+							+ SHeaderRow::Column(LogViewColumns::CategoryColumnName)
+							.ManualWidth(120.0f)
+							.DefaultLabel(LOCTEXT("CategoryColumn", "Category"))
+
+							+ SHeaderRow::Column(LogViewColumns::MessageColumnName)
+							.ManualWidth(880.0f)
+							.DefaultLabel(LOCTEXT("MessageColumn", "Message"))
+
+							+ SHeaderRow::Column(LogViewColumns::FileColumnName)
+							.ManualWidth(600.0f)
+							.DefaultLabel(LOCTEXT("FileColumn", "File"))
+
+							+ SHeaderRow::Column(LogViewColumns::LineColumnName)
+							.ManualWidth(60.0f)
+							.DefaultLabel(LOCTEXT("LineColumn", "Line"))
+						)
 					]
 				]
 
diff --git a/Engine/Source/Developer/TraceInsights/Private/Insights/Widgets/SStatsView.cpp b/Engine/Source/Developer/TraceInsights/Private/Insights/Widgets/SStatsView.cpp
--- a/Engine/Source/Developer/TraceInsights/Private/Insights/Widgets/SStatsView.cpp
+++ b/Engine/Source/Developer/TraceInsights/Private/Insights/Widgets/SStatsView.cpp
@@ -120,106 +120,102 @@ void SStatsView::Construct(const FArguments& InArgs)
 		+ SVerticalBox::Slot()
 		.VAlign(VAlign_Center)
 		.AutoHeight()
+		.Padding(2.0f, 2.0f, 2.0f, 2.0f)
 		[
-			SNew(SBorder)
-			.BorderImage(FEditorStyle::GetBrush("ToolPanel.GroupBorder"))
+			SNew(SVerticalBox)
+
+			+ SVerticalBox::Slot()
+			.VAlign(VAlign_Center)
 			.Padding(2.0f)
+			.AutoHeight()
 			[
-				SNew(SVerticalBox)
+				SNew(SHorizontalBox)
 
-				+ SVerticalBox::Slot()
+				// Search box
+				+ SHorizontalBox::Slot()
 				.VAlign(VAlign_Center)
 				.Padding(2.0f)
-				.AutoHeight()
+				.FillWidth(1.0f)
 				[
-					SNew(SHorizontalBox)
-
-					// Search box
-					+ SHorizontalBox::Slot()
-					.VAlign(VAlign_Center)
-					.Padding(2.0f)
-					.FillWidth(1.0f)
-					[
-						SAssignNew(SearchBox, SSearchBox)
-						.HintText(LOCTEXT("SearchBoxHint", "Search stats counters or groups"))
-						.OnTextChanged(this, &SStatsView::SearchBox_OnTextChanged)
-						.IsEnabled(this, &SStatsView::SearchBox_IsEnabled)
-						.ToolTipText(LOCTEXT("FilterSearchHint", "Type here to search stats counter or group"))
-					]
-
-					// Filter out timers with zero instance count
-					+ SHorizontalBox::Slot()
-					.VAlign(VAlign_Center)
-					.Padding(2.0f)
-					.AutoWidth()
-					[
-						SNew(SCheckBox)
-						.Style(FCoreStyle::Get(), "ToggleButtonCheckbox")
-						.HAlign(HAlign_Center)
-						.Padding(2.0f)
-						.OnCheckStateChanged(this, &SStatsView::FilterOutZeroCountStats_OnCheckStateChanged)
-						.IsChecked(this, &SStatsView::FilterOutZeroCountStats_IsChecked)
-						.ToolTipText(LOCTEXT("FilterOutZeroCountStats_Tooltip", "Filter out the stats counters having zero total instance count (aggregated stats)."))
-						[
-							SNew(STextBlock)
-							.Text(LOCTEXT("FilterOutZeroCountStats_Button", " !0 "))
-						]
-					]
+					SAssignNew(SearchBox, SSearchBox)
+					.HintText(LOCTEXT("SearchBoxHint", "Search stats counters or groups"))
+					.OnTextChanged(this, &SStatsView::SearchBox_OnTextChanged)
+					.IsEnabled(this, &SStatsView::SearchBox_IsEnabled)
+					.ToolTipText(LOCTEXT("FilterSearchHint", "Type here to search stats counter or group"))
 				]
 
-				// Group by
-				+ SVerticalBox::Slot()
+				// Filter out timers with zero instance count
+				+ SHorizontalBox::Slot()
 				.VAlign(VAlign_Center)
 				.Padding(2.0f)
-				.AutoHeight()
+				.AutoWidth()
 				[
-					SNew(SHorizontalBox)
-
-					+ SHorizontalBox::Slot()
-					.FillWidth(1.0f)
-					.VAlign(VAlign_Center)
+					SNew(SCheckBox)
+					.Style(FCoreStyle::Get(), "ToggleButtonCheckbox")
+					.HAlign(HAlign_Center)
+					.Padding(2.0f)
+					.OnCheckStateChanged(this, &SStatsView::FilterOutZeroCountStats_OnCheckStateChanged)
+					.IsChecked(this, &SStatsView::FilterOutZeroCountStats_IsChecked)
+					.ToolTipText(LOCTEXT("FilterOutZeroCountStats_Tooltip", "Filter out the stats counters having zero total instance count (aggregated stats)."))
 					[
 						SNew(STextBlock)
-						.Text(LOCTEXT("GroupByText", "Group by"))
-					]
-
-					+ SHorizontalBox::Slot()
-					.FillWidth(2.0f)
-					.VAlign(VAlign_Center)
-					[
-						SAssignNew(GroupByComboBox, SComboBox<TSharedPtr<EStatsGroupingMode>>)
-						.ToolTipText(this, &SStatsView::GroupBy_GetSelectedTooltipText)
-						.OptionsSource(&GroupByOptionsSource)
-						.OnSelectionChanged(this, &SStatsView::GroupBy_OnSelectionChanged)
-						.OnGenerateWidget(this, &SStatsView::GroupBy_OnGenerateWidget)
-						[
-							SNew(STextBlock)
-							.Text(this, &SStatsView::GroupBy_GetSelectedText)
-						]
+						.Text(LOCTEXT("FilterOutZeroCountStats_Button", " !0 "))
 					]
 				]
+			]
 
-				// Check boxes for: Int64, Float
-				+ SVerticalBox::Slot()
+			// Group by
+			+ SVerticalBox::Slot()
+			.VAlign(VAlign_Center)
+			.Padding(2.0f)
+			.AutoHeight()
+			[
+				SNew(SHorizontalBox)
+
+				+ SHorizontalBox::Slot()
+				.FillWidth(1.0f)
 				.VAlign(VAlign_Center)
-				.Padding(2.0f)
-				.AutoHeight()
 				[
-					SNew(SHorizontalBox)
+					SNew(STextBlock)
+					.Text(LOCTEXT("GroupByText", "Group by"))
+				]
 
-					+ SHorizontalBox::Slot()
-					.Padding(FMargin(0.0f,0.0f,1.0f,0.0f))
-					.FillWidth(1.0f)
+				+ SHorizontalBox::Slot()
+				.FillWidth(2.0f)
+				.VAlign(VAlign_Center)
+				[
+					SAssignNew(GroupByComboBox, SComboBox<TSharedPtr<EStatsGroupingMode>>)
+					.ToolTipText(this, &SStatsView::GroupBy_GetSelectedTooltipText)
+					.OptionsSource(&GroupByOptionsSource)
+					.OnSelectionChanged(this, &SStatsView::GroupBy_OnSelectionChanged)
+					.OnGenerateWidget(this, &SStatsView::GroupBy_OnGenerateWidget)
 					[
-						GetToggleButtonForDataType(EStatsNodeDataType::Int64)
+						SNew(STextBlock)
+						.Text(this, &SStatsView::GroupBy_GetSelectedText)
 					]
+				]
+			]
 
-					+ SHorizontalBox::Slot()
-					.Padding(FMargin(1.0f,0.0f,1.0f,0.0f))
-					.FillWidth(1.0f)
-					[
-						GetToggleButtonForDataType(EStatsNodeDataType::Double)
-					]
+			// Check boxes for: Int64, Float
+			+ SVerticalBox::Slot()
+			.VAlign(VAlign_Center)
+			.Padding(2.0f)
+			.AutoHeight()
+			[
+				SNew(SHorizontalBox)
+
+				+ SHorizontalBox::Slot()
+				.Padding(FMargin(0.0f,0.0f,1.0f,0.0f))
+				.FillWidth(1.0f)
+				[
+					GetToggleButtonForDataType(EStatsNodeDataType::Int64)
+				]
+
+				+ SHorizontalBox::Slot()
+				.Padding(FMargin(1.0f,0.0f,1.0f,0.0f))
+				.FillWidth(1.0f)
+				[
+					GetToggleButtonForDataType(EStatsNodeDataType::Double)
 				]
 			]
 		]
@@ -246,10 +242,6 @@ void SStatsView::Construct(const FArguments& InArgs)
 					.HAlign(HAlign_Fill)
 					.VAlign(VAlign_Fill)
 					[
-					//SNew(SBorder)
-					//.BorderImage(FEditorStyle::GetBrush("ToolPanel.GroupBorder"))
-					//.Padding(0.0f)
-					//[
 						SAssignNew(TreeView, STreeView<FStatsNodePtr>)
 						.ExternalScrollbar(ExternalScrollbar)
 						.SelectionMode(ESelectionMode::Multi)
@@ -265,7 +257,6 @@ void SStatsView::Construct(const FArguments& InArgs)
 							SAssignNew(TreeViewHeaderRow, SHeaderRow)
 							.Visibility(EVisibility::Visible)
 						)
-					//]
 					]
 
 					+ SOverlay::Slot()
diff --git a/Engine/Source/Developer/TraceInsights/Private/Insights/Widgets/STimersView.cpp b/Engine/Source/Developer/TraceInsights/Private/Insights/Widgets/STimersView.cpp
--- a/Engine/Source/Developer/TraceInsights/Private/Insights/Widgets/STimersView.cpp
+++ b/Engine/Source/Developer/TraceInsights/Private/Insights/Widgets/STimersView.cpp
@@ -122,113 +122,109 @@ void STimersView::Construct(const FArguments& InArgs)
 		+ SVerticalBox::Slot()
 		.VAlign(VAlign_Center)
 		.AutoHeight()
+		.Padding(2.0f, 2.0f, 2.0f, 2.0f)
 		[
-			SNew(SBorder)
-			.BorderImage(FEditorStyle::GetBrush("ToolPanel.GroupBorder"))
+			SNew(SVerticalBox)
+
+			+ SVerticalBox::Slot()
+			.VAlign(VAlign_Center)
 			.Padding(2.0f)
+			.AutoHeight()
 			[
-				SNew(SVerticalBox)
+				SNew(SHorizontalBox)
 
-				+ SVerticalBox::Slot()
+				// Search box
+				+ SHorizontalBox::Slot()
 				.VAlign(VAlign_Center)
 				.Padding(2.0f)
-				.AutoHeight()
+				.FillWidth(1.0f)
 				[
-					SNew(SHorizontalBox)
-
-					// Search box
-					+ SHorizontalBox::Slot()
-					.VAlign(VAlign_Center)
-					.Padding(2.0f)
-					.FillWidth(1.0f)
-					[
-						SAssignNew(SearchBox, SSearchBox)
-						.HintText(LOCTEXT("SearchBoxHint", "Search timers or groups"))
-						.OnTextChanged(this, &STimersView::SearchBox_OnTextChanged)
-						.IsEnabled(this, &STimersView::SearchBox_IsEnabled)
-						.ToolTipText(LOCTEXT("FilterSearchHint", "Type here to search timer or group"))
-					]
-
-					// Filter out timers with zero instance count
-					+ SHorizontalBox::Slot()
-					.VAlign(VAlign_Center)
-					.Padding(2.0f)
-					.AutoWidth()
-					[
-						SNew(SCheckBox)
-						.Style(FCoreStyle::Get(), "ToggleButtonCheckbox")
-						.HAlign(HAlign_Center)
-						.Padding(2.0f)
-						.OnCheckStateChanged(this, &STimersView::FilterOutZeroCountTimers_OnCheckStateChanged)
-						.IsChecked(this, &STimersView::FilterOutZeroCountTimers_IsChecked)
-						.ToolTipText(LOCTEXT("FilterOutZeroCountTimers_Tooltip", "Filter out the timers having zero total instance count (aggregated stats)."))
-						[
-							SNew(STextBlock)
-							.Text(LOCTEXT("FilterOutZeroCountTimers_Button", " !0 "))
-						]
-					]
+					SAssignNew(SearchBox, SSearchBox)
+					.HintText(LOCTEXT("SearchBoxHint", "Search timers or groups"))
+					.OnTextChanged(this, &STimersView::SearchBox_OnTextChanged)
+					.IsEnabled(this, &STimersView::SearchBox_IsEnabled)
+					.ToolTipText(LOCTEXT("FilterSearchHint", "Type here to search timer or group"))
 				]
 
-				// Group by
-				+ SVerticalBox::Slot()
+				// Filter out timers with zero instance count
+				+ SHorizontalBox::Slot()
 				.VAlign(VAlign_Center)
 				.Padding(2.0f)
-				.AutoHeight()
+				.AutoWidth()
 				[
-					SNew(SHorizontalBox)
-
-					+ SHorizontalBox::Slot()
-					.FillWidth(1.0f)
-					.VAlign(VAlign_Center)
+					SNew(SCheckBox)
+					.Style(FCoreStyle::Get(), "ToggleButtonCheckbox")
+					.HAlign(HAlign_Center)
+					.Padding(2.0f)
+					.OnCheckStateChanged(this, &STimersView::FilterOutZeroCountTimers_OnCheckStateChanged)
+					.IsChecked(this, &STimersView::FilterOutZeroCountTimers_IsChecked)
+					.ToolTipText(LOCTEXT("FilterOutZeroCountTimers_Tooltip", "Filter out the timers having zero total instance count (aggregated stats)."))
 					[
 						SNew(STextBlock)
-						.Text(LOCTEXT("GroupByText", "Group by"))
-					]
-
-					+ SHorizontalBox::Slot()
-					.FillWidth(2.0f)
-					.VAlign(VAlign_Center)
-					[
-						SAssignNew(GroupByComboBox, SComboBox<TSharedPtr<ETimerGroupingMode>>)
-						.ToolTipText(this, &STimersView::GroupBy_GetSelectedTooltipText)
-						.OptionsSource(&GroupByOptionsSource)
-						.OnSelectionChanged(this, &STimersView::GroupBy_OnSelectionChanged)
-						.OnGenerateWidget(this, &STimersView::GroupBy_OnGenerateWidget)
-						[
-							SNew(STextBlock)
-							.Text(this, &STimersView::GroupBy_GetSelectedText)
-						]
+						.Text(LOCTEXT("FilterOutZeroCountTimers_Button", " !0 "))
 					]
 				]
+			]
 
-				// Check boxes for: GpuScope, ComputeScope, CpuScope
-				+ SVerticalBox::Slot()
+			// Group by
+			+ SVerticalBox::Slot()
+			.VAlign(VAlign_Center)
+			.Padding(2.0f)
+			.AutoHeight()
+			[
+				SNew(SHorizontalBox)
+
+				+ SHorizontalBox::Slot()
+				.FillWidth(1.0f)
 				.VAlign(VAlign_Center)
-				.Padding(2.0f)
-				.AutoHeight()
 				[
-					SNew(SHorizontalBox)
+					SNew(STextBlock)
+					.Text(LOCTEXT("GroupByText", "Group by"))
+				]
 
-					+ SHorizontalBox::Slot()
-					.Padding(FMargin(0.0f,0.0f,1.0f,0.0f))
-					.FillWidth(1.0f)
+				+ SHorizontalBox::Slot()
+				.FillWidth(2.0f)
+				.VAlign(VAlign_Center)
+				[
+					SAssignNew(GroupByComboBox, SComboBox<TSharedPtr<ETimerGroupingMode>>)
+					.ToolTipText(this, &STimersView::GroupBy_GetSelectedTooltipText)
+					.OptionsSource(&GroupByOptionsSource)
+					.OnSelectionChanged(this, &STimersView::GroupBy_OnSelectionChanged)
+					.OnGenerateWidget(this, &STimersView::GroupBy_OnGenerateWidget)
 					[
-						GetToggleButtonForTimerType(ETimerNodeType::GpuScope)
+						SNew(STextBlock)
+						.Text(this, &STimersView::GroupBy_GetSelectedText)
 					]
+				]
+			]
+
+			// Check boxes for: GpuScope, ComputeScope, CpuScope
+			+ SVerticalBox::Slot()
+			.VAlign(VAlign_Center)
+			.Padding(2.0f)
+			.AutoHeight()
+			[
+				SNew(SHorizontalBox)
 
-					//+ SHorizontalBox::Slot()
-					//.Padding(FMargin(1.0f,0.0f,1.0f,0.0f))
-					//.FillWidth(1.0f)
-					//[
-					//	GetToggleButtonForTimerType(ETimerNodeType::ComputeScope)
-					//]
+				+ SHorizontalBox::Slot()
+				.Padding(FMargin(0.0f,0.0f,1.0f,0.0f))
+				.FillWidth(1.0f)
+				[
+					GetToggleButtonForTimerType(ETimerNodeType::GpuScope)
+				]
 
-					+ SHorizontalBox::Slot()
-					.Padding(FMargin(1.0f,0.0f,1.0f,0.0f))
-					.FillWidth(1.0f)
-					[
-						GetToggleButtonForTimerType(ETimerNodeType::CpuScope)
-					]
+				//+ SHorizontalBox::Slot()
+				//.Padding(FMargin(1.0f,0.0f,1.0f,0.0f))
+				//.FillWidth(1.0f)
+				//[
+				//	GetToggleButtonForTimerType(ETimerNodeType::ComputeScope)
+				//]
+
+				+ SHorizontalBox::Slot()
+				.Padding(FMargin(1.0f,0.0f,1.0f,0.0f))
+				.FillWidth(1.0f)
+				[
+					GetToggleButtonForTimerType(ETimerNodeType::CpuScope)
 				]
 			]
 		]
@@ -255,10 +251,6 @@ void STimersView::Construct(const FArguments& InArgs)
 					.HAlign(HAlign_Fill)
 					.VAlign(VAlign_Fill)
 					[
-					//SNew(SBorder)
-					//.BorderImage(FEditorStyle::GetBrush("ToolPanel.GroupBorder"))
-					//.Padding(0.0f)
-					//[
 						SAssignNew(TreeView, STreeView<FTimerNodePtr>)
 						.ExternalScrollbar(ExternalScrollbar)
 						.SelectionMode(ESelectionMode::Multi)
@@ -274,7 +266,6 @@ void STimersView::Construct(const FArguments& InArgs)
 							SAssignNew(TreeViewHeaderRow, SHeaderRow)
 							.Visibility(EVisibility::Visible)
 						)
-					//]
 					]
 
 					+ SOverlay::Slot()
