[Insights] Removed initial selection for TrackerComboBox.

diff --git a/Engine/Source/Developer/TraceInsights/Private/Insights/MemoryProfiler/Widgets/SMemTagTreeView.cpp b/Engine/Source/Developer/TraceInsights/Private/Insights/MemoryProfiler/Widgets/SMemTagTreeView.cpp
--- a/Engine/Source/Developer/TraceInsights/Private/Insights/MemoryProfiler/Widgets/SMemTagTreeView.cpp
+++ b/Engine/Source/Developer/TraceInsights/Private/Insights/MemoryProfiler/Widgets/SMemTagTreeView.cpp
@@ -295,14 +295,6 @@ TSharedRef<SWidget> SMemTagTreeView::ConstructTagsGroupingWidgetArea()
 		]
 	];
 
-	TSharedPtr<SMemoryProfilerWindow> ProfilerWindow = ProfilerWindowWeakPtr.Pin();
-	if (ProfilerWindow.IsValid())
-	{
-		FMemorySharedState& SharedState = ProfilerWindow->GetSharedState();
-		TrackerComboBox->SetSelectedItem(SharedState.GetCurrentTracker());
-		TrackerComboBox->RefreshOptions();
-	}
-
 	return Widget;
 }
 
