[Insights]
  - Changed FTimingProfilerManager::Tick trying to address a rare crash happening when app shutsdown after is left unattended for long time.
  - Fixed return of GetType in "node" classes.
  - Callers/Callees: Removed suffix after name for timer nodes.

diff --git a/Engine/Source/Developer/TraceInsights/Private/Insights/MemoryProfiler/ViewModels/MemAllocNode.h b/Engine/Source/Developer/TraceInsights/Private/Insights/MemoryProfiler/ViewModels/MemAllocNode.h
--- a/Engine/Source/Developer/TraceInsights/Private/Insights/MemoryProfiler/ViewModels/MemAllocNode.h
+++ b/Engine/Source/Developer/TraceInsights/Private/Insights/MemoryProfiler/ViewModels/MemAllocNode.h
@@ -73,7 +73,7 @@ public:
 	/**
 	 * @return a type of this MemAlloc node or EMemAllocNodeType::Group for group nodes.
 	 */
-	const EMemAllocNodeType& GetType() const { return Type; }
+	EMemAllocNodeType GetType() const { return Type; }
 
 	FMemAllocTable& GetMemTableChecked() const
 	{
diff --git a/Engine/Source/Developer/TraceInsights/Private/Insights/MemoryProfiler/ViewModels/MemTagNode.h b/Engine/Source/Developer/TraceInsights/Private/Insights/MemoryProfiler/ViewModels/MemTagNode.h
--- a/Engine/Source/Developer/TraceInsights/Private/Insights/MemoryProfiler/ViewModels/MemTagNode.h
+++ b/Engine/Source/Developer/TraceInsights/Private/Insights/MemoryProfiler/ViewModels/MemTagNode.h
@@ -86,7 +86,7 @@ public:
 	/**
 	 * @return a type of this MemTag node or EMemTagNodeType::Group for group nodes.
 	 */
-	const EMemTagNodeType& GetType() const { return Type; }
+	EMemTagNodeType GetType() const { return Type; }
 
 	bool IsValidStat() const { return MemTag != nullptr; }
 	Insights::FMemoryTag* GetMemTag() const { return MemTag; }
diff --git a/Engine/Source/Developer/TraceInsights/Private/Insights/NetworkingProfiler/ViewModels/NetEventNode.h b/Engine/Source/Developer/TraceInsights/Private/Insights/NetworkingProfiler/ViewModels/NetEventNode.h
--- a/Engine/Source/Developer/TraceInsights/Private/Insights/NetworkingProfiler/ViewModels/NetEventNode.h
+++ b/Engine/Source/Developer/TraceInsights/Private/Insights/NetworkingProfiler/ViewModels/NetEventNode.h
@@ -79,7 +79,7 @@ public:
 	/**
 	 * @return a type of this NetEvent node or ENetEventNodeType::Group for group nodes.
 	 */
-	const ENetEventNodeType& GetType() const { return Type; }
+	ENetEventNodeType GetType() const { return Type; }
 
 	void SetLevel(uint32 InLevel) const { return ; }
 	uint32 GetLevel() const { return Level; }
diff --git a/Engine/Source/Developer/TraceInsights/Private/Insights/TimingProfilerManager.cpp b/Engine/Source/Developer/TraceInsights/Private/Insights/TimingProfilerManager.cpp
--- a/Engine/Source/Developer/TraceInsights/Private/Insights/TimingProfilerManager.cpp
+++ b/Engine/Source/Developer/TraceInsights/Private/Insights/TimingProfilerManager.cpp
@@ -242,11 +242,14 @@ FTimingProfilerActionManager& FTimingProfilerManager::GetActionManager()
 
 bool FTimingProfilerManager::Tick(float DeltaTime)
 {
-	// Check if session has Timing events (to spawn the tab), but not too often.
-	if (!bIsAvailable && AvailabilityCheck.Tick())
+	TSharedPtr<FInsightsManager> InsightsManager = FInsightsManager::Get();
+	check(InsightsManager.IsValid());
+
+	TSharedPtr<const TraceServices::IAnalysisSession> Session = InsightsManager->GetSession();
+	if (Session.IsValid())
 	{
-		TSharedPtr<const TraceServices::IAnalysisSession> Session = FInsightsManager::Get()->GetSession();
-		if (Session.IsValid())
+		// Check if session has Timing events (to spawn the tab), but not too often.
+		if (!bIsAvailable && AvailabilityCheck.Tick())
 		{
 			bIsAvailable = true;
 
@@ -268,9 +271,9 @@ bool FTimingProfilerManager::Tick(float DeltaTime)
 			// Do not check again until the next session changed event (see OnSessionChanged).
 			AvailabilityCheck.Disable();
 		}
-	}
 
-	TimerButterflyAggregator->Tick(FInsightsManager::Get()->GetSession(), 0.0f, DeltaTime, [this]() { FinishTimerButterflyAggregation(); });
+		TimerButterflyAggregator->Tick(Session, 0.0f, DeltaTime, [this]() { FinishTimerButterflyAggregation(); });
+	}
 
 	return true;
 }
diff --git a/Engine/Source/Developer/TraceInsights/Private/Insights/ViewModels/TimerNode.h b/Engine/Source/Developer/TraceInsights/Private/Insights/ViewModels/TimerNode.h
--- a/Engine/Source/Developer/TraceInsights/Private/Insights/ViewModels/TimerNode.h
+++ b/Engine/Source/Developer/TraceInsights/Private/Insights/ViewModels/TimerNode.h
@@ -82,7 +82,7 @@ public:
 	/**
 	 * @return a type of this timer node or ETimerNodeType::Group for group nodes.
 	 */
-	const ETimerNodeType& GetType() const { return Type; }
+	ETimerNodeType GetType() const { return Type; }
 
 	/**
 	 * @return color of the node. Used when showing a graph series for a stats counter.
diff --git a/Engine/Source/Developer/TraceInsights/Private/Insights/Widgets/STimerTableCell.h b/Engine/Source/Developer/TraceInsights/Private/Insights/Widgets/STimerTableCell.h
--- a/Engine/Source/Developer/TraceInsights/Private/Insights/Widgets/STimerTableCell.h
+++ b/Engine/Source/Developer/TraceInsights/Private/Insights/Widgets/STimerTableCell.h
@@ -124,7 +124,7 @@ protected:
 
 	EVisibility HasExtraDisplayName() const
 	{
-		return TimerNodePtr->HasExtraDisplayName() ? EVisibility::Visible : EVisibility::Collapsed;
+		return TimerNodePtr->GetType() == ETimerNodeType::Group && TimerNodePtr->HasExtraDisplayName() ? EVisibility::Visible : EVisibility::Collapsed;
 	}
 
 	FText GetValueAsText() const;
