[Insights]
  - Fixed regression re selection change of a timer not working as expected in Callers/Callees. Also fixes regression re double clicking a timer/counter not adding/removing the graph track.
  - Added "Add/Remove series to/from graph track" in context menu of a counter in the Counters panel.

diff --git a/Engine/Source/Developer/TraceInsights/Private/Insights/Widgets/SStatsView.cpp b/Engine/Source/Developer/TraceInsights/Private/Insights/Widgets/SStatsView.cpp
--- a/Engine/Source/Developer/TraceInsights/Private/Insights/Widgets/SStatsView.cpp
+++ b/Engine/Source/Developer/TraceInsights/Private/Insights/Widgets/SStatsView.cpp
@@ -94,6 +94,8 @@ SStatsView::~SStatsView()
 	{
 		FInsightsManager::Get()->GetSessionChangedEvent().RemoveAll(this);
 	}
+
+	FStatsViewCommands::Unregister();
 }
 
 ////////////////////////////////////////////////////////////////////////////////////////////////////
@@ -141,10 +143,10 @@ void SStatsView::Construct(const FArguments& InArgs)
 					.HintText(LOCTEXT("SearchBoxHint", "Search stats counters or groups"))
 					.OnTextChanged(this, &SStatsView::SearchBox_OnTextChanged)
 					.IsEnabled(this, &SStatsView::SearchBox_IsEnabled)
-					.ToolTipText(LOCTEXT("FilterSearchHint", "Type here to search stats counter or group"))
+					.ToolTipText(LOCTEXT("FilterSearchHint", "Type here to search stats counter or group."))
 				]
 
-				// Filter out timers with zero instance count
+				// Filter out counters with zero instance count
 				+ SHorizontalBox::Slot()
 				.VAlign(VAlign_Center)
 				.Padding(2.0f)
@@ -366,6 +368,46 @@ TSharedPtr<SWidget> SStatsView::TreeView_GetMenuContent()
 	}
 	MenuBuilder.EndSection();
 
+	// Counter options section
+	MenuBuilder.BeginSection("CounterOptions", LOCTEXT("ContextMenu_Header_CounterOptions", "Counter Options"));
+	{
+		auto CanExecute = [NumSelectedNodes, SelectedNode]()
+		{
+			TSharedPtr<STimingProfilerWindow> Wnd = FTimingProfilerManager::Get()->GetProfilerWindow();
+			TSharedPtr<STimingView> TimingView = Wnd.IsValid() ? Wnd->GetTimingView() : nullptr;
+			return TimingView.IsValid() && NumSelectedNodes == 1 && SelectedNode.IsValid() && SelectedNode->GetType() != EStatsNodeType::Group;
+		};
+
+		/* Add/remove series to/from graph track */
+		{
+			FUIAction Action_ToggleCounterInGraphTrack;
+			Action_ToggleCounterInGraphTrack.CanExecuteAction = FCanExecuteAction::CreateLambda(CanExecute);
+			Action_ToggleCounterInGraphTrack.ExecuteAction = FExecuteAction::CreateSP(this, &SStatsView::ToggleTimingViewMainGraphEventSeries, SelectedNode);
+
+			if (SelectedNode.IsValid() &&
+				SelectedNode->GetType() != EStatsNodeType::Group &&
+				IsSeriesInTimingViewMainGraph(SelectedNode))
+			{
+				MenuBuilder.AddMenuEntry
+				(
+					LOCTEXT("ContextMenu_Header_CounterOptions_RemoveFromGraphTrack", "Remove series from graph track"),
+					LOCTEXT("ContextMenu_Header_CounterOptions_RemoveFromGraphTrack_Desc", "Remove the series containing event instances of this counter from the timing graph track."),
+					FSlateIcon(FEditorStyle::GetStyleSetName(), "ProfilerCommand.ToggleShowDataGraph"), Action_ToggleCounterInGraphTrack, NAME_None, EUserInterfaceActionType::Button
+				);
+			}
+			else
+			{
+				MenuBuilder.AddMenuEntry
+				(
+					LOCTEXT("ContextMenu_Header_CounterOptions_AddToGraphTrack", "Add series to graph track"),
+					LOCTEXT("ContextMenu_Header_CounterOptions_AddToGraphTrack_Desc", "Add a series containing event instances of this counter to the timing graph track."),
+					FSlateIcon(FEditorStyle::GetStyleSetName(), "ProfilerCommand.ToggleShowDataGraph"), Action_ToggleCounterInGraphTrack, NAME_None, EUserInterfaceActionType::Button
+				);
+			}
+		}
+	}
+	MenuBuilder.EndSection();
+
 	MenuBuilder.BeginSection("Misc", LOCTEXT("ContextMenu_Header_Misc", "Miscellaneous"));
 	{
 		MenuBuilder.AddMenuEntry
@@ -380,7 +422,7 @@ TSharedPtr<SWidget> SStatsView::TreeView_GetMenuContent()
 		MenuBuilder.AddSubMenu
 		(
 			LOCTEXT("ContextMenu_Header_Misc_Sort", "Sort By"),
-			LOCTEXT("ContextMenu_Header_Misc_Sort_Desc", "Sort by column"),
+			LOCTEXT("ContextMenu_Header_Misc_Sort_Desc", "Sort by column."),
 			FNewMenuDelegate::CreateSP(this, &SStatsView::TreeView_BuildSortByMenu),
 			false,
 			FSlateIcon(FEditorStyle::GetStyleSetName(), "Profiler.Misc.SortBy")
@@ -393,7 +435,7 @@ TSharedPtr<SWidget> SStatsView::TreeView_GetMenuContent()
 		MenuBuilder.AddSubMenu
 		(
 			LOCTEXT("ContextMenu_Header_Columns_View", "View Column"),
-			LOCTEXT("ContextMenu_Header_Columns_View_Desc", "Hides or shows columns"),
+			LOCTEXT("ContextMenu_Header_Columns_View_Desc", "Hides or shows columns."),
 			FNewMenuDelegate::CreateSP(this, &SStatsView::TreeView_BuildViewColumnMenu),
 			false,
 			FSlateIcon(FEditorStyle::GetStyleSetName(), "Profiler.EventGraph.ViewColumn")
@@ -407,7 +449,7 @@ TSharedPtr<SWidget> SStatsView::TreeView_GetMenuContent()
 		MenuBuilder.AddMenuEntry
 		(
 			LOCTEXT("ContextMenu_Header_Columns_ShowAllColumns", "Show All Columns"),
-			LOCTEXT("ContextMenu_Header_Columns_ShowAllColumns_Desc", "Resets tree view to show all columns"),
+			LOCTEXT("ContextMenu_Header_Columns_ShowAllColumns_Desc", "Resets tree view to show all columns."),
 			FSlateIcon(FEditorStyle::GetStyleSetName(), "Profiler.EventGraph.ResetColumn"), Action_ShowAllColumns, NAME_None, EUserInterfaceActionType::Button
 		);
 
@@ -419,7 +461,7 @@ TSharedPtr<SWidget> SStatsView::TreeView_GetMenuContent()
 		MenuBuilder.AddMenuEntry
 		(
 			LOCTEXT("ContextMenu_Header_Columns_ShowMinMaxMedColumns", "Reset Columns to Min/Max/Median Preset"),
-			LOCTEXT("ContextMenu_Header_Columns_ShowMinMaxMedColumns_Desc", "Resets columns to Min/Max/Median preset"),
+			LOCTEXT("ContextMenu_Header_Columns_ShowMinMaxMedColumns_Desc", "Resets columns to Min/Max/Median preset."),
 			FSlateIcon(FEditorStyle::GetStyleSetName(), "Profiler.EventGraph.ResetColumn"), Action_ShowMinMaxMedColumns, NAME_None, EUserInterfaceActionType::Button
 		);
 
@@ -431,7 +473,7 @@ TSharedPtr<SWidget> SStatsView::TreeView_GetMenuContent()
 		MenuBuilder.AddMenuEntry
 		(
 			LOCTEXT("ContextMenu_Header_Columns_ResetColumns", "Reset Columns to Default"),
-			LOCTEXT("ContextMenu_Header_Columns_ResetColumns_Desc", "Resets columns to default"),
+			LOCTEXT("ContextMenu_Header_Columns_ResetColumns_Desc", "Resets columns to default."),
 			FSlateIcon(FEditorStyle::GetStyleSetName(), "Profiler.EventGraph.ResetColumn"), Action_ResetColumns, NAME_None, EUserInterfaceActionType::Button
 		);
 	}
@@ -701,18 +743,18 @@ void SStatsView::ApplyFiltering()
 		const bool bIsGroupVisible = Filters->PassesAllFilters(GroupPtr);
 
 		const TArray<Insights::FBaseTreeNodePtr>& GroupChildren = GroupPtr->GetChildren();
-		const int32 NumChildren = GroupChildren.Num();
 		int32 NumVisibleChildren = 0;
-		for (int32 Cx = 0; Cx < NumChildren; ++Cx)
+		for (const Insights::FBaseTreeNodePtr& ChildPtr : GroupChildren)
 		{
-			// Add a child.
-			const FStatsNodePtr& NodePtr = StaticCastSharedPtr<FStatsNode, Insights::FBaseTreeNode>(GroupChildren[Cx]);
+			const FStatsNodePtr& NodePtr = StaticCastSharedPtr<FStatsNode, Insights::FBaseTreeNode>(ChildPtr);
+
 			const bool bIsChildVisible = (!bFilterOutZeroCountStats || NodePtr->GetAggregatedStats().Count > 0)
 									  && FilterByNodeType[static_cast<int>(NodePtr->GetType())]
 									  && FilterByDataType[static_cast<int>(NodePtr->GetDataType())]
 									  && Filters->PassesAllFilters(NodePtr);
 			if (bIsChildVisible)
 			{
+				// Add a child.
 				GroupPtr->AddFilteredChild(NodePtr);
 				NumVisibleChildren++;
 			}
@@ -741,9 +783,8 @@ void SStatsView::ApplyFiltering()
 			bExpansionSaved = true;
 		}
 
-		for (int32 Fx = 0; Fx < FilteredGroupNodes.Num(); Fx++)
+		for (const FStatsNodePtr& GroupPtr : FilteredGroupNodes)
 		{
-			const FStatsNodePtr& GroupPtr = FilteredGroupNodes[Fx];
 			TreeView->SetItemExpansion(GroupPtr, GroupPtr->IsExpanded());
 		}
 	}
@@ -971,7 +1012,7 @@ void SStatsView::TreeView_OnSelectionChanged(FStatsNodePtr SelectedItem, ESelect
 	if (SelectInfo != ESelectInfo::Direct)
 	{
 		//TArray<FStatsNodePtr> SelectedItems = TreeView->GetSelectedItems();
-		//if (SelectedItems.Num() == 1 && !SelectedItems[0]->IsGroup())
+		//if (SelectedItems.Num() == 1 && SelectedItems[0]->GetType() != EStatsNodeType::Group)
 		//{
 		//	FTimingProfilerManager::Get()->SetSelectedCounter(SelectedItems[0]->GetCounterId());
 		//}
@@ -994,44 +1035,14 @@ void SStatsView::TreeView_OnGetChildren(FStatsNodePtr InParent, TArray<FStatsNod
 
 void SStatsView::TreeView_OnMouseButtonDoubleClick(FStatsNodePtr NodePtr)
 {
-	if (NodePtr->IsGroup())
+	if (NodePtr->GetType() == EStatsNodeType::Group)
 	{
 		const bool bIsGroupExpanded = TreeView->IsItemExpanded(NodePtr);
 		TreeView->SetItemExpansion(NodePtr, !bIsGroupExpanded);
 	}
 	else
 	{
-		TSharedPtr<STimingProfilerWindow> Wnd = FTimingProfilerManager::Get()->GetProfilerWindow();
-		TSharedPtr<STimingView> TimingView = Wnd.IsValid() ? Wnd->GetTimingView() : nullptr;
-		if (TimingView.IsValid())
-		{
-			TSharedPtr<FTimingGraphTrack> GraphTrack = TimingView->GetMainTimingGraphTrack();
-			if (GraphTrack.IsValid())
-			{
-				ToggleGraphSeries(GraphTrack.ToSharedRef(), NodePtr.ToSharedRef());
-			}
-		}
-	}
-}
-
-////////////////////////////////////////////////////////////////////////////////////////////////////
-
-void SStatsView::ToggleGraphSeries(TSharedRef<FTimingGraphTrack> GraphTrack, FStatsNodeRef NodePtr)
-{
-	const uint32 CounterId = NodePtr->GetCounterId();
-	TSharedPtr<FTimingGraphSeries> Series = GraphTrack->GetStatsCounterSeries(CounterId);
-	if (Series.IsValid())
-	{
-		GraphTrack->RemoveStatsCounterSeries(CounterId);
-		GraphTrack->SetDirtyFlag();
-		NodePtr->SetAddedToGraphFlag(false);
-	}
-	else
-	{
-		GraphTrack->Show();
-		Series = GraphTrack->AddStatsCounterSeries(CounterId, NodePtr->GetColor());
-		GraphTrack->SetDirtyFlag();
-		NodePtr->SetAddedToGraphFlag(true);
+		ToggleTimingViewMainGraphEventSeries(NodePtr);
 	}
 }
 
@@ -1394,15 +1405,6 @@ void SStatsView::SortTreeNodesRec(FStatsNode& Node, const Insights::ITableCellVa
 	{
 		Node.SortChildrenAscending(Sorter);
 	}
-
-	//for (Insights::FBaseTreeNodePtr ChildPtr : Node.GetChildren())
-	//{
-	//	//if (ChildPtr->IsGroup())
-	//	if (ChildPtr->GetChildren().Num() > 0)
-	//	{
-	//		SortTreeNodesRec(*StaticCastSharedPtr<FStatsNode>(ChildPtr), Sorter);
-	//	}
-	//}
 }
 
 ////////////////////////////////////////////////////////////////////////////////////////////////////
@@ -1967,6 +1969,65 @@ void SStatsView::SelectCounterNode(uint32 CounterId)
 
 ////////////////////////////////////////////////////////////////////////////////////////////////////
 
+TSharedPtr<FTimingGraphTrack> SStatsView::GetTimingViewMainGraphTrack() const
+{
+	TSharedPtr<STimingProfilerWindow> Wnd = FTimingProfilerManager::Get()->GetProfilerWindow();
+	TSharedPtr<STimingView> TimingView = Wnd.IsValid() ? Wnd->GetTimingView() : nullptr;
+
+	return TimingView.IsValid() ? TimingView->GetMainTimingGraphTrack() : nullptr;
+}
+
+////////////////////////////////////////////////////////////////////////////////////////////////////
+
+void SStatsView::ToggleGraphSeries(TSharedRef<FTimingGraphTrack> GraphTrack, FStatsNodeRef NodePtr) const
+{
+	const uint32 CounterId = NodePtr->GetCounterId();
+	TSharedPtr<FTimingGraphSeries> Series = GraphTrack->GetStatsCounterSeries(CounterId);
+	if (Series.IsValid())
+	{
+		GraphTrack->RemoveStatsCounterSeries(CounterId);
+		GraphTrack->SetDirtyFlag();
+		NodePtr->SetAddedToGraphFlag(false);
+	}
+	else
+	{
+		GraphTrack->Show();
+		Series = GraphTrack->AddStatsCounterSeries(CounterId, NodePtr->GetColor());
+		GraphTrack->SetDirtyFlag();
+		NodePtr->SetAddedToGraphFlag(true);
+	}
+}
+
+////////////////////////////////////////////////////////////////////////////////////////////////////
+
+bool SStatsView::IsSeriesInTimingViewMainGraph(FStatsNodePtr CounterNode) const
+{
+	TSharedPtr<FTimingGraphTrack> GraphTrack = GetTimingViewMainGraphTrack();
+
+	if (GraphTrack.IsValid())
+	{
+		const uint32 CounterId = CounterNode->GetCounterId();
+		TSharedPtr<FTimingGraphSeries> Series = GraphTrack->GetStatsCounterSeries(CounterId);
+
+		return Series.IsValid();
+	}
+
+	return false;
+}
+
+////////////////////////////////////////////////////////////////////////////////////////////////////
+
+void SStatsView::ToggleTimingViewMainGraphEventSeries(FStatsNodePtr CounterNode) const
+{
+	TSharedPtr<FTimingGraphTrack> GraphTrack = GetTimingViewMainGraphTrack();
+	if (GraphTrack.IsValid())
+	{
+		ToggleGraphSeries(GraphTrack.ToSharedRef(), CounterNode.ToSharedRef());
+	}
+}
+
+////////////////////////////////////////////////////////////////////////////////////////////////////
+
 bool SStatsView::ContextMenu_CopySelectedToClipboard_CanExecute() const
 {
 	const TArray<FStatsNodePtr> SelectedNodes = TreeView->GetSelectedItems();
@@ -1984,9 +2045,9 @@ void SStatsView::ContextMenu_CopySelectedToClipboard_Execute()
 	}
 
 	TArray<Insights::FBaseTreeNodePtr> SelectedNodes;
-	for (FStatsNodePtr TimerPtr : TreeView->GetSelectedItems())
+	for (FStatsNodePtr CounterPtr : TreeView->GetSelectedItems())
 	{
-		SelectedNodes.Add(TimerPtr);
+		SelectedNodes.Add(CounterPtr);
 	}
 
 	if (SelectedNodes.Num() == 0)
diff --git a/Engine/Source/Developer/TraceInsights/Private/Insights/Widgets/SStatsView.h b/Engine/Source/Developer/TraceInsights/Private/Insights/Widgets/SStatsView.h
--- a/Engine/Source/Developer/TraceInsights/Private/Insights/Widgets/SStatsView.h
+++ b/Engine/Source/Developer/TraceInsights/Private/Insights/Widgets/SStatsView.h
@@ -83,8 +83,6 @@ public:
 	void ResetStats();
 	void UpdateStats(double StartTime, double EndTime);
 
-	void ToggleGraphSeries(TSharedRef<FTimingGraphTrack> GraphTrack, FStatsNodeRef NodePtr);
-
 	FStatsNodePtr GetCounterNode(uint32 CounterId) const;
 	void SelectCounterNode(uint32 CounterId);
 
@@ -260,6 +258,13 @@ private:
 
 	////////////////////////////////////////////////////////////////////////////////////////////////////
 
+	TSharedPtr<FTimingGraphTrack> GetTimingViewMainGraphTrack() const;
+	void ToggleGraphSeries(TSharedRef<FTimingGraphTrack> GraphTrack, FStatsNodeRef NodePtr) const;
+	bool IsSeriesInTimingViewMainGraph(FStatsNodePtr CounterNode) const;
+	void ToggleTimingViewMainGraphEventSeries(FStatsNodePtr CounterNode) const;
+
+	////////////////////////////////////////////////////////////////////////////////////////////////////
+
 	/**
 	 * Ticks this widget.  Override in derived classes, but always call the parent implementation.
 	 *
diff --git a/Engine/Source/Developer/TraceInsights/Private/Insights/Widgets/STimersView.cpp b/Engine/Source/Developer/TraceInsights/Private/Insights/Widgets/STimersView.cpp
--- a/Engine/Source/Developer/TraceInsights/Private/Insights/Widgets/STimersView.cpp
+++ b/Engine/Source/Developer/TraceInsights/Private/Insights/Widgets/STimersView.cpp
@@ -48,7 +48,7 @@ class FTimersViewCommands : public TCommands<FTimersViewCommands>
 {
 public:
 	FTimersViewCommands()
-		: TCommands<FTimersViewCommands>(TEXT("STimersViewCommands"), NSLOCTEXT("STimersViewCommands", "Timer View Commands", "Timer View Commands"), NAME_None, FEditorStyle::Get().GetStyleSetName())
+		: TCommands<FTimersViewCommands>(TEXT("FTimersViewCommands"), NSLOCTEXT("FTimersViewCommands", "Timer View Commands", "Timer View Commands"), NAME_None, FEditorStyle::Get().GetStyleSetName())
 	{
 	}
 
@@ -87,15 +87,6 @@ STimersView::STimersView()
 
 ////////////////////////////////////////////////////////////////////////////////////////////////////
 
-void STimersView::InitCommandList()
-{
-	FTimersViewCommands::Register();
-	CommandList = MakeShared<FUICommandList>();
-	CommandList->MapAction(FTimersViewCommands::Get().Command_CopyToClipboard, FExecuteAction::CreateSP(this, &STimersView::ContextMenu_CopySelectedToClipboard_Execute), FCanExecuteAction::CreateSP(this, &STimersView::ContextMenu_CopySelectedToClipboard_CanExecute));
-}
-
-////////////////////////////////////////////////////////////////////////////////////////////////////
-
 STimersView::~STimersView()
 {
 	// Remove ourselves from the Insights manager.
@@ -109,6 +100,15 @@ STimersView::~STimersView()
 
 ////////////////////////////////////////////////////////////////////////////////////////////////////
 
+void STimersView::InitCommandList()
+{
+	FTimersViewCommands::Register();
+	CommandList = MakeShared<FUICommandList>();
+	CommandList->MapAction(FTimersViewCommands::Get().Command_CopyToClipboard, FExecuteAction::CreateSP(this, &STimersView::ContextMenu_CopySelectedToClipboard_Execute), FCanExecuteAction::CreateSP(this, &STimersView::ContextMenu_CopySelectedToClipboard_CanExecute));
+}
+
+////////////////////////////////////////////////////////////////////////////////////////////////////
+
 BEGIN_SLATE_FUNCTION_BUILD_OPTIMIZATION
 void STimersView::Construct(const FArguments& InArgs)
 {
@@ -143,7 +143,7 @@ void STimersView::Construct(const FArguments& InArgs)
 					.HintText(LOCTEXT("SearchBoxHint", "Search timers or groups"))
 					.OnTextChanged(this, &STimersView::SearchBox_OnTextChanged)
 					.IsEnabled(this, &STimersView::SearchBox_IsEnabled)
-					.ToolTipText(LOCTEXT("FilterSearchHint", "Type here to search timer or group"))
+					.ToolTipText(LOCTEXT("FilterSearchHint", "Type here to search timer or group."))
 				]
 
 				// Filter out timers with zero instance count
@@ -382,10 +382,10 @@ TSharedPtr<SWidget> STimersView::TreeView_GetMenuContent()
 		{
 			TSharedPtr<STimingProfilerWindow> Wnd = FTimingProfilerManager::Get()->GetProfilerWindow();
 			TSharedPtr<STimingView> TimingView = Wnd.IsValid() ? Wnd->GetTimingView() : nullptr;
-			return TimingView.IsValid() && NumSelectedNodes == 1 && SelectedNode.IsValid() && !SelectedNode->IsGroup();
+			return TimingView.IsValid() && NumSelectedNodes == 1 && SelectedNode.IsValid() && SelectedNode->GetType() != ETimerNodeType::Group;
 		};
 
-		/*Highlight event*/
+		/* Highlight event */
 		{
 			FUIAction Action_ToggleHighlight;
 			Action_ToggleHighlight.CanExecuteAction = FCanExecuteAction::CreateLambda(CanExecute);
@@ -394,7 +394,10 @@ TSharedPtr<SWidget> STimersView::TreeView_GetMenuContent()
 			TSharedPtr<STimingProfilerWindow> Wnd = FTimingProfilerManager::Get()->GetProfilerWindow();
 			TSharedPtr<STimingView> TimingView = Wnd.IsValid() ? Wnd->GetTimingView() : nullptr;
 
-			if (SelectedNode.IsValid() && !SelectedNode->IsGroup() && TimingView.IsValid() && TimingView->IsFilterByEventType(SelectedNode->GetTimerId()))
+			if (SelectedNode.IsValid() &&
+				SelectedNode->GetType() != ETimerNodeType::Group &&
+				TimingView.IsValid() &&
+				TimingView->IsFilterByEventType(SelectedNode->GetTimerId()))
 			{
 				MenuBuilder.AddMenuEntry
 				(
@@ -414,13 +417,15 @@ TSharedPtr<SWidget> STimersView::TreeView_GetMenuContent()
 			}
 		}
 
-		/*Add series to timing view graph track*/
+		/* Add/remove series to/from graph track */
 		{
 			FUIAction Action_ToggleTimerInGraphTrack;
 			Action_ToggleTimerInGraphTrack.CanExecuteAction = FCanExecuteAction::CreateLambda(CanExecute);
 			Action_ToggleTimerInGraphTrack.ExecuteAction = FExecuteAction::CreateSP(this, &STimersView::ToggleTimingViewMainGraphEventSeries, SelectedNode);
 
-			if (SelectedNode.IsValid() && !SelectedNode->IsGroup() && IsSeriesInTimingViewMainGraph(SelectedNode))
+			if (SelectedNode.IsValid() &&
+				SelectedNode->GetType() != ETimerNodeType::Group &&
+				IsSeriesInTimingViewMainGraph(SelectedNode))
 			{
 				MenuBuilder.AddMenuEntry
 				(
@@ -456,7 +461,7 @@ TSharedPtr<SWidget> STimersView::TreeView_GetMenuContent()
 		MenuBuilder.AddSubMenu
 		(
 			LOCTEXT("ContextMenu_Header_Misc_Sort", "Sort By"),
-			LOCTEXT("ContextMenu_Header_Misc_Sort_Desc", "Sort by column"),
+			LOCTEXT("ContextMenu_Header_Misc_Sort_Desc", "Sort by column."),
 			FNewMenuDelegate::CreateSP(this, &STimersView::TreeView_BuildSortByMenu),
 			false,
 			FSlateIcon(FEditorStyle::GetStyleSetName(), "Profiler.Misc.SortBy")
@@ -852,7 +857,7 @@ void STimersView::ApplyFiltering()
 		for (const Insights::FBaseTreeNodePtr& ChildPtr : GroupChildren)
 		{
 			const FTimerNodePtr& NodePtr = StaticCastSharedPtr<FTimerNode, Insights::FBaseTreeNode>(ChildPtr);
-			TraceServices::FTimingProfilerAggregatedStats& NodeAggregatedStats = NodePtr->GetAggregatedStats();
+			const TraceServices::FTimingProfilerAggregatedStats& NodeAggregatedStats = NodePtr->GetAggregatedStats();
 
 			if (NodeAggregatedStats.InstanceCount > 0)
 			{
@@ -959,7 +964,7 @@ void STimersView::TreeView_OnSelectionChanged(FTimerNodePtr SelectedItem, ESelec
 	if (SelectInfo != ESelectInfo::Direct)
 	{
 		TArray<FTimerNodePtr> SelectedItems = TreeView->GetSelectedItems();
-		if (SelectedItems.Num() == 1 && !SelectedItems[0]->IsGroup())
+		if (SelectedItems.Num() == 1 && SelectedItems[0]->GetType() != ETimerNodeType::Group)
 		{
 			FTimingProfilerManager::Get()->SetSelectedTimer(SelectedItems[0]->GetTimerId());
 		}
@@ -982,7 +987,7 @@ void STimersView::TreeView_OnGetChildren(FTimerNodePtr InParent, TArray<FTimerNo
 
 void STimersView::TreeView_OnMouseButtonDoubleClick(FTimerNodePtr NodePtr)
 {
-	if (NodePtr->IsGroup())
+	if (NodePtr->GetType() == ETimerNodeType::Group)
 	{
 		const bool bIsGroupExpanded = TreeView->IsItemExpanded(NodePtr);
 		TreeView->SetItemExpansion(NodePtr, !bIsGroupExpanded);
@@ -1349,15 +1354,6 @@ void STimersView::SortTreeNodesRec(FTimerNode& Node, const Insights::ITableCellV
 	{
 		Node.SortChildrenAscending(Sorter);
 	}
-
-	//for (Insights::FBaseTreeNodePtr ChildPtr : Node.GetChildren())
-	//{
-	//	//if (ChildPtr->IsGroup())
-	//	if (ChildPtr->GetChildren().Num() > 0)
-	//	{
-	//		SortTreeNodesRec(*StaticCastSharedPtr<FTimerNode>(ChildPtr), Sorter);
-	//	}
-	//}
 }
 
 ////////////////////////////////////////////////////////////////////////////////////////////////////
