[Insights] Added SortChildren to FBaseTreeNode.

diff --git a/Engine/Source/Developer/TraceInsights/Private/Insights/Table/ViewModels/BaseTreeNode.h b/Engine/Source/Developer/TraceInsights/Private/Insights/Table/ViewModels/BaseTreeNode.h
--- a/Engine/Source/Developer/TraceInsights/Private/Insights/Table/ViewModels/BaseTreeNode.h
+++ b/Engine/Source/Developer/TraceInsights/Private/Insights/Table/ViewModels/BaseTreeNode.h
@@ -146,6 +146,12 @@ public:
 	void SortChildrenAscending(const ITableCellValueSorter& Sorter);
 	void SortChildrenDescending(const ITableCellValueSorter& Sorter);
 
+	template <class PREDICATE_CLASS>
+	void SortChildren(const PREDICATE_CLASS& Predicate)
+	{
+		GroupData->Children.Sort(Predicate);
+	}
+
 	/** Adds specified child to the children and sets group for it. */
 	FORCEINLINE_DEBUGGABLE void AddChildAndSetGroupPtr(const FBaseTreeNodePtr& ChildPtr)
 	{
diff --git a/Engine/Source/Developer/TraceInsights/Private/Insights/Widgets/STimerTreeView.cpp b/Engine/Source/Developer/TraceInsights/Private/Insights/Widgets/STimerTreeView.cpp
--- a/Engine/Source/Developer/TraceInsights/Private/Insights/Widgets/STimerTreeView.cpp
+++ b/Engine/Source/Developer/TraceInsights/Private/Insights/Widgets/STimerTreeView.cpp
@@ -119,10 +119,6 @@ void STimerTreeView::Construct(const FArguments& InArgs, const FText& InViewName
 				.HAlign(HAlign_Fill)
 				.VAlign(VAlign_Fill)
 				[
-				//SNew(SBorder)
-				//.BorderImage(FEditorStyle::GetBrush("ToolPanel.GroupBorder"))
-				//.Padding(0.0f)
-				//[
 					SAssignNew(TreeView, STreeView<FTimerNodePtr>)
 					.ExternalScrollbar(ExternalScrollbar)
 					.SelectionMode(ESelectionMode::Multi)
@@ -138,7 +134,6 @@ void STimerTreeView::Construct(const FArguments& InArgs, const FText& InViewName
 						SAssignNew(TreeViewHeaderRow, SHeaderRow)
 						.Visibility(EVisibility::Visible)
 					)
-				//]
 				]
 
 				+ SOverlay::Slot()
@@ -1048,8 +1043,7 @@ FTimerNodePtr STimerTreeView::CreateTimerNodeRec(const TraceServices::FTimingPro
 	}
 
 	// Sort children by InclTime (descending).
-	TArray<Insights::FBaseTreeNodePtr>& Children = const_cast<TArray<Insights::FBaseTreeNodePtr>&>(TimerNodePtr->GetChildren());
-	Children.Sort([](const Insights::FBaseTreeNodePtr& A, const Insights::FBaseTreeNodePtr& B) -> bool
+	TimerNodePtr->SortChildren([](const Insights::FBaseTreeNodePtr& A, const Insights::FBaseTreeNodePtr& B) -> bool
 	{
 		const double InclTimeA = StaticCastSharedPtr<FTimerNode>(A)->GetAggregatedStats().TotalInclusiveTime;
 		const double InclTimeB = StaticCastSharedPtr<FTimerNode>(B)->GetAggregatedStats().TotalInclusiveTime;
