[Insights] LLMTags: Changed "X" button to remove only the graph tracks for LLM tags  (and not the ones based on alloc-free events).

diff --git a/Engine/Source/Developer/TraceInsights/Private/Insights/MemoryProfiler/ViewModels/MemoryGraphTrack.cpp b/Engine/Source/Developer/TraceInsights/Private/Insights/MemoryProfiler/ViewModels/MemoryGraphTrack.cpp
--- a/Engine/Source/Developer/TraceInsights/Private/Insights/MemoryProfiler/ViewModels/MemoryGraphTrack.cpp
+++ b/Engine/Source/Developer/TraceInsights/Private/Insights/MemoryProfiler/ViewModels/MemoryGraphTrack.cpp
@@ -238,7 +238,8 @@ TSharedPtr<FMemoryGraphSeries> FMemoryGraphTrack::GetMemTagSeries(Insights::FMem
 	{
 		//TODO: if (Series->Is<FMemoryGraphSeries>())
 		const TSharedPtr<FMemoryGraphSeries> MemorySeries = StaticCastSharedPtr<FMemoryGraphSeries>(Series);
-		return MemorySeries->GetTagId() == MemTagId;
+		return MemorySeries->GetTimelineType() == FMemoryGraphSeries::ETimelineType::MemTag &&
+			   MemorySeries->GetTagId() == MemTagId;
 	});
 	return (Ptr != nullptr) ? StaticCastSharedPtr<FMemoryGraphSeries>(*Ptr) : nullptr;
 }
@@ -277,9 +278,10 @@ int32 FMemoryGraphTrack::RemoveMemTagSeries(Insights::FMemoryTagId MemTagId)
 	SetDirtyFlag();
 	return AllSeries.RemoveAll([MemTagId](const TSharedPtr<FGraphSeries>& GraphSeries)
 	{
-		//TODO: if (Series->Is<FMemoryGraphSeries>())
+		//TODO: if (GraphSeries->Is<FMemoryGraphSeries>())
 		const TSharedPtr<FMemoryGraphSeries> Series = StaticCastSharedPtr<FMemoryGraphSeries>(GraphSeries);
-		return Series->GetTagId() == MemTagId;
+		return Series->GetTimelineType() == FMemoryGraphSeries::ETimelineType::MemTag &&
+			   Series->GetTagId() == MemTagId;
 	});
 }
 
@@ -288,8 +290,12 @@ int32 FMemoryGraphTrack::RemoveMemTagSeries(Insights::FMemoryTagId MemTagId)
 int32 FMemoryGraphTrack::RemoveAllMemTagSeries()
 {
 	SetDirtyFlag();
-	int32 Count = AllSeries.Num();
-	AllSeries.Reset();
+	return AllSeries.RemoveAll([](const TSharedPtr<FGraphSeries>& GraphSeries)
+	{
+		//TODO: if (GraphSeries->Is<FMemoryGraphSeries>())
+		const TSharedPtr<FMemoryGraphSeries> Series = StaticCastSharedPtr<FMemoryGraphSeries>(GraphSeries);
+		return Series->GetTimelineType() == FMemoryGraphSeries::ETimelineType::MemTag;
+	});
 	return Count;
 }
 
diff --git a/Engine/Source/Developer/TraceInsights/Private/Insights/MemoryProfiler/ViewModels/MemorySharedState.cpp b/Engine/Source/Developer/TraceInsights/Private/Insights/MemoryProfiler/ViewModels/MemorySharedState.cpp
--- a/Engine/Source/Developer/TraceInsights/Private/Insights/MemoryProfiler/ViewModels/MemorySharedState.cpp
+++ b/Engine/Source/Developer/TraceInsights/Private/Insights/MemoryProfiler/ViewModels/MemorySharedState.cpp
@@ -455,8 +455,11 @@ int32 FMemorySharedState::RemoveMemoryGraphTrack(TSharedPtr<FMemoryGraphTrack> G
 	{
 		RemoveTrackFromMemTags(GraphTrack);
 		GraphTrack->RemoveAllMemTagSeries();
-		GraphTrack->Hide();
-		TimingView->OnTrackVisibilityChanged();
+		if (GraphTrack->GetSeries().Num() == 0)
+		{
+			GraphTrack->Hide();
+			TimingView->OnTrackVisibilityChanged();
+		}
 		return -1;
 	}
 
@@ -489,45 +492,6 @@ void FMemorySharedState::RemoveTrackFromMemTags(TSharedPtr<FMemoryGraphTrack>& G
 
 ////////////////////////////////////////////////////////////////////////////////////////////////////
 
-int32 FMemorySharedState::RemoveAllMemoryGraphTracks()
-{
-	if (!TimingView.IsValid() || !CurrentTracker)
-	{
-		return -1;
-	}
-
-	int32 TrackCount = 0;
-
-	for (TSharedPtr<FMemoryGraphTrack>& GraphTrack : AllTracks)
-	{
-		GraphTrack->RemoveAllMemTagSeries();
-		if (GraphTrack != MainGraphTrack)
-		{
-			++TrackCount;
-			TimingView->RemoveTrack(GraphTrack);
-		}
-	}
-
-	AllTracks.Reset();
-
-	// Hide the MainGraphTrack instead of removing it.
-	if (MainGraphTrack.IsValid())
-	{
-		AllTracks.Add(MainGraphTrack);
-		MainGraphTrack->Hide();
-		TimingView->OnTrackVisibilityChanged();
-	}
-
-	for (Insights::FMemoryTag* TagPtr : TagList.GetTags())
-	{
-		TagPtr->RemoveAllTracks();
-	}
-
-	return TrackCount;
-}
-
-////////////////////////////////////////////////////////////////////////////////////////////////////
-
 TSharedPtr<FMemoryGraphTrack> FMemorySharedState::GetMemTagGraphTrack(Insights::FMemoryTagId MemTagId)
 {
 	if (!TimingView.IsValid())
@@ -727,6 +691,49 @@ int32 FMemorySharedState::RemoveUnusedMemTagGraphTracks()
 
 ////////////////////////////////////////////////////////////////////////////////////////////////////
 
+int32 FMemorySharedState::RemoveAllMemTagGraphTracks()
+{
+	if (!TimingView.IsValid() || !CurrentTracker)
+	{
+		return -1;
+	}
+
+	int32 TrackCount = 0;
+
+	TArray<TSharedPtr<FMemoryGraphTrack>> TracksToRemove;
+	for (TSharedPtr<FMemoryGraphTrack>& GraphTrack : AllTracks)
+	{
+		GraphTrack->RemoveAllMemTagSeries();
+		if (GraphTrack->GetSeries().Num() == 0)
+		{
+			if (GraphTrack == MainGraphTrack)
+			{
+				GraphTrack->Hide();
+				TimingView->OnTrackVisibilityChanged();
+			}
+			else
+			{
+				++TrackCount;
+				TimingView->RemoveTrack(GraphTrack);
+				TracksToRemove.Add(GraphTrack);
+			}
+		}
+	}
+	for (TSharedPtr<FMemoryGraphTrack>& GraphTrack : TracksToRemove)
+	{
+		AllTracks.Remove(GraphTrack);
+	}
+
+	for (Insights::FMemoryTag* TagPtr : TagList.GetTags())
+	{
+		TagPtr->RemoveAllTracks();
+	}
+
+	return TrackCount;
+}
+
+////////////////////////////////////////////////////////////////////////////////////////////////////
+
 TSharedPtr<FMemoryGraphSeries> FMemorySharedState::ToggleMemTagGraphSeries(TSharedPtr<FMemoryGraphTrack> GraphTrack, Insights::FMemoryTagId MemTagId)
 {
 	if (!GraphTrack.IsValid() || !CurrentTracker)
diff --git a/Engine/Source/Developer/TraceInsights/Private/Insights/MemoryProfiler/ViewModels/MemorySharedState.h b/Engine/Source/Developer/TraceInsights/Private/Insights/MemoryProfiler/ViewModels/MemorySharedState.h
--- a/Engine/Source/Developer/TraceInsights/Private/Insights/MemoryProfiler/ViewModels/MemorySharedState.h
+++ b/Engine/Source/Developer/TraceInsights/Private/Insights/MemoryProfiler/ViewModels/MemorySharedState.h
@@ -111,13 +111,13 @@ public:
 
 	TSharedPtr<FMemoryGraphTrack> CreateMemoryGraphTrack();
 	int32 RemoveMemoryGraphTrack(TSharedPtr<FMemoryGraphTrack> GraphTrack);
-	int32 RemoveAllMemoryGraphTracks();
 
 	TSharedPtr<FMemoryGraphTrack> GetMemTagGraphTrack(Insights::FMemoryTagId MemTagId);
 	TSharedPtr<FMemoryGraphTrack> CreateMemTagGraphTrack(Insights::FMemoryTagId MemTagId);
 	void RemoveTrackFromMemTags(TSharedPtr<FMemoryGraphTrack>& GraphTrack);
 	int32 RemoveMemTagGraphTrack(Insights::FMemoryTagId MemTagId);
 	int32 RemoveUnusedMemTagGraphTracks();
+	int32 RemoveAllMemTagGraphTracks();
 
 	TSharedPtr<FMemoryGraphSeries> ToggleMemTagGraphSeries(TSharedPtr<FMemoryGraphTrack> GraphTrack, Insights::FMemoryTagId MemTagId);
 
diff --git a/Engine/Source/Developer/TraceInsights/Private/Insights/MemoryProfiler/Widgets/SMemTagTreeView.cpp b/Engine/Source/Developer/TraceInsights/Private/Insights/MemoryProfiler/Widgets/SMemTagTreeView.cpp
--- a/Engine/Source/Developer/TraceInsights/Private/Insights/MemoryProfiler/Widgets/SMemTagTreeView.cpp
+++ b/Engine/Source/Developer/TraceInsights/Private/Insights/MemoryProfiler/Widgets/SMemTagTreeView.cpp
@@ -310,7 +310,7 @@ TSharedRef<SWidget> SMemTagTreeView::ConstructTracksMiniToolbar()
 	.AutoWidth()
 	[
 		SNew(SButton)
-		.ToolTipText(LOCTEXT("HideAll_ToolTip", "Remove all memory graph tracks."))
+		.ToolTipText(LOCTEXT("HideAll_ToolTip", "Remove memory graph tracks for all LLM tags."))
 		.ContentPadding(FMargin(-14.0f, 0.0f, -17.0f, -1.0f))
 		.OnClicked(this, &SMemTagTreeView::HideAllTracks_OnClicked)
 		.Content()
@@ -643,7 +643,7 @@ void SMemTagTreeView::TreeView_BuildCreateGraphTracksMenu(FMenuBuilder& MenuBuil
 			FSlateIcon(), Action_CreateGraphTracksForFilteredMemTags, NAME_None, EUserInterfaceActionType::Button
 		);
 
-		// Create all graph series
+		// Create memory graph tracks for all LLM tags
 		FUIAction Action_CreateAllGraphTracks
 		(
 			FExecuteAction::CreateSP(this, &SMemTagTreeView::CreateAllGraphTracks),
@@ -652,7 +652,7 @@ void SMemTagTreeView::TreeView_BuildCreateGraphTracksMenu(FMenuBuilder& MenuBuil
 		MenuBuilder.AddMenuEntry
 		(
 			LOCTEXT("ContextMenu_Misc_CreateAllGraphTracks", "All"),
-			LOCTEXT("ContextMenu_Misc_CreateAllGraphTracks_Desc", "Create memory graph tracks for all LLM tags"),
+			LOCTEXT("ContextMenu_Misc_CreateAllGraphTracks_Desc", "Create memory graph tracks for all LLM tags."),
 			FSlateIcon(), Action_CreateAllGraphTracks, NAME_None, EUserInterfaceActionType::Button
 		);
 	}
@@ -674,7 +674,7 @@ void SMemTagTreeView::TreeView_BuildRemoveGraphTracksMenu(FMenuBuilder& MenuBuil
 		MenuBuilder.AddMenuEntry
 		(
 			LOCTEXT("ContextMenu_Misc_RemoveGraphTracksForSelectedMemTags", "Selected"),
-			LOCTEXT("ContextMenu_Misc_RemoveGraphTracksForSelectedMemTags_Desc", "Remove memory graph tracks for selected LLM tag(s)"),
+			LOCTEXT("ContextMenu_Misc_RemoveGraphTracksForSelectedMemTags_Desc", "Remove memory graph tracks for selected LLM tag(s)."),
 			FSlateIcon(), Action_RemoveGraphTracksForSelectedMemTags, NAME_None, EUserInterfaceActionType::Button
 		);
 
@@ -687,11 +687,11 @@ void SMemTagTreeView::TreeView_BuildRemoveGraphTracksMenu(FMenuBuilder& MenuBuil
 		MenuBuilder.AddMenuEntry
 		(
 			LOCTEXT("ContextMenu_Misc_RemoveGraphTracksForUnusedMemTags", "Unused"),
-			LOCTEXT("ContextMenu_Misc_RemoveGraphTracksForUnusedMemTags_Desc", "Remove memory graph tracks for LLM tags not used by the current tracker"),
+			LOCTEXT("ContextMenu_Misc_RemoveGraphTracksForUnusedMemTags_Desc", "Remove memory graph tracks for LLM tags not used by the current tracker."),
 			FSlateIcon(), Action_RemoveGraphTracksForUnusedMemTags, NAME_None, EUserInterfaceActionType::Button
 		);
 
-		// Remove all graph series
+		// Remove memory graph tracks for all LLM tags
 		FUIAction Action_RemoveAllGraphTracks
 		(
 			FExecuteAction::CreateSP(this, &SMemTagTreeView::RemoveAllGraphTracks),
@@ -700,7 +700,7 @@ void SMemTagTreeView::TreeView_BuildRemoveGraphTracksMenu(FMenuBuilder& MenuBuil
 		MenuBuilder.AddMenuEntry
 		(
 			LOCTEXT("ContextMenu_Misc_RemoveAllGraphTracks", "All"),
-			LOCTEXT("ContextMenu_Misc_RemoveAllGraphTracks_Desc", "Remove all memory graph tracks."),
+			LOCTEXT("ContextMenu_Misc_RemoveAllGraphTracks_Desc", "Remove memory graph tracks for all LLM tags."),
 			FSlateIcon(), Action_RemoveAllGraphTracks, NAME_None, EUserInterfaceActionType::Button
 		);
 	}
@@ -2234,7 +2234,7 @@ FReply SMemTagTreeView::LoadReportXML_OnClicked()
 		if (bOpened == true && OutFiles.Num() == 1)
 		{
 			FMemorySharedState& SharedState = ProfilerWindow->GetSharedState();
-			SharedState.RemoveAllMemoryGraphTracks();
+			SharedState.RemoveAllMemTagGraphTracks();
 			SharedState.CreateTracksFromReport(OutFiles[0]);
 		}
 	}
@@ -2503,7 +2503,7 @@ void SMemTagTreeView::RemoveAllGraphTracks()
 	if (ProfilerWindow.IsValid())
 	{
 		FMemorySharedState& SharedState = ProfilerWindow->GetSharedState();
-		SharedState.RemoveAllMemoryGraphTracks();
+		SharedState.RemoveAllMemTagGraphTracks();
 	}
 }
 
