[Insights] Initial work to show task subsequents and nested tasks in Timing View.

diff --git a/Engine/Source/Developer/TraceInsights/Private/Insights/TimingProfilerCommands.cpp b/Engine/Source/Developer/TraceInsights/Private/Insights/TimingProfilerCommands.cpp
--- a/Engine/Source/Developer/TraceInsights/Private/Insights/TimingProfilerCommands.cpp
+++ b/Engine/Source/Developer/TraceInsights/Private/Insights/TimingProfilerCommands.cpp
@@ -86,6 +86,7 @@ void FTimingViewCommands::RegisterCommands()
 {
 	UI_COMMAND(ShowAllGpuTracks, "GPU Track", "Show/hide the GPU track", EUserInterfaceActionType::ToggleButton, FInputChord(EKeys::Y));
 	UI_COMMAND(ShowAllCpuTracks, "CPU Thread Tracks", "Show/hide all CPU tracks (and all CPU thread groups)", EUserInterfaceActionType::ToggleButton, FInputChord(EKeys::U));
+	UI_COMMAND(ShowTaskDependencies, "Task Dependencies", "Show/hide all dependencies of the current task", EUserInterfaceActionType::ToggleButton, FInputChord(EKeys::T));
 }
 PRAGMA_ENABLE_OPTIMIZATION
 
diff --git a/Engine/Source/Developer/TraceInsights/Private/Insights/TimingProfilerCommands.h b/Engine/Source/Developer/TraceInsights/Private/Insights/TimingProfilerCommands.h
--- a/Engine/Source/Developer/TraceInsights/Private/Insights/TimingProfilerCommands.h
+++ b/Engine/Source/Developer/TraceInsights/Private/Insights/TimingProfilerCommands.h
@@ -71,6 +71,9 @@ public:
 
 	/** Toggles visibility for all CPU thread tracks at once. */
 	TSharedPtr<FUICommandInfo> ShowAllCpuTracks;
+
+	/** Toggles visibility for task graph dependency relations */
+	TSharedPtr<FUICommandInfo> ShowTaskDependencies;
 };
 
 ////////////////////////////////////////////////////////////////////////////////////////////////////
diff --git a/Engine/Source/Developer/TraceInsights/Private/Insights/ViewModels/TaskGraphRelation.cpp b/Engine/Source/Developer/TraceInsights/Private/Insights/ViewModels/TaskGraphRelation.cpp
new file mode 100644
--- /dev/null
+++ b/Engine/Source/Developer/TraceInsights/Private/Insights/ViewModels/TaskGraphRelation.cpp
@@ -0,0 +1,39 @@
+// Copyright Epic Games, Inc. All Rights Reserved.
+
+#include "TaskGraphRelation.h"
+
+#include "Insights/ViewModels/TimingTrackViewport.h"
+#include "Insights/ViewModels/TimingViewDrawHelper.h"
+#include "Insights/Common/PaintUtils.h"
+
+namespace Insights
+{
+
+FTaskGraphRelation::FTaskGraphRelation(double InSourceTime, int32 InSourceThreadId, double InTargetTime, int32 InTargetThreadId, ETaskGraphRelationType InType)
+{
+	SourceTime = InSourceTime;
+	SourceThreadId = InSourceThreadId;
+	TargetTime = InTargetTime;
+	TargetThreadId = InTargetThreadId;
+	Type = InType;
+}
+
+void FTaskGraphRelation::Draw(const FDrawContext& DrawContext, const FTimingTrackViewport& Viewport, const ITimingViewDrawHelper& Helper)
+{
+	float X1 = Viewport.TimeToSlateUnitsRounded(SourceTime);
+	float Y1 = SourceTrack.Pin()->GetPosY();
+	Y1 += Viewport.GetLayout().GetLaneY(SourceDepth) + Viewport.GetLayout().EventH / 2;
+
+	float X2 = Viewport.TimeToSlateUnitsRounded(TargetTime);
+	float Y2 = TargetTrack.Pin()->GetPosY();
+	Y2 += Viewport.GetLayout().GetLaneY(TargetDepth) + Viewport.GetLayout().EventH / 2;
+
+	TArray<FVector2D> Points;
+	Points.Emplace(0, 0);
+	Points.Emplace(X2 - X1, Y2 - Y1);
+
+	FLinearColor Color = Type == ETaskGraphRelationType::Subsequent ? FLinearColor::Red : FLinearColor::Blue;
+	DrawContext.DrawLines(Helper.GetRelationLayerId(), X1, Y1, Points, ESlateDrawEffect::None, Color);
+}
+
+} // namespace Insights
\ No newline at end of file
diff --git a/Engine/Source/Developer/TraceInsights/Private/Insights/ViewModels/TaskGraphRelation.h b/Engine/Source/Developer/TraceInsights/Private/Insights/ViewModels/TaskGraphRelation.h
new file mode 100644
--- /dev/null
+++ b/Engine/Source/Developer/TraceInsights/Private/Insights/ViewModels/TaskGraphRelation.h
@@ -0,0 +1,60 @@
+// Copyright Epic Games, Inc. All Rights Reserved.
+
+#pragma once
+
+#include "CoreMinimal.h"
+
+#include "Insights/ViewModels/ITimingEvent.h"
+
+namespace Insights
+{
+
+class FTaskGraphRelation : public ITimingEventRelation
+{
+public:
+	enum class ETaskGraphRelationType : int32
+	{
+		Subsequent = 0,
+		Nested = 1,
+	};
+	FTaskGraphRelation(double InSourceTime, int32 InSourceThreadId, double InTargetTime, int32 InTargetThreadId, ETaskGraphRelationType InType);
+
+	virtual bool IsSolved() override { return bIsSolved; }
+
+	virtual void Draw(const FDrawContext& DrawContext, const FTimingTrackViewport& Viewport, const ITimingViewDrawHelper& Helper) override;
+
+	void SetSourceTrack(TSharedPtr<const FBaseTimingTrack> InSourceTrack) { SourceTrack = InSourceTrack; }
+	TSharedPtr<const FBaseTimingTrack> GetSourceTrack() { return SourceTrack.Pin(); }
+
+	void SetTargetTrack(TSharedPtr<const FBaseTimingTrack> InTargetTrack) { TargetTrack = InTargetTrack; }
+	TSharedPtr<const FBaseTimingTrack> GetTargetTrack() { return TargetTrack.Pin(); }
+
+	double GetSourceTime() { return SourceTime; }
+	int32 GetSourceThreadId() { return SourceThreadId; }
+	void SetSourceDepth(int32 InDepth) { SourceDepth = InDepth; }
+	int32 GetSourceDepth() { return SourceDepth; }
+
+	double GetTargetTime() { return TargetTime; }
+	int32 GetTargetThreadId() { return TargetThreadId; }
+	void SetTargetDepth(int32 InDepth) { TargetDepth = InDepth; }
+	int32 GetTargetDepth() { return TargetDepth; }
+
+	void SetIsSolved(bool InValue) { bIsSolved = InValue; }
+	ETaskGraphRelationType GetType() { return Type; }
+
+private:
+	double SourceTime;
+	int32 SourceThreadId;
+	int32 SourceDepth = 0;
+	double TargetTime;
+	int32 TargetThreadId;
+	int32 TargetDepth = 0;
+	ETaskGraphRelationType Type;
+
+	bool bIsSolved = false;
+
+	TWeakPtr<const FBaseTimingTrack> SourceTrack;
+	TWeakPtr<const FBaseTimingTrack> TargetTrack;
+};
+
+} // namespace Insights
diff --git a/Engine/Source/Developer/TraceInsights/Private/Insights/ViewModels/ThreadTimingTrack.cpp b/Engine/Source/Developer/TraceInsights/Private/Insights/ViewModels/ThreadTimingTrack.cpp
--- a/Engine/Source/Developer/TraceInsights/Private/Insights/ViewModels/ThreadTimingTrack.cpp
+++ b/Engine/Source/Developer/TraceInsights/Private/Insights/ViewModels/ThreadTimingTrack.cpp
@@ -19,6 +19,7 @@
 #include "Insights/TimingProfilerManager.h"
 #include "Insights/ViewModels/Filters.h"
 #include "Insights/ViewModels/FilterConfigurator.h"
+#include "Insights/ViewModels/TaskGraphRelation.h"
 #include "Insights/ViewModels/TimerNode.h"
 #include "Insights/ViewModels/ThreadTrackEvent.h"
 #include "Insights/ViewModels/TimingEventSearch.h"
@@ -1059,8 +1060,17 @@ void FThreadTimingTrack::InitTooltip(FTooltipDrawState& InOutTooltip, const ITim
 				InOutTooltip.AddTextLine(FString::Printf(TEXT("%s (%s queue)"), NamedThreadsStr[ThreadIndex], QueueStr), FLinearColor::Green);
 			}
 
-			InOutTooltip.AddNameValueTextLine(TEXT("Created:"), FString::Printf(TEXT("%f on %s"), Task.CreatedTimestamp, *SharedState.GetCpuTrack(Task.CreatedThreadId)->GetName()));
-			InOutTooltip.AddNameValueTextLine(TEXT("Launched:"), FString::Printf(TEXT("%f (+%s) on %s"), Task.LaunchedTimestamp, *TimeUtils::FormatTimeAuto(Task.LaunchedTimestamp - Task.CreatedTimestamp), *SharedState.GetCpuTrack(Task.LaunchedThreadId)->GetName()));
+			{
+				TSharedPtr<FCpuTimingTrack> Track = SharedState.GetCpuTrack(Task.CreatedThreadId);
+				FString TrackName = Track.IsValid() ? Track->GetName() : TEXT("Unknown");
+				InOutTooltip.AddNameValueTextLine(TEXT("Created:"), FString::Printf(TEXT("%f on %s"), Task.CreatedTimestamp, *TrackName));
+			}
+			
+			{
+				TSharedPtr<FCpuTimingTrack> Track = SharedState.GetCpuTrack(Task.LaunchedThreadId);
+				FString TrackName = Track.IsValid() ? Track->GetName() : TEXT("Unknown");
+				InOutTooltip.AddNameValueTextLine(TEXT("Launched:"), FString::Printf(TEXT("%f (+%s) on %s"), Task.LaunchedTimestamp, *TimeUtils::FormatTimeAuto(Task.LaunchedTimestamp - Task.CreatedTimestamp), *TrackName));
+			}
 			InOutTooltip.AddNameValueTextLine(TEXT("Scheduled:"), FString::Printf(TEXT("%f (+%s) on %s"), Task.ScheduledTimestamp, *TimeUtils::FormatTimeAuto(Task.ScheduledTimestamp - Task.LaunchedTimestamp), *SharedState.GetCpuTrack(Task.ScheduledThreadId)->GetName()));
 			InOutTooltip.AddNameValueTextLine(TEXT("Started:"), FString::Printf(TEXT("%f (+%s)"), Task.StartedTimestamp, *TimeUtils::FormatTimeAuto(Task.StartedTimestamp - Task.ScheduledTimestamp)));
 			if (Task.FinishedTimestamp != TraceServices::FTaskInfo::InvalidTimestamp)
@@ -1069,9 +1079,13 @@ void FThreadTimingTrack::InitTooltip(FTooltipDrawState& InOutTooltip, const ITim
 
 				if (Task.CompletedTimestamp != TraceServices::FTaskInfo::InvalidTimestamp)
 				{
-					InOutTooltip.AddNameValueTextLine(TEXT("Completed:"), FString::Printf(TEXT("%f (+%s) on %s"), Task.FinishedTimestamp, *TimeUtils::FormatTimeAuto(Task.CompletedTimestamp - Task.FinishedTimestamp), *SharedState.GetCpuTrack(Task.CompletedThreadId)->GetName()));
+					TSharedPtr<FCpuTimingTrack> Track = SharedState.GetCpuTrack(Task.CompletedThreadId);
+					FString TrackName = Track.IsValid() ? Track->GetName() : TEXT("Unknown");
+					InOutTooltip.AddNameValueTextLine(TEXT("Completed:"), FString::Printf(TEXT("%f (+%s) on %s"), Task.FinishedTimestamp, *TimeUtils::FormatTimeAuto(Task.CompletedTimestamp - Task.FinishedTimestamp), *TrackName));
 				}
 			}
+			InOutTooltip.AddNameValueTextLine(TEXT("Subsequent Tasks"), FString::Printf(TEXT("%d"), Task.Subsequents.Num()));
+			InOutTooltip.AddNameValueTextLine(TEXT("Nested Tasks:"), FString::Printf(TEXT("%d"), Task.NestedTasks.Num()));
 		};
 
 		do // info about a task
@@ -1475,4 +1489,87 @@ bool FThreadTimingTrack::HasCustomFilter() const
 
 ////////////////////////////////////////////////////////////////////////////////////////////////////
 
+void FThreadTimingTrack::GetEventRelations(const ITimingEvent& InSelectedEvent, TArray<ITimingEventRelation*>& Relations) const
+{
+	double StartTime = InSelectedEvent.GetStartTime();
+
+	TSharedPtr<const TraceServices::IAnalysisSession> Session = FInsightsManager::Get()->GetSession();
+	const TraceServices::ITasksProvider* TasksProvider = TraceServices::ReadTasksProvider(*Session.Get());
+	if (TasksProvider)
+	{
+		const TraceServices::FTaskInfo* Task = TasksProvider->TryGetTask(ThreadId, StartTime);
+		if (Task)
+		{
+			for (const TraceServices::FTaskInfo::FRelationInfo& RelationInfo : Task->Subsequents)
+			{
+				const TraceServices::FTaskInfo* RelatedTask = TasksProvider->TryGetTask(RelationInfo.RelativeId);
+
+				if (RelatedTask)
+				{
+					FTaskGraphRelation* NewRelation = new FTaskGraphRelation(StartTime, ThreadId, RelatedTask->StartedTimestamp, RelatedTask->StartedThreadId, FTaskGraphRelation::ETaskGraphRelationType::Subsequent);
+					NewRelation->SetSourceTrack(SharedThis(this));
+					NewRelation->SetSourceDepth(InSelectedEvent.GetDepth());
+					Relations.Add(NewRelation);
+				}
+			}
+
+			for (const TraceServices::FTaskInfo::FRelationInfo& RelationInfo : Task->NestedTasks)
+			{
+				const TraceServices::FTaskInfo* RelatedTask = TasksProvider->TryGetTask(RelationInfo.RelativeId);
+
+				if (RelatedTask)
+				{
+					FTaskGraphRelation* NewRelation = new FTaskGraphRelation(StartTime, ThreadId, RelatedTask->StartedTimestamp, RelatedTask->StartedThreadId, FTaskGraphRelation::ETaskGraphRelationType::Nested);
+					NewRelation->SetSourceTrack(SharedThis(this));
+					NewRelation->SetSourceDepth(InSelectedEvent.GetDepth());
+					Relations.Add(NewRelation);
+				}
+			}
+		}
+	}
+}
+
+////////////////////////////////////////////////////////////////////////////////////////////////////
+
+void FThreadTimingTrack::SolveEventRelations(const TArray<ITimingEventRelation*>& Relations) const
+{
+	for (ITimingEventRelation* Relation : Relations)
+	{
+		if (!Relation->IsSolved())
+		{
+			// A check that we have the correct type will need to be added
+			FTaskGraphRelation* TaskRelation = static_cast<FTaskGraphRelation*>(Relation);
+			if (TaskRelation->GetTargetThreadId() == ThreadId)
+			{
+				TaskRelation->SetTargetTrack(StaticCastSharedRef<const FBaseTimingTrack>(SharedThis(this)));
+				TaskRelation->SetTargetDepth(GetDepthAt(TaskRelation->GetTargetTime()) - 1);
+				TaskRelation->SetIsSolved(true);
+			}
+		}
+	}
+}
+
+////////////////////////////////////////////////////////////////////////////////////////////////////
+
+int32 FThreadTimingTrack::GetDepthAt(double Time) const
+{
+	int32 Depth = 0;
+	TSharedPtr<const TraceServices::IAnalysisSession> Session = FInsightsManager::Get()->GetSession();
+	if (Session.IsValid() && TraceServices::ReadTimingProfilerProvider(*Session.Get()))
+	{
+		TraceServices::FAnalysisSessionReadScope SessionReadScope(*Session.Get());
+
+		const TraceServices::ITimingProfilerProvider& TimingProfilerProvider = *TraceServices::ReadTimingProfilerProvider(*Session.Get());
+
+		TimingProfilerProvider.ReadTimeline(TimelineIndex,
+			[Time, &Depth](const TraceServices::ITimingProfilerProvider::Timeline& Timeline)
+			{
+				Depth = Timeline.GetDepthAt(Time);
+			});
+	}
+	return Depth;
+}
+
+////////////////////////////////////////////////////////////////////////////////////////////////////
+
 #undef LOCTEXT_NAMESPACE
diff --git a/Engine/Source/Developer/TraceInsights/Private/Insights/ViewModels/ThreadTimingTrack.h b/Engine/Source/Developer/TraceInsights/Private/Insights/ViewModels/ThreadTimingTrack.h
--- a/Engine/Source/Developer/TraceInsights/Private/Insights/ViewModels/ThreadTimingTrack.h
+++ b/Engine/Source/Developer/TraceInsights/Private/Insights/ViewModels/ThreadTimingTrack.h
@@ -160,8 +160,12 @@ public:
 	virtual void OnClipboardCopyEvent(const ITimingEvent& InSelectedEvent) const override;
 	virtual void BuildContextMenu(FMenuBuilder& MenuBuilder) override;
 
+	virtual void GetEventRelations(const ITimingEvent& InSelectedEvent, TArray<ITimingEventRelation*>& Relations) const override;
+	virtual void SolveEventRelations(const TArray<ITimingEventRelation*>& Relations) const override;
+
 protected:
 	virtual bool HasCustomFilter() const override;
+	int32 GetDepthAt(double Time) const;
 
 private:
 	void DrawSelectedEventInfo(const FThreadTrackEvent& SelectedEvent, const FTimingTrackViewport& Viewport, const FDrawContext& DrawContext, const FSlateBrush* WhiteBrush, const FSlateFontInfo& Font) const;
diff --git a/Engine/Source/Developer/TraceInsights/Private/Insights/ViewModels/TimingEventsTrack.cpp b/Engine/Source/Developer/TraceInsights/Private/Insights/ViewModels/TimingEventsTrack.cpp
--- a/Engine/Source/Developer/TraceInsights/Private/Insights/ViewModels/TimingEventsTrack.cpp
+++ b/Engine/Source/Developer/TraceInsights/Private/Insights/ViewModels/TimingEventsTrack.cpp
@@ -148,6 +148,8 @@ void FTimingEventsTrack::PreUpdate(const ITimingTrackUpdateContext& Context)
 		SetNumLanes(MaxDepth + 1);
 	}
 
+	SolveEventRelations(Context.GetCurrentRelations());
+
 	UpdateTrackHeight(Context);
 }
 
diff --git a/Engine/Source/Developer/TraceInsights/Private/Insights/ViewModels/TimingViewDrawHelper.cpp b/Engine/Source/Developer/TraceInsights/Private/Insights/ViewModels/TimingViewDrawHelper.cpp
--- a/Engine/Source/Developer/TraceInsights/Private/Insights/ViewModels/TimingViewDrawHelper.cpp
+++ b/Engine/Source/Developer/TraceInsights/Private/Insights/ViewModels/TimingViewDrawHelper.cpp
@@ -653,3 +653,16 @@ FLinearColor FTimingViewDrawHelper::GetTrackNameTextColor(const FBaseTimingTrack
 }
 
 ////////////////////////////////////////////////////////////////////////////////////////////////////
+
+void FTimingViewDrawHelper::DrawRelations(const TArray<ITimingEventRelation*> Relations) const
+{
+	for (ITimingEventRelation* Relation : Relations)
+	{
+		if (Relation->IsSolved())
+		{
+			Relation->Draw(DrawContext, Viewport, *this);
+		}
+	}
+}
+
+////////////////////////////////////////////////////////////////////////////////////////////////////
diff --git a/Engine/Source/Developer/TraceInsights/Private/Insights/ViewModels/TimingViewDrawHelper.h b/Engine/Source/Developer/TraceInsights/Private/Insights/ViewModels/TimingViewDrawHelper.h
--- a/Engine/Source/Developer/TraceInsights/Private/Insights/ViewModels/TimingViewDrawHelper.h
+++ b/Engine/Source/Developer/TraceInsights/Private/Insights/ViewModels/TimingViewDrawHelper.h
@@ -145,6 +145,7 @@ private:
 		EventFill,
 		EventText,
 		EventHighlight,
+		Relation,
 		HeaderBackground,
 		HeaderText,
 
@@ -169,6 +170,7 @@ public:
 	virtual FLinearColor GetTrackNameTextColor(const FBaseTimingTrack& Track) const override;
 	virtual int32 GetHeaderBackgroundLayerId() const override { return ReservedLayerId + ToInt32(EDrawLayer::HeaderBackground); }
 	virtual int32 GetHeaderTextLayerId() const override { return ReservedLayerId + ToInt32(EDrawLayer::HeaderText); }
+	virtual int32 GetRelationLayerId() const override {return ReservedLayerId + ToInt32(EDrawLayer::Relation); };
 	virtual int32 GetFirstLayerId() const override { return ReservedLayerId; }
 	virtual int32 GetNumLayerIds() const override { return ToInt32(EDrawLayer::Count); }
 
@@ -193,6 +195,8 @@ public:
 
 	void EndDrawTracks() const;
 
+	void DrawRelations(const TArray<ITimingEventRelation*> Relations) const;
+
 	//////////////////////////////////////////////////
 
 	void DrawTimingEventHighlight(double StartTime, double EndTime, float Y, EDrawEventMode Mode) const;
diff --git a/Engine/Source/Developer/TraceInsights/Private/Insights/Widgets/STimingView.cpp b/Engine/Source/Developer/TraceInsights/Private/Insights/Widgets/STimingView.cpp
--- a/Engine/Source/Developer/TraceInsights/Private/Insights/Widgets/STimingView.cpp
+++ b/Engine/Source/Developer/TraceInsights/Private/Insights/Widgets/STimingView.cpp
@@ -10,6 +10,8 @@
 #include "Fonts/SlateFontInfo.h"
 #include "Framework/Application/MenuStack.h"
 #include "Framework/Application/SlateApplication.h"
+#include "Framework/Commands/Commands.h"
+#include "Framework/Commands/UICommandList.h"
 #include "Framework/MultiBox/MultiBoxBuilder.h"
 #include "HAL/PlatformTime.h"
 #include "Layout/WidgetPath.h"
@@ -113,6 +115,9 @@ STimingView::~STimingView()
 	ScrollableTracks.Reset();
 	ForegroundTracks.Reset();
 
+	SelectedEvent.Reset();
+	UpdateEventRelations();
+
 	for (Insights::ITimingViewExtender* Extender : GetExtenders())
 	{
 		Extender->OnEndSession(*this);
@@ -629,6 +634,7 @@ void STimingView::Tick(const FGeometry& AllottedGeometry, const double InCurrent
 		virtual const TSharedPtr<const ITimingEvent> GetHoveredEvent() const override { return TimingView->GetHoveredEvent(); }
 		virtual const TSharedPtr<const ITimingEvent> GetSelectedEvent() const override { return TimingView->GetSelectedEvent(); }
 		virtual const TSharedPtr<ITimingEventFilter> GetEventFilter() const override { return TimingView->GetEventFilter(); }
+		virtual const TArray<ITimingEventRelation*>& GetCurrentRelations() const override { return TimingView->GetCurrentRelations(); }
 		virtual double GetCurrentTime() const override { return CurrentTime; }
 		virtual float GetDeltaTime() const override { return DeltaTime; }
 
@@ -1265,6 +1271,9 @@ int32 STimingView::OnPaint(const FPaintArgs& Args, const FGeometry& AllottedGeom
 	// Draw the time range selection.
 	FDrawHelpers::DrawTimeRangeSelection(DrawContext, Viewport, SelectionStartTime, SelectionEndTime, WhiteBrush, MainFont);
 
+	const FTimingViewDrawHelper& TimingHelper = *static_cast<const FTimingViewDrawHelper*>(&TimingDrawContext.GetHelper());
+	TimingHelper.DrawRelations(CurrentRelations);
+
 	//////////////////////////////////////////////////
 	// Post-Draw
 	{
@@ -2643,7 +2652,7 @@ void STimingView::ShowContextMenu(const FPointerEvent& MouseEvent)
 
 	const bool bShouldCloseWindowAfterMenuSelection = true;
 	//FMenuBuilder MenuBuilder(bShouldCloseWindowAfterMenuSelection, ProfilerCommandList);
-	FMenuBuilder MenuBuilder(bShouldCloseWindowAfterMenuSelection, NULL);
+	FMenuBuilder MenuBuilder(bShouldCloseWindowAfterMenuSelection, CommandList);
 
 	if (HoveredTrack.IsValid())
 	{
@@ -2669,6 +2678,21 @@ void STimingView::ShowContextMenu(const FPointerEvent& MouseEvent)
 		MenuBuilder.EndSection();
 	}
 
+	{
+		MenuBuilder.BeginSection(TEXT("Event"), LOCTEXT("Event", "Event"));
+
+		MenuBuilder.AddMenuEntry
+		(
+			FTimingViewCommands::Get().ShowTaskDependencies,
+			NAME_None,
+			TAttribute<FText>(),
+			TAttribute<FText>(),
+			FSlateIcon(FEditorStyle::GetStyleSetName(), "Profiler.Type.Calls")
+		);
+
+		MenuBuilder.EndSection();
+	}
+
 	//MenuBuilder.BeginSection(TEXT("Misc"), LOCTEXT("Miscellaneous", "Miscellaneous"));
 	//{
 	//	MenuBuilder.AddMenuEntry(Commands.ShowAllGpuTracks);
@@ -2857,8 +2881,15 @@ bool STimingView::CanChangeTrackLocation(TSharedRef<FBaseTimingTrack> Track, ETi
 
 void STimingView::BindCommands()
 {
-	//FTimingViewCommands::Register();
-	//const FTimingViewCommands& Commands = FTimingViewCommands::Get();
+	FTimingViewCommands::Register();
+	CommandList = MakeShared<FUICommandList>();
+
+	CommandList->MapAction(
+		FTimingViewCommands::Get().ShowTaskDependencies, 
+		FExecuteAction::CreateSP(this, &STimingView::ContextMenu_ShowTaskDependecies_Execute), 
+		FCanExecuteAction::CreateSP(this, &STimingView::ContextMenu_ShowTaskDependecies_CanExecute),
+		FIsActionChecked::CreateSP(this, &STimingView::ContextMenu_ShowTaskDependecies_IsChecked));
+
 	//
 	//TSharedPtr<FUICommandList> CommandList = FTimingProfilerManager::Get()->GetCommandList();
 
@@ -3410,6 +3441,7 @@ void STimingView::OnSelectedTimingEventChanged()
 		}
 	}
 
+	UpdateEventRelations();
 	OnSelectedEventChangedDelegate.Broadcast(SelectedEvent);
 }
 
@@ -4057,5 +4089,44 @@ TArray<Insights::ITimingViewExtender*> STimingView::GetExtenders() const
 
 ////////////////////////////////////////////////////////////////////////////////////////////////////
 
+void STimingView::ContextMenu_ShowTaskDependecies_Execute()
+{
+	bShowEventRelations = !bShowEventRelations;
+	UpdateEventRelations();
+}
+
+////////////////////////////////////////////////////////////////////////////////////////////////////
+
+bool STimingView::ContextMenu_ShowTaskDependecies_CanExecute()
+{
+	return true;
+}
+
+////////////////////////////////////////////////////////////////////////////////////////////////////
+
+bool STimingView::ContextMenu_ShowTaskDependecies_IsChecked()
+{
+	return bShowEventRelations;
+}
+
+////////////////////////////////////////////////////////////////////////////////////////////////////
+
+void STimingView::UpdateEventRelations()
+{
+	for (ITimingEventRelation* Relation : CurrentRelations)
+	{
+		delete Relation;
+	}
+
+	CurrentRelations.Empty();
+
+	if (bShowEventRelations && SelectedEvent.IsValid())
+	{
+		SelectedEvent->GetTrack()->GetEventRelations(*SelectedEvent, CurrentRelations);
+	}
+}
+
+////////////////////////////////////////////////////////////////////////////////////////////////////
+
 #undef ACTIVATE_BENCHMARK
 #undef LOCTEXT_NAMESPACE
diff --git a/Engine/Source/Developer/TraceInsights/Private/Insights/Widgets/STimingView.h b/Engine/Source/Developer/TraceInsights/Private/Insights/Widgets/STimingView.h
--- a/Engine/Source/Developer/TraceInsights/Private/Insights/Widgets/STimingView.h
+++ b/Engine/Source/Developer/TraceInsights/Private/Insights/Widgets/STimingView.h
@@ -33,6 +33,7 @@ class FThreadTimingSharedState;
 class FTimeRulerTrack;
 class FTimingGraphTrack;
 class FTimingViewDrawHelper;
+class FUICommandList;
 class SOverlay;
 class SScrollBar;
 
@@ -339,6 +340,8 @@ public:
 
 	const TSharedPtr<FBaseTimingTrack> GetTrackAt(float InPosX, float InPosY) const;
 
+	const TArray<ITimingEventRelation*>& GetCurrentRelations() const { return CurrentRelations; }
+
 protected:
 	virtual FVector2D ComputeDesiredSize(float) const override
 	{
@@ -427,6 +430,12 @@ protected:
 	FReply AllowTracksToProcessOnMouseButtonUp(const FGeometry& MyGeometry, const FPointerEvent& MouseEvent);
 	FReply AllowTracksToProcessOnMouseButtonDoubleClick(const FGeometry& MyGeometry, const FPointerEvent& MouseEvent);
 
+	void ContextMenu_ShowTaskDependecies_Execute();
+	bool ContextMenu_ShowTaskDependecies_CanExecute();
+	bool ContextMenu_ShowTaskDependecies_IsChecked();
+
+	void UpdateEventRelations();
+
 protected:
 	/** The track's viewport. Encapsulates info about position and scale. */
 	FTimingTrackViewport Viewport;
@@ -616,4 +625,9 @@ protected:
 	Insights::FHoveredEventChangedDelegate OnHoveredEventChangedDelegate;
 	Insights::FSelectedTrackChangedDelegate OnSelectedTrackChangedDelegate;
 	Insights::FSelectedEventChangedDelegate OnSelectedEventChangedDelegate;
-};
+
+	TSharedPtr<FUICommandList> CommandList;
+
+	TArray<ITimingEventRelation*> CurrentRelations;
+	bool bShowEventRelations = false;
+};
\ No newline at end of file
diff --git a/Engine/Source/Developer/TraceInsights/Public/Insights/Common/PaintUtils.h b/Engine/Source/Developer/TraceInsights/Public/Insights/Common/PaintUtils.h
--- a/Engine/Source/Developer/TraceInsights/Public/Insights/Common/PaintUtils.h
+++ b/Engine/Source/Developer/TraceInsights/Public/Insights/Common/PaintUtils.h
@@ -97,6 +97,11 @@ struct TRACEINSIGHTS_API FDrawContext
 		FSlateDrawElement::MakeSpline(ElementList, InLayer, MAKE_PAINT_GEOMETRY_PT(Geometry, X, Y), InStart, InStartDir, InEnd, InEndDir, InThickness, DrawEffects, InTint);
 	}
 
+	inline void DrawLines(uint32 InLayer, const float X, const float Y, const TArray<FVector2D>& Points, ESlateDrawEffect InDrawEffects = ESlateDrawEffect::None, const FLinearColor& InTint = FLinearColor::White, bool bAntialias = true, float Thickness = 1.0f) const
+	{
+		FSlateDrawElement::MakeLines(ElementList, InLayer, MAKE_PAINT_GEOMETRY_PT(Geometry, X, Y), Points, InDrawEffects, InTint, bAntialias, Thickness);
+	}
+
 	//inline void IncrementLayer() const { ++LayerId; }
 
 	// Accessors
diff --git a/Engine/Source/Developer/TraceInsights/Public/Insights/ViewModels/BaseTimingTrack.h b/Engine/Source/Developer/TraceInsights/Public/Insights/ViewModels/BaseTimingTrack.h
--- a/Engine/Source/Developer/TraceInsights/Public/Insights/ViewModels/BaseTimingTrack.h
+++ b/Engine/Source/Developer/TraceInsights/Public/Insights/ViewModels/BaseTimingTrack.h
@@ -19,6 +19,7 @@ class FTimingTrackViewport;
 class ITimingViewDrawHelper;
 class FTooltipDrawState;
 class ITimingEvent;
+class ITimingEventRelation;
 class ITimingEventFilter;
 
 struct FDrawContext;
@@ -93,6 +94,7 @@ public:
 	virtual const TSharedPtr<const ITimingEvent> GetHoveredEvent() const = 0;
 	virtual const TSharedPtr<const ITimingEvent> GetSelectedEvent() const = 0;
 	virtual const TSharedPtr<ITimingEventFilter> GetEventFilter() const = 0;
+	virtual const TArray<ITimingEventRelation*>& GetCurrentRelations() const = 0;
 	virtual double GetCurrentTime() const = 0;
 	virtual float GetDeltaTime() const = 0;
 };
@@ -252,6 +254,10 @@ public:
 	// Called back from the timing view when an event is copied to the clipboard with Ctrl+C.
 	virtual void OnClipboardCopyEvent(const ITimingEvent& InSelectedEvent) const {}
 
+	virtual void GetEventRelations(const ITimingEvent& InSelectedEvent, TArray<ITimingEventRelation*>& Dependencies) const {}
+
+	virtual void SolveEventRelations(const TArray<ITimingEventRelation*>& Dependencies) const {}
+
 	// Returns number of text lines needed to display the debug string.
 	//TODO: virtual int GetDebugStringLineCount() const { return 0; }
 
diff --git a/Engine/Source/Developer/TraceInsights/Public/Insights/ViewModels/ITimingEvent.h b/Engine/Source/Developer/TraceInsights/Public/Insights/ViewModels/ITimingEvent.h
--- a/Engine/Source/Developer/TraceInsights/Public/Insights/ViewModels/ITimingEvent.h
+++ b/Engine/Source/Developer/TraceInsights/Public/Insights/ViewModels/ITimingEvent.h
@@ -6,7 +6,10 @@
 
 #include "Insights/Common/SimpleRtti.h"
 
+class ITimingViewDrawHelper;
 class FBaseTimingTrack;
+struct FDrawContext;
+class FTimingTrackViewport;
 
 ////////////////////////////////////////////////////////////////////////////////////////////////////
 
@@ -63,3 +66,14 @@ public:
 };
 
 ////////////////////////////////////////////////////////////////////////////////////////////////////
+
+class ITimingEventRelation
+{
+public:
+	ITimingEventRelation() {}
+	virtual ~ITimingEventRelation() {}
+	virtual bool IsSolved() = 0;
+	virtual void Draw(const FDrawContext& DrawContext, const FTimingTrackViewport& Viewport, const ITimingViewDrawHelper& Helper) = 0;
+};
+
+////////////////////////////////////////////////////////////////////////////////////////////////////
diff --git a/Engine/Source/Developer/TraceInsights/Public/Insights/ViewModels/ITimingViewDrawHelper.h b/Engine/Source/Developer/TraceInsights/Public/Insights/ViewModels/ITimingViewDrawHelper.h
--- a/Engine/Source/Developer/TraceInsights/Public/Insights/ViewModels/ITimingViewDrawHelper.h
+++ b/Engine/Source/Developer/TraceInsights/Public/Insights/ViewModels/ITimingViewDrawHelper.h
@@ -18,6 +18,7 @@ public:
 	virtual FLinearColor GetTrackNameTextColor(const FBaseTimingTrack& Track) const = 0;
 	virtual int32 GetHeaderBackgroundLayerId() const = 0;
 	virtual int32 GetHeaderTextLayerId() const = 0;
+	virtual int32 GetRelationLayerId() const = 0;
 	virtual int32 GetFirstLayerId() const = 0;
 	virtual int32 GetNumLayerIds() const = 0;
 };
diff --git a/Engine/Source/Developer/TraceServices/Public/Model/MonotonicTimeline.h b/Engine/Source/Developer/TraceServices/Public/Model/MonotonicTimeline.h
--- a/Engine/Source/Developer/TraceServices/Public/Model/MonotonicTimeline.h
+++ b/Engine/Source/Developer/TraceServices/Public/Model/MonotonicTimeline.h
@@ -742,6 +742,49 @@ public:
 		return true;
 	}
 
+	virtual int32 GetDepthAt(double Time) const override
+	{
+		const FDetailLevel& DetailLevel = DetailLevels[0];
+		if (DetailLevel.ScopeEntries.Num() == 0)
+		{
+			return 0;
+		}
+
+		uint64 FirstScopePageIndex = Algo::UpperBoundBy(DetailLevel.ScopeEntries, Time, [](const FEventScopeEntryPage& Page)
+			{
+				return Page.BeginTime;
+			});
+		if (FirstScopePageIndex > 0)
+		{
+			--FirstScopePageIndex;
+		}
+		auto ScopeEntryIterator = DetailLevel.ScopeEntries.GetIteratorFromPage(FirstScopePageIndex);
+		const FEventScopeEntryPage* ScopePage = ScopeEntryIterator.GetCurrentPage();
+		if (ScopePage->BeginTime > Time)
+		{
+			return 0;
+		}
+
+		int32 CurrentStackDepth = ScopePage->InitialStackCount;
+
+		const FEventScopeEntry* ScopeEntry = ScopeEntryIterator.GetCurrentItem();
+		while (ScopeEntry && FMath::Abs(ScopeEntry->Time) < Time)
+		{
+			if (ScopeEntry->Time < 0.0)
+			{
+				CurrentStackDepth++;
+			}
+			else
+			{
+				check(CurrentStackDepth > 0);
+				--CurrentStackDepth;
+			}
+			ScopeEntry = ScopeEntryIterator.NextItem();
+		}
+
+		return CurrentStackDepth;
+	}
+
 private:
 
 	struct FIterationState
diff --git a/Engine/Source/Developer/TraceServices/Public/TraceServices/Containers/Timelines.h b/Engine/Source/Developer/TraceServices/Public/TraceServices/Containers/Timelines.h
--- a/Engine/Source/Developer/TraceServices/Public/TraceServices/Containers/Timelines.h
+++ b/Engine/Source/Developer/TraceServices/Public/TraceServices/Containers/Timelines.h
@@ -58,6 +58,7 @@ public:
 	virtual void EnumerateEventsDownSampledAsync(const EnumerateAsyncParams& EnumerateAsyncParams) const {};
 	virtual void EnumerateEvents(double IntervalStart, double IntervalEnd, EventCallback Callback) const = 0;
 	virtual void EnumerateEvents(double IntervalStart, double IntervalEnd, EventRangeCallback Callback) const = 0;
+	virtual int32 GetDepthAt(double Time) const { return 0; };
 	
 	/**
 	 * Finds event information for the event closest to InTime from the interval [InTime - DeltaTime, InTime + DeltaTime]
