[Insights] Added warning when system memory usage exceeds 80% of total physical memory, with a hysteresis limit to 50% (after a first warning is displayed, the warning will show again if mem usage drops first to below 50% and then increases to over 80%).

diff --git a/Engine/Source/Developer/TraceInsights/Private/Insights/InsightsManager.cpp b/Engine/Source/Developer/TraceInsights/Private/Insights/InsightsManager.cpp
--- a/Engine/Source/Developer/TraceInsights/Private/Insights/InsightsManager.cpp
+++ b/Engine/Source/Developer/TraceInsights/Private/Insights/InsightsManager.cpp
@@ -3,6 +3,8 @@
 #include "InsightsManager.h"
 
 #include "Framework/Application/SlateApplication.h"
+#include "HAL/PlatformMemory.h"
+#include "Logging/MessageLog.h"
 #include "MessageLog/Public/MessageLogModule.h"
 #include "Misc/CString.h"
 #include "Modules/ModuleManager.h"
@@ -17,6 +19,7 @@
 #include "Insights/InsightsStyle.h"
 #include "Insights/IUnrealInsightsModule.h"
 #include "Insights/LoadingProfiler/LoadingProfilerManager.h"
+#include "Insights/Log.h"
 #include "Insights/NetworkingProfiler/NetworkingProfilerManager.h"
 #include "Insights/Tests/InsightsTestRunner.h"
 #include "Insights/TimingProfilerManager.h"
@@ -113,6 +116,9 @@ TSharedPtr<FInsightsManager> FInsightsManager::CreateInstance(TSharedRef<TraceSe
 FInsightsManager::FInsightsManager(TSharedRef<TraceServices::IAnalysisService> InTraceAnalysisService,
 								   TSharedRef<TraceServices::IModuleService> InTraceModuleService)
 	: bIsInitialized(false)
+	, bMemUsageLimitHysteresis(false)
+	, MemUsageLimitLastTimestamp(0)
+	, LogListingName(TEXT("UnrealInsights"))
 	, AnalysisService(InTraceAnalysisService)
 	, ModuleService(InTraceModuleService)
 	, StoreDir()
@@ -144,6 +150,10 @@ void FInsightsManager::Initialize(IUnrealInsightsModule& InsightsModule)
 
 	Insights::FFilterService::CreateInstance();
 
+	FMessageLogModule& MessageLogModule = FModuleManager::LoadModuleChecked<FMessageLogModule>("MessageLog");
+	MessageLogModule.RegisterLogListing(GetLogListingName(), LOCTEXT("UnrealInsights", "Unreal Insights"));
+	MessageLogModule.EnableMessageLogDisplay(true);
+
 	// Register tick functions.
 	OnTick = FTickerDelegate::CreateSP(this, &FInsightsManager::Tick);
 	OnTickHandle = FTicker::GetCoreTicker().AddTicker(OnTick, 0.0f);
@@ -169,6 +179,16 @@ void FInsightsManager::Shutdown()
 	// Unregister tick function.
 	FTicker::GetCoreTicker().RemoveTicker(OnTickHandle);
 
+	// If the MessageLog module was already unloaded as part of the global Shutdown process, do not load it again.
+	if (FModuleManager::Get().IsModuleLoaded("MessageLog"))
+	{
+		FMessageLogModule& MessageLogModule = FModuleManager::LoadModuleChecked<FMessageLogModule>("MessageLog");
+		if (MessageLogModule.IsRegisteredLogListing(GetLogListingName()))
+		{
+			MessageLogModule.UnregisterLogListing(GetLogListingName());
+		}
+	}
+
 	FInsightsManager::Instance.Reset();
 }
 
@@ -355,6 +375,10 @@ bool FInsightsManager::Tick(float DeltaTime)
 {
 	UpdateSessionDuration();
 
+#if !WITH_EDITOR
+	CheckMemoryUsage();
+#endif
+
 	return true;
 }
 
@@ -391,6 +415,57 @@ void FInsightsManager::UpdateSessionDuration()
 
 ////////////////////////////////////////////////////////////////////////////////////////////////////
 
+void FInsightsManager::CheckMemoryUsage()
+{
+	if (Session.IsValid()) // only check if we are in "viewer mode"
+	{
+		constexpr double MemUsageLimitPercent = 80.0;
+		constexpr double MemUsageLimitHysteresisPercent = 50.0;
+
+		const uint64 Time = FPlatformTime::Cycles64();
+		const double DurationSeconds = (Time - MemUsageLimitLastTimestamp) * FPlatformTime::GetSecondsPerCycle64();
+		if (DurationSeconds > 1.0) // only check once per second
+		{
+			MemUsageLimitLastTimestamp = Time;
+
+			FPlatformMemoryStats Stats = FPlatformMemory::GetStats();
+
+			constexpr double GiB = 1024.0 * 1024.0 * 1024.0;
+			const double UsedGiB = (double)(Stats.TotalPhysical - Stats.AvailablePhysical) / GiB;
+			const double TotalGiB = (double)(Stats.TotalPhysical) / GiB;
+			const double UsedPercent = (UsedGiB * 100.0) / TotalGiB;
+
+			if (!bMemUsageLimitHysteresis)
+			{
+				if (UsedPercent >= MemUsageLimitPercent)
+				{
+					bMemUsageLimitHysteresis = true;
+
+					const FText MessageBoxTextFmt = LOCTEXT("MemUsageWarning_TextFmt", "High System Memory Usage Detected: {0} / {1} GiB ({2}%)!\nUnreal Insights might need more memory!");
+					const FText MessageBoxText = FText::Format(MessageBoxTextFmt,
+						FText::AsNumber((uint32)(UsedGiB + 0.5)),
+						FText::AsNumber((uint32)(TotalGiB + 0.5)),
+						FText::AsNumber((uint32)(UsedPercent + 0.5)));
+
+					FMessageLog ReportMessageLog(GetLogListingName());
+					TSharedRef<FTokenizedMessage> Message = FTokenizedMessage::Create(EMessageSeverity::Warning, MessageBoxText);
+					ReportMessageLog.AddMessage(Message);
+					ReportMessageLog.Notify();
+				}
+			}
+			else
+			{
+				if (UsedPercent <= MemUsageLimitHysteresisPercent)
+				{
+					bMemUsageLimitHysteresis = false;
+				}
+			}
+		}
+	}
+}
+
+////////////////////////////////////////////////////////////////////////////////////////////////////
+
 void FInsightsManager::ResetSession(bool bNotify)
 {
 	if (Session.IsValid())
diff --git a/Engine/Source/Developer/TraceInsights/Private/Insights/InsightsManager.h b/Engine/Source/Developer/TraceInsights/Private/Insights/InsightsManager.h
--- a/Engine/Source/Developer/TraceInsights/Private/Insights/InsightsManager.h
+++ b/Engine/Source/Developer/TraceInsights/Private/Insights/InsightsManager.h
@@ -198,6 +198,7 @@ public:
 	void OpenSettings();
 
 	void UpdateSessionDuration();
+	void CheckMemoryUsage();
 
 	bool IsAnalysisComplete() const { return bIsAnalysisComplete; }
 	double GetSessionDuration() const { return SessionDuration; }
@@ -206,6 +207,8 @@ public:
 
 	TSharedPtr<FInsightsMenuBuilder> GetInsightsMenuBuilder() { return InsightsMenuBuilder; }
 
+	const FName& GetLogListingName() const { return LogListingName; }
+
 	////////////////////////////////////////////////////////////////////////////////////////////////////
 	// SessionChangedEvent
 
@@ -261,6 +264,15 @@ private:
 private:
 	bool bIsInitialized;
 
+	/** If true, the "high system memory usage warning" will be disabled until the system memory usage first drops below a certain threshold. */
+	bool bMemUsageLimitHysteresis;
+
+	/** The timestamp when has occurred the last check for system memory usage. */
+	uint64 MemUsageLimitLastTimestamp;
+
+	/** The name of the Unreal Insights log listing. */
+	FName LogListingName;
+
 	/** The delegate to be invoked when this manager ticks. */
 	FTickerDelegate OnTick;
 
