[Insights] Fix crash when updating a query during an async operation in MemAllocTreeView

diff --git a/Engine/Source/Developer/TraceInsights/Private/Insights/MemoryProfiler/Widgets/SMemAllocTableTreeView.cpp b/Engine/Source/Developer/TraceInsights/Private/Insights/MemoryProfiler/Widgets/SMemAllocTableTreeView.cpp
--- a/Engine/Source/Developer/TraceInsights/Private/Insights/MemoryProfiler/Widgets/SMemAllocTableTreeView.cpp
+++ b/Engine/Source/Developer/TraceInsights/Private/Insights/MemoryProfiler/Widgets/SMemAllocTableTreeView.cpp
@@ -85,6 +85,12 @@ void SMemAllocTableTreeView::Tick(const FGeometry& AllottedGeometry, const doubl
 {
 	STableTreeView::Tick(AllottedGeometry, InCurrentTime, InDeltaTime);
 
+	if (bHasPendingQueryReset && !bIsUpdateRunning)
+	{
+		ResetAndStartQuery();
+		bHasPendingQueryReset = false;
+	}
+
 	if (!bIsUpdateRunning)
 	{
 		RebuildTree(false);
@@ -184,6 +190,20 @@ void SMemAllocTableTreeView::OnQueryInvalidated()
 {
 	CancelQuery();
 
+	if (bIsUpdateRunning)
+	{
+		bHasPendingQueryReset = true;
+	}
+	else
+	{
+		ResetAndStartQuery();
+	}
+}
+
+////////////////////////////////////////////////////////////////////////////////////////////////////
+
+void SMemAllocTableTreeView::ResetAndStartQuery()
+{
 	TableTreeNodes.Reset();
 
 	TSharedPtr<Insights::FMemAllocTable> MemAllocTable = GetMemAllocTable();
diff --git a/Engine/Source/Developer/TraceInsights/Private/Insights/MemoryProfiler/Widgets/SMemAllocTableTreeView.h b/Engine/Source/Developer/TraceInsights/Private/Insights/MemoryProfiler/Widgets/SMemAllocTableTreeView.h
--- a/Engine/Source/Developer/TraceInsights/Private/Insights/MemoryProfiler/Widgets/SMemAllocTableTreeView.h
+++ b/Engine/Source/Developer/TraceInsights/Private/Insights/MemoryProfiler/Widgets/SMemAllocTableTreeView.h
@@ -98,6 +98,7 @@ private:
 	void StartQuery();
 	void UpdateQuery(TraceServices::IAllocationsProvider::EQueryStatus& OutStatus);
 	void CancelQuery();
+	void ResetAndStartQuery();
 
 	FReply OnDetailedViewClicked();
 	FReply OnSizeViewClicked();
@@ -125,6 +126,7 @@ private:
 	FText QueryInfo;
 	FText QueryInfoTooltip;
 	FStopwatch QueryStopwatch;
+	bool bHasPendingQueryReset = false;
 	bool bIsCallstackGroupingByFunction = true;
 };
 
diff --git a/Engine/Source/Developer/TraceServices/Private/Model/AllocationsQuery.cpp b/Engine/Source/Developer/TraceServices/Private/Model/AllocationsQuery.cpp
--- a/Engine/Source/Developer/TraceServices/Private/Model/AllocationsQuery.cpp
+++ b/Engine/Source/Developer/TraceServices/Private/Model/AllocationsQuery.cpp
@@ -25,7 +25,7 @@ public:
 
 	FORCEINLINE TStatId GetStatId() const { RETURN_QUICK_DECLARE_CYCLE_STAT(FAllocationsQueryAsyncTask, STATGROUP_TaskGraphTasks); }
 	ENamedThreads::Type GetDesiredThread() { return ENamedThreads::Type::AnyThread; }
-	static ESubsequentsMode::Type GetSubsequentsMode() { return ESubsequentsMode::FireAndForget; }
+	static ESubsequentsMode::Type GetSubsequentsMode() { return ESubsequentsMode::TrackSubsequents; }
 
 	void DoTask(ENamedThreads::Type CurrentThread, const FGraphEventRef& MyCompletionGraphEvent)
 	{
