[Insights] Support hex values in integer fields for FilterConfiguratorNode

diff --git a/Engine/Source/Developer/TraceInsights/Private/Insights/ViewModels/FilterConfiguratorNode.cpp b/Engine/Source/Developer/TraceInsights/Private/Insights/ViewModels/FilterConfiguratorNode.cpp
--- a/Engine/Source/Developer/TraceInsights/Private/Insights/ViewModels/FilterConfiguratorNode.cpp
+++ b/Engine/Source/Developer/TraceInsights/Private/Insights/ViewModels/FilterConfiguratorNode.cpp
@@ -1,9 +1,11 @@
 // Copyright Epic Games, Inc. All Rights Reserved.
 
 #include "FilterConfiguratorNode.h"
-#include "Insights/ViewModels/Filters.h"
+
+#include "Misc/Parse.h"
 
 // Insights
+#include "Insights/ViewModels/Filters.h"
 
 #define LOCTEXT_NAMESPACE "FilterConfiguratorNode"
 
@@ -161,7 +163,14 @@ void FFilterConfiguratorNode::ProcessFilter()
 		}
 		case EFilterDataType::Int64:
 		{
-			FilterValue.Set<int64>(FCString::Atoi64(*TextBoxValue));
+			if (TextBoxValue.Contains(TEXT("x")))
+			{
+				FilterValue.Set<int64>((int64)FParse::HexNumber64(*TextBoxValue));
+			}
+			else
+			{
+				FilterValue.Set<int64>(FCString::Atoi64(*TextBoxValue));
+			}
 			break;
 		}
 		case EFilterDataType::String:
