[Insights] Tree views: Fixed the issue where a second vertical scrollbar was visible.

diff --git a/Engine/Source/Developer/TraceInsights/Private/Insights/Table/Widgets/STableTreeView.cpp b/Engine/Source/Developer/TraceInsights/Private/Insights/Table/Widgets/STableTreeView.cpp
--- a/Engine/Source/Developer/TraceInsights/Private/Insights/Table/Widgets/STableTreeView.cpp
+++ b/Engine/Source/Developer/TraceInsights/Private/Insights/Table/Widgets/STableTreeView.cpp
@@ -234,42 +234,36 @@ void STableTreeView::ConstructWidget(TSharedPtr<FTable> InTablePtr)
 		.FillWidth(1.0f)
 		.Padding(0.0f)
 		[
-			//SNew(SScrollBox)
-			//.Orientation(Orient_Horizontal)
+			SNew(SOverlay)
 
-			//+ SScrollBox::Slot()
-			//[
-				SNew(SOverlay)
-
-				+ SOverlay::Slot()
-				.HAlign(HAlign_Fill)
-				.VAlign(VAlign_Fill)
-				[
-					SAssignNew(TreeView, STreeView<FTableTreeNodePtr>)
-					//.ExternalScrollbar(ExternalScrollbar)
-					.SelectionMode(ESelectionMode::Multi)
-					.TreeItemsSource(&FilteredGroupNodes)
-					.OnGetChildren(this, &STableTreeView::TreeView_OnGetChildren)
-					.OnGenerateRow(this, &STableTreeView::TreeView_OnGenerateRow)
-					.OnSelectionChanged(this, &STableTreeView::TreeView_OnSelectionChanged)
-					.OnMouseButtonDoubleClick(this, &STableTreeView::TreeView_OnMouseButtonDoubleClick)
-					.OnContextMenuOpening(FOnContextMenuOpening::CreateSP(this, &STableTreeView::TreeView_GetMenuContent))
-					.ItemHeight(12.0f)
-					.HeaderRow
-					(
-						SAssignNew(TreeViewHeaderRow, SHeaderRow)
-						.Visibility(EVisibility::Visible)
-					)
-				]
+			+ SOverlay::Slot()
+			.HAlign(HAlign_Fill)
+			.VAlign(VAlign_Fill)
+			[
+				SAssignNew(TreeView, STreeView<FTableTreeNodePtr>)
+				.ExternalScrollbar(ExternalScrollbar)
+				.SelectionMode(ESelectionMode::Multi)
+				.TreeItemsSource(&FilteredGroupNodes)
+				.OnGetChildren(this, &STableTreeView::TreeView_OnGetChildren)
+				.OnGenerateRow(this, &STableTreeView::TreeView_OnGenerateRow)
+				.OnSelectionChanged(this, &STableTreeView::TreeView_OnSelectionChanged)
+				.OnMouseButtonDoubleClick(this, &STableTreeView::TreeView_OnMouseButtonDoubleClick)
+				.OnContextMenuOpening(FOnContextMenuOpening::CreateSP(this, &STableTreeView::TreeView_GetMenuContent))
+				.ItemHeight(12.0f)
+				.HeaderRow
+				(
+					SAssignNew(TreeViewHeaderRow, SHeaderRow)
+					.Visibility(EVisibility::Visible)
+				)
+			]
 
-				+ SOverlay::Slot()
-				.HAlign(HAlign_Right)
-				.VAlign(VAlign_Bottom)
-				.Padding(16.0f)
-				[
-					SAssignNew(AsyncOperationStatus, Insights::SAsyncOperationStatus, SharedThis(this))
-				]
-			//]
+			+ SOverlay::Slot()
+			.HAlign(HAlign_Right)
+			.VAlign(VAlign_Bottom)
+			.Padding(16.0f)
+			[
+				SAssignNew(AsyncOperationStatus, Insights::SAsyncOperationStatus, SharedThis(this))
+			]
 		]
 
 		+ SHorizontalBox::Slot()
