[Insights]
  - Optimized string handling in time formatting functions.
  - Session Info: The "Session Duration" and "Analyzed" time are now displayed with "seconds" precision during analysis and with 10ms precision when analysis is completed.

diff --git a/Engine/Source/Developer/TraceInsights/Private/Insights/Common/TimeUtils.cpp b/Engine/Source/Developer/TraceInsights/Private/Insights/Common/TimeUtils.cpp
--- a/Engine/Source/Developer/TraceInsights/Private/Insights/Common/TimeUtils.cpp
+++ b/Engine/Source/Developer/TraceInsights/Private/Insights/Common/TimeUtils.cpp
@@ -1,6 +1,7 @@
 // Copyright Epic Games, Inc. All Rights Reserved.
 
 #include "Insights/Common/TimeUtils.h"
+#include "Misc/StringBuilder.h"
 
 #include <cmath>
 #include <limits>
@@ -19,6 +20,11 @@ FString FormatTimeValue(const double Duration, const int32 NumDigits)
 		return TEXT("NaN");
 	}
 
+	if (NumDigits <= 0)
+	{
+		return FString::Printf(TEXT("%.0f"), Duration);
+	}
+
 #if !PLATFORM_USE_GENERIC_STRING_IMPLEMENTATION
 	FString Str = FString::Printf(TEXT("%.*f"), NumDigits, Duration);
 #else
@@ -28,14 +34,13 @@ FString FormatTimeValue(const double Duration, const int32 NumDigits)
 	FString Str = FString::Printf(FormatString, Duration);
 #endif
 
-	if (NumDigits == 1)
+	int32 Index = Str.Len() - 1;
+	while (Index > 0 && Str[Index] == '0') --Index;
+	while (Index > 0 && Str[Index] == '.') --Index;
+	++Index;
+	if (Index < Str.Len())
 	{
-		Str.RemoveFromEnd(TEXT(".0"));
-	}
-	else
-	{
-		while (Str.RemoveFromEnd(TEXT("0"))) { /* keep removing the ending 0 */ }
-		Str.RemoveFromEnd(TEXT("."));
+		Str.LeftInline(Index, false);
 	}
 
 	return Str;
@@ -59,29 +64,29 @@ FString FormatTimeAuto(const double InDuration, const int32 NumDigits)
 		return TEXT("0");
 	}
 
-	FString Str;
+	TStringBuilder<64> StrBuilder;
 
 	if (FMath::IsNegativeDouble(Duration))
 	{
-		Str = TEXT("-");
 		Duration = -Duration;
+		StrBuilder.Append(TEXT('-'));
 	}
 
 	if (Duration == DBL_MAX || Duration == std::numeric_limits<double>::infinity())
 	{
-		Str += TEXT("∞");
+		StrBuilder.Append(TEXT('∞'));
 	}
 	else if (Duration < TimeUtils::Picosecond)
 	{
 		// (0 .. 1ps)
-		Str = TEXT("~0");
+		return TEXT("~0");
 	}
 	else if (Duration < TimeUtils::Nanosecond)
 	{
 		// [1ps .. 1ns)
 		if (Duration >= 999.5 * TimeUtils::Picosecond)
 		{
-			Str += TEXT("1 ns");
+			StrBuilder.Append(TEXT("1 ns"));
 		}
 		else
 		{
@@ -90,11 +95,11 @@ FString FormatTimeAuto(const double InDuration, const int32 NumDigits)
 			ensure(IntPicoseconds <= 999);
 			//if (IntPicoseconds > 999)
 			//{
-			//	Str += TEXT("1 ns");
+			//	StrBuilder.Append(TEXT("1 ns"));
 			//}
 			//else
 			{
-				Str += FString::Printf(TEXT("%d ps"), IntPicoseconds);
+				StrBuilder.Appendf(TEXT("%d ps"), IntPicoseconds);
 			}
 		}
 	}
@@ -103,7 +108,7 @@ FString FormatTimeAuto(const double InDuration, const int32 NumDigits)
 		// [1ns .. 1µs)
 		if (Duration >= 999.5 * TimeUtils::Nanosecond)
 		{
-			Str += TEXT("1 µs");
+			StrBuilder.Append(TEXT("1 µs"));
 		}
 		else
 		{
@@ -112,11 +117,11 @@ FString FormatTimeAuto(const double InDuration, const int32 NumDigits)
 			ensure(IntNanoseconds <= 999);
 			//if (IntNanoseconds > 999)
 			//{
-			//	Str += TEXT("1 µs");
+			//	StrBuilder.Append(TEXT("1 µs"));
 			//}
 			//else
 			{
-				Str += FString::Printf(TEXT("%d ns"), IntNanoseconds);
+				StrBuilder.Appendf(TEXT("%d ns"), IntNanoseconds);
 			}
 		}
 	}
@@ -126,12 +131,12 @@ FString FormatTimeAuto(const double InDuration, const int32 NumDigits)
 		const double Microseconds = Duration * 1000000.0;
 		if (Microseconds >= 999.95)
 		{
-			Str += TEXT("1 ms");
+			StrBuilder.Append(TEXT("1 ms"));
 		}
 		else
 		{
-			Str += FormatTimeValue(Microseconds, NumDigits);
-			Str += TEXT(" µs");
+			StrBuilder.Append(FormatTimeValue(Microseconds, NumDigits));
+			StrBuilder.Append(TEXT(" µs"));
 		}
 	}
 	else if (Duration < TimeUtils::Second)
@@ -140,12 +145,12 @@ FString FormatTimeAuto(const double InDuration, const int32 NumDigits)
 		const double Miliseconds = Duration * 1000.0;
 		if (Miliseconds >= 999.95)
 		{
-			Str += TEXT("1s");
+			StrBuilder.Append(TEXT("1s"));
 		}
 		else
 		{
-			Str += FormatTimeValue(Miliseconds, NumDigits);
-			Str += TEXT(" ms");
+			StrBuilder.Append(FormatTimeValue(Miliseconds, NumDigits));
+			StrBuilder.Append(TEXT(" ms"));
 		}
 	}
 	else if (Duration < TimeUtils::Minute)
@@ -153,61 +158,61 @@ FString FormatTimeAuto(const double InDuration, const int32 NumDigits)
 		// [1s .. 1m)
 		if (Duration >= 59.95)
 		{
-			Str += TEXT("1m");
+			StrBuilder.Append(TEXT("1m"));
 		}
 		else
 		{
-			Str += FormatTimeValue(Duration, NumDigits);
-			Str += TEXT("s");
+			StrBuilder.Append(FormatTimeValue(Duration, NumDigits));
+			StrBuilder.Append(TEXT('s'));
 		}
 	}
 	else if (Duration < TimeUtils::Hour)
 	{
 		// [1m .. 1h)
 		const double Minutes = FMath::FloorToDouble(Duration / TimeUtils::Minute);
-		Str += FString::Printf(TEXT("%dm"), static_cast<int32>(Minutes));
+		StrBuilder.Appendf(TEXT("%dm"), static_cast<int32>(Minutes));
 		Duration -= Minutes * TimeUtils::Minute;
 		if (NumDigits <= 1)
 		{
 			const double Seconds = FMath::FloorToDouble(Duration / TimeUtils::Second);
 			if (Seconds > 0.5)
 			{
-				Str += FString::Printf(TEXT(" %ds"), static_cast<int32>(Seconds));
+				StrBuilder.Appendf(TEXT(" %ds"), static_cast<int32>(Seconds));
 			}
 		}
 		else
 		{
-			Str += TEXT(" ");
-			Str += FormatTimeValue(Duration, NumDigits - 1);
-			Str += TEXT("s");
+			StrBuilder.Append(TEXT(' '));
+			StrBuilder.Append(FormatTimeValue(Duration, NumDigits - 1));
+			StrBuilder.Append(TEXT('s'));
 		}
 	}
 	else if (Duration < TimeUtils::Day)
 	{
 		// [1h .. 1d)
 		const double Hours = FMath::FloorToDouble(Duration / TimeUtils::Hour);
-		Str += FString::Printf(TEXT("%dh"), static_cast<int32>(Hours));
+		StrBuilder.Appendf(TEXT("%dh"), static_cast<int32>(Hours));
 		Duration -= Hours * TimeUtils::Hour;
 		const double Minutes = FMath::FloorToDouble(Duration / TimeUtils::Minute);
 		if (Minutes > 0.5)
 		{
-			Str += FString::Printf(TEXT(" %dm"), static_cast<int32>(Minutes));
+			StrBuilder.Appendf(TEXT(" %dm"), static_cast<int32>(Minutes));
 		}
 	}
 	else
 	{
 		// [1d .. ∞)
 		const double Days = FMath::FloorToDouble(Duration / TimeUtils::Day);
-		Str += FString::Printf(TEXT("%dd"), static_cast<int32>(Days));
+		StrBuilder.Appendf(TEXT("%dd"), static_cast<int32>(Days));
 		Duration -= Days * TimeUtils::Day;
 		const double Hours = FMath::FloorToDouble(Duration / TimeUtils::Hour);
 		if (Hours > 0.5)
 		{
-			Str += FString::Printf(TEXT(" %dh"), static_cast<int32>(Hours));
+			StrBuilder.Appendf(TEXT(" %dh"), static_cast<int32>(Hours));
 		}
 	}
 
-	return Str;
+	return StrBuilder.ToString();
 }
 
 ////////////////////////////////////////////////////////////////////////////////////////////////////
@@ -225,36 +230,43 @@ FString FormatTimeMs(const double InDuration, const int32 NumDigits, bool bAddTi
 	}
 
 	double Duration = InDuration;
-	FString Str;
+	TStringBuilder<64> StrBuilder;
 
 	if (FMath::IsNegativeDouble(Duration))
 	{
 		Duration = -Duration;
-		Str = TEXT("-");
+		StrBuilder.Append(TEXT('-'));
 	}
 
 	if (Duration == DBL_MAX || Duration == std::numeric_limits<double>::infinity())
 	{
-		Str += TEXT("∞");
+		StrBuilder.Append(TEXT('∞'));
 	}
 	else
 	{
+		if (NumDigits <= 0)
+		{
+			StrBuilder.Appendf(TEXT("%.0f"), Duration * 1000.0);
+		}
+		else
+		{
 #if !PLATFORM_USE_GENERIC_STRING_IMPLEMENTATION
-		Str += FString::Printf(TEXT("%.*f"), NumDigits, Duration * 1000.0);
+			StrBuilder.Appendf(TEXT("%.*f"), NumDigits, Duration * 1000.0);
 #else
-		// proper resolution is tracked as UE-79534
-		TCHAR FormatString[32];
-		FCString::Snprintf(FormatString, sizeof(FormatString), TEXT("%%.%df"), NumDigits);
-		Str += FString::Printf(FormatString, Duration * 1000.0);
+			// proper resolution is tracked as UE-79534
+			TCHAR FormatString[32];
+			FCString::Snprintf(FormatString, sizeof(FormatString), TEXT("%%.%df"), NumDigits);
+			StrBuilder.Appendf(FormatString, Duration * 1000.0);
 #endif
+		}
 
 		if (bAddTimeUnit)
 		{
-			Str += TEXT(" ms");
+			StrBuilder.Append(TEXT(" ms"));
 		}
 	}
 
-	return Str;
+	return StrBuilder.ToString();
 }
 
 ////////////////////////////////////////////////////////////////////////////////////////////////////
@@ -272,18 +284,18 @@ FString FormatTime(const double InTime, const double Precision)
 	}
 
 	double Time = InTime;
-	FString Str;
+	TStringBuilder<64> StrBuilder;
 
 	if (FMath::IsNegativeDouble(Time))
 	{
 		Time = -Time;
-		Str = TEXT("-");
+		StrBuilder.Append(TEXT('-'));
 	}
 
 	if (Time == DBL_MAX || Time == std::numeric_limits<double>::infinity())
 	{
-		Str += TEXT("∞");
-		return Str;
+		StrBuilder.Append(TEXT('∞'));
+		return StrBuilder.ToString();
 	}
 
 	bool bIsSpaceNeeded = false;
@@ -291,7 +303,7 @@ FString FormatTime(const double InTime, const double Precision)
 	int32 Days = static_cast<int32>(Time / TimeUtils::Day);
 	if (Days > 0)
 	{
-		Str += FString::Printf(TEXT("%dd"), Days);
+		StrBuilder.Appendf(TEXT("%dd"), Days);
 		bIsSpaceNeeded = true;
 		Time -= static_cast<double>(Days) * TimeUtils::Day;
 	}
@@ -300,9 +312,9 @@ FString FormatTime(const double InTime, const double Precision)
 	{
 		if (!bIsSpaceNeeded)
 		{
-			Str = TEXT("~0");
+			return TEXT("~0");
 		}
-		return Str;
+		return StrBuilder.ToString();
 	}
 
 	int32 Hours = static_cast<int32>(Time / TimeUtils::Hour);
@@ -310,9 +322,9 @@ FString FormatTime(const double InTime, const double Precision)
 	{
 		if (bIsSpaceNeeded)
 		{
-			Str += TEXT(" ");
+			StrBuilder.Append(TEXT(' '));
 		}
-		Str += FString::Printf(TEXT("%dh"), Hours);
+		StrBuilder.Appendf(TEXT("%dh"), Hours);
 		bIsSpaceNeeded = true;
 		Time -= static_cast<double>(Hours) * TimeUtils::Hour;
 	}
@@ -321,9 +333,9 @@ FString FormatTime(const double InTime, const double Precision)
 	{
 		if (!bIsSpaceNeeded)
 		{
-			Str = TEXT("~0");
+			return TEXT("~0");
 		}
-		return Str;
+		return StrBuilder.ToString();
 	}
 
 	int32 Minutes = static_cast<int32>(Time / TimeUtils::Minute);
@@ -331,9 +343,9 @@ FString FormatTime(const double InTime, const double Precision)
 	{
 		if (bIsSpaceNeeded)
 		{
-			Str += TEXT(" ");
+			StrBuilder.Append(TEXT(' '));
 		}
-		Str += FString::Printf(TEXT("%dm"), Minutes);
+		StrBuilder.Appendf(TEXT("%dm"), Minutes);
 		bIsSpaceNeeded = true;
 		Time -= static_cast<double>(Minutes) * TimeUtils::Minute;
 	}
@@ -342,9 +354,9 @@ FString FormatTime(const double InTime, const double Precision)
 	{
 		if (!bIsSpaceNeeded)
 		{
-			Str = TEXT("~0");
+			return TEXT("~0");
 		}
-		return Str;
+		return StrBuilder.ToString();
 	}
 
 	//TestOptimizationIssue();
@@ -383,43 +395,43 @@ FString FormatTime(const double InTime, const double Precision)
 		{
 			if (bIsSpaceNeeded)
 			{
-				Str += TEXT(" ");
+				StrBuilder.Append(TEXT(' '));
 			}
-			Str += FString::Printf(TEXT("%ds"), Seconds);
+			StrBuilder.Appendf(TEXT("%ds"), Seconds);
 		}
 		else if (!bIsSpaceNeeded)
 		{
-			Str += TEXT("~0");
+			return TEXT("~0");
 		}
 	}
 	//else if (Digits <= 9)
 	//{
 	//	if (bIsSpaceNeeded)
 	//	{
-	//		Str += TEXT(" ");
+	//		StrBuilder.Append(TEXT(' '));
 	//	}
 	//	int32 Seconds = static_cast<int32>(Time / TimeUtils::Second);
 	//	Time -= static_cast<double>(Seconds) * TimeUtils::Second;
 	//	int64 SubSeconds = static_cast<int64>(Time * FMath::Pow(10.0f, Digits) + 0.5);
-	//	Str += FString::Printf(TEXT("%d.%0*llds"), Seconds, Digits, SubSeconds);
+	//	StrBuilder.Appendf(TEXT("%d.%0*llds"), Seconds, Digits, SubSeconds);
 	//}
 	else
 	{
 		if (bIsSpaceNeeded)
 		{
-			Str += TEXT(" ");
+			StrBuilder.Append(TEXT(' '));
 		}
 #if !PLATFORM_USE_GENERIC_STRING_IMPLEMENTATION
-		Str += FString::Printf(TEXT("%.*fs"), Digits, Time);
+		StrBuilder.Appendf(TEXT("%.*fs"), Digits, Time);
 #else
 		// proper resolution is tracked as UE-79534
 		TCHAR FormatString[32];
 		FCString::Snprintf(FormatString, sizeof(FormatString), TEXT("%%.%dfs"), Digits);
-		Str += FString::Printf(FormatString, Time);
+		StrBuilder.Appendf(FormatString, Time);
 #endif
 	}
 
-	return Str;
+	return StrBuilder.ToString();
 }
 
 ////////////////////////////////////////////////////////////////////////////////////////////////////
@@ -545,36 +557,33 @@ void SplitTime(const double InTime, FTimeSplit& OutTimeSplit)
 
 FString FormatTimeSplit(const FTimeSplit& InTimeSplit, const double Precision)
 {
-	FString Str;
-
 	if (InTimeSplit.bIsNaN)
 	{
-		Str = TEXT("NaN");
-		return Str;
+		return TEXT("NaN");
 	}
 
 	if (InTimeSplit.bIsZero)
 	{
-		Str = TEXT("0");
-		return Str;
+		return TEXT("0");
 	}
 
-	if (InTimeSplit.bIsNegative)
+	if (InTimeSplit.bIsInfinite)
 	{
-		Str = TEXT("-");
+		return InTimeSplit.bIsNegative ? TEXT("-∞") : TEXT("∞");
 	}
 
-	if (InTimeSplit.bIsInfinite)
+	TStringBuilder<64> StrBuilder;
+
+	if (InTimeSplit.bIsNegative)
 	{
-		Str += TEXT("∞");
-		return Str;
+		StrBuilder.Append(TEXT('-'));
 	}
 
 	bool bIsSpaceNeeded = false;
 
 	if (InTimeSplit.Days > 0)
 	{
-		Str += FString::Printf(TEXT("%dd"), InTimeSplit.Days);
+		StrBuilder.Appendf(TEXT("%dd"), InTimeSplit.Days);
 		bIsSpaceNeeded = true;
 	}
 
@@ -582,18 +591,18 @@ FString FormatTimeSplit(const FTimeSplit& InTimeSplit, const double Precision)
 	{
 		if (!bIsSpaceNeeded)
 		{
-			Str = TEXT("~0");
+			return TEXT("~0");
 		}
-		return Str;
+		return StrBuilder.ToString();
 	}
 
 	if (InTimeSplit.Hours > 0)
 	{
 		if (bIsSpaceNeeded)
 		{
-			Str += TEXT(" ");
+			StrBuilder.Append(TEXT(' '));
 		}
-		Str += FString::Printf(TEXT("%dh"), InTimeSplit.Hours);
+		StrBuilder.Appendf(TEXT("%dh"), InTimeSplit.Hours);
 		bIsSpaceNeeded = true;
 	}
 
@@ -601,18 +610,18 @@ FString FormatTimeSplit(const FTimeSplit& InTimeSplit, const double Precision)
 	{
 		if (!bIsSpaceNeeded)
 		{
-			Str = TEXT("~0");
+			return TEXT("~0");
 		}
-		return Str;
+		return StrBuilder.ToString();
 	}
 
 	if (InTimeSplit.Minutes > 0)
 	{
 		if (bIsSpaceNeeded)
 		{
-			Str += TEXT(" ");
+			StrBuilder.Append(TEXT(' '));
 		}
-		Str += FString::Printf(TEXT("%dm"), InTimeSplit.Minutes);
+		StrBuilder.Appendf(TEXT("%dm"), InTimeSplit.Minutes);
 		bIsSpaceNeeded = true;
 	}
 
@@ -620,18 +629,18 @@ FString FormatTimeSplit(const FTimeSplit& InTimeSplit, const double Precision)
 	{
 		if (!bIsSpaceNeeded)
 		{
-			Str = TEXT("~0");
+			return TEXT("~0");
 		}
-		return Str;
+		return StrBuilder.ToString();
 	}
 
 	if (InTimeSplit.Seconds > 0)
 	{
 		if (bIsSpaceNeeded)
 		{
-			Str += TEXT(" ");
+			StrBuilder.Append(TEXT(' '));
 		}
-		Str += FString::Printf(TEXT("%ds"), InTimeSplit.Seconds);
+		StrBuilder.Appendf(TEXT("%ds"), InTimeSplit.Seconds);
 		bIsSpaceNeeded = true;
 	}
 
@@ -639,18 +648,18 @@ FString FormatTimeSplit(const FTimeSplit& InTimeSplit, const double Precision)
 	{
 		if (!bIsSpaceNeeded)
 		{
-			Str = TEXT("~0");
+			return TEXT("~0");
 		}
-		return Str;
+		return StrBuilder.ToString();
 	}
 
 	if (InTimeSplit.Miliseconds > 0)
 	{
 		if (bIsSpaceNeeded)
 		{
-			Str += TEXT(" ");
+			StrBuilder.Append(TEXT(' '));
 		}
-		Str += FString::Printf(TEXT("%dms"), InTimeSplit.Miliseconds);
+		StrBuilder.Appendf(TEXT("%dms"), InTimeSplit.Miliseconds);
 		bIsSpaceNeeded = true;
 	}
 
@@ -658,18 +667,18 @@ FString FormatTimeSplit(const FTimeSplit& InTimeSplit, const double Precision)
 	{
 		if (!bIsSpaceNeeded)
 		{
-			Str = TEXT("~0");
+			return TEXT("~0");
 		}
-		return Str;
+		return StrBuilder.ToString();
 	}
 
 	if (InTimeSplit.Microseconds > 0)
 	{
 		if (bIsSpaceNeeded)
 		{
-			Str += TEXT(" ");
+			StrBuilder.Append(TEXT(' '));
 		}
-		Str += FString::Printf(TEXT("%dµs"), InTimeSplit.Microseconds);
+		StrBuilder.Appendf(TEXT("%dµs"), InTimeSplit.Microseconds);
 		bIsSpaceNeeded = true;
 	}
 
@@ -677,18 +686,18 @@ FString FormatTimeSplit(const FTimeSplit& InTimeSplit, const double Precision)
 	{
 		if (!bIsSpaceNeeded)
 		{
-			Str = TEXT("~0");
+			return TEXT("~0");
 		}
-		return Str;
+		return StrBuilder.ToString();
 	}
 
 	if (InTimeSplit.Nanoseconds > 0)
 	{
 		if (bIsSpaceNeeded)
 		{
-			Str += TEXT(" ");
+			StrBuilder.Append(TEXT(' '));
 		}
-		Str += FString::Printf(TEXT("%dns"), InTimeSplit.Nanoseconds);
+		StrBuilder.Appendf(TEXT("%dns"), InTimeSplit.Nanoseconds);
 		bIsSpaceNeeded = true;
 	}
 
@@ -696,27 +705,27 @@ FString FormatTimeSplit(const FTimeSplit& InTimeSplit, const double Precision)
 	{
 		if (!bIsSpaceNeeded)
 		{
-			Str = TEXT("~0");
+			return TEXT("~0");
 		}
-		return Str;
+		return StrBuilder.ToString();
 	}
 
 	if (InTimeSplit.Picoseconds > 0)
 	{
 		if (bIsSpaceNeeded)
 		{
-			Str += TEXT(" ");
+			StrBuilder.Append(TEXT(' '));
 		}
-		Str += FString::Printf(TEXT("%dps"), InTimeSplit.Picoseconds);
+		StrBuilder.Appendf(TEXT("%dps"), InTimeSplit.Picoseconds);
 		bIsSpaceNeeded = true;
 	}
 
 	if (!bIsSpaceNeeded)
 	{
-		Str = TEXT("~0");
+		return TEXT("~0");
 	}
 
-	return Str;
+	return StrBuilder.ToString();
 }
 
 ////////////////////////////////////////////////////////////////////////////////////////////////////
diff --git a/Engine/Source/Developer/TraceInsights/Private/Insights/Widgets/SSessionInfoWindow.cpp b/Engine/Source/Developer/TraceInsights/Private/Insights/Widgets/SSessionInfoWindow.cpp
--- a/Engine/Source/Developer/TraceInsights/Private/Insights/Widgets/SSessionInfoWindow.cpp
+++ b/Engine/Source/Developer/TraceInsights/Private/Insights/Widgets/SSessionInfoWindow.cpp
@@ -437,10 +437,11 @@ FText SSessionInfoWindow::GetStatusText() const
 	FNumberFormattingOptions FormattingOptions;
 	FormattingOptions.MaximumFractionalDigits = 1;
 
+	const int32 NumDigits = InsightsManager->IsAnalysisComplete() ? 2 : 0;
 	FText Status = FText::Format(LOCTEXT("StatusFmt", "{0}\nSession Duration: {1}\nAnalyzed in {2} at {3}X speed."),
 		InsightsManager->IsAnalysisComplete() ? FText::FromString(FString(TEXT("ANALYSIS COMPLETED."))) : FText::FromString(FString(TEXT("ANALYZING..."))),
-		FText::FromString(TimeUtils::FormatTimeAuto(InsightsManager->GetSessionDuration(), 2)),
-		FText::FromString(TimeUtils::FormatTimeAuto(InsightsManager->GetAnalysisDuration(), 2)),
+		FText::FromString(TimeUtils::FormatTimeAuto(InsightsManager->GetSessionDuration(), NumDigits)),
+		FText::FromString(TimeUtils::FormatTimeAuto(InsightsManager->GetAnalysisDuration(), NumDigits)),
 		FText::AsNumber(InsightsManager->GetAnalysisSpeedFactor(), &FormattingOptions));
 
 	return Status;
