TasksInsights: visualising full task biography in addition to previously implemented visualisation of subsequent and nested tasks

diff --git a/Engine/Source/Developer/TraceInsights/Private/Insights/ViewModels/TaskGraphRelation.cpp b/Engine/Source/Developer/TraceInsights/Private/Insights/ViewModels/TaskGraphRelation.cpp
--- a/Engine/Source/Developer/TraceInsights/Private/Insights/ViewModels/TaskGraphRelation.cpp
+++ b/Engine/Source/Developer/TraceInsights/Private/Insights/ViewModels/TaskGraphRelation.cpp
@@ -32,8 +32,19 @@ void FTaskGraphRelation::Draw(const FDrawContext& DrawContext, const FTimingTrac
 	Points.Emplace(0, 0);
 	Points.Emplace(X2 - X1, Y2 - Y1);
 
-	FLinearColor Color = Type == ETaskGraphRelationType::Subsequent ? FLinearColor::Red : FLinearColor::Blue;
-	DrawContext.DrawLines(Helper.GetRelationLayerId(), X1, Y1, Points, ESlateDrawEffect::None, Color);
+	FLinearColor RelationColors[(int32)ETaskGraphRelationType::Count] = {
+		FLinearColor::Yellow, // Created
+		FLinearColor::Green, // Launched
+		FLinearColor::Red, // Prerequisite
+		FLinearColor::Blue, // Scheduled
+		FLinearColor::Blue, // AddedNested
+		FLinearColor::Red, // NestedCompleted
+		FLinearColor::Red, // Subsequent
+		FLinearColor::Yellow, // Completed
+	};
+
+	FLinearColor Color = RelationColors[(int32)Type];
+	DrawContext.DrawLines(Helper.GetRelationLayerId(), X1, Y1, Points, ESlateDrawEffect::None, Color, /*bAntialias=*/ true, /*Thickness=*/ 2.0);
 }
 
 } // namespace Insights
\ No newline at end of file
diff --git a/Engine/Source/Developer/TraceInsights/Private/Insights/ViewModels/TaskGraphRelation.h b/Engine/Source/Developer/TraceInsights/Private/Insights/ViewModels/TaskGraphRelation.h
--- a/Engine/Source/Developer/TraceInsights/Private/Insights/ViewModels/TaskGraphRelation.h
+++ b/Engine/Source/Developer/TraceInsights/Private/Insights/ViewModels/TaskGraphRelation.h
@@ -14,9 +14,17 @@ class FTaskGraphRelation : public ITimingEventRelation
 public:
 	enum class ETaskGraphRelationType : int32
 	{
-		Subsequent = 0,
-		Nested = 1,
+		Created,
+		Launched,
+		Prerequisite,
+		Scheduled,
+		AddedNested,
+		NestedCompleted,
+		Subsequent,
+		Completed,
+		Count
 	};
+
 	FTaskGraphRelation(double InSourceTime, int32 InSourceThreadId, double InTargetTime, int32 InTargetThreadId, ETaskGraphRelationType InType);
 
 	virtual bool IsSolved() override { return bIsSolved; }
diff --git a/Engine/Source/Developer/TraceInsights/Private/Insights/ViewModels/ThreadTimingTrack.cpp b/Engine/Source/Developer/TraceInsights/Private/Insights/ViewModels/ThreadTimingTrack.cpp
--- a/Engine/Source/Developer/TraceInsights/Private/Insights/ViewModels/ThreadTimingTrack.cpp
+++ b/Engine/Source/Developer/TraceInsights/Private/Insights/ViewModels/ThreadTimingTrack.cpp
@@ -1106,27 +1106,24 @@ void FThreadTimingTrack::InitTooltip(FTooltipDrawState& InOutTooltip, const ITim
 
 		do // info about blocking
 		{
-			if (TimerName.StartsWith(TEXT("WaitUntilTasksComplete")) || TimerName.StartsWith(TEXT("GameThreadWaitForTask")))
+			const TraceServices::FWaitingForTasks* Waiting = TasksProvider->TryGetWaiting(*TimerName, ThreadId, TooltipEvent.GetStartTime());
+			if (Waiting == nullptr)
 			{
-				const TraceServices::FWaitingForTasks* Waiting = TasksProvider->TryGetWaiting(ThreadId, TooltipEvent.GetStartTime());
-				if (Waiting == nullptr)
-				{
-					break;
-				}
+				break;
+			}
 
-				InOutTooltip.AddTextLine(TEXT("-------- Wating for tasks --------"), FLinearColor::Red);
-				FString TaskIdsStr = FString::JoinBy(Waiting->Tasks, TEXT(", "), [](TaskTrace::FId TaskId) { return FString::FromInt(TaskId); });
-				InOutTooltip.AddNameValueTextLine(TEXT("Tasks:"), FString::Printf(TEXT("[%s]"), *TaskIdsStr));
-				InOutTooltip.AddNameValueTextLine(TEXT("Started waiting:"), FString::Printf(TEXT("%f"), Waiting->StartedTimestamp));
-				InOutTooltip.AddNameValueTextLine(TEXT("Finished waiting:"), FString::Printf(TEXT("%f (+%s)"), Waiting->FinishedTimestamp, *TimeUtils::FormatTimeAuto(Waiting->FinishedTimestamp - Waiting->StartedTimestamp)));
+			InOutTooltip.AddTextLine(TEXT("-------- Wating for tasks --------"), FLinearColor::Red);
+			FString TaskIdsStr = FString::JoinBy(Waiting->Tasks, TEXT(", "), [](TaskTrace::FId TaskId) { return FString::FromInt(TaskId); });
+			InOutTooltip.AddNameValueTextLine(TEXT("Tasks:"), FString::Printf(TEXT("[%s]"), *TaskIdsStr));
+			InOutTooltip.AddNameValueTextLine(TEXT("Started waiting:"), FString::Printf(TEXT("%f"), Waiting->StartedTimestamp));
+			InOutTooltip.AddNameValueTextLine(TEXT("Finished waiting:"), FString::Printf(TEXT("%f (+%s)"), Waiting->FinishedTimestamp, *TimeUtils::FormatTimeAuto(Waiting->FinishedTimestamp - Waiting->StartedTimestamp)));
 
-				for (TaskTrace::FId TaskId : Waiting->Tasks)
+			for (TaskTrace::FId TaskId : Waiting->Tasks)
+			{
+				const TraceServices::FTaskInfo* Task = TasksProvider->TryGetTask(TaskId);
+				if (Task != nullptr)
 				{
-					const TraceServices::FTaskInfo* Task = TasksProvider->TryGetTask(TaskId);
-					if (Task != nullptr)
-					{
-						AddTaskInfo(*Task);
-					}
+					AddTaskInfo(*Task);
 				}
 			}
 		} while (false);
@@ -1497,32 +1494,98 @@ void FThreadTimingTrack::GetEventRelations(const ITimingEvent& InSelectedEvent,
 	const TraceServices::ITasksProvider* TasksProvider = TraceServices::ReadTasksProvider(*Session.Get());
 	if (TasksProvider)
 	{
+		auto AddRelation = [this, &InSelectedEvent, &Relations](double SourceTimestamp, uint32 SourceThreadId, double TargetTimestamp, uint32 TargetThreadId, FTaskGraphRelation::ETaskGraphRelationType Type)
+		{
+			if (SourceTimestamp == TraceServices::FTaskInfo::InvalidTimestamp || TargetTimestamp == TraceServices::FTaskInfo::InvalidTimestamp)
+			{
+				return;
+			}
+
+			FTaskGraphRelation* Relation = new FTaskGraphRelation(SourceTimestamp, SourceThreadId, TargetTimestamp, TargetThreadId, Type);
+
+			if (Relation->GetSourceThreadId() == ThreadId)
+			{
+				Relation->SetSourceTrack(SharedThis(this));
+				Relation->SetSourceDepth(InSelectedEvent.GetDepth());
+			}
+			if (Relation->GetTargetThreadId() == ThreadId)
+			{
+				Relation->SetTargetTrack(SharedThis(this));
+				Relation->SetTargetDepth(InSelectedEvent.GetDepth());
+			}
+			if (Relation->GetSourceThreadId() == ThreadId && Relation->GetTargetThreadId() == ThreadId)
+			{
+				Relation->SetIsSolved(true);
+			}
+
+			Relations.Add(Relation);
+		};
+
 		const TraceServices::FTaskInfo* Task = TasksProvider->TryGetTask(ThreadId, StartTime);
 		if (Task)
 		{
-			for (const TraceServices::FTaskInfo::FRelationInfo& RelationInfo : Task->Subsequents)
+			if (Task->CreatedTimestamp != Task->LaunchedTimestamp || Task->CreatedThreadId != Task->LaunchedThreadId)
 			{
-				const TraceServices::FTaskInfo* RelatedTask = TasksProvider->TryGetTask(RelationInfo.RelativeId);
+				AddRelation(Task->CreatedTimestamp, Task->CreatedThreadId, Task->LaunchedTimestamp, Task->LaunchedThreadId, FTaskGraphRelation::ETaskGraphRelationType::Created);
+			}
 
-				if (RelatedTask)
-				{
-					FTaskGraphRelation* NewRelation = new FTaskGraphRelation(StartTime, ThreadId, RelatedTask->StartedTimestamp, RelatedTask->StartedThreadId, FTaskGraphRelation::ETaskGraphRelationType::Subsequent);
-					NewRelation->SetSourceTrack(SharedThis(this));
-					NewRelation->SetSourceDepth(InSelectedEvent.GetDepth());
-					Relations.Add(NewRelation);
-				}
+			AddRelation(Task->LaunchedTimestamp, Task->LaunchedThreadId, Task->ScheduledTimestamp, Task->ScheduledThreadId, FTaskGraphRelation::ETaskGraphRelationType::Launched);
+
+			for (const TraceServices::FTaskInfo::FRelationInfo& RelationInfo : Task->Prerequisites)
+			{
+				const TraceServices::FTaskInfo* Prerequisite = TasksProvider->TryGetTask(RelationInfo.RelativeId);
+				check(Prerequisite != nullptr);
+				AddRelation(Prerequisite->CompletedTimestamp, Prerequisite->CompletedThreadId, Task->StartedTimestamp, Task->StartedThreadId, FTaskGraphRelation::ETaskGraphRelationType::Prerequisite);
+			}
+
+			if (Task->LaunchedTimestamp != Task->ScheduledTimestamp || Task->LaunchedThreadId != Task->ScheduledThreadId)
+			{
+				AddRelation(Task->ScheduledTimestamp, Task->ScheduledThreadId, Task->StartedTimestamp, Task->StartedThreadId, FTaskGraphRelation::ETaskGraphRelationType::Scheduled);
 			}
 
 			for (const TraceServices::FTaskInfo::FRelationInfo& RelationInfo : Task->NestedTasks)
 			{
-				const TraceServices::FTaskInfo* RelatedTask = TasksProvider->TryGetTask(RelationInfo.RelativeId);
+				const TraceServices::FTaskInfo* NestedTask = TasksProvider->TryGetTask(RelationInfo.RelativeId);
+				check(NestedTask != nullptr);
+
+				AddRelation(RelationInfo.Timestamp, ThreadId, NestedTask->StartedTimestamp, NestedTask->StartedThreadId, FTaskGraphRelation::ETaskGraphRelationType::AddedNested);
+
+				AddRelation(NestedTask->CompletedTimestamp, NestedTask->CompletedThreadId, NestedTask->CompletedTimestamp, Task->StartedThreadId, FTaskGraphRelation::ETaskGraphRelationType::NestedCompleted);
+			}
+
+			for (const TraceServices::FTaskInfo::FRelationInfo& RelationInfo : Task->Subsequents)
+			{
+				const TraceServices::FTaskInfo* Subsequent = TasksProvider->TryGetTask(RelationInfo.RelativeId);
+				check(Subsequent != nullptr);
+				AddRelation(Task->CompletedTimestamp, ThreadId, Subsequent->StartedTimestamp, Subsequent->StartedThreadId, FTaskGraphRelation::ETaskGraphRelationType::Subsequent);
+			}
+
+			if (Task->FinishedTimestamp != Task->CompletedTimestamp || Task->CompletedThreadId != ThreadId)
+			{
+				AddRelation(Task->FinishedTimestamp, ThreadId, Task->CompletedTimestamp, ThreadId, FTaskGraphRelation::ETaskGraphRelationType::Completed);
+			}
+		}
+
+		// if it's an event waiting for tasks completeness, add relations to these tasks
+		const TraceServices::ITimingProfilerProvider& TimingProfilerProvider = *TraceServices::ReadTimingProfilerProvider(*Session.Get());
+		const TraceServices::ITimingProfilerTimerReader* TimerReader;
+		TraceServices::FAnalysisSessionReadScope SessionReadScope(*Session.Get());
+		TimingProfilerProvider.ReadTimers([&TimerReader](const TraceServices::ITimingProfilerTimerReader& Out) { TimerReader = &Out; });
 
-				if (RelatedTask)
+		const FThreadTrackEvent& ThreadTrackEvent = *static_cast<const FThreadTrackEvent*>(&InSelectedEvent);
+		const TraceServices::FTimingProfilerTimer* Timer = TimerReader->GetTimer(ThreadTrackEvent.GetTimerIndex());
+		check(Timer != nullptr);
+
+		const TraceServices::FWaitingForTasks* Waiting = TasksProvider->TryGetWaiting(Timer->Name, ThreadId, StartTime);
+		if (Waiting != nullptr)
+		{
+			for (TaskTrace::FId TaskId : Waiting->Tasks)
+			{
+				const TraceServices::FTaskInfo* WaitedTask = TasksProvider->TryGetTask(TaskId);
+				if (WaitedTask != nullptr)
 				{
-					FTaskGraphRelation* NewRelation = new FTaskGraphRelation(StartTime, ThreadId, RelatedTask->StartedTimestamp, RelatedTask->StartedThreadId, FTaskGraphRelation::ETaskGraphRelationType::Nested);
-					NewRelation->SetSourceTrack(SharedThis(this));
-					NewRelation->SetSourceDepth(InSelectedEvent.GetDepth());
-					Relations.Add(NewRelation);
+					AddRelation(StartTime, ThreadId, WaitedTask->StartedTimestamp, WaitedTask->StartedThreadId, FTaskGraphRelation::ETaskGraphRelationType::AddedNested);
+					AddRelation(WaitedTask->CompletedTimestamp, WaitedTask->CompletedThreadId, WaitedTask->CompletedTimestamp, ThreadId, FTaskGraphRelation::ETaskGraphRelationType::NestedCompleted);
 				}
 			}
 		}
@@ -1539,10 +1602,21 @@ void FThreadTimingTrack::SolveEventRelations(const TArray<ITimingEventRelation*>
 		{
 			// A check that we have the correct type will need to be added
 			FTaskGraphRelation* TaskRelation = static_cast<FTaskGraphRelation*>(Relation);
-			if (TaskRelation->GetTargetThreadId() == ThreadId)
+
+			if (TaskRelation->GetSourceThreadId() == ThreadId && !TaskRelation->GetSourceTrack())
+			{
+				TaskRelation->SetSourceTrack(StaticCastSharedRef<const FBaseTimingTrack>(SharedThis(this)));
+				TaskRelation->SetSourceDepth(GetDepthAt(TaskRelation->GetSourceTime()) - 1);
+			}
+
+			if (TaskRelation->GetTargetThreadId() == ThreadId && !TaskRelation->GetTargetTrack())
 			{
 				TaskRelation->SetTargetTrack(StaticCastSharedRef<const FBaseTimingTrack>(SharedThis(this)));
 				TaskRelation->SetTargetDepth(GetDepthAt(TaskRelation->GetTargetTime()) - 1);
+			}
+
+			if (TaskRelation->GetSourceTrack() && TaskRelation->GetTargetTrack())
+			{
 				TaskRelation->SetIsSolved(true);
 			}
 		}
diff --git a/Engine/Source/Developer/TraceInsights/Private/Insights/Widgets/STimingView.cpp b/Engine/Source/Developer/TraceInsights/Private/Insights/Widgets/STimingView.cpp
--- a/Engine/Source/Developer/TraceInsights/Private/Insights/Widgets/STimingView.cpp
+++ b/Engine/Source/Developer/TraceInsights/Private/Insights/Widgets/STimingView.cpp
@@ -3411,7 +3411,7 @@ void STimingView::UpdateHoveredTimingEvent(float InMousePosX, float InMousePosY)
 			const double T3 = Stopwatch.GetAccumulatedTime();
 			if (T3 > 0.001)
 			{
-				UE_LOG(TimingProfiler, Log, TEXT("HoveredTrack [%g, %g] Tooltip: %.1f ms (%.1f + %.1f + %.1f)"),
+				UE_LOG(TimingProfiler, Verbose, TEXT("HoveredTrack [%g, %g] Tooltip: %.1f ms (%.1f + %.1f + %.1f)"),
 					InMousePosX, InMousePosY, T3 * 1000.0, T1 * 1000.0, (T2 - T1) * 1000.0, (T3 - T2) * 1000.0);
 			}
 		}
diff --git a/Engine/Source/Developer/TraceInsights/Private/Insights/Widgets/STimingView.h b/Engine/Source/Developer/TraceInsights/Private/Insights/Widgets/STimingView.h
--- a/Engine/Source/Developer/TraceInsights/Private/Insights/Widgets/STimingView.h
+++ b/Engine/Source/Developer/TraceInsights/Private/Insights/Widgets/STimingView.h
@@ -629,5 +629,5 @@ protected:
 	TSharedPtr<FUICommandList> CommandList;
 
 	TArray<ITimingEventRelation*> CurrentRelations;
-	bool bShowEventRelations = false;
+	bool bShowEventRelations = true;
 };
\ No newline at end of file
diff --git a/Engine/Source/Developer/TraceServices/Private/Model/TasksProfiler.cpp b/Engine/Source/Developer/TraceServices/Private/Model/TasksProfiler.cpp
--- a/Engine/Source/Developer/TraceServices/Private/Model/TasksProfiler.cpp
+++ b/Engine/Source/Developer/TraceServices/Private/Model/TasksProfiler.cpp
@@ -197,6 +197,9 @@ namespace TraceServices
 		Task->Id = TaskId;
 
 		AddRelative(TEXT("Subsequent"), TaskId, &FTaskInfo::Subsequents, SubsequentId, Timestamp, ThreadId);
+
+		// make a backward link from the subsequent task to this task (prerequisite)
+		AddRelative(TEXT("Prerequisite"), SubsequentId, &FTaskInfo::Prerequisites, TaskId, Timestamp, ThreadId);
 	}
 
 	void FTasksProvider::TaskStarted(TaskTrace::FId TaskId, double Timestamp, uint32 ThreadId)
@@ -347,7 +350,7 @@ namespace TraceServices
 			return false;
 		}
 
-		checkf(Task->*TimestampPtr == FTaskInfo::InvalidTimestamp, TEXT("TaskId %d, old TS %.6f, new TS %.6f"), TaskId, Task->*TimestampPtr, TimestampValue);
+		checkf(Task->*TimestampPtr == FTaskInfo::InvalidTimestamp, TEXT("%s: TaskId %d, old TS %.6f, new TS %.6f"), EventName, TaskId, Task->*TimestampPtr, TimestampValue);
 		Task->*TimestampPtr = TimestampValue;
 		if (ThreadIdPtr != nullptr)
 		{
@@ -400,8 +403,13 @@ namespace TraceServices
 		return Task != nullptr && Task->FinishedTimestamp > Timestamp ? Task : nullptr;
 	}
 
-	const FWaitingForTasks* FTasksProvider::TryGetWaiting(uint32 ThreadId, double Timestamp) const
+	const FWaitingForTasks* FTasksProvider::TryGetWaiting(const TCHAR* TimerName, uint32 ThreadId, double Timestamp) const
 	{
+		if (FCString::Strcmp(TimerName, TEXT("WaitUntilTasksComplete")) != 0 && FCString::Strcmp(TimerName, TEXT("GameThreadWaitForTask")) != 0)
+		{
+			return nullptr;
+		}
+
 		const TArray64<FWaitingForTasks>* Thread = WaitingThreads.Find(ThreadId);
 		if (Thread == nullptr)
 		{
diff --git a/Engine/Source/Developer/TraceServices/Private/Model/TasksProfilerPrivate.h b/Engine/Source/Developer/TraceServices/Private/Model/TasksProfilerPrivate.h
--- a/Engine/Source/Developer/TraceServices/Private/Model/TasksProfilerPrivate.h
+++ b/Engine/Source/Developer/TraceServices/Private/Model/TasksProfilerPrivate.h
@@ -36,7 +36,7 @@ namespace TraceServices
 		// ITasksProvider impl
 		virtual const FTaskInfo* TryGetTask(uint32 ThreadId, double Timestamp) const override;
 		virtual const FTaskInfo* TryGetTask(TaskTrace::FId TaskId) const;
-		virtual const FWaitingForTasks* TryGetWaiting(uint32 ThreadId, double Timestamp) const override;
+		virtual const FWaitingForTasks* TryGetWaiting(const TCHAR* TimerName, uint32 ThreadId, double Timestamp) const override;
 
 	private:
 		// initializes conversion of task id to its index in the internal container, if it's not initialized already
diff --git a/Engine/Source/Developer/TraceServices/Public/TraceServices/Model/TasksProfiler.h b/Engine/Source/Developer/TraceServices/Public/TraceServices/Model/TasksProfiler.h
--- a/Engine/Source/Developer/TraceServices/Public/TraceServices/Model/TasksProfiler.h
+++ b/Engine/Source/Developer/TraceServices/Public/TraceServices/Model/TasksProfiler.h
@@ -57,6 +57,8 @@ namespace TraceServices
 			uint32 ThreadId;
 		};
 
+		// other tasks that must complete before this taks can be scheduled
+		TArray<FRelationInfo> Prerequisites;
 		// other tasks that wait for this task completion before they'll start execution
 		TArray<FRelationInfo> Subsequents;
 		// the task is completed only after all nested tasks are completed
@@ -82,7 +84,7 @@ namespace TraceServices
 		virtual const FTaskInfo* TryGetTask(TaskTrace::FId TaskId) const = 0;
 
 		// returns an info about waiting for tasks completion for given thread and timestamp, if any, otherwise `nullptr`
-		virtual const FWaitingForTasks* TryGetWaiting(uint32 ThreadId, double Timestamp) const = 0;
+		virtual const FWaitingForTasks* TryGetWaiting(const TCHAR* TimerName, uint32 ThreadId, double Timestamp) const = 0;
 	};
 
 	TRACESERVICES_API const ITasksProvider* ReadTasksProvider(const IAnalysisSession& Session);
