[Insights] Fixed bad integration in STableTreeView.cpp.

diff --git a/Engine/Source/Developer/TraceInsights/Private/Insights/Table/Widgets/STableTreeView.cpp b/Engine/Source/Developer/TraceInsights/Private/Insights/Table/Widgets/STableTreeView.cpp
--- a/Engine/Source/Developer/TraceInsights/Private/Insights/Table/Widgets/STableTreeView.cpp
+++ b/Engine/Source/Developer/TraceInsights/Private/Insights/Table/Widgets/STableTreeView.cpp
@@ -210,30 +210,14 @@ void STableTreeView::ConstructWidget(TSharedPtr<FTable> InTablePtr)
 			.FillWidth(1.0f)
 			.VAlign(VAlign_Center)
 			[
-				SAssignNew(ToolbarBox, SHorizontalBox)
-
-				+ SHorizontalBox::Slot()
-				.AutoWidth()
-				.VAlign(VAlign_Center)
-				[
-					SNew(STextBlock)
-					.Text(LOCTEXT("GroupByText", "Hierarchy:"))
-					.Margin(FMargin(0.0f, 0.0f, 4.0f, 0.0f))
-				]
-
-				+ SHorizontalBox::Slot()
-				.FillWidth(1.0f)
-				.VAlign(VAlign_Center)
-				[
-					SAssignNew(GroupingBreadcrumbTrail, SBreadcrumbTrail<TSharedPtr<FTreeNodeGrouping>>)
-					.ButtonContentPadding(FMargin(1.0f, 1.0f))
-					//.DelimiterImage(FEditorStyle::GetBrush("SlateFileDialogs.PathDelimiter"))
-					//.TextStyle(FEditorStyle::Get(), "Tutorials.Browser.PathText")
-					//.ShowLeadingDelimiter(true)
-					//.PersistentBreadcrumbs(true)
-					.OnCrumbClicked(this, &STableTreeView::OnGroupingCrumbClicked)
-					.GetCrumbMenuContent(this, &STableTreeView::GetGroupingCrumbMenuContent)
-				]
+				SAssignNew(GroupingBreadcrumbTrail, SBreadcrumbTrail<TSharedPtr<FTreeNodeGrouping>>)
+				.ButtonContentPadding(FMargin(1.0f, 1.0f))
+				//.DelimiterImage(FEditorStyle::GetBrush("SlateFileDialogs.PathDelimiter"))
+				//.TextStyle(FEditorStyle::Get(), "Tutorials.Browser.PathText")
+				//.ShowLeadingDelimiter(true)
+				//.PersistentBreadcrumbs(true)
+				.OnCrumbClicked(this, &STableTreeView::OnGroupingCrumbClicked)
+				.GetCrumbMenuContent(this, &STableTreeView::GetGroupingCrumbMenuContent)
 			]
 		]
 	]
