[Insights]
  - Timing view: Added command bindings for few actions ("Auto Hide Empty Tracks" - V key shortcut, "Show/Hide Graph Track" - G key shortcut and for "Compact Mode" - C key shortcut).
  - Timing view: Fxed command binding for "Show Task Dependencies" (T key shortcut). Also added this action in Tracks drop down menu (it is still available also in the context menu).
  - Timing view: Fixed "Auto Hide Empty Tracks" incorrectly hiding the graph tracks when toggling on/off.
  - Timing view: Fixed "hovered" and "selected" track states to be restored for all visible tracks (i.e. also for top and bottom docked tracks).
  - Timing view: Added "Move TimeMarker 'X' Here" options in context menu of Time Ruler track.
  - Timing view: Restored log message re slow tootip to be on Log channel instead of Verbose.

diff --git a/Engine/Source/Developer/TraceInsights/Private/Insights/TimingProfilerCommands.cpp b/Engine/Source/Developer/TraceInsights/Private/Insights/TimingProfilerCommands.cpp
--- a/Engine/Source/Developer/TraceInsights/Private/Insights/TimingProfilerCommands.cpp
+++ b/Engine/Source/Developer/TraceInsights/Private/Insights/TimingProfilerCommands.cpp
@@ -84,9 +84,12 @@ FTimingViewCommands::FTimingViewCommands()
 PRAGMA_DISABLE_OPTIMIZATION
 void FTimingViewCommands::RegisterCommands()
 {
-	UI_COMMAND(ShowAllGpuTracks, "GPU Track", "Show/hide the GPU track", EUserInterfaceActionType::ToggleButton, FInputChord(EKeys::Y));
-	UI_COMMAND(ShowAllCpuTracks, "CPU Thread Tracks", "Show/hide all CPU tracks (and all CPU thread groups)", EUserInterfaceActionType::ToggleButton, FInputChord(EKeys::U));
-	UI_COMMAND(ShowTaskDependencies, "Task Dependencies", "Show/hide all dependencies of the current task", EUserInterfaceActionType::ToggleButton, FInputChord(EKeys::T));
+	UI_COMMAND(AutoHideEmptyTracks, "Auto Hide Empty Tracks", "Auto hide empty tracks (ones without timing events in current viewport).", EUserInterfaceActionType::ToggleButton, FInputChord(EKeys::V));
+	UI_COMMAND(ToggleCompactMode, "Compact Mode", "Toggle compact mode for supporting tracks.\n(ex. timing tracks will be displayed with reduced height)", EUserInterfaceActionType::ToggleButton, FInputChord(EKeys::C));
+	UI_COMMAND(ShowMainGraphTrack, "Graph Track", "Show/hide the main Graph track.", EUserInterfaceActionType::ToggleButton, FInputChord(EKeys::G));
+	UI_COMMAND(ShowTaskDependencies, "Task Dependencies", "Show/hide all dependencies of the current task (for a selected cpu timing event).", EUserInterfaceActionType::ToggleButton, FInputChord(EKeys::T));
+	//UI_COMMAND(ShowAllGpuTracks, "GPU Track(s)", "Show/hide the GPU track(s).", EUserInterfaceActionType::ToggleButton, FInputChord(EKeys::Y));
+	//UI_COMMAND(ShowAllCpuTracks, "CPU Thread Tracks", "Show/hide all CPU tracks (and all CPU thread groups).", EUserInterfaceActionType::ToggleButton, FInputChord(EKeys::U));
 }
 PRAGMA_ENABLE_OPTIMIZATION
 
diff --git a/Engine/Source/Developer/TraceInsights/Private/Insights/TimingProfilerCommands.h b/Engine/Source/Developer/TraceInsights/Private/Insights/TimingProfilerCommands.h
--- a/Engine/Source/Developer/TraceInsights/Private/Insights/TimingProfilerCommands.h
+++ b/Engine/Source/Developer/TraceInsights/Private/Insights/TimingProfilerCommands.h
@@ -66,14 +66,23 @@ public:
 	virtual void RegisterCommands() override;
 
 public:
-	/** Toggles visibility for GPU thread track. */
-	TSharedPtr<FUICommandInfo> ShowAllGpuTracks;
+	/** Toggles visibility of empty tracks. */
+	TSharedPtr<FUICommandInfo> AutoHideEmptyTracks;
 
-	/** Toggles visibility for all CPU thread tracks at once. */
-	TSharedPtr<FUICommandInfo> ShowAllCpuTracks;
+	/** Toggles 'compact mode' for timing tracks. */
+	TSharedPtr<FUICommandInfo> ToggleCompactMode;
+
+	/** Toggles visibility for Main Graph track. */
+	TSharedPtr<FUICommandInfo> ShowMainGraphTrack;
 
 	/** Toggles visibility for task graph dependency relations */
 	TSharedPtr<FUICommandInfo> ShowTaskDependencies;
+
+	/** Toggles visibility for GPU thread track. */
+	//TSharedPtr<FUICommandInfo> ShowAllGpuTracks;
+
+	/** Toggles visibility for all CPU thread tracks at once. */
+	//TSharedPtr<FUICommandInfo> ShowAllCpuTracks;
 };
 
 ////////////////////////////////////////////////////////////////////////////////////////////////////
diff --git a/Engine/Source/Developer/TraceInsights/Private/Insights/ViewModels/TimeRulerTrack.cpp b/Engine/Source/Developer/TraceInsights/Private/Insights/ViewModels/TimeRulerTrack.cpp
--- a/Engine/Source/Developer/TraceInsights/Private/Insights/ViewModels/TimeRulerTrack.cpp
+++ b/Engine/Source/Developer/TraceInsights/Private/Insights/ViewModels/TimeRulerTrack.cpp
@@ -5,6 +5,7 @@
 #include "Fonts/FontMeasure.h"
 #include "Fonts/SlateFontInfo.h"
 #include "Framework/Application/SlateApplication.h"
+#include "Framework/MultiBox/MultiBoxBuilder.h"
 #include "Rendering/DrawElements.h"
 #include "Styling/CoreStyle.h"
 #include <limits>
@@ -28,6 +29,7 @@ FTimeRulerTrack::FTimeRulerTrack()
 	: FBaseTimingTrack(TEXT("Time Ruler"))
 	, WhiteBrush(FInsightsStyle::Get().GetBrush("WhiteBrush"))
 	, Font(FCoreStyle::GetDefaultFontStyle("Regular", 8))
+	, CrtMousePosTime(0.0)
 	, CrtMousePosTextWidth(0.0f)
 {
 	SetValidLocations(ETimingTrackLocation::TopDocked);
@@ -249,6 +251,8 @@ void FTimeRulerTrack::PostDraw(const ITimingTrackDrawContext& Context) const
 		// Time at current mouse position.
 		FString MousePosText;
 		const double MousePosTime = Viewport.SlateUnitsToTime(MousePosition.X);
+		CrtMousePosTime = MousePosTime;
+
 		const double DT = 100.0 / Viewport.GetScaleX();
 		const double MousePosPrecision = FMath::Max(DT / 100.0, TimeUtils::Nanosecond);
 		if (bIsMouseOver)
@@ -343,7 +347,7 @@ void FTimeRulerTrack::DrawTimeMarker(const ITimingTrackDrawContext& Context, con
 
 	FDrawContext& DrawContext = Context.GetDrawContext();
 
-	// Draw the orange vertical line.
+	// Draw the vertical line.
 	DrawContext.DrawBox(TimeMarkerX, TimeMarkerY, 1.0f, Viewport.GetHeight() - TimeMarkerY, WhiteBrush, TimeMarker.GetColor());
 	DrawContext.LayerId++;
 
@@ -420,4 +424,55 @@ void FTimeRulerTrack::DrawTimeMarker(const ITimingTrackDrawContext& Context, con
 
 ////////////////////////////////////////////////////////////////////////////////////////////////////
 
+void FTimeRulerTrack::BuildContextMenu(FMenuBuilder& MenuBuilder)
+{
+	TArray<TSharedRef<Insights::FTimeMarker>> VisibleTimeMarkers;
+	for (TSharedRef<Insights::FTimeMarker>& TimeMarker : TimeMarkers)
+	{
+		if (TimeMarker->IsVisible())
+		{
+			VisibleTimeMarkers.Add(TimeMarker);
+		}
+	}
+
+	if (VisibleTimeMarkers.Num() > 0)
+	{
+		// Sort TimeMarkers by name.
+		VisibleTimeMarkers.Sort([](const TSharedRef<Insights::FTimeMarker>& A, const TSharedRef<Insights::FTimeMarker>& B) -> bool { return A->GetName().Compare(B->GetName()) <= 0; });
+
+		MenuBuilder.BeginSection("TimeMarkers", LOCTEXT("TimeMarkersSection", "Time Markers"));
+
+		for (TSharedRef<Insights::FTimeMarker>& TimeMarker : VisibleTimeMarkers)
+		{
+			FUIAction Action_MoveTimeMarker
+			(
+				FExecuteAction::CreateSP(this, &FTimeRulerTrack::ContextMenu_MoveTimeMarker_Execute, TimeMarker),
+				FCanExecuteAction()
+			);
+			const FString& MarkerNameString = TimeMarker->GetName();
+			const FText MarkerNameText = FText::FromString((MarkerNameString.Len() > 0) ? MarkerNameString : TEXT("T"));
+			MenuBuilder.AddMenuEntry
+			(
+				FText::Format(LOCTEXT("ContextMenu_MoveTimeMerker", "Move Time Marker '{0}' Here"), MarkerNameText),
+				FText::Format(LOCTEXT("ContextMenu_MoveTimeMerker_Desc", "Move the time marker '{0}' at the current mouse position."), MarkerNameText),
+				FSlateIcon(),
+				Action_MoveTimeMarker,
+				NAME_None,
+				EUserInterfaceActionType::Button
+			);
+		}
+
+		MenuBuilder.EndSection();
+	}
+}
+
+////////////////////////////////////////////////////////////////////////////////////////////////////
+
+void FTimeRulerTrack::ContextMenu_MoveTimeMarker_Execute(TSharedRef<Insights::FTimeMarker> InTimeMarker)
+{
+	InTimeMarker->SetTime(CrtMousePosTime);
+}
+
+////////////////////////////////////////////////////////////////////////////////////////////////////
+
 #undef LOCTEXT_NAMESPACE
diff --git a/Engine/Source/Developer/TraceInsights/Private/Insights/ViewModels/TimeRulerTrack.h b/Engine/Source/Developer/TraceInsights/Private/Insights/ViewModels/TimeRulerTrack.h
--- a/Engine/Source/Developer/TraceInsights/Private/Insights/ViewModels/TimeRulerTrack.h
+++ b/Engine/Source/Developer/TraceInsights/Private/Insights/ViewModels/TimeRulerTrack.h
@@ -9,6 +9,7 @@
 #include "Insights/ViewModels/BaseTimingTrack.h"
 
 struct FDrawContext;
+class FMenuBuilder;
 struct FSlateBrush;
 class FTimingTrackViewport;
 
@@ -92,22 +93,27 @@ public:
 	void Draw(const ITimingTrackDrawContext& Context) const override;
 	void PostDraw(const ITimingTrackDrawContext& Context) const override;
 
+	virtual void BuildContextMenu(FMenuBuilder& MenuBuilder) override;
+
 private:
 	void DrawTimeMarker(const ITimingTrackDrawContext& Context, const Insights::FTimeMarker& TimeMarker) const;
+	void ContextMenu_MoveTimeMarker_Execute(TSharedRef<Insights::FTimeMarker> InTimeMarker);
 
-public:
+private:
 	// Slate resources
 	const FSlateBrush* WhiteBrush;
 	const FSlateFontInfo Font;
 
-	// Smoothed mouse pos text width to avoid flickering
-	mutable float CrtMousePosTextWidth;
-
-private:
 	bool bIsSelecting;
 	double SelectionStartTime;
 	double SelectionEndTime;
 
+	// The last time value at mouse postion. Updated in PostDraw.
+	mutable double CrtMousePosTime;
+
+	// The smoothed width of "the text at mouse position" to avoid flickering. Updated in PostDraw.
+	mutable float CrtMousePosTextWidth;
+
 	/**
 	 * The sorted list of all registered time markers. It defines the draw order of time markers.
 	 * The time marker currently scrubbing will be moved at the end of the list in order to be displayed on top of other markers.
diff --git a/Engine/Source/Developer/TraceInsights/Private/Insights/Widgets/STimingView.cpp b/Engine/Source/Developer/TraceInsights/Private/Insights/Widgets/STimingView.cpp
--- a/Engine/Source/Developer/TraceInsights/Private/Insights/Widgets/STimingView.cpp
+++ b/Engine/Source/Developer/TraceInsights/Private/Insights/Widgets/STimingView.cpp
@@ -907,7 +907,7 @@ void STimingView::Tick(const FGeometry& AllottedGeometry, const double InCurrent
 		UpdateTracksStopwatch.Stop();
 		UpdateTracksDurationHistory.AddValue(UpdateTracksStopwatch.AccumulatedTime);
 	}
-	//////////////////////////////////////////////////
+	////////////////////////////////////////////////////////////////////////////////////////////////////
 	// Post-Update.
 	{
 		FStopwatch PostUpdateTracksStopwatch;
@@ -948,7 +948,7 @@ void STimingView::Tick(const FGeometry& AllottedGeometry, const double InCurrent
 		PostUpdateTracksStopwatch.Stop();
 		PostUpdateTracksDurationHistory.AddValue(PostUpdateTracksStopwatch.AccumulatedTime);
 	}
-	//////////////////////////////////////////////////
+	////////////////////////////////////////////////////////////////////////////////////////////////////
 
 	Tooltip.Update();
 	if (!MousePosition.IsZero())
@@ -956,9 +956,27 @@ void STimingView::Tick(const FGeometry& AllottedGeometry, const double InCurrent
 		Tooltip.SetPosition(MousePosition, 0.0f, Viewport.GetWidth() - 12.0f, 0.0f, Viewport.GetHeight() - 12.0f); // -12.0f is to avoid overlaping the scrollbars
 	}
 
-	//////////////////////////////////////////////////
+	////////////////////////////////////////////////////////////////////////////////////////////////////
+	// Update "hovered" and "selected" flags for all visible tracks.
+	//TODO: Move this before PreUpdate (so a track could adjust its size based on hovered/selected flags).
 
 	// Reset hovered/selected flags for all tracks.
+	for (TSharedPtr<FBaseTimingTrack>& TrackPtr : TopDockedTracks)
+	{
+		if (TrackPtr->IsVisible())
+		{
+			TrackPtr->SetHoveredState(false);
+			TrackPtr->SetSelectedFlag(false);
+		}
+	}
+	for (TSharedPtr<FBaseTimingTrack>& TrackPtr : BottomDockedTracks)
+	{
+		if (TrackPtr->IsVisible())
+		{
+			TrackPtr->SetHoveredState(false);
+			TrackPtr->SetSelectedFlag(false);
+		}
+	}
 	for (TSharedPtr<FBaseTimingTrack>& TrackPtr : ScrollableTracks)
 	{
 		if (TrackPtr->IsVisible())
@@ -967,6 +985,14 @@ void STimingView::Tick(const FGeometry& AllottedGeometry, const double InCurrent
 			TrackPtr->SetSelectedFlag(false);
 		}
 	}
+	for (TSharedPtr<FBaseTimingTrack>& TrackPtr : ForegroundTracks)
+	{
+		if (TrackPtr->IsVisible())
+		{
+			TrackPtr->SetHoveredState(false);
+			TrackPtr->SetSelectedFlag(false);
+		}
+	}
 
 	// Set the hovered flag for the actual hovered track, if any.
 	if (HoveredTrack.IsValid())
@@ -980,7 +1006,7 @@ void STimingView::Tick(const FGeometry& AllottedGeometry, const double InCurrent
 		SelectedTrack->SetSelectedFlag(true);
 	}
 
-	//////////////////////////////////////////////////
+	////////////////////////////////////////////////////////////////////////////////////////////////////
 
 	Viewport.ResetDirtyFlags();
 
@@ -1970,7 +1996,7 @@ FReply STimingView::OnMouseButtonUp(const FGeometry& MyGeometry, const FPointerE
 			}
 			else if (TimeRulerTrack->IsScrubbing())
 			{
-				RaiseTimeMarkerChanged();
+				RaiseTimeMarkerChanged(TimeRulerTrack->GetScrubbingTimeMarker());
 				TimeRulerTrack->StopScrubbing();
 			}
 
@@ -2015,7 +2041,7 @@ FReply STimingView::OnMouseButtonUp(const FGeometry& MyGeometry, const FPointerE
 			}
 			else if (TimeRulerTrack->IsScrubbing())
 			{
-				RaiseTimeMarkerChanged();
+				RaiseTimeMarkerChanged(TimeRulerTrack->GetScrubbingTimeMarker());
 				TimeRulerTrack->StopScrubbing();
 			}
 
@@ -2186,7 +2212,7 @@ FReply STimingView::OnMouseMove(const FGeometry& MyGeometry, const FPointerEvent
 
 				TSharedRef<Insights::FTimeMarker> ScrubbingTimeMarker = TimeRulerTrack->GetScrubbingTimeMarker();
 				ScrubbingTimeMarker->SetTime(Viewport.SlateUnitsToTime(MousePosition.X));
-				RaiseTimeMarkerChanging();
+				RaiseTimeMarkerChanging(ScrubbingTimeMarker);
 			}
 		}
 		else
@@ -2442,18 +2468,9 @@ FReply STimingView::OnKeyDown(const FGeometry& MyGeometry, const FKeyEvent& InKe
 			{
 				SelectedEvent->GetTrack()->OnClipboardCopyEvent(*SelectedEvent);
 			}
-		}
-		else
-		{
-			Viewport.SwitchLayoutCompactMode();
 			return FReply::Handled();
 		}
 	}
-	else if (InKeyEvent.GetKey() == EKeys::V)
-	{
-		ToggleAutoHideEmptyTracks();
-		return FReply::Handled();
-	}
 	else if (InKeyEvent.GetKey() == EKeys::Equals ||
 			 InKeyEvent.GetKey() == EKeys::Add)
 	{
@@ -2553,11 +2570,6 @@ FReply STimingView::OnKeyDown(const FGeometry& MyGeometry, const FKeyEvent& InKe
 		Viewport.AddDirtyFlags(ETimingTrackViewportDirtyFlags::HInvalidated);
 		return FReply::Handled();
 	}
-	else if (InKeyEvent.GetKey() == EKeys::G)
-	{
-		ShowHideGraphTrack_Execute();
-		return FReply::Handled();
-	}
 	else if (InKeyEvent.GetKey() == EKeys::R)
 	{
 		FrameSharedState->ShowHideAllFrameTracks();
@@ -2618,13 +2630,18 @@ FReply STimingView::OnKeyDown(const FGeometry& MyGeometry, const FKeyEvent& InKe
 	}
 #if ACTIVATE_BENCHMARK
 	else if (InKeyEvent.GetKey() == EKeys::Z)
-		{
-			FTimingProfilerTests::FCheckValues CheckValues;
-			FTimingProfilerTests::RunEnumerateBenchmark(FTimingProfilerTests::FEnumerateTestParams(), CheckValues);
-			return FReply::Handled();
-		}
+	{
+		FTimingProfilerTests::FCheckValues CheckValues;
+		FTimingProfilerTests::RunEnumerateBenchmark(FTimingProfilerTests::FEnumerateTestParams(), CheckValues);
+		return FReply::Handled();
+	}
 #endif
 
+	if (CommandList->ProcessCommandBindings(InKeyEvent))
+	{
+		return FReply::Handled();
+	}
+
 	return SCompoundWidget::OnKeyDown(MyGeometry, InKeyEvent);
 }
 
@@ -2646,12 +2663,10 @@ FReply STimingView::OnKeyUp(const FGeometry& MyGeometry, const FKeyEvent& InKeyE
 
 void STimingView::ShowContextMenu(const FPointerEvent& MouseEvent)
 {
-	//TSharedPtr<FUICommandList> ProfilerCommandList = FTimingProfilerManager::Get()->GetCommandList();
-	//const FTimingProfilerCommands& ProfilerCommands = FTimingProfilerManager::GetCommands();
-	//const FTimingViewCommands& Commands = FTimingViewCommands::Get();
+	const FTimingViewCommands& Commands = FTimingViewCommands::Get();
 
 	const bool bShouldCloseWindowAfterMenuSelection = true;
-	//FMenuBuilder MenuBuilder(bShouldCloseWindowAfterMenuSelection, ProfilerCommandList);
+
 	FMenuBuilder MenuBuilder(bShouldCloseWindowAfterMenuSelection, CommandList);
 
 	if (HoveredTrack.IsValid())
@@ -2668,7 +2683,7 @@ void STimingView::ShowContextMenu(const FPointerEvent& MouseEvent)
 		{
 			MenuBuilder.AddMenuEntry(
 				LOCTEXT("ContextMenu_NA", "N/A"),
-				LOCTEXT("ContextMenu_NA_Desc", "No actions available."),
+				LOCTEXT("ContextMenu_NA_Desc", "No action available."),
 				FSlateIcon(),
 				FUIAction(FExecuteAction(), FCanExecuteAction::CreateLambda([](){ return false; })),
 				NAME_None,
@@ -2683,7 +2698,7 @@ void STimingView::ShowContextMenu(const FPointerEvent& MouseEvent)
 
 		MenuBuilder.AddMenuEntry
 		(
-			FTimingViewCommands::Get().ShowTaskDependencies,
+			Commands.ShowTaskDependencies,
 			NAME_None,
 			TAttribute<FText>(),
 			TAttribute<FText>(),
@@ -2693,15 +2708,6 @@ void STimingView::ShowContextMenu(const FPointerEvent& MouseEvent)
 		MenuBuilder.EndSection();
 	}
 
-	//MenuBuilder.BeginSection(TEXT("Misc"), LOCTEXT("Miscellaneous", "Miscellaneous"));
-	//{
-	//	MenuBuilder.AddMenuEntry(Commands.ShowAllGpuTracks);
-	//	MenuBuilder.AddMenuEntry(Commands.ShowAllCpuTracks);
-	//	MenuBuilder.AddMenuEntry(ProfilerCommands.ToggleTimersViewVisibility);
-	//	MenuBuilder.AddMenuEntry(ProfilerCommands.ToggleStatsCountersViewVisibility);
-	//}
-	//MenuBuilder.EndSection();
-
 	TSharedRef<SWidget> MenuWidget = MenuBuilder.MakeWidget();
 
 	FWidgetPath EventPath = MouseEvent.GetEventPath() != nullptr ? *MouseEvent.GetEventPath() : FWidgetPath();
@@ -2882,16 +2888,37 @@ bool STimingView::CanChangeTrackLocation(TSharedRef<FBaseTimingTrack> Track, ETi
 void STimingView::BindCommands()
 {
 	FTimingViewCommands::Register();
+
+	const FTimingViewCommands& Commands = FTimingViewCommands::Get();
+
 	CommandList = MakeShared<FUICommandList>();
 
 	CommandList->MapAction(
-		FTimingViewCommands::Get().ShowTaskDependencies, 
-		FExecuteAction::CreateSP(this, &STimingView::ContextMenu_ShowTaskDependecies_Execute), 
+		Commands.AutoHideEmptyTracks,
+		FExecuteAction::CreateSP(this, &STimingView::ToggleAutoHideEmptyTracks),
+		FCanExecuteAction(),
+		FIsActionChecked::CreateSP(this, &STimingView::IsAutoHideEmptyTracksEnabled));
+
+	CommandList->MapAction(
+		Commands.ToggleCompactMode,
+		FExecuteAction::CreateSP(this, &STimingView::ToggleCompactMode),
+		FCanExecuteAction(),
+		FIsActionChecked::CreateSP(this, &STimingView::IsCompactModeEnabled));
+
+	CommandList->MapAction(
+		Commands.ShowMainGraphTrack,
+		FExecuteAction::CreateSP(this, &STimingView::ShowHideGraphTrack_Execute),
+		//FCanExecuteAction(),
+		FCanExecuteAction::CreateLambda([] { return true; }),
+		FIsActionChecked::CreateSP(this, &STimingView::ShowHideGraphTrack_IsChecked));
+
+	CommandList->MapAction(
+		Commands.ShowTaskDependencies,
+		FExecuteAction::CreateSP(this, &STimingView::ContextMenu_ShowTaskDependecies_Execute),
 		FCanExecuteAction::CreateSP(this, &STimingView::ContextMenu_ShowTaskDependecies_CanExecute),
 		FIsActionChecked::CreateSP(this, &STimingView::ContextMenu_ShowTaskDependecies_IsChecked));
 
-	//
-	//TSharedPtr<FUICommandList> CommandList = FTimingProfilerManager::Get()->GetCommandList();
+	//TODO: let extenders to bind other commands
 
 	//CommandList->MapAction(
 	//	Commands.ShowAllGpuTracks,
@@ -3155,18 +3182,24 @@ void STimingView::RaiseSelectionChanged()
 
 ////////////////////////////////////////////////////////////////////////////////////////////////////
 
-void STimingView::RaiseTimeMarkerChanging()
+void STimingView::RaiseTimeMarkerChanging(TSharedRef<Insights::FTimeMarker> InTimeMarker)
 {
-	const double TimeMarker = GetTimeMarker();
-	OnTimeMarkerChangedDelegate.Broadcast(Insights::ETimeChangedFlags::Interactive, TimeMarker);
+	if (InTimeMarker == DefaultTimeMarker)
+	{
+		const double Time = DefaultTimeMarker->GetTime();
+		OnTimeMarkerChangedDelegate.Broadcast(Insights::ETimeChangedFlags::Interactive, Time);
+	}
 }
 
 ////////////////////////////////////////////////////////////////////////////////////////////////////
 
-void STimingView::RaiseTimeMarkerChanged()
+void STimingView::RaiseTimeMarkerChanged(TSharedRef<Insights::FTimeMarker> InTimeMarker)
 {
-	const double TimeMarker = GetTimeMarker();
-	OnTimeMarkerChangedDelegate.Broadcast(Insights::ETimeChangedFlags::None, TimeMarker);
+	if (InTimeMarker == DefaultTimeMarker)
+	{
+		const double Time = DefaultTimeMarker->GetTime();
+		OnTimeMarkerChangedDelegate.Broadcast(Insights::ETimeChangedFlags::None, Time);
+	}
 }
 
 ////////////////////////////////////////////////////////////////////////////////////////////////////
@@ -3181,7 +3214,7 @@ double STimingView::GetTimeMarker() const
 void STimingView::SetTimeMarker(double InMarkerTime)
 {
 	DefaultTimeMarker->SetTime(InMarkerTime);
-	RaiseTimeMarkerChanged();
+	RaiseTimeMarkerChanged(DefaultTimeMarker);
 }
 
 ////////////////////////////////////////////////////////////////////////////////////////////////////
@@ -3411,7 +3444,7 @@ void STimingView::UpdateHoveredTimingEvent(float InMousePosX, float InMousePosY)
 			const double T3 = Stopwatch.GetAccumulatedTime();
 			if (T3 > 0.001)
 			{
-				UE_LOG(TimingProfiler, Verbose, TEXT("HoveredTrack [%g, %g] Tooltip: %.1f ms (%.1f + %.1f + %.1f)"),
+				UE_LOG(TimingProfiler, Log, TEXT("HoveredTrack [%g, %g] Tooltip: %.1f ms (%.1f + %.1f + %.1f)"),
 					InMousePosX, InMousePosY, T3 * 1000.0, T1 * 1000.0, (T2 - T1) * 1000.0, (T3 - T2) * 1000.0);
 			}
 		}
@@ -3711,7 +3744,7 @@ bool STimingView::IsFilterByEventType(const uint64 EventType) const
 
 TSharedRef<SWidget> STimingView::MakeAutoScrollOptionsMenu()
 {
-	FMenuBuilder MenuBuilder(/*bInShouldCloseWindowAfterMenuSelection=*/true, nullptr);
+	FMenuBuilder MenuBuilder(/*bInShouldCloseWindowAfterMenuSelection=*/true, CommandList);
 
 	MenuBuilder.BeginSection("QuickFilter", LOCTEXT("AutoScrollHeading", "Auto Scroll Options"));
 	{
@@ -3837,39 +3870,18 @@ TSharedRef<SWidget> STimingView::MakeAutoScrollOptionsMenu()
 
 TSharedRef<SWidget> STimingView::MakeTracksFilterMenu()
 {
-	FMenuBuilder MenuBuilder(/*bInShouldCloseWindowAfterMenuSelection=*/true, nullptr);
+	FMenuBuilder MenuBuilder(/*bInShouldCloseWindowAfterMenuSelection=*/true, CommandList);
 
 	CreateAllTracksMenu(MenuBuilder);
 
 	MenuBuilder.BeginSection("QuickFilter", LOCTEXT("TracksFilterHeading", "Quick Filter"));
 	{
-		//const FTimingViewCommands& Commands = FTimingViewCommands::Get();
+		const FTimingViewCommands& Commands = FTimingViewCommands::Get();
 
-		//TODO: MenuBuilder.AddMenuEntry(Commands.AutoHideGraphTrack);
-		MenuBuilder.AddMenuEntry(
-			LOCTEXT("ShowGraphTrack", "Graph Track - G"),
-			LOCTEXT("ShowGraphTrack_Tooltip", "Show/hide the Graph track"),
-			FSlateIcon(),
-			FUIAction(FExecuteAction::CreateSP(this, &STimingView::ShowHideGraphTrack_Execute),
-					  FCanExecuteAction(),
-					  FIsActionChecked::CreateSP(this, &STimingView::ShowHideGraphTrack_IsChecked)),
-			NAME_None,
-			EUserInterfaceActionType::ToggleButton
-		);
-
-		MenuBuilder.AddMenuSeparator("QuickFilterSeparator");
-
-		//TODO: MenuBuilder.AddMenuEntry(Commands.AutoHideEmptyTracks);
-		MenuBuilder.AddMenuEntry(
-			LOCTEXT("AutoHideEmptyTracks", "Auto Hide Empty Tracks - V"),
-			LOCTEXT("AutoHideEmptyTracks_Tooltip", "Auto hide empty tracks (ones without timing events in current viewport)"),
-			FSlateIcon(),
-			FUIAction(FExecuteAction::CreateSP(this, &STimingView::ToggleAutoHideEmptyTracks),
-					  FCanExecuteAction(),
-					  FIsActionChecked::CreateSP(this, &STimingView::IsAutoHideEmptyTracksEnabled)),
-			NAME_None,
-			EUserInterfaceActionType::ToggleButton
-		);
+		MenuBuilder.AddMenuEntry(Commands.ShowMainGraphTrack);
+		MenuBuilder.AddMenuEntry(Commands.ShowTaskDependencies);
+		MenuBuilder.AddMenuEntry(Commands.ToggleCompactMode);
+		MenuBuilder.AddMenuEntry(Commands.AutoHideEmptyTracks);
 	}
 	MenuBuilder.EndSection();
 
@@ -3987,6 +3999,20 @@ void STimingView::ShowHideGraphTrack_Execute()
 
 ////////////////////////////////////////////////////////////////////////////////////////////////////
 
+bool STimingView::IsCompactModeEnabled() const
+{
+	return Viewport.IsLayoutCompactModeEnabled();
+}
+
+////////////////////////////////////////////////////////////////////////////////////////////////////
+
+void STimingView::ToggleCompactMode()
+{
+	Viewport.SwitchLayoutCompactMode();
+}
+
+////////////////////////////////////////////////////////////////////////////////////////////////////
+
 bool STimingView::IsAutoHideEmptyTracksEnabled() const
 {
 	return (Viewport.GetLayout().TargetMinTimelineH == 0.0f);
@@ -4000,10 +4026,11 @@ void STimingView::ToggleAutoHideEmptyTracks()
 
 	for (const TSharedPtr<FBaseTimingTrack>& TrackPtr : ScrollableTracks)
 	{
-		TrackPtr->SetHeight(0.0f);
+		if (TrackPtr->Is<FTimingEventsTrack>())
+		{
+			TrackPtr->SetHeight(0.0f);
+		}
 	}
-
-	ScrollAtPosY(0.0f);
 }
 
 ////////////////////////////////////////////////////////////////////////////////////////////////////
diff --git a/Engine/Source/Developer/TraceInsights/Private/Insights/Widgets/STimingView.h b/Engine/Source/Developer/TraceInsights/Private/Insights/Widgets/STimingView.h
--- a/Engine/Source/Developer/TraceInsights/Private/Insights/Widgets/STimingView.h
+++ b/Engine/Source/Developer/TraceInsights/Private/Insights/Widgets/STimingView.h
@@ -74,6 +74,9 @@ public:
 	bool ShowHideGraphTrack_IsChecked() const;
 	void ShowHideGraphTrack_Execute();
 
+	bool IsCompactModeEnabled() const;
+	void ToggleCompactMode();
+
 	bool IsAutoHideEmptyTracksEnabled() const;
 	void ToggleAutoHideEmptyTracks();
 
@@ -404,8 +407,8 @@ protected:
 	void RaiseSelectionChanging();
 	void RaiseSelectionChanged();
 
-	void RaiseTimeMarkerChanging();
-	void RaiseTimeMarkerChanged();
+	void RaiseTimeMarkerChanging(TSharedRef<Insights::FTimeMarker> InTimeMarker);
+	void RaiseTimeMarkerChanged(TSharedRef<Insights::FTimeMarker> InTimeMarker);
 
 	void UpdateAggregatedStats();
 
diff --git a/Engine/Source/Developer/TraceInsights/Public/Insights/ViewModels/TimingTrackViewport.h b/Engine/Source/Developer/TraceInsights/Public/Insights/ViewModels/TimingTrackViewport.h
--- a/Engine/Source/Developer/TraceInsights/Public/Insights/ViewModels/TimingTrackViewport.h
+++ b/Engine/Source/Developer/TraceInsights/Public/Insights/ViewModels/TimingTrackViewport.h
@@ -204,6 +204,8 @@ public:
 	const FTimingViewLayout& GetLayout() const { return Layout; }
 	void UpdateLayout();
 
+	bool IsLayoutCompactModeEnabled() const { return Layout.bIsCompactMode; }
+
 	void SwitchLayoutCompactMode()
 	{
 		Layout.bIsCompactMode = !Layout.bIsCompactMode;
