[Insights] Draw relations between tasks as splines with arrows

diff --git a/Engine/Source/Developer/TraceInsights/Private/Insights/ViewModels/TaskGraphRelation.cpp b/Engine/Source/Developer/TraceInsights/Private/Insights/ViewModels/TaskGraphRelation.cpp
--- a/Engine/Source/Developer/TraceInsights/Private/Insights/ViewModels/TaskGraphRelation.cpp
+++ b/Engine/Source/Developer/TraceInsights/Private/Insights/ViewModels/TaskGraphRelation.cpp
@@ -22,15 +22,19 @@ void FTaskGraphRelation::Draw(const FDrawContext& DrawContext, const FTimingTrac
 {
 	float X1 = Viewport.TimeToSlateUnitsRounded(SourceTime);
 	float Y1 = SourceTrack.Pin()->GetPosY();
-	Y1 += Viewport.GetLayout().GetLaneY(SourceDepth) + Viewport.GetLayout().EventH / 2;
+	Y1 += Viewport.GetLayout().GetLaneY(SourceDepth) + Viewport.GetLayout().EventH / 2.0f;
 
 	float X2 = Viewport.TimeToSlateUnitsRounded(TargetTime);
 	float Y2 = TargetTrack.Pin()->GetPosY();
-	Y2 += Viewport.GetLayout().GetLaneY(TargetDepth) + Viewport.GetLayout().EventH / 2;
+	Y2 += Viewport.GetLayout().GetLaneY(TargetDepth) + Viewport.GetLayout().EventH / 2.0f;
 
-	TArray<FVector2D> Points;
-	Points.Emplace(0, 0);
-	Points.Emplace(X2 - X1, Y2 - Y1);
+	FVector2D StartPoint = FVector2D(X1, Y1);
+	FVector2D EndPoint = FVector2D(X2, Y2);
+	float Distance = FVector2D::Distance(StartPoint, EndPoint);
+
+	constexpr float DirectionFactor = 2.0f;
+	constexpr float LineLenghtAtEnds = 20.0f;
+	FVector2D StartDir((X2 - X1) / DirectionFactor, 0.0f);
 
 	FLinearColor RelationColors[(int32)ETaskGraphRelationType::Count] = {
 		FLinearColor::Yellow, // Created
@@ -42,9 +46,48 @@ void FTaskGraphRelation::Draw(const FDrawContext& DrawContext, const FTimingTrac
 		FLinearColor::Red, // Subsequent
 		FLinearColor::Yellow, // Completed
 	};
-
+	
 	FLinearColor Color = RelationColors[(int32)Type];
-	DrawContext.DrawLines(Helper.GetRelationLayerId(), X1, Y1, Points, ESlateDrawEffect::None, Color, /*bAntialias=*/ true, /*Thickness=*/ 2.0);
+	TArray<FVector2D> ArrowPoints;
+
+	if (Distance > LineLenghtAtEnds)
+	{
+		FVector2D SplineStart(StartPoint.X + LineLenghtAtEnds, StartPoint.Y);
+		FVector2D SplineEnd(EndPoint.X - LineLenghtAtEnds, EndPoint.Y);
+
+		DrawContext.DrawSpline(Helper.GetRelationLayerId(), 0.0f, 0.0f, SplineStart, StartDir, SplineEnd, StartDir, /*Thickness=*/ 2.0f, Color);
+
+		ArrowPoints.Add(StartPoint);
+		ArrowPoints.Add(SplineStart);
+
+		DrawContext.DrawLines(Helper.GetRelationLayerId(), 0.0f, 0.0f, ArrowPoints, ESlateDrawEffect::None, Color, /*bAntialias=*/ true, /*Thickness=*/ 2.0f);
+
+		ArrowPoints.Empty();
+		ArrowPoints.Add(SplineEnd);
+		ArrowPoints.Add(EndPoint);
+		DrawContext.DrawLines(Helper.GetRelationLayerId(), 0.0f, 0.0f, ArrowPoints, ESlateDrawEffect::None, Color, /*bAntialias=*/ true, /*Thickness=*/ 2.0f);
+	}
+	else
+	{
+		ArrowPoints.Empty();
+		ArrowPoints.Add(StartPoint);
+		ArrowPoints.Add(EndPoint);
+		DrawContext.DrawLines(Helper.GetRelationLayerId(), 0.0f, 0.0f, ArrowPoints, ESlateDrawEffect::None, Color, /*bAntialias=*/ true, /*Thickness=*/ 2.0f);
+	}
+
+	FVector2D ArrowOrigin = EndPoint;
+
+	constexpr float ArrowDirectionLen = -15.0f;
+	constexpr float ArrowRotationAngle = 20.0f;
+	FVector2D ArrowDirection(ArrowDirectionLen, 0.0f);
+
+	ArrowPoints.Empty();
+	ArrowPoints.Add(ArrowOrigin + ArrowDirection.GetRotated(ArrowRotationAngle));
+	ArrowPoints.Add(ArrowOrigin);
+	ArrowPoints.Add(ArrowOrigin + ArrowDirection.GetRotated(-ArrowRotationAngle));
+
+	constexpr float YOffset = 1.0f; // Needed to align the arrow with the line. 
+	DrawContext.DrawLines(Helper.GetRelationLayerId(), 0.0f, YOffset, ArrowPoints, ESlateDrawEffect::None, Color, /*bAntialias=*/ true, /*Thickness=*/ 2.0f);
 }
 
 } // namespace Insights
\ No newline at end of file
diff --git a/Engine/Source/Developer/TraceInsights/Private/Insights/ViewModels/ThreadTimingTrack.cpp b/Engine/Source/Developer/TraceInsights/Private/Insights/ViewModels/ThreadTimingTrack.cpp
--- a/Engine/Source/Developer/TraceInsights/Private/Insights/ViewModels/ThreadTimingTrack.cpp
+++ b/Engine/Source/Developer/TraceInsights/Private/Insights/ViewModels/ThreadTimingTrack.cpp
@@ -1071,7 +1071,11 @@ void FThreadTimingTrack::InitTooltip(FTooltipDrawState& InOutTooltip, const ITim
 				FString TrackName = Track.IsValid() ? Track->GetName() : TEXT("Unknown");
 				InOutTooltip.AddNameValueTextLine(TEXT("Launched:"), FString::Printf(TEXT("%f (+%s) on %s"), Task.LaunchedTimestamp, *TimeUtils::FormatTimeAuto(Task.LaunchedTimestamp - Task.CreatedTimestamp), *TrackName));
 			}
-			InOutTooltip.AddNameValueTextLine(TEXT("Scheduled:"), FString::Printf(TEXT("%f (+%s) on %s"), Task.ScheduledTimestamp, *TimeUtils::FormatTimeAuto(Task.ScheduledTimestamp - Task.LaunchedTimestamp), *SharedState.GetCpuTrack(Task.ScheduledThreadId)->GetName()));
+			{
+				TSharedPtr<FCpuTimingTrack> Track = SharedState.GetCpuTrack(Task.ScheduledThreadId);
+				FString TrackName = Track.IsValid() ? Track->GetName() : TEXT("Unknown");
+				InOutTooltip.AddNameValueTextLine(TEXT("Scheduled:"), FString::Printf(TEXT("%f (+%s) on %s"), Task.ScheduledTimestamp, *TimeUtils::FormatTimeAuto(Task.ScheduledTimestamp - Task.LaunchedTimestamp), *TrackName));
+			}
 			InOutTooltip.AddNameValueTextLine(TEXT("Started:"), FString::Printf(TEXT("%f (+%s)"), Task.StartedTimestamp, *TimeUtils::FormatTimeAuto(Task.StartedTimestamp - Task.ScheduledTimestamp)));
 			if (Task.FinishedTimestamp != TraceServices::FTaskInfo::InvalidTimestamp)
 			{
