[Insights] Disable Task Dependencies menu option when task graph profiling data does not exist in the session.

diff --git a/Engine/Source/Developer/TraceInsights/Private/Insights/TaskGraphProfiler/TaskGraphProfilerManager.cpp b/Engine/Source/Developer/TraceInsights/Private/Insights/TaskGraphProfiler/TaskGraphProfilerManager.cpp
new file mode 100644
--- /dev/null
+++ b/Engine/Source/Developer/TraceInsights/Private/Insights/TaskGraphProfiler/TaskGraphProfilerManager.cpp
@@ -0,0 +1,155 @@
+// Copyright Epic Games, Inc. All Rights Reserved.
+
+#include "TaskGraphProfilerManager.h"
+
+#include "Modules/ModuleManager.h"
+#include "TraceServices/AnalysisService.h"
+#include "TraceServices/Model/TasksProfiler.h"
+#include "Async/TaskGraphInterfaces.h"
+
+////////////////////////////////////////////////////////////////////////////////////////////////////
+
+#define LOCTEXT_NAMESPACE "TaskGraphProfilerManager"
+
+TSharedPtr<FTaskGraphProfilerManager> FTaskGraphProfilerManager::Instance = nullptr;
+
+////////////////////////////////////////////////////////////////////////////////////////////////////
+
+TSharedPtr<FTaskGraphProfilerManager> FTaskGraphProfilerManager::Get()
+{
+	return FTaskGraphProfilerManager::Instance;
+}
+
+////////////////////////////////////////////////////////////////////////////////////////////////////
+
+TSharedPtr<FTaskGraphProfilerManager> FTaskGraphProfilerManager::CreateInstance()
+{
+	ensure(!FTaskGraphProfilerManager::Instance.IsValid());
+	if (FTaskGraphProfilerManager::Instance.IsValid())
+	{
+		FTaskGraphProfilerManager::Instance.Reset();
+	}
+
+	FTaskGraphProfilerManager::Instance = MakeShared<FTaskGraphProfilerManager>(FInsightsManager::Get()->GetCommandList());
+
+	return FTaskGraphProfilerManager::Instance;
+}
+
+////////////////////////////////////////////////////////////////////////////////////////////////////
+
+FTaskGraphProfilerManager::FTaskGraphProfilerManager(TSharedRef<FUICommandList> InCommandList)
+	: bIsInitialized(false)
+	, bIsAvailable(false)
+{
+}
+
+////////////////////////////////////////////////////////////////////////////////////////////////////
+
+void FTaskGraphProfilerManager::Initialize(IUnrealInsightsModule& InsightsModule)
+{
+	ensure(!bIsInitialized);
+	if (bIsInitialized)
+	{
+		return;
+	}
+	bIsInitialized = true;
+
+	// Register tick functions.
+	OnTick = FTickerDelegate::CreateSP(this, &FTaskGraphProfilerManager::Tick);
+	OnTickHandle = FTicker::GetCoreTicker().AddTicker(OnTick, 0.0f);
+
+	FInsightsManager::Get()->GetSessionChangedEvent().AddSP(this, &FTaskGraphProfilerManager::OnSessionChanged);
+	OnSessionChanged();
+}
+
+////////////////////////////////////////////////////////////////////////////////////////////////////
+
+void FTaskGraphProfilerManager::Shutdown()
+{
+	if (!bIsInitialized)
+	{
+		return;
+	}
+	bIsInitialized = false;
+
+	FInsightsManager::Get()->GetSessionChangedEvent().RemoveAll(this);
+
+	// Unregister tick function.
+	FTicker::GetCoreTicker().RemoveTicker(OnTickHandle);
+
+	FTaskGraphProfilerManager::Instance.Reset();
+}
+
+////////////////////////////////////////////////////////////////////////////////////////////////////
+
+FTaskGraphProfilerManager::~FTaskGraphProfilerManager()
+{
+	ensure(!bIsInitialized);
+}
+
+////////////////////////////////////////////////////////////////////////////////////////////////////
+
+void FTaskGraphProfilerManager::RegisterMajorTabs(IUnrealInsightsModule& InsightsModule)
+{
+}
+
+////////////////////////////////////////////////////////////////////////////////////////////////////
+
+void FTaskGraphProfilerManager::UnregisterMajorTabs()
+{
+}
+
+////////////////////////////////////////////////////////////////////////////////////////////////////
+
+bool FTaskGraphProfilerManager::Tick(float DeltaTime)
+{
+	// Check if session has Memory events (to spawn the tab), but not too often.
+	if (!bIsAvailable && AvailabilityCheck.Tick())
+	{
+		bool bShouldBeAvailable = false;
+
+		TSharedPtr<const TraceServices::IAnalysisSession> Session = FInsightsManager::Get()->GetSession();
+		if (Session.IsValid())
+		{
+			TraceServices::FAnalysisSessionReadScope SessionReadScope(*Session.Get());
+
+			const TraceServices::ITasksProvider* TasksProvider = TraceServices::ReadTasksProvider(*Session.Get());
+			if (TasksProvider && TasksProvider->GetNumTasks() > 0)
+			{
+				bIsAvailable = true;
+			}
+
+			if (Session->IsAnalysisComplete())
+			{
+				// Never check again during this session.
+				AvailabilityCheck.Disable();
+			}
+		}
+		else
+		{
+			// Do not check again until the next session changed event (see OnSessionChanged).
+			AvailabilityCheck.Disable();
+		}
+	}
+
+	return true;
+}
+
+////////////////////////////////////////////////////////////////////////////////////////////////////
+
+void FTaskGraphProfilerManager::OnSessionChanged()
+{
+	bIsAvailable = false;
+	if (FInsightsManager::Get()->GetSession().IsValid())
+	{
+		AvailabilityCheck.Enable(0.5);
+	}
+	else
+	{
+		AvailabilityCheck.Disable();
+	}
+}
+
+////////////////////////////////////////////////////////////////////////////////////////////////////
+
+#undef LOCTEXT_NAMESPACE
diff --git a/Engine/Source/Developer/TraceInsights/Private/Insights/TaskGraphProfiler/TaskGraphProfilerManager.h b/Engine/Source/Developer/TraceInsights/Private/Insights/TaskGraphProfiler/TaskGraphProfilerManager.h
new file mode 100644
--- /dev/null
+++ b/Engine/Source/Developer/TraceInsights/Private/Insights/TaskGraphProfiler/TaskGraphProfilerManager.h
@@ -0,0 +1,69 @@
+// Copyright Epic Games, Inc. All Rights Reserved.
+
+#pragma once
+
+#include "Containers/Ticker.h"
+#include "CoreMinimal.h"
+#include "Framework/Commands/UICommandList.h"
+
+// Insights
+#include "Insights/InsightsManager.h"
+#include "Insights/IUnrealInsightsModule.h"
+
+////////////////////////////////////////////////////////////////////////////////////////////////////
+/**
+ * This class manages the Task Graph Profiler state and settings.
+ */
+class FTaskGraphProfilerManager : public TSharedFromThis<FTaskGraphProfilerManager>, public IInsightsComponent
+{
+public:
+	/** Creates the Memory Profiler manager, only one instance can exist. */
+	FTaskGraphProfilerManager(TSharedRef<FUICommandList> InCommandList);
+
+	/** Virtual destructor. */
+	virtual ~FTaskGraphProfilerManager();
+
+	/** Creates an instance of the Memory Profiler manager. */
+	static TSharedPtr<FTaskGraphProfilerManager> CreateInstance();
+
+	/**
+	 * @return the global instance of the Task Graph Profiler manager.
+	 * This is an internal singleton and cannot be used outside TraceInsights.
+	 * For external use:
+	 *     IUnrealInsightsModule& Module = FModuleManager::Get().LoadModuleChecked<IUnrealInsightsModule>("TraceInsights");
+	 *     Module.GetTaskGraphProfilerManager();
+	 */
+	static TSharedPtr<FTaskGraphProfilerManager> Get();
+
+	//////////////////////////////////////////////////
+	// IInsightsComponent
+
+	virtual void Initialize(IUnrealInsightsModule& InsightsModule) override;
+	virtual void Shutdown() override;
+	virtual void RegisterMajorTabs(IUnrealInsightsModule& InsightsModule) override;
+	virtual void UnregisterMajorTabs() override;
+	virtual void OnWindowClosedEvent() override {}
+
+	//////////////////////////////////////////////////
+	bool GetIsAvailable() { return bIsAvailable; }
+
+	void OnSessionChanged();
+
+private:
+	/** Updates this manager, done through FCoreTicker. */
+	bool Tick(float DeltaTime);
+
+private:
+	bool bIsInitialized;
+	bool bIsAvailable;
+	FAvailabilityCheck AvailabilityCheck;
+
+	/** The delegate to be invoked when this manager ticks. */
+	FTickerDelegate OnTick;
+
+	/** Handle to the registered OnTick. */
+	FDelegateHandle OnTickHandle;
+
+	/** A shared pointer to the global instance of the Task Graph Profiler manager. */
+	static TSharedPtr<FTaskGraphProfilerManager> Instance;
+};
diff --git a/Engine/Source/Developer/TraceInsights/Private/Insights/TraceInsightsModule.cpp b/Engine/Source/Developer/TraceInsights/Private/Insights/TraceInsightsModule.cpp
--- a/Engine/Source/Developer/TraceInsights/Private/Insights/TraceInsightsModule.cpp
+++ b/Engine/Source/Developer/TraceInsights/Private/Insights/TraceInsightsModule.cpp
@@ -20,6 +20,7 @@
 #include "Insights/Log.h"
 #include "Insights/MemoryProfiler/MemoryProfilerManager.h"
 #include "Insights/NetworkingProfiler/NetworkingProfilerManager.h"
+#include "Insights/TaskGraphProfiler/TaskGraphProfilerManager.h"
 #include "Insights/Tests/InsightsTestRunner.h"
 #include "Insights/TimingProfilerManager.h"
 #include "Insights/Widgets/SStartPageWindow.h"
@@ -52,6 +53,7 @@ void FTraceInsightsModule::StartupModule()
 	RegisterComponent(FLoadingProfilerManager::CreateInstance());
 	RegisterComponent(FNetworkingProfilerManager::CreateInstance());
 	RegisterComponent(FMemoryProfilerManager::CreateInstance());
+	RegisterComponent(FTaskGraphProfilerManager::CreateInstance());
 
 	UnrealInsightsLayoutIni = FConfigCacheIni::GetDestIniFilename(TEXT("UnrealInsightsLayout"), nullptr, *FPaths::GeneratedConfigDir());
 }
diff --git a/Engine/Source/Developer/TraceInsights/Private/Insights/Widgets/STimingView.cpp b/Engine/Source/Developer/TraceInsights/Private/Insights/Widgets/STimingView.cpp
--- a/Engine/Source/Developer/TraceInsights/Private/Insights/Widgets/STimingView.cpp
+++ b/Engine/Source/Developer/TraceInsights/Private/Insights/Widgets/STimingView.cpp
@@ -39,6 +39,7 @@
 #include "Insights/LoadingProfiler/LoadingProfilerManager.h"
 #include "Insights/LoadingProfiler/Widgets/SLoadingProfilerWindow.h"
 #include "Insights/Table/Widgets/STableTreeView.h"
+#include "Insights/TaskGraphProfiler/TaskGraphProfilerManager.h"
 #include "Insights/Tests/TimingProfilerTests.h"
 #include "Insights/TimingProfilerCommon.h"
 #include "Insights/TimingProfilerManager.h"
@@ -4126,14 +4127,14 @@ void STimingView::ContextMenu_ShowTaskDependecies_Execute()
 
 bool STimingView::ContextMenu_ShowTaskDependecies_CanExecute()
 {
-	return true;
+	return FTaskGraphProfilerManager::Get()->GetIsAvailable();
 }
 
 ////////////////////////////////////////////////////////////////////////////////////////////////////
 
 bool STimingView::ContextMenu_ShowTaskDependecies_IsChecked()
 {
-	return bShowEventRelations;
+	return FTaskGraphProfilerManager::Get()->GetIsAvailable() && bShowEventRelations;
 }
 
 ////////////////////////////////////////////////////////////////////////////////////////////////////
diff --git a/Engine/Source/Developer/TraceServices/Private/Model/TasksProfiler.cpp b/Engine/Source/Developer/TraceServices/Private/Model/TasksProfiler.cpp
--- a/Engine/Source/Developer/TraceServices/Private/Model/TasksProfiler.cpp
+++ b/Engine/Source/Developer/TraceServices/Private/Model/TasksProfiler.cpp
@@ -431,4 +431,9 @@ namespace TraceServices
 		const FWaitingForTasks& Waiting = (*Thread)[NextWaitingIndex - 1];
 		return Waiting.FinishedTimestamp > Timestamp || Waiting.FinishedTimestamp == FTaskInfo::InvalidTimestamp ? &Waiting : nullptr;
 	}
+
+	int64 FTasksProvider::GetNumTasks() const
+	{
+		return Tasks.Num();
+	}
 }
diff --git a/Engine/Source/Developer/TraceServices/Private/Model/TasksProfilerPrivate.h b/Engine/Source/Developer/TraceServices/Private/Model/TasksProfilerPrivate.h
--- a/Engine/Source/Developer/TraceServices/Private/Model/TasksProfilerPrivate.h
+++ b/Engine/Source/Developer/TraceServices/Private/Model/TasksProfilerPrivate.h
@@ -37,6 +37,7 @@ namespace TraceServices
 		virtual const FTaskInfo* TryGetTask(uint32 ThreadId, double Timestamp) const override;
 		virtual const FTaskInfo* TryGetTask(TaskTrace::FId TaskId) const;
 		virtual const FWaitingForTasks* TryGetWaiting(const TCHAR* TimerName, uint32 ThreadId, double Timestamp) const override;
+		virtual int64 GetNumTasks() const override;
 
 	private:
 		// initializes conversion of task id to its index in the internal container, if it's not initialized already
diff --git a/Engine/Source/Developer/TraceServices/Public/TraceServices/Model/TasksProfiler.h b/Engine/Source/Developer/TraceServices/Public/TraceServices/Model/TasksProfiler.h
--- a/Engine/Source/Developer/TraceServices/Public/TraceServices/Model/TasksProfiler.h
+++ b/Engine/Source/Developer/TraceServices/Public/TraceServices/Model/TasksProfiler.h
@@ -85,6 +85,9 @@ namespace TraceServices
 
 		// returns an info about waiting for tasks completion for given thread and timestamp, if any, otherwise `nullptr`
 		virtual const FWaitingForTasks* TryGetWaiting(const TCHAR* TimerName, uint32 ThreadId, double Timestamp) const = 0;
+
+		// returns the number of tasks stored in the provider
+		virtual int64 GetNumTasks() const = 0;
 	};
 
 	TRACESERVICES_API const ITasksProvider* ReadTasksProvider(const IAnalysisSession& Session);
