[Insights] Removed duplicated include.

diff --git a/Engine/Source/Developer/TraceInsights/Private/Insights/Tests/InsightsTestRunner.cpp b/Engine/Source/Developer/TraceInsights/Private/Insights/Tests/InsightsTestRunner.cpp
--- a/Engine/Source/Developer/TraceInsights/Private/Insights/Tests/InsightsTestRunner.cpp
+++ b/Engine/Source/Developer/TraceInsights/Private/Insights/Tests/InsightsTestRunner.cpp
@@ -19,8 +19,6 @@
 #include "Insights/Tests/UITests.h"
 #include "Insights/InsightsStyle.h"
 
-#include "Widgets/Docking/SDockTab.h"
-
 ////////////////////////////////////////////////////////////////////////////////////////////////////
 
 DEFINE_LOG_CATEGORY(InsightsTestRunner);
