[Insights] Fix event relation drawing over docked tracks

diff --git a/Engine/Source/Developer/TraceInsights/Private/Insights/ViewModels/TaskGraphRelation.cpp b/Engine/Source/Developer/TraceInsights/Private/Insights/ViewModels/TaskGraphRelation.cpp
--- a/Engine/Source/Developer/TraceInsights/Private/Insights/ViewModels/TaskGraphRelation.cpp
+++ b/Engine/Source/Developer/TraceInsights/Private/Insights/ViewModels/TaskGraphRelation.cpp
@@ -18,8 +18,28 @@ FTaskGraphRelation::FTaskGraphRelation(double InSourceTime, int32 InSourceThread
 	Type = InType;
 }
 
-void FTaskGraphRelation::Draw(const FDrawContext& DrawContext, const FTimingTrackViewport& Viewport, const ITimingViewDrawHelper& Helper)
+void FTaskGraphRelation::Draw(const FDrawContext& DrawContext, const FTimingTrackViewport& Viewport, const ITimingViewDrawHelper& Helper, const ITimingEventRelation::EDrawFilter Filter)
 {
+	int32 LayerId = Helper.GetRelationLayerId();
+
+	if (Filter == ITimingEventRelation::EDrawFilter::BetweenScrollableTracks)
+	{
+		if (SourceTrack.Pin()->GetLocation() > ETimingTrackLocation::Scrollable || TargetTrack.Pin()->GetLocation() > ETimingTrackLocation::Scrollable)
+		{
+			return;
+		}
+	}
+
+	if (Filter == ITimingEventRelation::EDrawFilter::BetweenDockedTracks)
+	{
+		if (SourceTrack.Pin()->GetLocation() <= ETimingTrackLocation::Scrollable && TargetTrack.Pin()->GetLocation() <= ETimingTrackLocation::Scrollable)
+		{
+			return;
+		}
+
+		LayerId = DrawContext.LayerId;
+	}
+
 	float X1 = Viewport.TimeToSlateUnitsRounded(SourceTime);
 	float Y1 = SourceTrack.Pin()->GetPosY();
 	Y1 += Viewport.GetLayout().GetLaneY(SourceDepth) + Viewport.GetLayout().EventH / 2.0f;
@@ -55,24 +75,24 @@ void FTaskGraphRelation::Draw(const FDrawContext& DrawContext, const FTimingTrac
 		FVector2D SplineStart(StartPoint.X + LineLenghtAtEnds, StartPoint.Y);
 		FVector2D SplineEnd(EndPoint.X - LineLenghtAtEnds, EndPoint.Y);
 
-		DrawContext.DrawSpline(Helper.GetRelationLayerId(), 0.0f, 0.0f, SplineStart, StartDir, SplineEnd, StartDir, /*Thickness=*/ 2.0f, Color);
+		DrawContext.DrawSpline(LayerId, 0.0f, 0.0f, SplineStart, StartDir, SplineEnd, StartDir, /*Thickness=*/ 2.0f, Color);
 
 		ArrowPoints.Add(StartPoint);
 		ArrowPoints.Add(SplineStart);
 
-		DrawContext.DrawLines(Helper.GetRelationLayerId(), 0.0f, 0.0f, ArrowPoints, ESlateDrawEffect::None, Color, /*bAntialias=*/ true, /*Thickness=*/ 2.0f);
+		DrawContext.DrawLines(LayerId, 0.0f, 0.0f, ArrowPoints, ESlateDrawEffect::None, Color, /*bAntialias=*/ true, /*Thickness=*/ 2.0f);
 
 		ArrowPoints.Empty();
 		ArrowPoints.Add(SplineEnd);
 		ArrowPoints.Add(EndPoint);
-		DrawContext.DrawLines(Helper.GetRelationLayerId(), 0.0f, 0.0f, ArrowPoints, ESlateDrawEffect::None, Color, /*bAntialias=*/ true, /*Thickness=*/ 2.0f);
+		DrawContext.DrawLines(LayerId, 0.0f, 0.0f, ArrowPoints, ESlateDrawEffect::None, Color, /*bAntialias=*/ true, /*Thickness=*/ 2.0f);
 	}
 	else
 	{
 		ArrowPoints.Empty();
 		ArrowPoints.Add(StartPoint);
 		ArrowPoints.Add(EndPoint);
-		DrawContext.DrawLines(Helper.GetRelationLayerId(), 0.0f, 0.0f, ArrowPoints, ESlateDrawEffect::None, Color, /*bAntialias=*/ true, /*Thickness=*/ 2.0f);
+		DrawContext.DrawLines(LayerId, 0.0f, 0.0f, ArrowPoints, ESlateDrawEffect::None, Color, /*bAntialias=*/ true, /*Thickness=*/ 2.0f);
 	}
 
 	FVector2D ArrowOrigin = EndPoint;
@@ -87,7 +107,7 @@ void FTaskGraphRelation::Draw(const FDrawContext& DrawContext, const FTimingTrac
 	ArrowPoints.Add(ArrowOrigin + ArrowDirection.GetRotated(-ArrowRotationAngle));
 
 	constexpr float YOffset = 1.0f; // Needed to align the arrow with the line. 
-	DrawContext.DrawLines(Helper.GetRelationLayerId(), 0.0f, YOffset, ArrowPoints, ESlateDrawEffect::None, Color, /*bAntialias=*/ true, /*Thickness=*/ 2.0f);
+	DrawContext.DrawLines(LayerId, 0.0f, YOffset, ArrowPoints, ESlateDrawEffect::None, Color, /*bAntialias=*/ true, /*Thickness=*/ 2.0f);
 }
 
 } // namespace Insights
\ No newline at end of file
diff --git a/Engine/Source/Developer/TraceInsights/Private/Insights/ViewModels/TaskGraphRelation.h b/Engine/Source/Developer/TraceInsights/Private/Insights/ViewModels/TaskGraphRelation.h
--- a/Engine/Source/Developer/TraceInsights/Private/Insights/ViewModels/TaskGraphRelation.h
+++ b/Engine/Source/Developer/TraceInsights/Private/Insights/ViewModels/TaskGraphRelation.h
@@ -29,7 +29,7 @@ public:
 
 	virtual bool IsSolved() override { return bIsSolved; }
 
-	virtual void Draw(const FDrawContext& DrawContext, const FTimingTrackViewport& Viewport, const ITimingViewDrawHelper& Helper) override;
+	virtual void Draw(const FDrawContext& DrawContext, const FTimingTrackViewport& Viewport, const ITimingViewDrawHelper& Helper, const ITimingEventRelation::EDrawFilter Filter) override;
 
 	void SetSourceTrack(TSharedPtr<const FBaseTimingTrack> InSourceTrack) { SourceTrack = InSourceTrack; }
 	TSharedPtr<const FBaseTimingTrack> GetSourceTrack() { return SourceTrack.Pin(); }
diff --git a/Engine/Source/Developer/TraceInsights/Private/Insights/ViewModels/TimingViewDrawHelper.cpp b/Engine/Source/Developer/TraceInsights/Private/Insights/ViewModels/TimingViewDrawHelper.cpp
--- a/Engine/Source/Developer/TraceInsights/Private/Insights/ViewModels/TimingViewDrawHelper.cpp
+++ b/Engine/Source/Developer/TraceInsights/Private/Insights/ViewModels/TimingViewDrawHelper.cpp
@@ -654,13 +654,13 @@ FLinearColor FTimingViewDrawHelper::GetTrackNameTextColor(const FBaseTimingTrack
 
 ////////////////////////////////////////////////////////////////////////////////////////////////////
 
-void FTimingViewDrawHelper::DrawRelations(const TArray<ITimingEventRelation*> Relations) const
+void FTimingViewDrawHelper::DrawRelations(const TArray<ITimingEventRelation*> Relations, ITimingEventRelation::EDrawFilter Filter) const
 {
 	for (ITimingEventRelation* Relation : Relations)
 	{
 		if (Relation->IsSolved())
 		{
-			Relation->Draw(DrawContext, Viewport, *this);
+			Relation->Draw(DrawContext, Viewport, *this, Filter);
 		}
 	}
 }
diff --git a/Engine/Source/Developer/TraceInsights/Private/Insights/ViewModels/TimingViewDrawHelper.h b/Engine/Source/Developer/TraceInsights/Private/Insights/ViewModels/TimingViewDrawHelper.h
--- a/Engine/Source/Developer/TraceInsights/Private/Insights/ViewModels/TimingViewDrawHelper.h
+++ b/Engine/Source/Developer/TraceInsights/Private/Insights/ViewModels/TimingViewDrawHelper.h
@@ -7,8 +7,9 @@
 #include "Math/Color.h"
 
 // Insights
-#include "Insights/ViewModels/TimingEventsTrack.h" // for ITimingEventsTrackDrawStateBuilder
+#include "Insights/ViewModels/ITimingEvent.h"
 #include "Insights/ViewModels/ITimingViewDrawHelper.h"
+#include "Insights/ViewModels/TimingEventsTrack.h" // for ITimingEventsTrackDrawStateBuilder
 
 ////////////////////////////////////////////////////////////////////////////////////////////////////
 
@@ -195,7 +196,7 @@ public:
 
 	void EndDrawTracks() const;
 
-	void DrawRelations(const TArray<ITimingEventRelation*> Relations) const;
+	void DrawRelations(const TArray<ITimingEventRelation*> Relations, ITimingEventRelation::EDrawFilter Filter) const;
 
 	//////////////////////////////////////////////////
 
diff --git a/Engine/Source/Developer/TraceInsights/Private/Insights/Widgets/STimingView.cpp b/Engine/Source/Developer/TraceInsights/Private/Insights/Widgets/STimingView.cpp
--- a/Engine/Source/Developer/TraceInsights/Private/Insights/Widgets/STimingView.cpp
+++ b/Engine/Source/Developer/TraceInsights/Private/Insights/Widgets/STimingView.cpp
@@ -1148,6 +1148,9 @@ int32 STimingView::OnPaint(const FPaintArgs& Args, const FGeometry& AllottedGeom
 		PreDrawTracksDurationHistory.AddValue(PreDrawTracksStopwatch.AccumulatedTime);
 	}
 
+	const FVector2D Position = AllottedGeometry.GetAbsolutePosition();
+	const float Scale = AllottedGeometry.GetAccumulatedLayoutTransform().GetScale();
+
 	//////////////////////////////////////////////////
 	// Draw
 	{
@@ -1156,9 +1159,6 @@ int32 STimingView::OnPaint(const FPaintArgs& Args, const FGeometry& AllottedGeom
 
 		Helper.BeginDrawTracks();
 
-		const FVector2D Position = AllottedGeometry.GetAbsolutePosition();
-		const float Scale = AllottedGeometry.GetAccumulatedLayoutTransform().GetScale();
-
 		// Draw the scrollable tracks.
 		{
 			const float TopY = Viewport.GetTopOffset();
@@ -1187,6 +1187,10 @@ int32 STimingView::OnPaint(const FPaintArgs& Args, const FGeometry& AllottedGeom
 				}
 			}
 
+			// Draw relations between scrollable tracks.
+			const FTimingViewDrawHelper& TimingHelper = *static_cast<const FTimingViewDrawHelper*>(&TimingDrawContext.GetHelper());
+			TimingHelper.DrawRelations(CurrentRelations, ITimingEventRelation::EDrawFilter::BetweenScrollableTracks);
+
 			DrawContext.ElementList.PopClip();
 		}
 
@@ -1298,9 +1302,6 @@ int32 STimingView::OnPaint(const FPaintArgs& Args, const FGeometry& AllottedGeom
 	// Draw the time range selection.
 	FDrawHelpers::DrawTimeRangeSelection(DrawContext, Viewport, SelectionStartTime, SelectionEndTime, WhiteBrush, MainFont);
 
-	const FTimingViewDrawHelper& TimingHelper = *static_cast<const FTimingViewDrawHelper*>(&TimingDrawContext.GetHelper());
-	TimingHelper.DrawRelations(CurrentRelations);
-
 	//////////////////////////////////////////////////
 	// Post-Draw
 	{
@@ -1402,6 +1403,30 @@ int32 STimingView::OnPaint(const FPaintArgs& Args, const FGeometry& AllottedGeom
 
 	//////////////////////////////////////////////////
 
+	// Draw relations between docked/foreground tracks.
+	if (GraphTrack->IsVisible())
+	{
+		const float TopY = GraphTrack->GetPosY() + GraphTrack->GetHeight();
+		const float BottomY = Viewport.GetHeight();
+
+		const float L = Position.X;
+		const float R = Position.X + (Viewport.GetWidth() * Scale);
+		const float T = Position.Y + (TopY * Scale);
+		const float B = Position.Y + (BottomY * Scale);
+		const FSlateClippingZone ClipZone(FVector2D(L, T), FVector2D(R, T), FVector2D(L, B), FVector2D(R, B));
+		DrawContext.ElementList.PushClip(ClipZone);
+	}
+
+	const FTimingViewDrawHelper& TimingHelper = *static_cast<const FTimingViewDrawHelper*>(&TimingDrawContext.GetHelper());
+	TimingHelper.DrawRelations(CurrentRelations, ITimingEventRelation::EDrawFilter::BetweenDockedTracks);
+
+	if (GraphTrack->IsVisible())
+	{
+		DrawContext.ElementList.PopClip();
+	}
+
+	//////////////////////////////////////////////////
+
 	const bool bShouldDisplayDebugInfo = FInsightsManager::Get()->IsDebugInfoEnabled();
 	if (bShouldDisplayDebugInfo)
 	{
diff --git a/Engine/Source/Developer/TraceInsights/Public/Insights/ViewModels/ITimingEvent.h b/Engine/Source/Developer/TraceInsights/Public/Insights/ViewModels/ITimingEvent.h
--- a/Engine/Source/Developer/TraceInsights/Public/Insights/ViewModels/ITimingEvent.h
+++ b/Engine/Source/Developer/TraceInsights/Public/Insights/ViewModels/ITimingEvent.h
@@ -70,10 +70,16 @@ public:
 class ITimingEventRelation
 {
 public:
+	enum class EDrawFilter
+	{
+		BetweenScrollableTracks, //  Only draw relations between 2 scrollable tracks.
+		BetweenDockedTracks, // Only draw relation if the source or target track are docked.
+	};
+
 	ITimingEventRelation() {}
 	virtual ~ITimingEventRelation() {}
 	virtual bool IsSolved() = 0;
-	virtual void Draw(const FDrawContext& DrawContext, const FTimingTrackViewport& Viewport, const ITimingViewDrawHelper& Helper) = 0;
+	virtual void Draw(const FDrawContext& DrawContext, const FTimingTrackViewport& Viewport, const ITimingViewDrawHelper& Helper, const EDrawFilter Filter) = 0;
 };
 
 ////////////////////////////////////////////////////////////////////////////////////////////////////
