TasksInsights: added info about prerequisites to task tooltips, limited the max number of task relations drawn to avoid UI stalls

diff --git a/Engine/Source/Developer/TraceInsights/Private/Insights/ViewModels/ThreadTimingTrack.cpp b/Engine/Source/Developer/TraceInsights/Private/Insights/ViewModels/ThreadTimingTrack.cpp
--- a/Engine/Source/Developer/TraceInsights/Private/Insights/ViewModels/ThreadTimingTrack.cpp
+++ b/Engine/Source/Developer/TraceInsights/Private/Insights/ViewModels/ThreadTimingTrack.cpp
@@ -1088,8 +1088,9 @@ void FThreadTimingTrack::InitTooltip(FTooltipDrawState& InOutTooltip, const ITim
 					InOutTooltip.AddNameValueTextLine(TEXT("Completed:"), FString::Printf(TEXT("%f (+%s) on %s"), Task.FinishedTimestamp, *TimeUtils::FormatTimeAuto(Task.CompletedTimestamp - Task.FinishedTimestamp), *TrackName));
 				}
 			}
-			InOutTooltip.AddNameValueTextLine(TEXT("Subsequent Tasks"), FString::Printf(TEXT("%d"), Task.Subsequents.Num()));
-			InOutTooltip.AddNameValueTextLine(TEXT("Nested Tasks:"), FString::Printf(TEXT("%d"), Task.NestedTasks.Num()));
+			InOutTooltip.AddNameValueTextLine(TEXT("Prerequisite tasks:"), FString::Printf(TEXT("%d"), Task.Prerequisites.Num()));
+			InOutTooltip.AddNameValueTextLine(TEXT("Subsequent tasks:"), FString::Printf(TEXT("%d"), Task.Subsequents.Num()));
+			InOutTooltip.AddNameValueTextLine(TEXT("Nested tasks:"), FString::Printf(TEXT("%d"), Task.NestedTasks.Num()));
 		};
 
 		do // info about a task
@@ -1117,19 +1118,31 @@ void FThreadTimingTrack::InitTooltip(FTooltipDrawState& InOutTooltip, const ITim
 			}
 
 			InOutTooltip.AddTextLine(TEXT("-------- Wating for tasks --------"), FLinearColor::Red);
-			FString TaskIdsStr = FString::JoinBy(Waiting->Tasks, TEXT(", "), [](TaskTrace::FId TaskId) { return FString::FromInt(TaskId); });
-			InOutTooltip.AddNameValueTextLine(TEXT("Tasks:"), FString::Printf(TEXT("[%s]"), *TaskIdsStr));
+			const int32 MaxWaitedTasksToList = 5;
+			FString TaskIdsStr = Waiting->Tasks.Num() <= MaxWaitedTasksToList ?
+				FString::JoinBy(Waiting->Tasks, TEXT(", "), [](TaskTrace::FId TaskId) { return FString::FromInt(TaskId); }) :
+				FString::Printf(TEXT("[%d]"), Waiting->Tasks.Num());
+			InOutTooltip.AddNameValueTextLine(TEXT("Tasks:"), FString::Printf(TEXT("%s"), *TaskIdsStr));
 			InOutTooltip.AddNameValueTextLine(TEXT("Started waiting:"), FString::Printf(TEXT("%f"), Waiting->StartedTimestamp));
-			InOutTooltip.AddNameValueTextLine(TEXT("Finished waiting:"), FString::Printf(TEXT("%f (+%s)"), Waiting->FinishedTimestamp, *TimeUtils::FormatTimeAuto(Waiting->FinishedTimestamp - Waiting->StartedTimestamp)));
+			InOutTooltip.AddNameValueTextLine(TEXT("Finished waiting:"), 
+				FString::Printf(TEXT("%s (+%s)"), 
+					Waiting->FinishedTimestamp == TraceServices::FTaskInfo::InvalidTimestamp ? TEXT("[not set]") : *FString::SanitizeFloat(Waiting->FinishedTimestamp),
+					*TimeUtils::FormatTimeAuto(Waiting->FinishedTimestamp - Waiting->StartedTimestamp)));
 
-			for (TaskTrace::FId TaskId : Waiting->Tasks)
+			int32 NumTasksToList = FMath::Min(Waiting->Tasks.Num(), MaxWaitedTasksToList);
+			for (int32 i = 0; i != NumTasksToList; ++i)
 			{
-				const TraceServices::FTaskInfo* Task = TasksProvider->TryGetTask(TaskId);
+				const TraceServices::FTaskInfo* Task = TasksProvider->TryGetTask(Waiting->Tasks[i]);
 				if (Task != nullptr)
 				{
 					AddTaskInfo(*Task);
 				}
 			}
+			if (NumTasksToList < Waiting->Tasks.Num())
+			{
+				InOutTooltip.AddTextLine(TEXT("[...]"), FLinearColor::Green);
+			}
+
 		} while (false);
 	}
 
@@ -1492,6 +1505,7 @@ bool FThreadTimingTrack::HasCustomFilter() const
 
 void FThreadTimingTrack::GetEventRelations(const ITimingEvent& InSelectedEvent, TArray<ITimingEventRelation*>& Relations) const
 {
+	const int32 MaxTasksToShow = 30;
 	double StartTime = InSelectedEvent.GetStartTime();
 
 	TSharedPtr<const TraceServices::IAnalysisSession> Session = FInsightsManager::Get()->GetSession();
@@ -1535,9 +1549,10 @@ void FThreadTimingTrack::GetEventRelations(const ITimingEvent& InSelectedEvent,
 
 			AddRelation(Task->LaunchedTimestamp, Task->LaunchedThreadId, Task->ScheduledTimestamp, Task->ScheduledThreadId, FTaskGraphRelation::ETaskGraphRelationType::Launched);
 
-			for (const TraceServices::FTaskInfo::FRelationInfo& RelationInfo : Task->Prerequisites)
+			int32 NumPrerequisitesToShow = FMath::Min(Task->Prerequisites.Num(), MaxTasksToShow);
+			for (int32 i = 0; i != NumPrerequisitesToShow; ++i)
 			{
-				const TraceServices::FTaskInfo* Prerequisite = TasksProvider->TryGetTask(RelationInfo.RelativeId);
+				const TraceServices::FTaskInfo* Prerequisite = TasksProvider->TryGetTask(Task->Prerequisites[i].RelativeId);
 				check(Prerequisite != nullptr);
 				AddRelation(Prerequisite->CompletedTimestamp, Prerequisite->CompletedThreadId, Task->StartedTimestamp, Task->StartedThreadId, FTaskGraphRelation::ETaskGraphRelationType::Prerequisite);
 			}
@@ -1547,8 +1562,10 @@ void FThreadTimingTrack::GetEventRelations(const ITimingEvent& InSelectedEvent,
 				AddRelation(Task->ScheduledTimestamp, Task->ScheduledThreadId, Task->StartedTimestamp, Task->StartedThreadId, FTaskGraphRelation::ETaskGraphRelationType::Scheduled);
 			}
 
-			for (const TraceServices::FTaskInfo::FRelationInfo& RelationInfo : Task->NestedTasks)
+			int32 NumNestedToShow = FMath::Min(Task->NestedTasks.Num(), MaxTasksToShow);
+			for (int32 i = 0; i != NumNestedToShow; ++i)
 			{
+				const TraceServices::FTaskInfo::FRelationInfo& RelationInfo = Task->NestedTasks[i];
 				const TraceServices::FTaskInfo* NestedTask = TasksProvider->TryGetTask(RelationInfo.RelativeId);
 				check(NestedTask != nullptr);
 
@@ -1557,9 +1574,10 @@ void FThreadTimingTrack::GetEventRelations(const ITimingEvent& InSelectedEvent,
 				AddRelation(NestedTask->CompletedTimestamp, NestedTask->CompletedThreadId, NestedTask->CompletedTimestamp, Task->StartedThreadId, FTaskGraphRelation::ETaskGraphRelationType::NestedCompleted);
 			}
 
-			for (const TraceServices::FTaskInfo::FRelationInfo& RelationInfo : Task->Subsequents)
+			int32 NumSubsequentsToShow = FMath::Min(Task->NestedTasks.Num(), MaxTasksToShow);
+			for (int32 i = 0; i != NumSubsequentsToShow; ++i)
 			{
-				const TraceServices::FTaskInfo* Subsequent = TasksProvider->TryGetTask(RelationInfo.RelativeId);
+				const TraceServices::FTaskInfo* Subsequent = TasksProvider->TryGetTask(Task->Subsequents[i].RelativeId);
 				check(Subsequent != nullptr);
 				AddRelation(Task->CompletedTimestamp, ThreadId, Subsequent->StartedTimestamp, Subsequent->StartedThreadId, FTaskGraphRelation::ETaskGraphRelationType::Subsequent);
 			}
@@ -1583,9 +1601,10 @@ void FThreadTimingTrack::GetEventRelations(const ITimingEvent& InSelectedEvent,
 		const TraceServices::FWaitingForTasks* Waiting = TasksProvider->TryGetWaiting(Timer->Name, ThreadId, StartTime);
 		if (Waiting != nullptr)
 		{
-			for (TaskTrace::FId TaskId : Waiting->Tasks)
+			int32 NumWaitedTasksToShow = FMath::Min(Waiting->Tasks.Num(), MaxTasksToShow);
+			for (int32 i = 0; i != NumWaitedTasksToShow; ++i)
 			{
-				const TraceServices::FTaskInfo* WaitedTask = TasksProvider->TryGetTask(TaskId);
+				const TraceServices::FTaskInfo* WaitedTask = TasksProvider->TryGetTask(Waiting->Tasks[i]);
 				if (WaitedTask != nullptr)
 				{
 					AddRelation(StartTime, ThreadId, WaitedTask->StartedTimestamp, WaitedTask->StartedThreadId, FTaskGraphRelation::ETaskGraphRelationType::AddedNested);
