[Insights] Fix color scheme for Session Browser and a few other places

diff --git a/Engine/Source/Developer/TraceInsights/Private/Insights/LoadingProfiler/Widgets/SLoadingProfilerWindow.cpp b/Engine/Source/Developer/TraceInsights/Private/Insights/LoadingProfiler/Widgets/SLoadingProfilerWindow.cpp
--- a/Engine/Source/Developer/TraceInsights/Private/Insights/LoadingProfiler/Widgets/SLoadingProfilerWindow.cpp
+++ b/Engine/Source/Developer/TraceInsights/Private/Insights/LoadingProfiler/Widgets/SLoadingProfilerWindow.cpp
@@ -2,11 +2,11 @@
 
 #include "SLoadingProfilerWindow.h"
 
-#include "EditorStyleSet.h"
 #include "Framework/Docking/LayoutService.h"
 #include "Framework/Docking/WorkspaceItem.h"
 #include "Framework/MultiBox/MultiBoxBuilder.h"
 #include "SlateOptMacros.h"
+#include "Styling/CoreStyle.h"
 #include "Widgets/Docking/SDockTab.h"
 #include "Widgets/Images/SImage.h"
 #include "Widgets/Layout/SBorder.h"
@@ -643,7 +643,7 @@ void SLoadingProfilerWindow::Construct(const FArguments& InArgs, const TSharedRe
 				[
 					SNew(SBorder)
 						.Visibility(this, &SLoadingProfilerWindow::IsSessionOverlayVisible)
-						.BorderImage(FEditorStyle::GetBrush("NotificationList.ItemBackground"))
+						.BorderImage(FCoreStyle::Get().GetBrush("PopupText.Background"))
 						.Padding(8.0f)
 						[
 							SNew(STextBlock)
diff --git a/Engine/Source/Developer/TraceInsights/Private/Insights/NetworkingProfiler/Widgets/SNetworkingProfilerWindow.cpp b/Engine/Source/Developer/TraceInsights/Private/Insights/NetworkingProfiler/Widgets/SNetworkingProfilerWindow.cpp
--- a/Engine/Source/Developer/TraceInsights/Private/Insights/NetworkingProfiler/Widgets/SNetworkingProfilerWindow.cpp
+++ b/Engine/Source/Developer/TraceInsights/Private/Insights/NetworkingProfiler/Widgets/SNetworkingProfilerWindow.cpp
@@ -2,11 +2,11 @@
 
 #include "SNetworkingProfilerWindow.h"
 
-#include "EditorStyleSet.h"
 #include "Framework/Commands/UICommandList.h"
 #include "Framework/Docking/WorkspaceItem.h"
 #include "Framework/MultiBox/MultiBoxBuilder.h"
 #include "SlateOptMacros.h"
+#include "Styling/CoreStyle.h"
 #include "Widgets/Docking/SDockTab.h"
 #include "Widgets/Images/SImage.h"
 #include "Widgets/Layout/SBorder.h"
@@ -349,7 +349,7 @@ void SNetworkingProfilerWindow::Construct(const FArguments& InArgs, const TShare
 				[
 					SNew(SBorder)
 						.Visibility(this, &SNetworkingProfilerWindow::IsSessionOverlayVisible)
-						.BorderImage(FEditorStyle::GetBrush("NotificationList.ItemBackground"))
+						.BorderImage(FCoreStyle::Get().GetBrush("PopupText.Background"))
 						.Padding(8.0f)
 						[
 							SNew(STextBlock)
diff --git a/Engine/Source/Developer/TraceInsights/Private/Insights/Widgets/SAsyncOperationStatus.cpp b/Engine/Source/Developer/TraceInsights/Private/Insights/Widgets/SAsyncOperationStatus.cpp
--- a/Engine/Source/Developer/TraceInsights/Private/Insights/Widgets/SAsyncOperationStatus.cpp
+++ b/Engine/Source/Developer/TraceInsights/Private/Insights/Widgets/SAsyncOperationStatus.cpp
@@ -2,7 +2,6 @@
 
 #include "SAsyncOperationStatus.h"
 
-#include "EditorStyleSet.h"
 #include "Styling/CoreStyle.h"
 #include "Widgets/Layout/SBorder.h"
 #include "Widgets/Layout/SBox.h"
@@ -29,7 +28,7 @@ void SAsyncOperationStatus::Construct(const FArguments& InArgs, TSharedRef<IAsyn
 		.Visibility(this, &SAsyncOperationStatus::GetContentVisibility)
 		[
 			SNew(SBorder)
-			.BorderImage(FEditorStyle::GetBrush("NotificationList.ItemBackground"))
+			.BorderImage(FCoreStyle::Get().GetBrush("PopupText.Background"))
 			.BorderBackgroundColor(this, &SAsyncOperationStatus::GetBackgroundColorAndOpacity)
 			.Padding(FMargin(16.0f, 8.0f, 16.0f, 8.0f))
 			.HAlign(HAlign_Fill)
diff --git a/Engine/Source/Developer/TraceInsights/Private/Insights/Widgets/SStartPageWindow.cpp b/Engine/Source/Developer/TraceInsights/Private/Insights/Widgets/SStartPageWindow.cpp
--- a/Engine/Source/Developer/TraceInsights/Private/Insights/Widgets/SStartPageWindow.cpp
+++ b/Engine/Source/Developer/TraceInsights/Private/Insights/Widgets/SStartPageWindow.cpp
@@ -591,7 +591,7 @@ void SStartPageWindow::Construct(const FArguments& InArgs)
 					SNew(SBox)
 					[
 						SNew(SBorder)
-						.BorderImage(FEditorStyle::GetBrush("NotificationList.ItemBackground"))
+						.BorderImage(FCoreStyle::Get().GetBrush("PopupText.Background"))
 						.Padding(8.0f)
 						[
 							ConstructSessionsPanel()
@@ -607,7 +607,7 @@ void SStartPageWindow::Construct(const FArguments& InArgs)
 					SNew(SBox)
 					[
 						SNew(SBorder)
-						.BorderImage(FEditorStyle::GetBrush("NotificationList.ItemBackground"))
+						.BorderImage(FCoreStyle::Get().GetBrush("PopupText.Background"))
 						.Padding(8.0f)
 						[
 							ConstructTraceStoreDirectoryPanel()
@@ -624,7 +624,7 @@ void SStartPageWindow::Construct(const FArguments& InArgs)
 					.Visibility(this, &SStartPageWindow::StopTraceRecorder_Visibility)
 					[
 						SNew(SBorder)
-						.BorderImage(FEditorStyle::GetBrush("NotificationList.ItemBackground"))
+						.BorderImage(FCoreStyle::Get().GetBrush("PopupText.Background"))
 						.Padding(8.0f)
 						[
 							ConstructConnectPanel()
@@ -643,7 +643,7 @@ void SStartPageWindow::Construct(const FArguments& InArgs)
 				.Visibility(this, &SStartPageWindow::SplashScreenOverlay_Visibility)
 				[
 					SNew(SBorder)
-					.BorderImage(FEditorStyle::GetBrush("NotificationList.ItemBackground"))
+					.BorderImage(FCoreStyle::Get().GetBrush("PopupText.Background"))
 					.BorderBackgroundColor(this, &SStartPageWindow::SplashScreenOverlay_ColorAndOpacity)
 					.Padding(0.0f)
 					.HAlign(HAlign_Fill)
@@ -1873,7 +1873,7 @@ void SStartPageWindow::OpenSettings()
 	(*OverlaySettingsSlot)
 	[
 		SNew(SBorder)
-		.BorderImage(FEditorStyle::GetBrush("NotificationList.ItemBackground"))
+		.BorderImage(FCoreStyle::Get().GetBrush("PopupText.Background"))
 		.Padding(8.0f)
 		[
 			SNew(SInsightsSettings)
diff --git a/Engine/Source/Developer/TraceInsights/Private/Insights/Widgets/STimingProfilerWindow.cpp b/Engine/Source/Developer/TraceInsights/Private/Insights/Widgets/STimingProfilerWindow.cpp
--- a/Engine/Source/Developer/TraceInsights/Private/Insights/Widgets/STimingProfilerWindow.cpp
+++ b/Engine/Source/Developer/TraceInsights/Private/Insights/Widgets/STimingProfilerWindow.cpp
@@ -2,12 +2,12 @@
 
 #include "STimingProfilerWindow.h"
 
-#include "EditorStyleSet.h"
 #include "Framework/Docking/LayoutExtender.h"
 #include "Framework/Docking/WorkspaceItem.h"
 #include "Framework/MultiBox/MultiBoxBuilder.h"
 #include "Framework/Docking/LayoutService.h"
 #include "SlateOptMacros.h"
+#include "Styling/CoreStyle.h"
 #include "Widgets/Docking/SDockTab.h"
 #include "Widgets/Images/SImage.h"
 #include "Widgets/Layout/SBorder.h"
@@ -595,7 +595,7 @@ void STimingProfilerWindow::Construct(const FArguments& InArgs, const TSharedRef
 				[
 					SNew(SBorder)
 						.Visibility(this, &STimingProfilerWindow::IsSessionOverlayVisible)
-						.BorderImage(FEditorStyle::GetBrush("NotificationList.ItemBackground"))
+						.BorderImage(FCoreStyle::Get().GetBrush("PopupText.Background"))
 						.Padding(8.0f)
 						[
 							SNew(STextBlock)
