[Insights]
  - Timers/Counters: Added warning/note message "Select a time region to update the aggregated statistics!" to the bottom of the panel.
  - Timers/Counters: Added auto-aggregation for entire session when analysis is completed. It also restores the default sorting (if auto disabled by the "10000 items" limit).
  - Timers/Counters: Improved layout of the top area of the panel (moved filtering "by type" buttons to same line with search and filter "by instance count").
  - Timing view: Fixed size of the yellow selected border for a timing event.

diff --git a/Engine/Source/Developer/TraceInsights/Private/Insights/ViewModels/StatsAggregator.h b/Engine/Source/Developer/TraceInsights/Private/Insights/ViewModels/StatsAggregator.h
--- a/Engine/Source/Developer/TraceInsights/Private/Insights/ViewModels/StatsAggregator.h
+++ b/Engine/Source/Developer/TraceInsights/Private/Insights/ViewModels/StatsAggregator.h
@@ -52,6 +52,7 @@ public:
 
 	double GetIntervalStartTime() const { return IntervalStartTime; }
 	double GetIntervalEndTime() const { return IntervalEndTime; }
+	bool IsEmptyTimeInterval() const { return IntervalStartTime >= IntervalEndTime; }
 
 	void SetTimeInterval(double InStartTime, double InEndTime)
 	{
diff --git a/Engine/Source/Developer/TraceInsights/Private/Insights/ViewModels/TimingViewDrawHelper.cpp b/Engine/Source/Developer/TraceInsights/Private/Insights/ViewModels/TimingViewDrawHelper.cpp
--- a/Engine/Source/Developer/TraceInsights/Private/Insights/ViewModels/TimingViewDrawHelper.cpp
+++ b/Engine/Source/Developer/TraceInsights/Private/Insights/ViewModels/TimingViewDrawHelper.cpp
@@ -622,12 +622,14 @@ void FTimingViewDrawHelper::DrawTimingEventHighlight(double StartTime, double En
 
 	const int32 LayerId = ReservedLayerId + ToInt32(EDrawLayer::EventHighlight);
 
+	constexpr float BorderOffset = 1.0f;
+
 	if (Mode == EDrawEventMode::Hovered)
 	{
 		const FLinearColor Color(1.0f, 1.0f, 0.0f, 1.0f); // yellow
 
 		// Draw border around the timing event box.
-		DrawContext.DrawBox(LayerId, EventX1 - 2.0f, Y - 2.0f, EventW + 4.0f, Layout.EventH + 4.0f, HoveredEventBorderBrush, Color);
+		DrawContext.DrawBox(LayerId, EventX1 - BorderOffset, Y - BorderOffset, EventW + 2 * BorderOffset, Layout.EventH + 2 * BorderOffset, HoveredEventBorderBrush, Color);
 	}
 	else // EDrawEventMode::Selected or EDrawEventMode::SelectedAndHovered
 	{
@@ -639,7 +641,7 @@ void FTimingViewDrawHelper::DrawTimingEventHighlight(double StartTime, double En
 		const FLinearColor Color(S, S, Blue, 1.0f);
 
 		// Draw border around the timing event box.
-		DrawContext.DrawBox(LayerId, EventX1 - 2.0f, Y - 2.0f, EventW + 4.0f, Layout.EventH + 4.0f, SelectedEventBorderBrush, Color);
+		DrawContext.DrawBox(LayerId, EventX1 - BorderOffset, Y - BorderOffset, EventW + 2 * BorderOffset, Layout.EventH + 2 * BorderOffset, SelectedEventBorderBrush, Color);
 	}
 }
 
diff --git a/Engine/Source/Developer/TraceInsights/Private/Insights/Widgets/SStatsView.cpp b/Engine/Source/Developer/TraceInsights/Private/Insights/Widgets/SStatsView.cpp
--- a/Engine/Source/Developer/TraceInsights/Private/Insights/Widgets/SStatsView.cpp
+++ b/Engine/Source/Developer/TraceInsights/Private/Insights/Widgets/SStatsView.cpp
@@ -21,6 +21,7 @@
 // Insights
 #include "Insights/Common/Stopwatch.h"
 #include "Insights/Common/TimeUtils.h"
+#include "Insights/InsightsStyle.h"
 #include "Insights/Table/ViewModels/Table.h"
 #include "Insights/Table/ViewModels/TableColumn.h"
 #include "Insights/TimingProfilerCommon.h"
@@ -93,6 +94,7 @@ SStatsView::~SStatsView()
 	if (FInsightsManager::Get().IsValid())
 	{
 		FInsightsManager::Get()->GetSessionChangedEvent().RemoveAll(this);
+		FInsightsManager::Get()->GetSessionAnalysisCompletedEvent().RemoveAll(this);
 	}
 
 	FStatsViewCommands::Unregister();
@@ -135,9 +137,9 @@ void SStatsView::Construct(const FArguments& InArgs)
 
 				// Search box
 				+ SHorizontalBox::Slot()
-				.VAlign(VAlign_Center)
 				.Padding(2.0f)
 				.FillWidth(1.0f)
+				.VAlign(VAlign_Center)
 				[
 					SAssignNew(SearchBox, SSearchBox)
 					.HintText(LOCTEXT("SearchBoxHint", "Search stats counters or groups"))
@@ -146,11 +148,29 @@ void SStatsView::Construct(const FArguments& InArgs)
 					.ToolTipText(LOCTEXT("FilterSearchHint", "Type here to search stats counter or group."))
 				]
 
-				// Filter out counters with zero instance count
+				// Filter by type (Int64)
+				+ SHorizontalBox::Slot()
+				.Padding(4.0f, 0.0f, 4.0f, 0.0f)
+				.AutoWidth()
+				.VAlign(VAlign_Center)
+				[
+					GetToggleButtonForDataType(EStatsNodeDataType::Int64)
+				]
+
+				// Filter by type (Double)
 				+ SHorizontalBox::Slot()
+				.Padding(4.0f, 0.0f, 4.0f, 0.0f)
+				.AutoWidth()
 				.VAlign(VAlign_Center)
+				[
+					GetToggleButtonForDataType(EStatsNodeDataType::Double)
+				]
+
+				// Filter out counters with zero instance count
+				+ SHorizontalBox::Slot()
 				.Padding(2.0f)
 				.AutoWidth()
+				.VAlign(VAlign_Center)
 				[
 					SNew(SCheckBox)
 					.Style(FCoreStyle::Get(), "ToggleButtonCheckbox")
@@ -175,7 +195,8 @@ void SStatsView::Construct(const FArguments& InArgs)
 				SNew(SHorizontalBox)
 
 				+ SHorizontalBox::Slot()
-				.FillWidth(1.0f)
+				.AutoWidth()
+				.Padding(0.0f, 0.0f, 4.0f, 0.0f)
 				.VAlign(VAlign_Center)
 				[
 					SNew(STextBlock)
@@ -183,49 +204,30 @@ void SStatsView::Construct(const FArguments& InArgs)
 				]
 
 				+ SHorizontalBox::Slot()
-				.FillWidth(2.0f)
+				.AutoWidth()
 				.VAlign(VAlign_Center)
 				[
-					SAssignNew(GroupByComboBox, SComboBox<TSharedPtr<EStatsGroupingMode>>)
-					.ToolTipText(this, &SStatsView::GroupBy_GetSelectedTooltipText)
-					.OptionsSource(&GroupByOptionsSource)
-					.OnSelectionChanged(this, &SStatsView::GroupBy_OnSelectionChanged)
-					.OnGenerateWidget(this, &SStatsView::GroupBy_OnGenerateWidget)
+					SNew(SBox)
+					.MinDesiredWidth(128.0f)
 					[
-						SNew(STextBlock)
-						.Text(this, &SStatsView::GroupBy_GetSelectedText)
+						SAssignNew(GroupByComboBox, SComboBox<TSharedPtr<EStatsGroupingMode>>)
+						.ToolTipText(this, &SStatsView::GroupBy_GetSelectedTooltipText)
+						.OptionsSource(&GroupByOptionsSource)
+						.OnSelectionChanged(this, &SStatsView::GroupBy_OnSelectionChanged)
+						.OnGenerateWidget(this, &SStatsView::GroupBy_OnGenerateWidget)
+						[
+							SNew(STextBlock)
+							.Text(this, &SStatsView::GroupBy_GetSelectedText)
+						]
 					]
 				]
 			]
-
-			// Check boxes for: Int64, Float
-			+ SVerticalBox::Slot()
-			.VAlign(VAlign_Center)
-			.Padding(2.0f)
-			.AutoHeight()
-			[
-				SNew(SHorizontalBox)
-
-				+ SHorizontalBox::Slot()
-				.Padding(FMargin(0.0f,0.0f,1.0f,0.0f))
-				.FillWidth(1.0f)
-				[
-					GetToggleButtonForDataType(EStatsNodeDataType::Int64)
-				]
-
-				+ SHorizontalBox::Slot()
-				.Padding(FMargin(1.0f,0.0f,1.0f,0.0f))
-				.FillWidth(1.0f)
-				[
-					GetToggleButtonForDataType(EStatsNodeDataType::Double)
-				]
-			]
 		]
 
 		// Tree view
 		+ SVerticalBox::Slot()
 		.FillHeight(1.0f)
-		.Padding(0.0f, 6.0f, 0.0f, 0.0f)
+		.Padding(0.0f, 2.0f, 0.0f, 0.0f)
 		[
 			SNew(SHorizontalBox)
 
@@ -282,6 +284,25 @@ void SStatsView::Construct(const FArguments& InArgs)
 				]
 			]
 		]
+
+		// Status bar
+		+ SVerticalBox::Slot()
+		.AutoHeight()
+		.HAlign(HAlign_Fill)
+		.Padding(0.0f)
+		[
+			SNew(SBorder)
+			.BorderImage(FInsightsStyle::Get().GetBrush("WhiteBrush"))
+			.BorderBackgroundColor(FLinearColor(0.05f, 0.1f, 0.2f, 1.0f))
+			.HAlign(HAlign_Center)
+			[
+				SNew(STextBlock)
+				.Margin(FMargin(4.0f, 1.0f, 4.0f, 1.0f))
+				.Text(LOCTEXT("EmptyAggregationNote", "-- Select a time region to update the aggregated statistics! --"))
+				.ColorAndOpacity(FLinearColor(1.0f, 0.75f, 0.5f, 1.0f))
+				.Visibility_Lambda([this]() { return Aggregator->IsEmptyTimeInterval() && !Aggregator->IsRunning() ? EVisibility::Visible : EVisibility::Collapsed; })
+			]
+		]
 	];
 
 	InitializeAndShowHeaderColumns();
@@ -299,6 +320,7 @@ void SStatsView::Construct(const FArguments& InArgs)
 
 	// Register ourselves with the Insights manager.
 	FInsightsManager::Get()->GetSessionChangedEvent().AddSP(this, &SStatsView::InsightsManager_OnSessionChanged);
+	FInsightsManager::Get()->GetSessionAnalysisCompletedEvent().AddSP(this, &SStatsView::InsightsManager_OnSessionAnalysisCompleted);
 
 	// Update the Session (i.e. when analysis session was already started).
 	InsightsManager_OnSessionChanged();
@@ -702,6 +724,31 @@ void SStatsView::InsightsManager_OnSessionChanged()
 
 ////////////////////////////////////////////////////////////////////////////////////////////////////
 
+void SStatsView::InsightsManager_OnSessionAnalysisCompleted()
+{
+	if (Aggregator->IsEmptyTimeInterval() && !Aggregator->IsRunning())
+	{
+		TSharedPtr<FInsightsManager> InsightsManager = FInsightsManager::Get();
+		InsightsManager->UpdateSessionDuration();
+		const double SessionDuration = InsightsManager->GetSessionDuration();
+
+		constexpr double Delta = 0.0; // session padding
+
+		Aggregator->Cancel();
+		Aggregator->SetTimeInterval(0.0 - Delta, SessionDuration + Delta);
+		Aggregator->Start();
+
+		if (ColumnBeingSorted == NAME_None)
+		{
+			// Restore sorting...
+			SetSortModeForColumn(GetDefaultColumnBeingSorted(), GetDefaultColumnSortMode());
+			TreeView_Refresh();
+		}
+	}
+}
+
+////////////////////////////////////////////////////////////////////////////////////////////////////
+
 void SStatsView::UpdateTree()
 {
 	FStopwatch Stopwatch;
@@ -890,8 +937,8 @@ TSharedRef<SWidget> SStatsView::GetToggleButtonForNodeType(const EStatsNodeType
 {
 	return SNew(SCheckBox)
 		.Style(FCoreStyle::Get(), "ToggleButtonCheckbox")
+		.Padding(FMargin(4.0f, 2.0f, 4.0f, 2.0f))
 		.HAlign(HAlign_Center)
-		.Padding(2.0f)
 		.OnCheckStateChanged(this, &SStatsView::FilterByStatsType_OnCheckStateChanged, NodeType)
 		.IsChecked(this, &SStatsView::FilterByStatsType_IsChecked, NodeType)
 		.ToolTipText(StatsNodeTypeHelper::ToDescription(NodeType))
@@ -901,13 +948,13 @@ TSharedRef<SWidget> SStatsView::GetToggleButtonForNodeType(const EStatsNodeType
 			+ SHorizontalBox::Slot()
 			.AutoWidth()
 			.VAlign(VAlign_Center)
+			.Padding(0.0f, 0.0f, 2.0f, 0.0f)
 			[
 				SNew(SImage)
 				.Image(StatsNodeTypeHelper::GetIcon(NodeType))
 			]
 
 			+ SHorizontalBox::Slot()
-			.Padding(2.0f, 0.0f, 0.0f, 0.0f)
 			.VAlign(VAlign_Center)
 			[
 				SNew(STextBlock)
@@ -922,8 +969,8 @@ TSharedRef<SWidget> SStatsView::GetToggleButtonForDataType(const EStatsNodeDataT
 {
 	return SNew(SCheckBox)
 		.Style(FCoreStyle::Get(), "ToggleButtonCheckbox")
+		.Padding(FMargin(4.0f, 2.0f, 4.0f, 2.0f))
 		.HAlign(HAlign_Center)
-		.Padding(2.0f)
 		.OnCheckStateChanged(this, &SStatsView::FilterByStatsDataType_OnCheckStateChanged, DataType)
 		.IsChecked(this, &SStatsView::FilterByStatsDataType_IsChecked, DataType)
 		.ToolTipText(StatsNodeDataTypeHelper::ToDescription(DataType))
@@ -931,6 +978,7 @@ TSharedRef<SWidget> SStatsView::GetToggleButtonForDataType(const EStatsNodeDataT
 			SNew(SHorizontalBox)
 
 			//+ SHorizontalBox::Slot()
+			//.Padding(0.0f, 0.0f, 2.0f, 0.0f)
 			//.AutoWidth()
 			//.VAlign(VAlign_Center)
 			//[
@@ -939,7 +987,6 @@ TSharedRef<SWidget> SStatsView::GetToggleButtonForDataType(const EStatsNodeDataT
 			//]
 
 			+ SHorizontalBox::Slot()
-			.Padding(2.0f, 0.0f, 0.0f, 0.0f)
 			.VAlign(VAlign_Center)
 			[
 				SNew(STextBlock)
diff --git a/Engine/Source/Developer/TraceInsights/Private/Insights/Widgets/SStatsView.h b/Engine/Source/Developer/TraceInsights/Private/Insights/Widgets/SStatsView.h
--- a/Engine/Source/Developer/TraceInsights/Private/Insights/Widgets/SStatsView.h
+++ b/Engine/Source/Developer/TraceInsights/Private/Insights/Widgets/SStatsView.h
@@ -97,9 +97,12 @@ private:
 
 	void FinishAggregation();
 
-	/** Called when the analysis session has changed. */
+	/** Called when the session has changed. */
 	void InsightsManager_OnSessionChanged();
 
+	/** Called when the analysis was completed. */
+	void InsightsManager_OnSessionAnalysisCompleted();
+
 	/**
 	 * Populates OutSearchStrings with the strings that should be used in searching.
 	 *
diff --git a/Engine/Source/Developer/TraceInsights/Private/Insights/Widgets/STimersView.cpp b/Engine/Source/Developer/TraceInsights/Private/Insights/Widgets/STimersView.cpp
--- a/Engine/Source/Developer/TraceInsights/Private/Insights/Widgets/STimersView.cpp
+++ b/Engine/Source/Developer/TraceInsights/Private/Insights/Widgets/STimersView.cpp
@@ -21,6 +21,7 @@
 // Insights
 #include "Insights/Common/Stopwatch.h"
 #include "Insights/Common/TimeUtils.h"
+#include "Insights/InsightsStyle.h"
 #include "Insights/Table/ViewModels/Table.h"
 #include "Insights/Table/ViewModels/TableColumn.h"
 #include "Insights/TimingProfilerCommon.h"
@@ -93,6 +94,7 @@ STimersView::~STimersView()
 	if (FInsightsManager::Get().IsValid())
 	{
 		FInsightsManager::Get()->GetSessionChangedEvent().RemoveAll(this);
+		FInsightsManager::Get()->GetSessionAnalysisCompletedEvent().RemoveAll(this);
 	}
 
 	FTimersViewCommands::Unregister();
@@ -135,9 +137,9 @@ void STimersView::Construct(const FArguments& InArgs)
 
 				// Search box
 				+ SHorizontalBox::Slot()
-				.VAlign(VAlign_Center)
 				.Padding(2.0f)
 				.FillWidth(1.0f)
+				.VAlign(VAlign_Center)
 				[
 					SAssignNew(SearchBox, SSearchBox)
 					.HintText(LOCTEXT("SearchBoxHint", "Search timers or groups"))
@@ -146,11 +148,38 @@ void STimersView::Construct(const FArguments& InArgs)
 					.ToolTipText(LOCTEXT("FilterSearchHint", "Type here to search timer or group."))
 				]
 
-				// Filter out timers with zero instance count
+				// Filter by type (Gpu)
 				+ SHorizontalBox::Slot()
+				.Padding(4.0f, 0.0f, 4.0f, 0.0f)
+				.AutoWidth()
 				.VAlign(VAlign_Center)
+				[
+					GetToggleButtonForTimerType(ETimerNodeType::GpuScope)
+				]
+
+				// Filter by type (Compute)
+				//+ SHorizontalBox::Slot()
+				//.Padding(4.0f, 0.0f, 4.0f, 0.0f)
+				//.AutoWidth()
+				//.VAlign(VAlign_Center)
+				//[
+				//	GetToggleButtonForTimerType(ETimerNodeType::ComputeScope)
+				//]
+
+				// Filter by type (Cpu)
+				+ SHorizontalBox::Slot()
+				.Padding(4.0f, 0.0f, 4.0f, 0.0f)
+				.AutoWidth()
+				.VAlign(VAlign_Center)
+				[
+					GetToggleButtonForTimerType(ETimerNodeType::CpuScope)
+				]
+
+				// Filter out timers with zero instance count
+				+ SHorizontalBox::Slot()
 				.Padding(2.0f)
 				.AutoWidth()
+				.VAlign(VAlign_Center)
 				[
 					SNew(SCheckBox)
 					.Style(FCoreStyle::Get(), "ToggleButtonCheckbox")
@@ -175,7 +204,8 @@ void STimersView::Construct(const FArguments& InArgs)
 				SNew(SHorizontalBox)
 
 				+ SHorizontalBox::Slot()
-				.FillWidth(1.0f)
+				.AutoWidth()
+				.Padding(0.0f, 0.0f, 4.0f, 0.0f)
 				.VAlign(VAlign_Center)
 				[
 					SNew(STextBlock)
@@ -183,56 +213,30 @@ void STimersView::Construct(const FArguments& InArgs)
 				]
 
 				+ SHorizontalBox::Slot()
-				.FillWidth(2.0f)
+				.AutoWidth()
 				.VAlign(VAlign_Center)
 				[
-					SAssignNew(GroupByComboBox, SComboBox<TSharedPtr<ETimerGroupingMode>>)
-					.ToolTipText(this, &STimersView::GroupBy_GetSelectedTooltipText)
-					.OptionsSource(&GroupByOptionsSource)
-					.OnSelectionChanged(this, &STimersView::GroupBy_OnSelectionChanged)
-					.OnGenerateWidget(this, &STimersView::GroupBy_OnGenerateWidget)
+					SNew(SBox)
+					.MinDesiredWidth(128.0f)
 					[
-						SNew(STextBlock)
-						.Text(this, &STimersView::GroupBy_GetSelectedText)
+						SAssignNew(GroupByComboBox, SComboBox<TSharedPtr<ETimerGroupingMode>>)
+						.ToolTipText(this, &STimersView::GroupBy_GetSelectedTooltipText)
+						.OptionsSource(&GroupByOptionsSource)
+						.OnSelectionChanged(this, &STimersView::GroupBy_OnSelectionChanged)
+						.OnGenerateWidget(this, &STimersView::GroupBy_OnGenerateWidget)
+						[
+							SNew(STextBlock)
+							.Text(this, &STimersView::GroupBy_GetSelectedText)
+						]
 					]
 				]
 			]
-
-			// Check boxes for: GpuScope, ComputeScope, CpuScope
-			+ SVerticalBox::Slot()
-			.VAlign(VAlign_Center)
-			.Padding(2.0f)
-			.AutoHeight()
-			[
-				SNew(SHorizontalBox)
-
-				+ SHorizontalBox::Slot()
-				.Padding(FMargin(0.0f,0.0f,1.0f,0.0f))
-				.FillWidth(1.0f)
-				[
-					GetToggleButtonForTimerType(ETimerNodeType::GpuScope)
-				]
-
-				//+ SHorizontalBox::Slot()
-				//.Padding(FMargin(1.0f,0.0f,1.0f,0.0f))
-				//.FillWidth(1.0f)
-				//[
-				//	GetToggleButtonForTimerType(ETimerNodeType::ComputeScope)
-				//]
-
-				+ SHorizontalBox::Slot()
-				.Padding(FMargin(1.0f,0.0f,1.0f,0.0f))
-				.FillWidth(1.0f)
-				[
-					GetToggleButtonForTimerType(ETimerNodeType::CpuScope)
-				]
-			]
 		]
 
 		// Tree view
 		+ SVerticalBox::Slot()
 		.FillHeight(1.0f)
-		.Padding(0.0f, 6.0f, 0.0f, 0.0f)
+		.Padding(0.0f, 2.0f, 0.0f, 0.0f)
 		[
 			SNew(SHorizontalBox)
 
@@ -289,6 +293,25 @@ void STimersView::Construct(const FArguments& InArgs)
 				]
 			]
 		]
+
+		// Status bar
+		+ SVerticalBox::Slot()
+		.AutoHeight()
+		.HAlign(HAlign_Fill)
+		.Padding(0.0f)
+		[
+			SNew(SBorder)
+			.BorderImage(FInsightsStyle::Get().GetBrush("WhiteBrush"))
+			.BorderBackgroundColor(FLinearColor(0.05f, 0.1f, 0.2f, 1.0f))
+			.HAlign(HAlign_Center)
+			[
+				SNew(STextBlock)
+				.Margin(FMargin(4.0f, 1.0f, 4.0f, 1.0f))
+				.Text(LOCTEXT("EmptyAggregationNote", "-- Select a time region to update the aggregated statistics! --"))
+				.ColorAndOpacity(FLinearColor(1.0f, 0.75f, 0.5f, 1.0f))
+				.Visibility_Lambda([this]() { return Aggregator->IsEmptyTimeInterval() && !Aggregator->IsRunning() ? EVisibility::Visible : EVisibility::Collapsed; })
+			]
+		]
 	];
 
 	InitializeAndShowHeaderColumns();
@@ -306,6 +329,7 @@ void STimersView::Construct(const FArguments& InArgs)
 
 	// Register ourselves with the Insights manager.
 	FInsightsManager::Get()->GetSessionChangedEvent().AddSP(this, &STimersView::InsightsManager_OnSessionChanged);
+	FInsightsManager::Get()->GetSessionAnalysisCompletedEvent().AddSP(this, &STimersView::InsightsManager_OnSessionAnalysisCompleted);
 
 	// Update the Session (i.e. when analysis session was already started).
 	InsightsManager_OnSessionChanged();
@@ -741,6 +765,31 @@ void STimersView::InsightsManager_OnSessionChanged()
 
 ////////////////////////////////////////////////////////////////////////////////////////////////////
 
+void STimersView::InsightsManager_OnSessionAnalysisCompleted()
+{
+	if (Aggregator->IsEmptyTimeInterval() && !Aggregator->IsRunning())
+	{
+		TSharedPtr<FInsightsManager> InsightsManager = FInsightsManager::Get();
+		InsightsManager->UpdateSessionDuration();
+		const double SessionDuration = InsightsManager->GetSessionDuration();
+
+		constexpr double Delta = 0.0; // session padding
+
+		Aggregator->Cancel();
+		Aggregator->SetTimeInterval(0.0 - Delta, SessionDuration + Delta);
+		Aggregator->Start();
+
+		if (ColumnBeingSorted == NAME_None)
+		{
+			// Restore sorting...
+			SetSortModeForColumn(GetDefaultColumnBeingSorted(), GetDefaultColumnSortMode());
+			TreeView_Refresh();
+		}
+	}
+}
+
+////////////////////////////////////////////////////////////////////////////////////////////////////
+
 void STimersView::UpdateTree()
 {
 	FStopwatch Stopwatch;
@@ -889,8 +938,8 @@ TSharedRef<SWidget> STimersView::GetToggleButtonForTimerType(const ETimerNodeTyp
 {
 	return SNew(SCheckBox)
 		.Style(FCoreStyle::Get(), "ToggleButtonCheckbox")
+		.Padding(FMargin(4.0f, 2.0f, 4.0f, 2.0f))
 		.HAlign(HAlign_Center)
-		.Padding(2.0f)
 		.OnCheckStateChanged(this, &STimersView::FilterByTimerType_OnCheckStateChanged, NodeType)
 		.IsChecked(this, &STimersView::FilterByTimerType_IsChecked, NodeType)
 		.ToolTipText(TimerNodeTypeHelper::ToDescription(NodeType))
@@ -898,6 +947,7 @@ TSharedRef<SWidget> STimersView::GetToggleButtonForTimerType(const ETimerNodeTyp
 			SNew(SHorizontalBox)
 
 			//+ SHorizontalBox::Slot()
+			//.Padding(0.0f, 0.0f, 2.0f, 0.0f)
 			//.AutoWidth()
 			//.VAlign(VAlign_Center)
 			//[
@@ -906,7 +956,6 @@ TSharedRef<SWidget> STimersView::GetToggleButtonForTimerType(const ETimerNodeTyp
 			//]
 
 			+ SHorizontalBox::Slot()
-			.Padding(2.0f, 0.0f, 0.0f, 0.0f)
 			.VAlign(VAlign_Center)
 			[
 				SNew(STextBlock)
diff --git a/Engine/Source/Developer/TraceInsights/Private/Insights/Widgets/STimersView.h b/Engine/Source/Developer/TraceInsights/Private/Insights/Widgets/STimersView.h
--- a/Engine/Source/Developer/TraceInsights/Private/Insights/Widgets/STimersView.h
+++ b/Engine/Source/Developer/TraceInsights/Private/Insights/Widgets/STimersView.h
@@ -96,9 +96,12 @@ private:
 	void FinishAggregation();
 	void ApplyAggregation(TraceServices::ITable<TraceServices::FTimingProfilerAggregatedStats>* AggregatedStatsTable);
 
-	/** Called when the analysis session has changed. */
+	/** Called when the session has changed. */
 	void InsightsManager_OnSessionChanged();
 
+	/** Called when the analysis was completed. */
+	void InsightsManager_OnSessionAnalysisCompleted();
+
 	/**
 	 * Populates OutSearchStrings with the strings that should be used in searching.
 	 *
