[Insights] Changed "Task Dependencies" action to appear only in the context menu of the "timing events" tracks.

diff --git a/Engine/Source/Developer/TraceInsights/Private/Insights/Widgets/STimingView.cpp b/Engine/Source/Developer/TraceInsights/Private/Insights/Widgets/STimingView.cpp
--- a/Engine/Source/Developer/TraceInsights/Private/Insights/Widgets/STimingView.cpp
+++ b/Engine/Source/Developer/TraceInsights/Private/Insights/Widgets/STimingView.cpp
@@ -2719,6 +2719,7 @@ void STimingView::ShowContextMenu(const FPointerEvent& MouseEvent)
 		MenuBuilder.EndSection();
 	}
 
+	if (SelectedTrack.IsValid() && SelectedTrack->Is<FTimingEventsTrack>())
 	{
 		MenuBuilder.BeginSection(TEXT("Event"), LOCTEXT("Event", "Event"));
 
