[Insights] Memory Insights: Changed the default value for A, B, C, D time markers to 10s, 20s, 30s, 40s (instead of 0, 1s, 2s, 3s).

diff --git a/Engine/Source/Developer/TraceInsights/Private/Insights/MemoryProfiler/Widgets/SMemoryProfilerWindow.cpp b/Engine/Source/Developer/TraceInsights/Private/Insights/MemoryProfiler/Widgets/SMemoryProfilerWindow.cpp
--- a/Engine/Source/Developer/TraceInsights/Private/Insights/MemoryProfiler/Widgets/SMemoryProfilerWindow.cpp
+++ b/Engine/Source/Developer/TraceInsights/Private/Insights/MemoryProfiler/Widgets/SMemoryProfilerWindow.cpp
@@ -150,7 +150,7 @@ void SMemoryProfilerWindow::ResetTimingViewMarkers()
 		const FLinearColor Color = FLinearColor::MakeFromHSV8(H, S, V);
 		TimeMarker->SetColor(Color);
 
-		TimeMarker->SetTime(static_cast<float>(Index)); // 0.0f, 1.0f, 2.0f, etc.
+		TimeMarker->SetTime((Index + 1) * 10.0); // 10s, 20s, 30s, etc.
 
 		TimeRulerTrack->AddTimeMarker(TimeMarker);
 		CustomTimeMarkers.Add(TimeMarker);
