[Insights] Further improved string handling (ex. in time formating functions) by avoiding some hidden strlen calls.

diff --git a/Engine/Source/Developer/TraceInsights/Private/Insights/Common/TimeUtils.cpp b/Engine/Source/Developer/TraceInsights/Private/Insights/Common/TimeUtils.cpp
--- a/Engine/Source/Developer/TraceInsights/Private/Insights/Common/TimeUtils.cpp
+++ b/Engine/Source/Developer/TraceInsights/Private/Insights/Common/TimeUtils.cpp
@@ -212,7 +212,7 @@ FString FormatTimeAuto(const double InDuration, const int32 NumDigits)
 		}
 	}
 
-	return StrBuilder.ToString();
+	return FString(StrBuilder);
 }
 
 ////////////////////////////////////////////////////////////////////////////////////////////////////
@@ -266,7 +266,7 @@ FString FormatTimeMs(const double InDuration, const int32 NumDigits, bool bAddTi
 		}
 	}
 
-	return StrBuilder.ToString();
+	return FString(StrBuilder);
 }
 
 ////////////////////////////////////////////////////////////////////////////////////////////////////
@@ -295,7 +295,7 @@ FString FormatTime(const double InTime, const double Precision)
 	if (Time == DBL_MAX || Time == std::numeric_limits<double>::infinity())
 	{
 		StrBuilder.Append(TEXT('∞'));
-		return StrBuilder.ToString();
+		return FString(StrBuilder);
 	}
 
 	bool bIsSpaceNeeded = false;
@@ -314,7 +314,7 @@ FString FormatTime(const double InTime, const double Precision)
 		{
 			return TEXT("~0");
 		}
-		return StrBuilder.ToString();
+		return FString(StrBuilder);
 	}
 
 	int32 Hours = static_cast<int32>(Time / TimeUtils::Hour);
@@ -335,7 +335,7 @@ FString FormatTime(const double InTime, const double Precision)
 		{
 			return TEXT("~0");
 		}
-		return StrBuilder.ToString();
+		return FString(StrBuilder);
 	}
 
 	int32 Minutes = static_cast<int32>(Time / TimeUtils::Minute);
@@ -356,7 +356,7 @@ FString FormatTime(const double InTime, const double Precision)
 		{
 			return TEXT("~0");
 		}
-		return StrBuilder.ToString();
+		return FString(StrBuilder);
 	}
 
 	//TestOptimizationIssue();
@@ -431,7 +431,7 @@ FString FormatTime(const double InTime, const double Precision)
 #endif
 	}
 
-	return StrBuilder.ToString();
+	return FString(StrBuilder);
 }
 
 ////////////////////////////////////////////////////////////////////////////////////////////////////
@@ -593,7 +593,7 @@ FString FormatTimeSplit(const FTimeSplit& InTimeSplit, const double Precision)
 		{
 			return TEXT("~0");
 		}
-		return StrBuilder.ToString();
+		return FString(StrBuilder);
 	}
 
 	if (InTimeSplit.Hours > 0)
@@ -612,7 +612,7 @@ FString FormatTimeSplit(const FTimeSplit& InTimeSplit, const double Precision)
 		{
 			return TEXT("~0");
 		}
-		return StrBuilder.ToString();
+		return FString(StrBuilder);
 	}
 
 	if (InTimeSplit.Minutes > 0)
@@ -631,7 +631,7 @@ FString FormatTimeSplit(const FTimeSplit& InTimeSplit, const double Precision)
 		{
 			return TEXT("~0");
 		}
-		return StrBuilder.ToString();
+		return FString(StrBuilder);
 	}
 
 	if (InTimeSplit.Seconds > 0)
@@ -650,7 +650,7 @@ FString FormatTimeSplit(const FTimeSplit& InTimeSplit, const double Precision)
 		{
 			return TEXT("~0");
 		}
-		return StrBuilder.ToString();
+		return FString(StrBuilder);
 	}
 
 	if (InTimeSplit.Miliseconds > 0)
@@ -669,7 +669,7 @@ FString FormatTimeSplit(const FTimeSplit& InTimeSplit, const double Precision)
 		{
 			return TEXT("~0");
 		}
-		return StrBuilder.ToString();
+		return FString(StrBuilder);
 	}
 
 	if (InTimeSplit.Microseconds > 0)
@@ -688,7 +688,7 @@ FString FormatTimeSplit(const FTimeSplit& InTimeSplit, const double Precision)
 		{
 			return TEXT("~0");
 		}
-		return StrBuilder.ToString();
+		return FString(StrBuilder);
 	}
 
 	if (InTimeSplit.Nanoseconds > 0)
@@ -707,7 +707,7 @@ FString FormatTimeSplit(const FTimeSplit& InTimeSplit, const double Precision)
 		{
 			return TEXT("~0");
 		}
-		return StrBuilder.ToString();
+		return FString(StrBuilder);
 	}
 
 	if (InTimeSplit.Picoseconds > 0)
@@ -725,7 +725,7 @@ FString FormatTimeSplit(const FTimeSplit& InTimeSplit, const double Precision)
 		return TEXT("~0");
 	}
 
-	return StrBuilder.ToString();
+	return FString(StrBuilder);
 }
 
 ////////////////////////////////////////////////////////////////////////////////////////////////////
diff --git a/Engine/Source/Developer/TraceInsights/Private/Insights/NetworkingProfiler/ViewModels/PacketContentViewDrawHelper.cpp b/Engine/Source/Developer/TraceInsights/Private/Insights/NetworkingProfiler/ViewModels/PacketContentViewDrawHelper.cpp
--- a/Engine/Source/Developer/TraceInsights/Private/Insights/NetworkingProfiler/ViewModels/PacketContentViewDrawHelper.cpp
+++ b/Engine/Source/Developer/TraceInsights/Private/Insights/NetworkingProfiler/ViewModels/PacketContentViewDrawHelper.cpp
@@ -246,7 +246,7 @@ void FPacketContentViewDrawStateBuilder::AddEvent(const TraceServices::FNetProfi
 			}
 		}
 
-		FString Name(Builder.ToString());
+		const FString Name(Builder);
 
 		const TSharedRef<FSlateFontMeasure> FontMeasureService = FSlateApplication::Get().GetRenderer()->GetFontMeasureService();
 		const int32 LastWholeCharacterIndex = FontMeasureService->FindLastWholeCharacterIndexBeforeOffset(Name, EventFont, FMath::RoundToInt(EventW - 2.0f));
