[Insights] Renamed "UE4 stat node" to "Stat node".

diff --git a/Engine/Source/Developer/TraceInsights/Private/Insights/ViewModels/StatsNodeHelper.cpp b/Engine/Source/Developer/TraceInsights/Private/Insights/ViewModels/StatsNodeHelper.cpp
--- a/Engine/Source/Developer/TraceInsights/Private/Insights/ViewModels/StatsNodeHelper.cpp
+++ b/Engine/Source/Developer/TraceInsights/Private/Insights/ViewModels/StatsNodeHelper.cpp
@@ -29,7 +29,7 @@ FText StatsNodeTypeHelper::ToDescription(const EStatsNodeType NodeType)
 	switch (NodeType)
 	{
 		case EStatsNodeType::Counter:	return LOCTEXT("Stats_Desc_Counter", "Counter node");
-		case EStatsNodeType::Stat:		return LOCTEXT("Stats_Desc_Stat", "UE4 stat node");
+		case EStatsNodeType::Stat:		return LOCTEXT("Stats_Desc_Stat", "Stat node");
 		case EStatsNodeType::Group:		return LOCTEXT("Stats_Desc_Group", "Group node");
 		default:						return LOCTEXT("InvalidOrMax", "InvalidOrMax");
 	}
