[Insights] Memory: Added double click functionality to time marker labels (A, B, C, D) in the Mem Investigation panel --> moves the marker to the center time value of the Timing view.

diff --git a/Engine/Source/Developer/TraceInsights/Private/Insights/MemoryProfiler/Widgets/SMemInvestigationView.cpp b/Engine/Source/Developer/TraceInsights/Private/Insights/MemoryProfiler/Widgets/SMemInvestigationView.cpp
--- a/Engine/Source/Developer/TraceInsights/Private/Insights/MemoryProfiler/Widgets/SMemInvestigationView.cpp
+++ b/Engine/Source/Developer/TraceInsights/Private/Insights/MemoryProfiler/Widgets/SMemInvestigationView.cpp
@@ -257,6 +257,7 @@ TSharedRef<SWidget> SMemInvestigationView::ConstructTimeMarkerWidget(uint32 Time
 			SNew(STextBlock)
 			.Text(FText::FromString(TimeMarker->GetName()))
 			.ColorAndOpacity(FSlateColor(TimeMarker->GetColor()))
+			.OnDoubleClicked(this, &SMemInvestigationView::OnTimeMarkerLabelDoubleClicked, TimeMarkerIndex)
 		]
 
 		+ SHorizontalBox::Slot()
@@ -533,4 +534,27 @@ FReply SMemInvestigationView::RunQuery()
 
 ////////////////////////////////////////////////////////////////////////////////////////////////////
 
+FReply SMemInvestigationView::OnTimeMarkerLabelDoubleClicked(const FGeometry& MyGeometry, const FPointerEvent& PointerEvent, uint32 TimeMarkerIndex)
+{
+	TSharedPtr<SMemoryProfilerWindow> ProfilerWindow = ProfilerWindowWeakPtr.Pin();
+	if (ProfilerWindow.IsValid())
+	{
+		const uint32 NumTimeMarkers = ProfilerWindow->GetNumCustomTimeMarkers();
+		if (TimeMarkerIndex < NumTimeMarkers)
+		{
+			const TSharedRef<Insights::FTimeMarker>& TimeMarker = ProfilerWindow->GetCustomTimeMarker(TimeMarkerIndex);
+			TSharedPtr<STimingView> TimingView = ProfilerWindow->GetTimingView();
+			if (TimingView.IsValid())
+			{
+				// Move timer to the center of the timing view.
+				const double Time = (TimingView->GetViewport().GetStartTime() + TimingView->GetViewport().GetEndTime()) / 2.0;
+				TimeMarker->SetTime(Time);
+			}
+		}
+	}
+	return FReply::Handled();
+}
+
+////////////////////////////////////////////////////////////////////////////////////////////////////
+
 #undef LOCTEXT_NAMESPACE
diff --git a/Engine/Source/Developer/TraceInsights/Private/Insights/MemoryProfiler/Widgets/SMemInvestigationView.h b/Engine/Source/Developer/TraceInsights/Private/Insights/MemoryProfiler/Widgets/SMemInvestigationView.h
--- a/Engine/Source/Developer/TraceInsights/Private/Insights/MemoryProfiler/Widgets/SMemInvestigationView.h
+++ b/Engine/Source/Developer/TraceInsights/Private/Insights/MemoryProfiler/Widgets/SMemInvestigationView.h
@@ -77,6 +77,8 @@ private:
 	 */
 	virtual void Tick(const FGeometry& AllottedGeometry, const double InCurrentTime, const float InDeltaTime) override;
 
+	FReply OnTimeMarkerLabelDoubleClicked(const FGeometry& MyGeometry, const FPointerEvent& PointerEvent, uint32 TimeMarkerIndex);
+
 private:
 	TWeakPtr<SMemoryProfilerWindow> ProfilerWindowWeakPtr;
 
