[Insights]
  - Frames track: Added "start time" and "fps" to the hovered frame tooltip. Also improved its appearance.
  - Frames track: Improved distribution of values on vertical axis when zooming in vertically.
  - Frames track: Also improved horizontal axis (ex. labels are now displayed on top; not hidden anymore when zooming in vertically).
  - Frames track: Changed labels of vertical axis to move on the left or right side of view if the user hovers a frame in their area.
  - Frames track: Added "Zoom Timing View on Frame Selection" toggle option in context menu of Frames track. If enabled, the Timing view will also be zoomed when a frame is selected.
  - Added default ini file with following settings:
      [Insights.TimingProfiler]
      DefaultZoomLevel=5.0
      bAutoHideEmptyTracks=True
      bAutoZoomOnFrameSelection=False
      AutoScrollFrameAlignment=game
      AutoScrollViewportOffsetPercent=0.1
      AutoScrollMinDelay=0.3
    (see Engine\Programs\UnrealInsights\Config\DefaultUnrealInsightsSettings.ini for detailed description of each setting)
    Above settings (except the DefaultZoomLevel) can be changed directly from UI and Unreal Insights will remember that settings.
  - Timing view: Changed all settings related to "auto scroll" to be persistent in the ini file.
  - Added SetScaleX(ScaleX) to FTimingTrackViewport.

diff --git a/Engine/Programs/UnrealInsights/Config/DefaultUnrealInsightsSettings.ini b/Engine/Programs/UnrealInsights/Config/DefaultUnrealInsightsSettings.ini
new file mode 100644
--- /dev/null
+++ b/Engine/Programs/UnrealInsights/Config/DefaultUnrealInsightsSettings.ini
@@ -0,0 +1,21 @@
+[Insights.TimingProfiler]
+; The default (initial) zoom level of the Timing view.
+; This option can be changed from the "Tracks" drop down menu.
+; It is defined as time between major tick marks, in seconds.
+; Valid values: [0.00000001 .. 3600] i.e. [10ns to 1h]
+DefaultZoomLevel=5.0
+; Auto hide empty tracks (ex.: ones without timing events in the current viewport).
+; This option can be changed from the "Tracks" drop down menu.
+bAutoHideEmptyTracks=True
+; If enabled, the Timing view will also be zoomed when a frame is selected in the Frames track.
+; This option can be changed from the context menu of the Frames track.
+bAutoZoomOnFrameSelection=False
+; Frame Alignment. Controls if auto-scroll should align center of the viewport with start of a frame or not.
+; Valid options: 'none' to disable frame alignment, 'game' / 'rendering' to align with Game / Rendering frames
+AutoScrollFrameAlignment=game
+; Viewport offset while auto-scrolling, as percent of viewport width.
+; If positive, it offsets the viewport forward, allowing an empty space at the right side of the viewport (i.e. after end of session).
+; If negative, it offsets the viewport backward (i.e. end of session will be outside viewport).
+AutoScrollViewportOffsetPercent=0.1
+; Minimum time between two auto-scroll updates, in [seconds].
+AutoScrollMinDelay=0.3
diff --git a/Engine/Source/Developer/TraceInsights/Private/Insights/InsightsSettings.h b/Engine/Source/Developer/TraceInsights/Private/Insights/InsightsSettings.h
--- a/Engine/Source/Developer/TraceInsights/Private/Insights/InsightsSettings.h
+++ b/Engine/Source/Developer/TraceInsights/Private/Insights/InsightsSettings.h
@@ -8,11 +8,18 @@
 /** Contains all settings for the Unreal Insights, accessible through the main manager. */
 class FInsightsSettings
 {
+	friend class SInsightsSettings;
+
 public:
 	FInsightsSettings(bool bInIsDefault = false)
 		: bIsEditing(false)
 		, bIsDefault(bInIsDefault)
-		, bShowEmptyTracksByDefault(false)
+		, DefaultZoomLevel(5.0) // 5 seconds between major tick marks
+		, bAutoHideEmptyTracks(true)
+		, bAutoZoomOnFrameSelection(false)
+		, AutoScrollFrameAlignment((int32)TraceFrameType_Game) // -1 = none, 0 = game, 1 = rendering
+		, AutoScrollViewportOffsetPercent(0.1) // scrolls forward 10% of viewport's width
+		, AutoScrollMinDelay(0.3) // [seconds]
 	{
 		if (!bIsDefault)
 		{
@@ -30,14 +37,50 @@ public:
 
 	void LoadFromConfig()
 	{
-		FConfigCacheIni::LoadGlobalIniFile(SettingsIni, TEXT("UnrealInsightsSettings"));
+		if (!FConfigCacheIni::LoadGlobalIniFile(SettingsIni, TEXT("UnrealInsightsSettings")))
+		{
+			return;
+		}
 
-		GConfig->GetBool(TEXT("Insights.TimingProfiler"), TEXT("bShowEmptyTracksByDefault"), bShowEmptyTracksByDefault, SettingsIni);
+		GConfig->GetDouble(TEXT("Insights.TimingProfiler"), TEXT("DefaultZoomLevel"), DefaultZoomLevel, SettingsIni);
+		GConfig->GetBool(TEXT("Insights.TimingProfiler"), TEXT("bAutoHideEmptyTracks"), bAutoHideEmptyTracks, SettingsIni);
+		GConfig->GetBool(TEXT("Insights.TimingProfiler"), TEXT("bAutoZoomOnFrameSelection"), bAutoZoomOnFrameSelection, SettingsIni);
+
+		// Auto-scroll options
+		FString FrameAlignment;
+		if (GConfig->GetString(TEXT("Insights.TimingProfiler"), TEXT("AutoScrollFrameAlignment"), FrameAlignment, SettingsIni))
+		{
+			FrameAlignment.TrimStartAndEndInline();
+			if (FrameAlignment.Equals(TEXT("game"), ESearchCase::IgnoreCase))
+			{
+				static_assert((int32)TraceFrameType_Game == 0, "ETraceFrameType");
+				AutoScrollFrameAlignment = (int32)TraceFrameType_Game;
+			}
+			else if (FrameAlignment.Equals(TEXT("rendering"), ESearchCase::IgnoreCase))
+			{
+				static_assert((int32)TraceFrameType_Rendering == 1, "ETraceFrameType");
+				AutoScrollFrameAlignment = (int32)TraceFrameType_Rendering;
+			}
+			else
+			{
+				AutoScrollFrameAlignment = -1;
+			}
+		}
+		GConfig->GetDouble(TEXT("Insights.TimingProfiler"), TEXT("AutoScrollViewportOffsetPercent"), AutoScrollViewportOffsetPercent, SettingsIni);
+		GConfig->GetDouble(TEXT("Insights.TimingProfiler"), TEXT("AutoScrollMinDelay"), AutoScrollMinDelay, SettingsIni);
 	}
 
 	void SaveToConfig()
 	{
-		GConfig->SetBool(TEXT("Insights.TimingProfiler"), TEXT("bShowEmptyTracksByDefault"), bShowEmptyTracksByDefault, SettingsIni);
+		GConfig->SetDouble(TEXT("Insights.TimingProfiler"), TEXT("DefaultZoomLevel"), DefaultZoomLevel, SettingsIni);
+		GConfig->SetBool(TEXT("Insights.TimingProfiler"), TEXT("bAutoHideEmptyTracks"), bAutoHideEmptyTracks, SettingsIni);
+		GConfig->SetBool(TEXT("Insights.TimingProfiler"), TEXT("bAutoZoomOnFrameSelection"), bAutoZoomOnFrameSelection, SettingsIni);
+
+		// Auto-scroll options
+		const TCHAR* FrameAlignment = (AutoScrollFrameAlignment == 0) ? TEXT("game") : (AutoScrollFrameAlignment == 1) ? TEXT("rendering") : TEXT("none");
+		GConfig->SetString(TEXT("Insights.TimingProfiler"), TEXT("AutoScrollFrameAlignment"), FrameAlignment, SettingsIni);
+		GConfig->SetDouble(TEXT("Insights.TimingProfiler"), TEXT("AutoScrollViewportOffsetPercent"), AutoScrollViewportOffsetPercent, SettingsIni);
+		GConfig->SetDouble(TEXT("Insights.TimingProfiler"), TEXT("AutoScrollMinDelay"), AutoScrollMinDelay, SettingsIni);
 
 		GConfig->Flush(false, SettingsIni);
 	}
@@ -64,11 +107,43 @@ public:
 
 	void ResetToDefaults()
 	{
-		bShowEmptyTracksByDefault = Defaults.bShowEmptyTracksByDefault;
+		DefaultZoomLevel = Defaults.DefaultZoomLevel;
+		bAutoHideEmptyTracks = Defaults.bAutoHideEmptyTracks;
+		bAutoZoomOnFrameSelection = Defaults.bAutoZoomOnFrameSelection;
+		AutoScrollFrameAlignment = Defaults.AutoScrollFrameAlignment;
+		AutoScrollViewportOffsetPercent = Defaults.AutoScrollViewportOffsetPercent;
+		AutoScrollMinDelay = Defaults.AutoScrollMinDelay;
 	}
 
-//protected:
-public:
+	#define SET_AND_SAVE(Option, Value) { if (Option != Value) { Option = Value; SaveToConfig(); } }
+
+	double GetDefaultZoomLevel() const { return DefaultZoomLevel; }
+	void SetDefaultZoomLevel(double ZoomLevel) { DefaultZoomLevel = ZoomLevel; }
+	void SetAndSaveDefaultZoomLevel(double ZoomLevel) { SET_AND_SAVE(DefaultZoomLevel, ZoomLevel); }
+
+	bool IsAutoHideEmptyTracksEnabled() const { return bAutoHideEmptyTracks; }
+	void SetAutoHideEmptyTracks(bool bOnOff) { bAutoHideEmptyTracks = bOnOff; }
+	void SetAndSaveAutoHideEmptyTracks(bool bOnOff) { SET_AND_SAVE(bAutoHideEmptyTracks, bOnOff); }
+
+	bool IsAutoZoomOnFrameSelectionEnabled() const { return bAutoZoomOnFrameSelection; }
+	void SetAutoZoomOnFrameSelection(bool bOnOff) { bAutoZoomOnFrameSelection = bOnOff; }
+	void SetAndSaveAutoZoomOnFrameSelection(bool bOnOff) { SET_AND_SAVE(bAutoZoomOnFrameSelection, bOnOff); }
+
+	int32 GetAutoScrollFrameAlignment() const { return AutoScrollFrameAlignment; }
+	void SetAutoScrollFrameAlignment(int32 FrameType) { AutoScrollFrameAlignment = FrameType; }
+	void SetAndSaveAutoScrollFrameAlignment(int32 FrameType) { SET_AND_SAVE(AutoScrollFrameAlignment, FrameType); }
+
+	double GetAutoScrollViewportOffsetPercent() const { return AutoScrollViewportOffsetPercent; }
+	void SetAutoScrollViewportOffsetPercent(double OffsetPercent) { AutoScrollViewportOffsetPercent = OffsetPercent; }
+	void SetAndSaveAutoScrollViewportOffsetPercent(double OffsetPercent) { SET_AND_SAVE(AutoScrollViewportOffsetPercent, OffsetPercent); }
+
+	double GetAutoScrollMinDelay() const { return AutoScrollMinDelay; }
+	void SetAutoScrollMinDelay(double Delay) { AutoScrollMinDelay = Delay; }
+	void SetAndSaveAutoScrollMinDelay(double Delay) { SET_AND_SAVE(AutoScrollMinDelay, Delay); }
+
+	#undef SET_AND_SAVE
+
+private:
 	/** Contains default settings. */
 	static FInsightsSettings Defaults;
 
@@ -84,6 +159,25 @@ public:
 	//////////////////////////////////////////////////
 	// Actual settings.
 
-	/** If True, the Timing View will show empty tracks by default. */
-	bool bShowEmptyTracksByDefault;
+	/** The default (initial) zoom level of the Timing view. */
+	double DefaultZoomLevel;
+
+	/** Auto hide empty tracks (ex.: ones without timing events in the current viewport). */
+	bool bAutoHideEmptyTracks;
+
+	/** If enabled, the Timing View will also be zoomed when a new frame is selected in the Frames track. */
+	bool bAutoZoomOnFrameSelection;
+
+	/** -1 to disable frame alignment or the type of frame to align with (0 = Game or 1 = Rendering). */
+	int32 AutoScrollFrameAlignment;
+
+	/**
+	 * Viewport offset while auto-scrolling, as percent of viewport width.
+	 * If positive, it offsets the viewport forward, allowing an empty space at the right side of the viewport (i.e. after end of session).
+	 * If negative, it offsets the viewport backward (i.e. end of session will be outside viewport).
+	 */
+	double AutoScrollViewportOffsetPercent;
+
+	/** Minimum time between two auto-scroll updates, in [seconds]. */
+	double AutoScrollMinDelay;
 };
diff --git a/Engine/Source/Developer/TraceInsights/Private/Insights/TimingProfilerCommands.cpp b/Engine/Source/Developer/TraceInsights/Private/Insights/TimingProfilerCommands.cpp
--- a/Engine/Source/Developer/TraceInsights/Private/Insights/TimingProfilerCommands.cpp
+++ b/Engine/Source/Developer/TraceInsights/Private/Insights/TimingProfilerCommands.cpp
@@ -84,10 +84,30 @@ FTimingViewCommands::FTimingViewCommands()
 PRAGMA_DISABLE_OPTIMIZATION
 void FTimingViewCommands::RegisterCommands()
 {
-	UI_COMMAND(AutoHideEmptyTracks, "Auto Hide Empty Tracks", "Auto hide empty tracks (ones without timing events in current viewport).", EUserInterfaceActionType::ToggleButton, FInputChord(EKeys::V));
-	UI_COMMAND(ToggleCompactMode, "Compact Mode", "Toggle compact mode for supporting tracks.\n(ex. timing tracks will be displayed with reduced height)", EUserInterfaceActionType::ToggleButton, FInputChord(EKeys::C));
-	UI_COMMAND(ShowMainGraphTrack, "Graph Track", "Show/hide the main Graph track.", EUserInterfaceActionType::ToggleButton, FInputChord(EKeys::G));
-	UI_COMMAND(ShowTaskDependencies, "Task Dependencies", "Show/hide all dependencies of the current task (for a selected cpu timing event).", EUserInterfaceActionType::ToggleButton, FInputChord(EKeys::T));
+	UI_COMMAND(AutoHideEmptyTracks,
+		"Auto Hide Empty Tracks",
+		"Auto hide empty tracks (ex.: ones without timing events in the current viewport).\nThis option is persistent to UnrealInsightsSettings.ini file.",
+		EUserInterfaceActionType::ToggleButton,
+		FInputChord(EKeys::V));
+
+	UI_COMMAND(ToggleCompactMode,
+		"Compact Mode",
+		"Toggle compact mode for supporting tracks.\n(ex.: the timing tracks will be displayed with reduced height)",
+		EUserInterfaceActionType::ToggleButton,
+		FInputChord(EKeys::C));
+
+	UI_COMMAND(ShowMainGraphTrack,
+		"Graph Track",
+		"Show/hide the main Graph track.",
+		EUserInterfaceActionType::ToggleButton,
+		FInputChord(EKeys::G));
+
+	UI_COMMAND(ShowTaskDependencies,
+		"Task Dependencies",
+		"Show/hide all dependencies of the current task (for a selected cpu timing event).",
+		EUserInterfaceActionType::ToggleButton,
+		FInputChord(EKeys::T));
+
 	//UI_COMMAND(ShowAllGpuTracks, "GPU Track(s)", "Show/hide the GPU track(s).", EUserInterfaceActionType::ToggleButton, FInputChord(EKeys::Y));
 	//UI_COMMAND(ShowAllCpuTracks, "CPU Thread Tracks", "Show/hide all CPU tracks (and all CPU thread groups).", EUserInterfaceActionType::ToggleButton, FInputChord(EKeys::U));
 }
diff --git a/Engine/Source/Developer/TraceInsights/Private/Insights/ViewModels/TimingTrackViewport.cpp b/Engine/Source/Developer/TraceInsights/Private/Insights/ViewModels/TimingTrackViewport.cpp
--- a/Engine/Source/Developer/TraceInsights/Private/Insights/ViewModels/TimingTrackViewport.cpp
+++ b/Engine/Source/Developer/TraceInsights/Private/Insights/ViewModels/TimingTrackViewport.cpp
@@ -119,6 +119,20 @@ bool FTimingTrackViewport::ZoomWithFixedX(const double NewScaleX, const float X)
 
 ////////////////////////////////////////////////////////////////////////////////////////////////////
 
+bool FTimingTrackViewport::SetScaleX(const double NewScaleX)
+{
+	const double LocalNewScaleX = FMath::Clamp(NewScaleX, MinScaleX, MaxScaleX);
+	if (LocalNewScaleX != ScaleX)
+	{
+		ScaleX = LocalNewScaleX;
+		AddDirtyFlags(ETimingTrackViewportDirtyFlags::HScaleChanged);
+		return true;
+	}
+	return false;
+}
+
+////////////////////////////////////////////////////////////////////////////////////////////////////
+
 double FTimingTrackViewport::RestrictEndTime(const double InEndTime) const
 {
 	if (InEndTime == DBL_MAX || InEndTime == std::numeric_limits<double>::infinity())
diff --git a/Engine/Source/Developer/TraceInsights/Private/Insights/Widgets/SFrameTrack.cpp b/Engine/Source/Developer/TraceInsights/Private/Insights/Widgets/SFrameTrack.cpp
--- a/Engine/Source/Developer/TraceInsights/Private/Insights/Widgets/SFrameTrack.cpp
+++ b/Engine/Source/Developer/TraceInsights/Private/Insights/Widgets/SFrameTrack.cpp
@@ -7,6 +7,7 @@
 #include "Framework/Application/SlateApplication.h"
 #include "Framework/MultiBox/MultiBoxBuilder.h"
 #include "HAL/PlatformTime.h"
+#include "Misc/StringBuilder.h"
 #include "Rendering/DrawElements.h"
 #include "Styling/CoreStyle.h"
 #include "TraceServices/AnalysisService.h"
@@ -47,6 +48,8 @@ SFrameTrack::~SFrameTrack()
 
 void SFrameTrack::Reset()
 {
+	const FInsightsSettings& Settings = FInsightsManager::Get()->GetSettings();
+
 	Viewport.Reset();
 	FAxisViewportInt32& ViewportX = Viewport.GetHorizontalAxisViewport();
 	ViewportX.SetScaleLimits(0.0001f, 16.0f); // 10000 [sample/px] to 16 [px/sample]
@@ -71,6 +74,8 @@ void SFrameTrack::Reset()
 	AutoZoomViewportScale = ViewportX.GetScale();
 	AutoZoomViewportSize = 0.0f;
 
+	bZoomTimingViewOnFrameSelection = Settings.IsAutoZoomOnFrameSelectionEnabled();
+
 	AnalysisSyncNextTimestamp = 0;
 
 	MousePosition = FVector2D::ZeroVector;
@@ -85,8 +90,10 @@ void SFrameTrack::Reset()
 
 	bIsScrolling = false;
 
+	bDrawVerticalAxisLabelsOnLeftSide = false;
+
 	HoveredSample.Reset();
-	TooltipDesiredOpacity = 0.9f;
+	TooltipDesiredOpacity = 1.0f;
 	TooltipOpacity = 0.0f;
 
 	//ThisGeometry
@@ -376,7 +383,18 @@ void SFrameTrack::SelectFrameAtMousePosition(float X, float Y)
 			{
 				const double StartTime = SampleRef.Sample->LargestFrameStartTime;
 				const double Duration = SampleRef.Sample->LargestFrameDuration;
-				TimingView->CenterOnTimeInterval(StartTime, Duration);
+
+				if (bZoomTimingViewOnFrameSelection)
+				{
+					const double EndTime = FMath::Min(StartTime + Duration, TimingView->GetViewport().GetMaxValidTime());
+					const double AdjustedDuration = EndTime - StartTime;
+					TimingView->ZoomOnTimeInterval(StartTime - AdjustedDuration * 0.1, AdjustedDuration * 1.2);
+				}
+				else
+				{
+					TimingView->CenterOnTimeInterval(StartTime, Duration);
+				}
+
 				TimingView->SelectTimeInterval(StartTime, Duration);
 				FSlateApplication::Get().SetKeyboardFocus(TimingView);
 			}
@@ -411,8 +429,8 @@ int32 SFrameTrack::OnPaint(const FPaintArgs& Args, const FGeometry& AllottedGeom
 
 		Helper.DrawBackground();
 
-		// Draw the horizontal axis grid.
-		DrawHorizontalAxisGrid(DrawContext, WhiteBrush, SummaryFont);
+		// Draw the horizontal axis grid (background layer).
+		DrawHorizontalAxisGrid(DrawContext, WhiteBrush, SummaryFont, true);
 
 		// Draw frames, for each visible Series.
 		for (int32 SeriesIndex = 0; SeriesIndex < SeriesOrder.Num(); ++SeriesIndex)
@@ -434,15 +452,40 @@ int32 SFrameTrack::OnPaint(const FPaintArgs& Args, const FGeometry& AllottedGeom
 
 		NumDrawSamples = Helper.GetNumDrawSamples();
 
-		// Highlight the mouse hovered sample (frame).
-		if (HoveredSample.IsValid())
+		TSharedPtr<STimingProfilerWindow> Window = FTimingProfilerManager::Get()->GetProfilerWindow();
+		if (Window)
 		{
-			Helper.DrawHoveredSample(*HoveredSample.Sample);
+			TSharedPtr<STimingView> TimingView = Window->GetTimingView();
+			if (TimingView)
+			{
+				TSharedPtr<FFrameTrackSeries> SeriesPtr = nullptr;
+				for (const TPair<int32, TSharedPtr<FFrameTrackSeries>>& KeyValuePair : SeriesMap)
+				{
+					SeriesPtr = KeyValuePair.Value;
+					break; // stop at first enumerated Series
+				}
+				if (SeriesPtr.IsValid())
+				{
+					// Highlight the area corresponding to viewport of Timing View.
+					const double StartTime = TimingView->GetViewport().GetStartTime();
+					const double EndTime = TimingView->GetViewport().GetEndTime();
+					Helper.DrawHighlightedInterval(*SeriesPtr, StartTime, EndTime);
+				}
+			}
 		}
 
+		// Draw the horizontal axis grid (foreground layer).
+		DrawHorizontalAxisGrid(DrawContext, WhiteBrush, SummaryFont, false);
+
 		// Draw the vertical axis grid.
 		DrawVerticalAxisGrid(DrawContext, WhiteBrush, SummaryFont);
 
+		// Highlight the mouse hovered sample (frame).
+		if (HoveredSample.IsValid())
+		{
+			Helper.DrawHoveredSample(*HoveredSample.Sample);
+		}
+
 		// Draw tooltip for hovered sample (frame).
 		if (HoveredSample.IsValid())
 		{
@@ -455,59 +498,68 @@ int32 SFrameTrack::OnPaint(const FPaintArgs& Args, const FGeometry& AllottedGeom
 				TooltipOpacity = TooltipDesiredOpacity;
 			}
 
-			const FString Text = FString::Format(TEXT("{0} frame {1} ({2})"),
-				{
-					FFrameTrackDrawHelper::FrameTypeToString(HoveredSample.Series->FrameType),
-					FText::AsNumber(HoveredSample.Sample->LargestFrameIndex).ToString(),
-					TimeUtils::FormatTimeAuto(HoveredSample.Sample->LargestFrameDuration)
-				});
+			// First line: "Rendering Frame 1,234"
+			TStringBuilder<512> StringBuilder;
+			StringBuilder.Append(FFrameTrackDrawHelper::FrameTypeToString(HoveredSample.Series->FrameType));
+			StringBuilder.Append(TEXT(" Frame "));
+			StringBuilder.Append(FText::AsNumber(HoveredSample.Sample->LargestFrameIndex).ToString());
+			const FString Text1(StringBuilder);
 
-			FVector2D TextSize = FontMeasureService->Measure(Text, SummaryFont);
+			// Second line: "1m 2.34s + 16.67ms (60 fps)"
+			StringBuilder.Reset();
+			StringBuilder.Append(TimeUtils::FormatTimeAuto(HoveredSample.Sample->LargestFrameStartTime, HoveredSample.Sample->LargestFrameStartTime > 60.0 ? 3 : 2));
+			StringBuilder.Append(TEXT(" + "));
+			StringBuilder.Append(TimeUtils::FormatTimeAuto(HoveredSample.Sample->LargestFrameDuration, 2));
+			StringBuilder.Appendf(TEXT(" (%.1f fps)"), 1.0 / HoveredSample.Sample->LargestFrameDuration);
+			const FString Text2(StringBuilder);
 
-			const float DX = 2.0f;
-			const float W2 = TextSize.X / 2 + DX;
+			const FVector2D TextSize1 = FontMeasureService->Measure(Text1, SummaryFont);
+			const FVector2D TextSize2 = FontMeasureService->Measure(Text2, SummaryFont);
 
 			const FAxisViewportInt32& ViewportX = Viewport.GetHorizontalAxisViewport();
 
-			float X1 = ViewportX.GetOffsetForValue(HoveredSample.Sample->LargestFrameIndex);
-			float CX = X1 + FMath::RoundToFloat(Viewport.GetSampleWidth() / 2);
-			if (CX + W2 > ViewportX.GetSize())
+			const float FrameX = ViewportX.GetOffsetForValue(HoveredSample.Sample->LargestFrameIndex);
+			const float CX0 = FMath::RoundToFloat(FrameX + Viewport.GetSampleWidth() / 2);
+
+			constexpr float DX = 3.0f;
+			const float DX1 = FMath::RoundToFloat(TextSize1.X / 2);
+			const float DX2 = FMath::RoundToFloat(TextSize2.X / 2);
+			const float BoxDX = FMath::Max(DX1, DX2) + DX;
+
+			float CX = CX0;
+			if (CX > ViewportX.GetSize() - BoxDX)
 			{
-				CX = FMath::RoundToFloat(ViewportX.GetSize() - W2);
+				CX = FMath::RoundToFloat(ViewportX.GetSize() - BoxDX);
 			}
-			if (CX - W2 < 0)
+			if (CX - BoxDX < 0)
 			{
-				CX = W2;
+				CX = BoxDX;
 			}
 
-			const float Y = 10.0f;
-			const float H = 14.0f;
-			DrawContext.DrawBox(CX - W2, Y, 2 * W2, H, WhiteBrush, FLinearColor(0.7, 0.7, 0.7, TooltipOpacity));
-			DrawContext.LayerId++;
-			DrawContext.DrawText(CX - W2 + DX, Y + 1.0f, Text, SummaryFont, FLinearColor(0.0, 0.0, 0.0, TooltipOpacity));
-			DrawContext.LayerId++;
-		}
+			constexpr float BoxY = 11.0f;
+			constexpr float BoxH = 26.0f;
+			constexpr float LineDY = 12.0f;
 
-		TSharedPtr<STimingProfilerWindow> Window = FTimingProfilerManager::Get()->GetProfilerWindow();
-		if (Window)
-		{
-			TSharedPtr<STimingView> TimingView = Window->GetTimingView();
-			if (TimingView)
+			const FLinearColor BackgroundColor(0.9f, 0.9f, 0.9f, TooltipOpacity);
+			DrawContext.DrawBox(CX - BoxDX, BoxY, 2 * BoxDX, BoxH, WhiteBrush, BackgroundColor);
+			const int32 ArrowSize = 4;
+			for (int32 ArrowY = 0; ArrowY < ArrowSize; ++ArrowY)
 			{
-				TSharedPtr<FFrameTrackSeries> SeriesPtr = nullptr;
-				for (const TPair<int32, TSharedPtr<FFrameTrackSeries>>& KeyValuePair : SeriesMap)
-				{
-					SeriesPtr = KeyValuePair.Value;
-					break; // stop at first enumerated Series
-				}
-				if (SeriesPtr.IsValid())
-				{
-					// Highlight the area corresponding to viewport of Timing View.
-					const double StartTime = TimingView->GetViewport().GetStartTime();
-					const double EndTime = TimingView->GetViewport().GetEndTime();
-					Helper.DrawHighlightedInterval(*SeriesPtr, StartTime, EndTime);
-				}
+				const int32 LineWidth = ArrowSize - ArrowY;
+				DrawContext.DrawBox(CX0 - float(LineWidth), BoxY + BoxH + float(ArrowY), float(2 * LineWidth - 1), 1.0f, WhiteBrush, BackgroundColor);
 			}
+			DrawContext.LayerId++;
+
+			const FLinearColor TextColor1 =
+				HoveredSample.Series->FrameType == TraceFrameType_Rendering ?
+					FLinearColor(0.5f, 0.1f, 0.1f, TooltipOpacity) :
+				HoveredSample.Series->FrameType == TraceFrameType_Game ?
+					FLinearColor(0.1f, 0.1f, 0.5f, TooltipOpacity) :
+					FLinearColor(0.1f, 0.1f, 0.1f, TooltipOpacity);
+			const FLinearColor TextColor2(0.05f, 0.05f, 0.05f, TooltipOpacity);
+			DrawContext.DrawText(CX - DX1, BoxY          + 1.0f, Text1, SummaryFont, TextColor1);
+			DrawContext.DrawText(CX - DX2, BoxY + LineDY + 1.0f, Text2, SummaryFont, TextColor2);
+			DrawContext.LayerId++;
 		}
 
 		Stopwatch.Stop();
@@ -526,12 +578,12 @@ int32 SFrameTrack::OnPaint(const FPaintArgs& Args, const FGeometry& AllottedGeom
 		const float DbgX = ViewWidth - DbgW - 20.0f;
 		float DbgY = 7.0f;
 
-		DrawContext.LayerId++;
+		const FLinearColor DbgBackgroundColor(1.0f, 1.0f, 1.0f, 0.9f);
+		const FLinearColor DbgTextColor(0.0f, 0.0f, 0.0f, 0.9f);
 
-		DrawContext.DrawBox(DbgX - 2.0f, DbgY - 2.0f, DbgW, DbgH, WhiteBrush, FLinearColor(1.0, 1.0, 1.0, 0.9));
 		DrawContext.LayerId++;
-
-		FLinearColor DbgTextColor(0.0, 0.0, 0.0, 0.9);
+		DrawContext.DrawBox(DbgX - 2.0f, DbgY - 2.0f, DbgW, DbgH, WhiteBrush, DbgBackgroundColor);
+		DrawContext.LayerId++;
 
 		// Time interval since last OnPaint call.
 		const uint64 CurrentTime = FPlatformTime::Cycles64();
@@ -601,94 +653,142 @@ void SFrameTrack::DrawVerticalAxisGrid(FDrawContext& DrawContext, const FSlateBr
 	const FAxisViewportDouble& ViewportY = Viewport.GetVerticalAxisViewport();
 	const float RoundedViewHeight = FMath::RoundToFloat(ViewportY.GetSize());
 
-	constexpr float TextH = 14.0f;
-	constexpr float MinDY = 12.0f; // min vertical distance between horizontal grid lines
-
 	const FLinearColor GridColor(0.0f, 0.0f, 0.0f, 0.1f);
 	const FLinearColor TextBgColor(0.05f, 0.05f, 0.05f, 1.0f);
-	//const FLinearColor TextColor(1.0f, 1.0f, 1.0f, 1.0f);
 
 	const TSharedRef<FSlateFontMeasure> FontMeasureService = FSlateApplication::Get().GetRenderer()->GetFontMeasureService();
 
-	const double GridValues[] =
+	// Available axis, pre-ordered by value.
+	struct FAxis
+	{
+		int32 Priority; // lower value means higher priority
+		double Value; // time value
+	};
+	const FAxis AvailableAxis[] =
 	{
-		0.0,
-		1.0 / 200.0, //    5 ms (200 fps)
-		1.0 / 120.0, //  8.3 ms (120 fps)
-		1.0 / 90.0,  // 11.1 ms (90 fps)
-		1.0 / 72.0,  // 13.9 ms (72 fps)
-		1.0 / 60.0,  // 16.7 ms (60 fps)
-		1.0 / 30.0,  // 33.3 ms (30 fps)
-		1.0 / 20.0,  //   50 ms (20 fps)
-		1.0 / 15.0,  // 66.7 ms (15 fps)
-		1.0 / 10.0,  //  100 ms (10 fps)
-		1.0 / 5.0,   //  200 ms (5 fps)
-		1.0,   // 1s
-		10.0,  // 10s
-		60.0,  // 1m
-		600.0, // 10m
-		3600.0 // 1h
+		{ 0, 0.0 },
+		{ 3, 0.001 },       //    1 ms (1000 fps)
+		{ 3, 0.002 },       //    2 ms (500 fps)
+		{ 3, 0.003 },       //    3 ms (333 fps)
+		{ 3, 0.004 },       //    4 ms (250 fps)
+		{ 2, 0.005 },       //    5 ms (200 fps)
+		{ 2, 1.0 / 150.0 }, //  6.6 ms (150 fps)
+		{ 1, 1.0 / 120.0 }, //  8.3 ms (120 fps)
+		{ 2, 1.0 / 100.0 }, //   10 ms (100 fps)
+		{ 3, 1.0 / 90.0 },  // 11.1 ms (90 fps)
+		{ 4, 1.0 / 80.0 },  // 12.5 ms (80 fps)
+		{ 4, 1.0 / 70.0 },  // 14.3 ms (70 fps)
+		{ 1, 1.0 / 60.0 },  // 16.7 ms (60 fps)
+		{ 2, 1.0 / 50.0 },  //   20 ms (50 fps)
+		{ 3, 1.0 / 40.0 },  //   25 ms (40 fps)
+		{ 1, 1.0 / 30.0 },  // 33.3 ms (30 fps)
+		{ 2, 1.0 / 20.0 },  //   50 ms (20 fps)
+		{ 3, 1.0 / 15.0 },  // 66.7 ms (15 fps)
+		{ 2, 1.0 / 10.0 },  //  100 ms (10 fps)
+		{ 3, 1.0 / 5.0 },   //  200 ms (5 fps)
+		{ 3, 1.0 },         // 1s
+		{ 3, 10.0 },        // 10s
+		{ 3, 60.0 },        // 1m
+		{ 3, 600.0 },       // 10m
+		{ 3, 3600.0 },      // 1h
 	};
-	constexpr int32 NumGridValues = sizeof(GridValues) / sizeof(double);
+	constexpr int32 NumAvailableAxis = UE_ARRAY_COUNT(AvailableAxis);
 
-	double PreviousY = -MinDY;
+	struct FVisibleAxis
+	{
+		double Value;
+		float Y;
+		float LabelY;
+	};
+	FVisibleAxis VisibleAxis[NumAvailableAxis];
+	int32 NumVisibleAxis = 0;
+
+	constexpr float TextH = 14.0f;
+	constexpr float MinDY = 13.0f; // min vertical distance between horizontal grid lines
 
-	for (int32 Index = 0; Index < NumGridValues; ++Index)
+	int32 PreviousPriority = 0;
+	float PreviousLabelY = -MinDY;
+	for (int32 Index = 0; Index < NumAvailableAxis; ++Index)
 	{
-		const double Value = GridValues[Index];
+		const FAxis& Axis = AvailableAxis[Index];
+
+		const float Y = RoundedViewHeight - FMath::RoundToFloat(ViewportY.GetOffsetForValue(Axis.Value));
+		const float LabelY = FMath::Clamp(Y - TextH / 2, 0.0f, RoundedViewHeight - TextH);
+
+		if (Y < 0)
+		{
+			break; // we are done; the rest of axis are offscreen
+		}
+		if (Y > RoundedViewHeight + TextH)
+		{
+			continue; // skip the current axis
+		}
+
+		// Does the label overlaps with the label of the previous axis?
+		if (FMath::Abs(PreviousLabelY - LabelY) < MinDY)
+		{
+			if (Axis.Priority < PreviousPriority)
+			{
+				--NumVisibleAxis; // the current axis replaces the previous axis
+			}
+			else
+			{
+				continue; // skip the current axis
+			}
+		}
+
+		PreviousPriority = Axis.Priority;
+		PreviousLabelY = LabelY;
+
+		FVisibleAxis& CurrentVisibleAxis = VisibleAxis[NumVisibleAxis++];
+		CurrentVisibleAxis.Value = Axis.Value;
+		CurrentVisibleAxis.Y = Y;
+		CurrentVisibleAxis.LabelY = LabelY;
+	}
+
+	for (int32 Index = 0; Index < NumVisibleAxis; ++Index)
+	{
+		const FVisibleAxis& Axis = VisibleAxis[Index];
 
 		constexpr double Time60fps = 1.0 / 60.0;
 		constexpr double Time30fps = 1.0 / 30.0;
 
 		FLinearColor TextColor;
-		if (Value <= Time60fps)
+		if (Axis.Value <= Time60fps)
 		{
 			TextColor = FLinearColor(0.5f, 1.0f, 0.5f, 1.0f);
 		}
-		else if (Value <= Time30fps)
+		else if (Axis.Value <= Time30fps)
 		{
 			TextColor = FLinearColor(1.0f, 1.0f, 0.5f, 1.0f);
-			//const float U = (Value - Time60fps) * 0.5 / (Time30fps - Time60fps);
-			//TextColor = FLinearColor(0.5f + U, 1.0f - U, 0.5f, 1.0f);
 		}
 		else
 		{
 			TextColor = FLinearColor(1.0f, 0.5f, 0.5f, 1.0f);
 		}
 
-		const float Y = RoundedViewHeight - FMath::RoundToFloat(ViewportY.GetOffsetForValue(Value));
-		if (Y < 0)
-		{
-			break;
-		}
-		if (Y > RoundedViewHeight + TextH || FMath::Abs(PreviousY - Y) < MinDY)
-		{
-			continue;
-		}
-		PreviousY = Y;
-
 		// Draw horizontal grid line.
-		DrawContext.DrawBox(0, Y, ViewWidth, 1, Brush, GridColor);
+		DrawContext.DrawBox(0, Axis.Y, ViewWidth, 1.0f, Brush, GridColor);
+
+		const FString LabelText = (Axis.Value == 0.0) ? TEXT("0") :
+								  (Axis.Value <= 1.0) ? FString::Printf(TEXT("%s (%.0f fps)"), *TimeUtils::FormatTimeAuto(Axis.Value), 1.0 / Axis.Value) :
+														TimeUtils::FormatTimeAuto(Axis.Value);
 
-		const FString LabelText = (Value == 0.0) ? TEXT("0") :
-								  (Value <= 1.0) ? FString::Printf(TEXT("%s (%.0f fps)"), *TimeUtils::FormatTimeAuto(Value), 1.0 / Value) :
-												   TimeUtils::FormatTimeAuto(Value);
 		const FVector2D LabelTextSize = FontMeasureService->Measure(LabelText, Font);
-		float LabelX = ViewWidth - LabelTextSize.X - 4.0f;
-		float LabelY = FMath::Clamp(Y - TextH / 2, 0.0f, RoundedViewHeight - TextH);
+		const float LabelX = bDrawVerticalAxisLabelsOnLeftSide ? 0.0f : ViewWidth - LabelTextSize.X - 4.0f;
 
 		// Draw background for value text.
-		DrawContext.DrawBox(LabelX, LabelY, LabelTextSize.X + 4.0f, TextH, Brush, TextBgColor);
+		DrawContext.DrawBox(LabelX, Axis.LabelY, LabelTextSize.X + 4.0f, TextH, Brush, TextBgColor);
 
 		// Draw value text.
-		DrawContext.DrawText(LabelX + 2.0f, LabelY + 1.0f, LabelText, Font, TextColor);
+		DrawContext.DrawText(LabelX + 2.0f, Axis.LabelY + 1.0f, LabelText, Font, TextColor);
 	}
 	DrawContext.LayerId++;
 }
 
 ////////////////////////////////////////////////////////////////////////////////////////////////////
 
-void SFrameTrack::DrawHorizontalAxisGrid(FDrawContext& DrawContext, const FSlateBrush* Brush, const FSlateFontInfo& Font) const
+void SFrameTrack::DrawHorizontalAxisGrid(FDrawContext& DrawContext, const FSlateBrush* Brush, const FSlateFontInfo& Font, bool bDrawBackgroundLayer) const
 {
 	const FAxisViewportInt32& ViewportX = Viewport.GetHorizontalAxisViewport();
 
@@ -729,23 +829,36 @@ void SFrameTrack::DrawHorizontalAxisGrid(FDrawContext& DrawContext, const FSlate
 			StartIndex += Grid;
 		}
 
-		const float ViewHeight = Viewport.GetHeight();
-
-		const FLinearColor GridColor(0.0f, 0.0f, 0.0f, 0.1f);
-		const FLinearColor TopTextColor(1.0f, 1.0f, 1.0f, 0.7f);
-
-		for (int32 Index = StartIndex; Index < RightIndex; Index += Grid)
+		if (bDrawBackgroundLayer)
 		{
-			const float X = FMath::RoundToFloat(ViewportX.GetOffsetForValue(Index));
+			const float ViewHeight = Viewport.GetHeight();
 
-			// Draw vertical grid line.
-			DrawContext.DrawBox(X, 0, 1, ViewHeight, Brush, GridColor);
+			// Draw vertical grid lines.
+			const FLinearColor GridColor(0.0f, 0.0f, 0.0f, 0.1f);
+			for (int32 Index = StartIndex; Index < RightIndex; Index += Grid)
+			{
+				const float X = FMath::RoundToFloat(ViewportX.GetOffsetForValue(Index));
+				DrawContext.DrawBox(X, 0.0f, 1.0f, ViewHeight, Brush, GridColor);
+			}
+			DrawContext.LayerId++;
+		}
+		else
+		{
+			const TSharedRef<FSlateFontMeasure> FontMeasureService = FSlateApplication::Get().GetRenderer()->GetFontMeasureService();
 
-			// Draw label.
-			const FString LabelText = FText::AsNumber(Index).ToString();
-			DrawContext.DrawText(X + 2.0f, 10.0f, LabelText, Font, TopTextColor);
+			// Draw labels.
+			const FLinearColor LabelBoxColor(0.05f, 0.05f, 0.05f, 1.0f);
+			const FLinearColor LabelTextColor(1.0f, 1.0f, 1.0f, 0.7f);
+			for (int32 Index = StartIndex; Index < RightIndex; Index += Grid)
+			{
+				const float X = FMath::RoundToFloat(ViewportX.GetOffsetForValue(Index));
+				const FString LabelText = FText::AsNumber(Index).ToString();
+				const FVector2D LabelTextSize = FontMeasureService->Measure(LabelText, Font);
+				DrawContext.DrawBox(X, 10.0f, LabelTextSize.X + 4.0f, 12.0f, Brush, LabelBoxColor);
+				DrawContext.DrawText(X + 2.0f, 10.0f, LabelText, Font, LabelTextColor);
+			}
+			DrawContext.LayerId++;
 		}
-		DrawContext.LayerId++;
 	}
 }
 
@@ -875,6 +988,19 @@ FReply SFrameTrack::OnMouseMove(const FGeometry& MyGeometry, const FPointerEvent
 			{
 				HoveredSample = GetSampleAtMousePosition(MousePosition.X + 1.0f, MousePosition.Y);
 			}
+			if (HoveredSample.IsValid())
+			{
+				constexpr float VerticalAxisLabelAreaWidth = 100.0f;
+				FAxisViewportInt32& ViewportX = Viewport.GetHorizontalAxisViewport();
+				if (MousePosition.X > ViewportX.GetSize() - VerticalAxisLabelAreaWidth)
+				{
+					bDrawVerticalAxisLabelsOnLeftSide = true;
+				}
+				else if (MousePosition.X < VerticalAxisLabelAreaWidth)
+				{
+					bDrawVerticalAxisLabelsOnLeftSide = false;
+				}
+			}
 		}
 
 		Reply = FReply::Handled();
@@ -984,7 +1110,7 @@ void SFrameTrack::ShowContextMenu(const FPointerEvent& MouseEvent)
 	const bool bShouldCloseWindowAfterMenuSelection = true;
 	FMenuBuilder MenuBuilder(bShouldCloseWindowAfterMenuSelection, NULL);
 
-	MenuBuilder.BeginSection("Series", LOCTEXT("ContextMenu_Header_Series", "Series"));
+	MenuBuilder.BeginSection("FramesTrack", LOCTEXT("ContextMenu_Header", "Frames Track"));
 	{
 		struct FLocal
 		{
@@ -1003,7 +1129,7 @@ void SFrameTrack::ShowContextMenu(const FPointerEvent& MouseEvent)
 		MenuBuilder.AddMenuEntry
 		(
 			LOCTEXT("ContextMenu_ShowGameFrames", "Game Frames"),
-			LOCTEXT("ContextMenu_ShowGameFrames_Desc", "Show/hide the Game frames"),
+			LOCTEXT("ContextMenu_ShowGameFrames_Desc", "Show/hide the Game frames."),
 			FSlateIcon(),
 			Action_ShowGameFrames,
 			NAME_None,
@@ -1019,17 +1145,13 @@ void SFrameTrack::ShowContextMenu(const FPointerEvent& MouseEvent)
 		MenuBuilder.AddMenuEntry
 		(
 			LOCTEXT("ContextMenu_ShowRenderingFrames", "Rendering Frames"),
-			LOCTEXT("ContextMenu_ShowRenderingFrames_Desc", "Show/hide the Rendering frames"),
+			LOCTEXT("ContextMenu_ShowRenderingFrames_Desc", "Show/hide the Rendering frames."),
 			FSlateIcon(),
 			Action_ShowRenderingFrames,
 			NAME_None,
 			EUserInterfaceActionType::ToggleButton
 		);
-	}
-	MenuBuilder.EndSection();
 
-	MenuBuilder.BeginSection("Misc");
-	{
 		FUIAction Action_AutoZoom
 		(
 			FExecuteAction::CreateSP(this, &SFrameTrack::ContextMenu_AutoZoom_Execute),
@@ -1039,12 +1161,28 @@ void SFrameTrack::ShowContextMenu(const FPointerEvent& MouseEvent)
 		MenuBuilder.AddMenuEntry
 		(
 			LOCTEXT("ContextMenu_AutoZoom", "Auto Zoom"),
-			LOCTEXT("ContextMenu_AutoZoom_Desc", "Enable auto zoom. Makes entire session time range to fit into view."),
+			LOCTEXT("ContextMenu_AutoZoom_Desc", "Enable auto zoom. Makes the entire session time range to fit into the Frames track's view."),
 			FSlateIcon(),
 			Action_AutoZoom,
 			NAME_None,
 			EUserInterfaceActionType::ToggleButton
 		);
+
+		FUIAction Action_ZoomTimingViewOnFrameSelection
+		(
+			FExecuteAction::CreateSP(this, &SFrameTrack::ContextMenu_ZoomTimingViewOnFrameSelection_Execute),
+			FCanExecuteAction::CreateSP(this, &SFrameTrack::ContextMenu_ZoomTimingViewOnFrameSelection_CanExecute),
+			FIsActionChecked::CreateSP(this, &SFrameTrack::ContextMenu_ZoomTimingViewOnFrameSelection_IsChecked)
+		);
+		MenuBuilder.AddMenuEntry
+		(
+			LOCTEXT("ContextMenu_ZoomTimingViewOnFrameSelection", "Zoom Timing View on Frame Selection"),
+			LOCTEXT("ContextMenu_ZoomTimingViewOnFrameSelection_Desc", "If enabled, the Timing view will also be zoomed when a frame is selected.\n(This option is persistent to UnrealInsightsSettings.ini file.)"),
+			FSlateIcon(),
+			Action_ZoomTimingViewOnFrameSelection,
+			NAME_None,
+			EUserInterfaceActionType::ToggleButton
+		);
 	}
 	MenuBuilder.EndSection();
 
@@ -1174,6 +1312,31 @@ void SFrameTrack::AutoZoom()
 
 ////////////////////////////////////////////////////////////////////////////////////////////////////
 
+void SFrameTrack::ContextMenu_ZoomTimingViewOnFrameSelection_Execute()
+{
+	bZoomTimingViewOnFrameSelection = !bZoomTimingViewOnFrameSelection;
+
+	// Persistent option. Save it to the config file.
+	FInsightsSettings& Settings = FInsightsManager::Get()->GetSettings();
+	Settings.SetAndSaveAutoZoomOnFrameSelection(bZoomTimingViewOnFrameSelection);
+}
+
+////////////////////////////////////////////////////////////////////////////////////////////////////
+
+bool SFrameTrack::ContextMenu_ZoomTimingViewOnFrameSelection_CanExecute()
+{
+	return true;
+}
+
+////////////////////////////////////////////////////////////////////////////////////////////////////
+
+bool SFrameTrack::ContextMenu_ZoomTimingViewOnFrameSelection_IsChecked()
+{
+	return bZoomTimingViewOnFrameSelection;
+}
+
+////////////////////////////////////////////////////////////////////////////////////////////////////
+
 void SFrameTrack::BindCommands()
 {
 }
diff --git a/Engine/Source/Developer/TraceInsights/Private/Insights/Widgets/SFrameTrack.h b/Engine/Source/Developer/TraceInsights/Private/Insights/Widgets/SFrameTrack.h
--- a/Engine/Source/Developer/TraceInsights/Private/Insights/Widgets/SFrameTrack.h
+++ b/Engine/Source/Developer/TraceInsights/Private/Insights/Widgets/SFrameTrack.h
@@ -133,7 +133,7 @@ protected:
 	TSharedPtr<FFrameTrackSeries> FindSeries(int32 FrameType) const;
 	void UpdateState();
 
-	void DrawHorizontalAxisGrid(FDrawContext& DrawContext, const FSlateBrush* Brush, const FSlateFontInfo& Font) const;
+	void DrawHorizontalAxisGrid(FDrawContext& DrawContext, const FSlateBrush* Brush, const FSlateFontInfo& Font, bool bDrawBackgroundLayer) const;
 	void DrawVerticalAxisGrid(FDrawContext& DrawContext, const FSlateBrush* Brush, const FSlateFontInfo& Font) const;
 
 	FFrameTrackSampleRef GetSampleAtMousePosition(float X, float Y);
@@ -154,6 +154,10 @@ protected:
 	bool ContextMenu_AutoZoom_IsChecked();
 	void AutoZoom();
 
+	void ContextMenu_ZoomTimingViewOnFrameSelection_Execute();
+	bool ContextMenu_ZoomTimingViewOnFrameSelection_CanExecute();
+	bool ContextMenu_ZoomTimingViewOnFrameSelection_IsChecked();
+
 	/** Binds our UI commands to delegates. */
 	void BindCommands();
 
@@ -179,11 +183,15 @@ protected:
 
 	bool bShowGameFrames;
 	bool bShowRenderingFrames;
+
 	bool bIsAutoZoomEnabled;
+
 	float AutoZoomViewportPos;
 	float AutoZoomViewportScale;
 	float AutoZoomViewportSize;
 
+	bool bZoomTimingViewOnFrameSelection;
+
 	uint64 AnalysisSyncNextTimestamp;
 
 	//////////////////////////////////////////////////
@@ -209,6 +217,8 @@ protected:
 	/** True, if the user is currently interactively scrolling the view (ex.: by holding the left mouse button and dragging). */
 	bool bIsScrolling;
 
+	mutable bool bDrawVerticalAxisLabelsOnLeftSide;
+
 	//////////////////////////////////////////////////
 	// Selection
 
diff --git a/Engine/Source/Developer/TraceInsights/Private/Insights/Widgets/SInsightsSettings.cpp b/Engine/Source/Developer/TraceInsights/Private/Insights/Widgets/SInsightsSettings.cpp
--- a/Engine/Source/Developer/TraceInsights/Private/Insights/Widgets/SInsightsSettings.cpp
+++ b/Engine/Source/Developer/TraceInsights/Private/Insights/Widgets/SInsightsSettings.cpp
@@ -38,13 +38,22 @@ void SInsightsSettings::Construct(const FArguments& InArgs)
 
 	AddTitle(LOCTEXT("SettingsTitle","Unreal Insights - Settings"), SettingsGrid, CurrentRowPos);
 	AddSeparator(SettingsGrid, CurrentRowPos);
-	AddHeader(LOCTEXT("TimingProfilerTitle","TimingProfiler"), SettingsGrid, CurrentRowPos);
+	AddHeader(LOCTEXT("TimingProfilerTitle","Timing Insights - defaults (applies when a new analysis session starts)"), SettingsGrid, CurrentRowPos);
 	AddOption
 	(
-		LOCTEXT("bShowEmptyTracksByDefault_T","Show empty CPU/GPU tracks in Timing View, by default"),
-		LOCTEXT("bShowEmptyTracksByDefault_TT","If True, empty CPU/GPU tracks will be visible in Timing View. Applies when session starts."),
-		SettingPtr->bShowEmptyTracksByDefault,
-		SettingPtr->GetDefaults().bShowEmptyTracksByDefault,
+		LOCTEXT("bAutoHideEmptyTracks_T","Auto hide empty CPU/GPU tracks in Timing View"),
+		LOCTEXT("bAutoHideEmptyTracks_TT","If enabled, the empty CPU/GPU tracks will be hidden in the Timing View."),
+		SettingPtr->bAutoHideEmptyTracks,
+		SettingPtr->GetDefaults().bAutoHideEmptyTracks,
+		SettingsGrid,
+		CurrentRowPos
+	);
+	AddOption
+	(
+		LOCTEXT("bAutoZoomOnFrameSelection_T", "Auto zoom on frame selection"),
+		LOCTEXT("bAutoZoomOnFrameSelection_TT", "If enabled, the Timing View will also be zoomed when a new frame is selected in the Frames track."),
+		SettingPtr->bAutoZoomOnFrameSelection,
+		SettingPtr->GetDefaults().bAutoZoomOnFrameSelection,
 		SettingsGrid,
 		CurrentRowPos
 	);
diff --git a/Engine/Source/Developer/TraceInsights/Private/Insights/Widgets/STimingView.cpp b/Engine/Source/Developer/TraceInsights/Private/Insights/Widgets/STimingView.cpp
--- a/Engine/Source/Developer/TraceInsights/Private/Insights/Widgets/STimingView.cpp
+++ b/Engine/Source/Developer/TraceInsights/Private/Insights/Widgets/STimingView.cpp
@@ -269,6 +269,8 @@ void STimingView::HideAllDefaultTracks()
 
 void STimingView::Reset(bool bIsFirstReset)
 {
+	const FInsightsSettings& Settings = FInsightsManager::Get()->GetSettings();
+
 	if (!bIsFirstReset)
 	{
 		for (Insights::ITimingViewExtender* Extender : GetExtenders())
@@ -281,6 +283,13 @@ void STimingView::Reset(bool bIsFirstReset)
 
 	Viewport.Reset();
 
+	if (IsAutoHideEmptyTracksEnabled() != Settings.IsAutoHideEmptyTracksEnabled())
+	{
+		ToggleAutoHideEmptyTracks();
+	}
+
+	Viewport.SetScaleX(100.0 / FMath::Clamp(Settings.GetDefaultZoomLevel(), 0.00000001, 3600.0));
+
 	//////////////////////////////////////////////////
 
 	for (auto& KV : AllTracks)
@@ -361,10 +370,9 @@ void STimingView::Reset(bool bIsFirstReset)
 	bIsDragging = false;
 
 	bAutoScroll = false;
-	bIsAutoScrollFrameAligned = true;
-	AutoScrollFrameType = TraceFrameType_Game;
-	AutoScrollViewportOffsetPercent = 0.1; // scrolls forward 10% of viewport's width
-	AutoScrollMinDelay = 0.3; // [seconds]
+	AutoScrollFrameAlignment = Settings.GetAutoScrollFrameAlignment();
+	AutoScrollViewportOffsetPercent = Settings.GetAutoScrollViewportOffsetPercent();
+	AutoScrollMinDelay = Settings.GetAutoScrollMinDelay();
 	LastAutoScrollTime = 0;
 
 	bIsPanning = false;
@@ -557,14 +565,15 @@ void STimingView::Tick(const FGeometry& AllottedGeometry, const double InCurrent
 			// By default, align the current session time with the offseted right side of the viewport.
 			double MinStartTime = Viewport.GetMaxValidTime() - ViewportDuration + AutoScrollViewportOffsetTime;
 
-			if (bIsAutoScrollFrameAligned)
+			if (AutoScrollFrameAlignment >= 0)
 			{
 				if (Session.IsValid())
 				{
 					TraceServices::FAnalysisSessionReadScope SessionReadScope(*Session.Get());
 					const TraceServices::IFrameProvider& FramesProvider = TraceServices::ReadFrameProvider(*Session.Get());
 
-					const uint64 FrameCount = FramesProvider.GetFrameCount(AutoScrollFrameType);
+					const ETraceFrameType FrameTpe = static_cast<ETraceFrameType>(AutoScrollFrameAlignment);
+					const uint64 FrameCount = FramesProvider.GetFrameCount(FrameTpe);
 
 					if (FrameCount > 0)
 					{
@@ -572,7 +581,7 @@ void STimingView::Tick(const FGeometry& AllottedGeometry, const double InCurrent
 						uint64 FrameIndex = FrameCount;
 						while (FrameIndex > 0)
 						{
-							const TraceServices::FFrame* FramePtr = FramesProvider.GetFrame(AutoScrollFrameType, --FrameIndex);
+							const TraceServices::FFrame* FramePtr = FramesProvider.GetFrame(FrameTpe, --FrameIndex);
 							if (FramePtr && FramePtr->EndTime <= Viewport.GetMaxValidTime())
 							{
 								// Align the start time of the frame with the right side of the viewport.
@@ -584,7 +593,7 @@ void STimingView::Tick(const FGeometry& AllottedGeometry, const double InCurrent
 						// Get the frame at the center of the viewport.
 						TraceServices::FFrame Frame;
 						const double ViewportCenter = MinStartTime + ViewportDuration / 2;
-						if (FramesProvider.GetFrameFromTime(AutoScrollFrameType, ViewportCenter, Frame))
+						if (FramesProvider.GetFrameFromTime(FrameTpe, ViewportCenter, Frame))
 						{
 							if (Frame.EndTime > ViewportCenter)
 							{
@@ -612,7 +621,7 @@ void STimingView::Tick(const FGeometry& AllottedGeometry, const double InCurrent
 			Extender->Tick(*this, *Session.Get());
 		}
 
-		// Re-sort now if we need to
+		// Re-sort now (if we need to).
 		UpdateScrollableTracksOrder();
 	}
 
@@ -2977,63 +2986,54 @@ ECheckBoxState STimingView::AutoScroll_IsChecked() const
 
 ////////////////////////////////////////////////////////////////////////////////////////////////////
 
-void STimingView::AutoScrollFrameAligned_Execute()
+void STimingView::SetAutoScrollFrameAlignment(int32 FrameType)
 {
-	bIsAutoScrollFrameAligned = !bIsAutoScrollFrameAligned;
-}
-
-////////////////////////////////////////////////////////////////////////////////////////////////////
-
-bool STimingView::AutoScrollFrameAligned_IsChecked() const
-{
-	return bIsAutoScrollFrameAligned;
-}
-
-////////////////////////////////////////////////////////////////////////////////////////////////////
-
-void STimingView::AutoScrollFrameType_Execute(ETraceFrameType FrameType)
-{
-	AutoScrollFrameType = FrameType;
-}
-
-////////////////////////////////////////////////////////////////////////////////////////////////////
+	AutoScrollFrameAlignment = FrameType;
 
-bool STimingView::AutoScrollFrameType_CanExecute(ETraceFrameType FrameType) const
-{
-	return bIsAutoScrollFrameAligned;
+	// Persistent option. Save it to the config file.
+	FInsightsSettings& Settings = FInsightsManager::Get()->GetSettings();
+	Settings.SetAndSaveAutoScrollFrameAlignment(FrameType);
 }
 
 ////////////////////////////////////////////////////////////////////////////////////////////////////
 
-bool STimingView::AutoScrollFrameType_IsChecked(ETraceFrameType FrameType) const
+bool STimingView::CompareAutoScrollFrameAlignment(int32 FrameType) const
 {
-	return AutoScrollFrameType == FrameType;
+	return AutoScrollFrameAlignment == FrameType;
 }
 
 ////////////////////////////////////////////////////////////////////////////////////////////////////
 
-void STimingView::AutoScrollViewportOffset_Execute(double Percent)
+void STimingView::SetAutoScrollViewportOffset(double Percent)
 {
 	AutoScrollViewportOffsetPercent = Percent;
+
+	// Persistent option. Save it to the config file.
+	FInsightsSettings& Settings = FInsightsManager::Get()->GetSettings();
+	Settings.SetAndSaveAutoScrollViewportOffsetPercent(Percent);
 }
 
 ////////////////////////////////////////////////////////////////////////////////////////////////////
 
-bool STimingView::AutoScrollViewportOffset_IsChecked(double Percent) const
+bool STimingView::CompareAutoScrollViewportOffset(double Percent) const
 {
 	return AutoScrollViewportOffsetPercent == Percent;
 }
 
 ////////////////////////////////////////////////////////////////////////////////////////////////////
 
-void STimingView::AutoScrollDelay_Execute(double Delay)
+void STimingView::SetAutoScrollDelay(double Delay)
 {
 	AutoScrollMinDelay = Delay;
+
+	// Persistent option. Save it to the config file.
+	FInsightsSettings& Settings = FInsightsManager::Get()->GetSettings();
+	Settings.SetAndSaveAutoScrollMinDelay(Delay);
 }
 
 ////////////////////////////////////////////////////////////////////////////////////////////////////
 
-bool STimingView::AutoScrollDelay_IsChecked(double Delay) const
+bool STimingView::CompareAutoScrollDelay(double Delay) const
 {
 	return AutoScrollMinDelay == Delay;
 }
@@ -3152,6 +3152,17 @@ void STimingView::CenterOnTimeInterval(double IntervalStartTime, double Interval
 
 ////////////////////////////////////////////////////////////////////////////////////////////////////
 
+void STimingView::ZoomOnTimeInterval(double IntervalStartTime, double IntervalDuration)
+{
+	if (Viewport.ZoomOnTimeInterval(IntervalStartTime, IntervalDuration))
+	{
+		Viewport.EnforceHorizontalScrollLimits(1.0); // 1.0 is to disable interpolation
+		UpdateHorizontalScrollBar();
+	}
+}
+
+////////////////////////////////////////////////////////////////////////////////////////////////////
+
 void STimingView::BringIntoView(double StartTime, double EndTime)
 {
 	EndTime = Viewport.RestrictEndTime(EndTime);
@@ -3776,34 +3787,34 @@ TSharedRef<SWidget> STimingView::MakeAutoScrollOptionsMenu()
 	MenuBuilder.BeginSection("QuickFilter", LOCTEXT("AutoScrollHeading", "Auto Scroll Options"));
 	{
 		MenuBuilder.AddMenuEntry(
-			LOCTEXT("AutoScrollFrameAligned", "Frame Aligned"),
-			LOCTEXT("AutoScrollFrameAligned_Tooltip", "Align the viewport's center position with the start time of a frame."),
+			LOCTEXT("AutoScrollNoFrameAlignment", "No Frame Alignment"),
+			LOCTEXT("AutoScrollNoFrameAlignment_Tooltip", "Disable frame alignment."),
 			FSlateIcon(),
-			FUIAction(FExecuteAction::CreateSP(this, &STimingView::AutoScrollFrameAligned_Execute),
+			FUIAction(FExecuteAction::CreateSP(this, &STimingView::SetAutoScrollFrameAlignment, -1),
 				FCanExecuteAction(),
-				FIsActionChecked::CreateSP(this, &STimingView::AutoScrollFrameAligned_IsChecked)),
+				FIsActionChecked::CreateSP(this, &STimingView::CompareAutoScrollFrameAlignment, -1)),
 			NAME_None,
-			EUserInterfaceActionType::ToggleButton
+			EUserInterfaceActionType::RadioButton
 		);
 
 		MenuBuilder.AddMenuEntry(
-			LOCTEXT("AutoScrollFrameTypeGame", "Align with Game Frames"),
-			LOCTEXT("AutoScrollFrameTypeGame_Tooltip", "Align the viewport's center position with the start time of a Game frame."),
+			LOCTEXT("AutoScrollAlignWithGameFrames", "Align with Game Frames"),
+			LOCTEXT("AutoScrollAlignWithGameFrames_Tooltip", "Align the viewport's center position with the start time of a Game frame."),
 			FSlateIcon(),
-			FUIAction(FExecuteAction::CreateSP(this, &STimingView::AutoScrollFrameType_Execute, TraceFrameType_Game),
-				FCanExecuteAction::CreateSP(this, &STimingView::AutoScrollFrameType_CanExecute, TraceFrameType_Game),
-				FIsActionChecked::CreateSP(this, &STimingView::AutoScrollFrameType_IsChecked, TraceFrameType_Game)),
+			FUIAction(FExecuteAction::CreateSP(this, &STimingView::SetAutoScrollFrameAlignment, (int32)TraceFrameType_Game),
+				FCanExecuteAction(),
+				FIsActionChecked::CreateSP(this, &STimingView::CompareAutoScrollFrameAlignment, (int32)TraceFrameType_Game)),
 			NAME_None,
 			EUserInterfaceActionType::RadioButton
 		);
 
 		MenuBuilder.AddMenuEntry(
-			LOCTEXT("AutoScrollFrameTypeRendering", "Align with Rendering Frames"),
-			LOCTEXT("AutoScrollFrameTypeRendering_Tooltip", "Align the viewport's center position with the start time of a Rendering frame."),
+			LOCTEXT("AutoScrollAlignWithRenderingFrames", "Align with Rendering Frames"),
+			LOCTEXT("AutoScrollAlignWithRenderingFrames_Tooltip", "Align the viewport's center position with the start time of a Rendering frame."),
 			FSlateIcon(),
-			FUIAction(FExecuteAction::CreateSP(this, &STimingView::AutoScrollFrameType_Execute, TraceFrameType_Rendering),
-				FCanExecuteAction::CreateSP(this, &STimingView::AutoScrollFrameType_CanExecute, TraceFrameType_Rendering),
-				FIsActionChecked::CreateSP(this, &STimingView::AutoScrollFrameType_IsChecked, TraceFrameType_Rendering)),
+			FUIAction(FExecuteAction::CreateSP(this, &STimingView::SetAutoScrollFrameAlignment, (int32)TraceFrameType_Rendering),
+				FCanExecuteAction(),
+				FIsActionChecked::CreateSP(this, &STimingView::CompareAutoScrollFrameAlignment, (int32)TraceFrameType_Rendering)),
 			NAME_None,
 			EUserInterfaceActionType::RadioButton
 		);
@@ -3814,9 +3825,9 @@ TSharedRef<SWidget> STimingView::MakeAutoScrollOptionsMenu()
 			LOCTEXT("AutoScrollViewportOffset-10", "Viewport Offset: -10%"),
 			LOCTEXT("AutoScrollViewportOffset-10_Tooltip", "Set the viewport offset to -10% (i.e. backward) of the viewport's width.\nAvoids flickering as the end of session will be outside of the viewport."),
 			FSlateIcon(),
-			FUIAction(FExecuteAction::CreateSP(this, &STimingView::AutoScrollViewportOffset_Execute, -0.1),
+			FUIAction(FExecuteAction::CreateSP(this, &STimingView::SetAutoScrollViewportOffset, -0.1),
 				FCanExecuteAction(),
-				FIsActionChecked::CreateSP(this, &STimingView::AutoScrollViewportOffset_IsChecked, -0.1)),
+				FIsActionChecked::CreateSP(this, &STimingView::CompareAutoScrollViewportOffset, -0.1)),
 			NAME_None,
 			EUserInterfaceActionType::RadioButton
 		);
@@ -3825,9 +3836,9 @@ TSharedRef<SWidget> STimingView::MakeAutoScrollOptionsMenu()
 			LOCTEXT("AutoScrollViewportOffset0", "Viewport Offset: 0"),
 			LOCTEXT("AutoScrollViewportOffset0_Tooltip", "Set the viewport offset to 0."),
 			FSlateIcon(),
-			FUIAction(FExecuteAction::CreateSP(this, &STimingView::AutoScrollViewportOffset_Execute, 0.0),
+			FUIAction(FExecuteAction::CreateSP(this, &STimingView::SetAutoScrollViewportOffset, 0.0),
 				FCanExecuteAction(),
-				FIsActionChecked::CreateSP(this, &STimingView::AutoScrollViewportOffset_IsChecked, 0.0)),
+				FIsActionChecked::CreateSP(this, &STimingView::CompareAutoScrollViewportOffset, 0.0)),
 			NAME_None,
 			EUserInterfaceActionType::RadioButton
 		);
@@ -3836,9 +3847,9 @@ TSharedRef<SWidget> STimingView::MakeAutoScrollOptionsMenu()
 			LOCTEXT("AutoScrollViewportOffset+10", "Viewport Offset: +10%"),
 			LOCTEXT("AutoScrollViewportOffset+10_Tooltip", "Set the viewport offset to +10% (i.e. forward) of the viewport's width.\nAllows 10% empty space on the right side of the viewport."),
 			FSlateIcon(),
-			FUIAction(FExecuteAction::CreateSP(this, &STimingView::AutoScrollViewportOffset_Execute, +0.1),
+			FUIAction(FExecuteAction::CreateSP(this, &STimingView::SetAutoScrollViewportOffset, +0.1),
 				FCanExecuteAction(),
-				FIsActionChecked::CreateSP(this, &STimingView::AutoScrollViewportOffset_IsChecked, +0.1)),
+				FIsActionChecked::CreateSP(this, &STimingView::CompareAutoScrollViewportOffset, +0.1)),
 			NAME_None,
 			EUserInterfaceActionType::RadioButton
 		);
@@ -3849,9 +3860,9 @@ TSharedRef<SWidget> STimingView::MakeAutoScrollOptionsMenu()
 			LOCTEXT("AutoScrollDelay0", "Delay: 0"),
 			LOCTEXT("AutoScrollDelay0_Tooltip", "Set the time delay of the auto-scroll update to 0."),
 			FSlateIcon(),
-			FUIAction(FExecuteAction::CreateSP(this, &STimingView::AutoScrollDelay_Execute, 0.0),
+			FUIAction(FExecuteAction::CreateSP(this, &STimingView::SetAutoScrollDelay, 0.0),
 				FCanExecuteAction(),
-				FIsActionChecked::CreateSP(this, &STimingView::AutoScrollDelay_IsChecked, 0.0)),
+				FIsActionChecked::CreateSP(this, &STimingView::CompareAutoScrollDelay, 0.0)),
 			NAME_None,
 			EUserInterfaceActionType::RadioButton
 		);
@@ -3860,9 +3871,9 @@ TSharedRef<SWidget> STimingView::MakeAutoScrollOptionsMenu()
 			LOCTEXT("AutoScrollDelay300ms", "Delay: 300ms"),
 			LOCTEXT("AutoScrollDelay300ms_Tooltip", "Set the time delay of the auto-scroll update to 300ms."),
 			FSlateIcon(),
-			FUIAction(FExecuteAction::CreateSP(this, &STimingView::AutoScrollDelay_Execute, 0.3),
+			FUIAction(FExecuteAction::CreateSP(this, &STimingView::SetAutoScrollDelay, 0.3),
 				FCanExecuteAction(),
-				FIsActionChecked::CreateSP(this, &STimingView::AutoScrollDelay_IsChecked, 0.3)),
+				FIsActionChecked::CreateSP(this, &STimingView::CompareAutoScrollDelay, 0.3)),
 			NAME_None,
 			EUserInterfaceActionType::RadioButton
 		);
@@ -3871,9 +3882,9 @@ TSharedRef<SWidget> STimingView::MakeAutoScrollOptionsMenu()
 			LOCTEXT("AutoScrollDelay1s", "Delay: 1s"),
 			LOCTEXT("AutoScrollDelay1s_Tooltip", "Set the time delay of the auto-scroll update to 1s."),
 			FSlateIcon(),
-			FUIAction(FExecuteAction::CreateSP(this, &STimingView::AutoScrollDelay_Execute, 1.0),
+			FUIAction(FExecuteAction::CreateSP(this, &STimingView::SetAutoScrollDelay, 1.0),
 				FCanExecuteAction(),
-				FIsActionChecked::CreateSP(this, &STimingView::AutoScrollDelay_IsChecked, 1.0)),
+				FIsActionChecked::CreateSP(this, &STimingView::CompareAutoScrollDelay, 1.0)),
 			NAME_None,
 			EUserInterfaceActionType::RadioButton
 		);
@@ -3882,9 +3893,9 @@ TSharedRef<SWidget> STimingView::MakeAutoScrollOptionsMenu()
 			LOCTEXT("AutoScrollDelay3s", "Delay: 3s"),
 			LOCTEXT("AutoScrollDelay3s_Tooltip", "Set the time delay of the auto-scroll update to 3s."),
 			FSlateIcon(),
-			FUIAction(FExecuteAction::CreateSP(this, &STimingView::AutoScrollDelay_Execute, 3.0),
+			FUIAction(FExecuteAction::CreateSP(this, &STimingView::SetAutoScrollDelay, 3.0),
 				FCanExecuteAction(),
-				FIsActionChecked::CreateSP(this, &STimingView::AutoScrollDelay_IsChecked, 3.0)),
+				FIsActionChecked::CreateSP(this, &STimingView::CompareAutoScrollDelay, 3.0)),
 			NAME_None,
 			EUserInterfaceActionType::RadioButton
 		);
@@ -4058,6 +4069,11 @@ void STimingView::ToggleAutoHideEmptyTracks()
 			TrackPtr->SetHeight(0.0f);
 		}
 	}
+
+	// Persistent option. Save it to the config file.
+	FInsightsSettings& Settings = FInsightsManager::Get()->GetSettings();
+	const bool bIsEnabled = IsAutoHideEmptyTracksEnabled();
+	Settings.SetAndSaveAutoHideEmptyTracks(bIsEnabled);
 }
 
 ////////////////////////////////////////////////////////////////////////////////////////////////////
diff --git a/Engine/Source/Developer/TraceInsights/Private/Insights/Widgets/STimingView.h b/Engine/Source/Developer/TraceInsights/Private/Insights/Widgets/STimingView.h
--- a/Engine/Source/Developer/TraceInsights/Private/Insights/Widgets/STimingView.h
+++ b/Engine/Source/Developer/TraceInsights/Private/Insights/Widgets/STimingView.h
@@ -317,6 +317,7 @@ public:
 	void ScrollAtPosY(float ScrollPosY);
 	void ScrollAtTime(double StartTime);
 	void CenterOnTimeInterval(double IntervalStartTime, double IntervalDuration);
+	void ZoomOnTimeInterval(double IntervalStartTime, double IntervalDuration);
 	void BringIntoView(double StartTime, double EndTime);
 	void SelectTimeInterval(double IntervalStartTime, double IntervalDuration);
 	void SelectToTimeMarker(double InTimeMarker);
@@ -366,18 +367,14 @@ protected:
 	void AutoScroll_OnCheckStateChanged(ECheckBoxState NewRadioState);
 	ECheckBoxState AutoScroll_IsChecked() const;
 
-	void AutoScrollFrameAligned_Execute();
-	bool AutoScrollFrameAligned_IsChecked() const;
+	void SetAutoScrollFrameAlignment(int32 FrameType);
+	bool CompareAutoScrollFrameAlignment(int32 FrameType) const;
 
-	void AutoScrollFrameType_Execute(ETraceFrameType FrameType);
-	bool AutoScrollFrameType_CanExecute(ETraceFrameType FrameType) const;
-	bool AutoScrollFrameType_IsChecked(ETraceFrameType FrameType) const;
+	void SetAutoScrollViewportOffset(double Percent);
+	bool CompareAutoScrollViewportOffset(double Percent) const;
 
-	void AutoScrollViewportOffset_Execute(double Percent);
-	bool AutoScrollViewportOffset_IsChecked(double Percent) const;
-
-	void AutoScrollDelay_Execute(double Delay);
-	bool AutoScrollDelay_IsChecked(double Delay) const;
+	void SetAutoScrollDelay(double Delay);
+	bool CompareAutoScrollDelay(double Delay) const;
 
 	////////////////////////////////////////////////////////////////////////////////////////////////////
 
@@ -525,11 +522,11 @@ protected:
 	/** True if the viewport scrolls automatically. */
 	bool bAutoScroll;
 
-	/** True, if auto-scroll should align center of viewport with start of a frame. */
-	bool bIsAutoScrollFrameAligned;
-
-	/** Type of frame to align with (Game or Rendering), if bIsAutoScrollFrameAligned is enabled. */
-	ETraceFrameType AutoScrollFrameType;
+	/**
+	 * Frame Alignment. Controls if auto-scroll should align center of the viewport with start of a frame or not.
+	 * Valid options: -1 to disable frame alignment or the type of frame to align with (0 = Game or 1 = Rendering; see ETraceFrameType).
+	 */
+	int32 AutoScrollFrameAlignment;
 
 	/**
 	 * Viewport offset while auto-scrolling, as percent of viewport width.
diff --git a/Engine/Source/Developer/TraceInsights/Public/Insights/ViewModels/TimingTrackViewport.h b/Engine/Source/Developer/TraceInsights/Public/Insights/ViewModels/TimingTrackViewport.h
--- a/Engine/Source/Developer/TraceInsights/Public/Insights/ViewModels/TimingTrackViewport.h
+++ b/Engine/Source/Developer/TraceInsights/Public/Insights/ViewModels/TimingTrackViewport.h
@@ -177,6 +177,7 @@ public:
 	bool ZoomOnTimeInterval(const double Time, const double Duration);
 	bool RelativeZoomWithFixedX(const float Delta, const float X);
 	bool ZoomWithFixedX(const double NewScaleX, const float X);
+	bool SetScaleX(const double NewScaleX);
 
 	//////////////////////////////////////////////////
 	// Vertical axis
