[Insights] Fixed NonUnity compile error.

diff --git a/Engine/Source/Developer/TraceInsights/Private/Insights/InsightsSettings.h b/Engine/Source/Developer/TraceInsights/Private/Insights/InsightsSettings.h
--- a/Engine/Source/Developer/TraceInsights/Private/Insights/InsightsSettings.h
+++ b/Engine/Source/Developer/TraceInsights/Private/Insights/InsightsSettings.h
@@ -4,6 +4,7 @@
 
 #include "CoreMinimal.h"
 #include "Misc/ConfigCacheIni.h"
+#include "ProfilingDebugging/MiscTrace.h"
 
 /** Contains all settings for the Unreal Insights, accessible through the main manager. */
 class FInsightsSettings
