[Insights] Initial implementation for TaskTableTreeView

diff --git a/Engine/Source/Developer/TraceInsights/Private/Insights/TaskGraphProfiler/TaskGraphProfilerManager.cpp b/Engine/Source/Developer/TraceInsights/Private/Insights/TaskGraphProfiler/TaskGraphProfilerManager.cpp
--- a/Engine/Source/Developer/TraceInsights/Private/Insights/TaskGraphProfiler/TaskGraphProfilerManager.cpp
+++ b/Engine/Source/Developer/TraceInsights/Private/Insights/TaskGraphProfiler/TaskGraphProfilerManager.cpp
@@ -7,10 +7,23 @@
 #include "TraceServices/Model/TasksProfiler.h"
 #include "Async/TaskGraphInterfaces.h"
 
+#include "Insights/InsightsStyle.h"
+#include "Insights/TaskGraphProfiler/ViewModels/TaskTable.h"
+#include "Insights/TaskGraphProfiler/Widgets/STaskTableTreeView.h"
+#include "Insights/TimingProfilerManager.h"
+#include "Insights/Widgets/STimingProfilerWindow.h"
+
 ////////////////////////////////////////////////////////////////////////////////////////////////////
 
 #define LOCTEXT_NAMESPACE "TaskGraphProfilerManager"
 
+namespace Insights
+{
+
+const FName FTaskGraphProfilerTabs::TaskTableTreeViewTabID(TEXT("TaskTableTreeView"));
+
+////////////////////////////////////////////////////////////////////////////////////////////////////
+
 TSharedPtr<FTaskGraphProfilerManager> FTaskGraphProfilerManager::Instance = nullptr;
 
 ////////////////////////////////////////////////////////////////////////////////////////////////////
@@ -58,6 +71,12 @@ void FTaskGraphProfilerManager::Initialize(IUnrealInsightsModule& InsightsModule
 	OnTick = FTickerDelegate::CreateSP(this, &FTaskGraphProfilerManager::Tick);
 	OnTickHandle = FTicker::GetCoreTicker().AddTicker(OnTick, 0.0f);
 
+	FOnRegisterMajorTabExtensions* TimingProfilerLayoutExtension = InsightsModule.FindMajorTabLayoutExtension(FInsightsManagerTabs::TimingProfilerTabId);
+	if (TimingProfilerLayoutExtension)
+	{
+		TimingProfilerLayoutExtension->AddRaw(this, &FTaskGraphProfilerManager::RegisterTimingProfilerLayoutExtensions);
+	}
+
 	FInsightsManager::Get()->GetSessionChangedEvent().AddSP(this, &FTaskGraphProfilerManager::OnSessionChanged);
 	OnSessionChanged();
 }
@@ -114,8 +133,10 @@ bool FTaskGraphProfilerManager::Tick(float DeltaTime)
 			TraceServices::FAnalysisSessionReadScope SessionReadScope(*Session.Get());
 
 			const TraceServices::ITasksProvider* TasksProvider = TraceServices::ReadTasksProvider(*Session.Get());
-			if (TasksProvider && TasksProvider->GetNumTasks() > 0)
+			TSharedPtr<FTabManager> TabManagerShared = TimingTabManager.Pin();
+			if (TasksProvider && TasksProvider->GetNumTasks() > 0 && TabManagerShared.IsValid())
 			{
+				TabManagerShared->TryInvokeTab(FTaskGraphProfilerTabs::TaskTableTreeViewTabID);
 				bIsAvailable = true;
 			}
 
@@ -152,4 +173,50 @@ void FTaskGraphProfilerManager::OnSessionChanged()
 
 ////////////////////////////////////////////////////////////////////////////////////////////////////
 
+void FTaskGraphProfilerManager::RegisterTimingProfilerLayoutExtensions(FInsightsMajorTabExtender& InOutExtender)
+{
+	TimingTabManager = InOutExtender.GetTabManager();
+
+	FInsightsMinorTabConfig& MinorTabConfig = InOutExtender.AddMinorTabConfig();
+	MinorTabConfig.TabId = FTaskGraphProfilerTabs::TaskTableTreeViewTabID;
+	MinorTabConfig.TabLabel = LOCTEXT("TaskTableTreeViewTabTitle", "Tasks");
+	MinorTabConfig.TabTooltip = LOCTEXT("TaskTableTreeViewTabTitleTooltip", "Opens the Task Table Tree View tab, that allows Task Graph profilling.");
+	MinorTabConfig.TabIcon = FSlateIcon(FSlateIcon(FInsightsStyle::GetStyleSetName(), "TimersView.Icon.Small"));
+	MinorTabConfig.OnSpawnTab = FOnSpawnTab::CreateRaw(this, &FTaskGraphProfilerManager::SpawnTab_TaskTableTreeView);
+
+	InOutExtender.GetLayoutExtender().ExtendLayout(FTimingProfilerTabs::StatsCountersID
+		, ELayoutExtensionPosition::After
+		, FTabManager::FTab(FTaskGraphProfilerTabs::TaskTableTreeViewTabID, ETabState::ClosedTab));
+}
+
+////////////////////////////////////////////////////////////////////////////////////////////////////
+
+TSharedRef<SDockTab> FTaskGraphProfilerManager::SpawnTab_TaskTableTreeView(const FSpawnTabArgs& Args)
+{
+	TSharedRef<FTaskTable> TaskTable = MakeShared<FTaskTable>();
+	TaskTable->Reset();
+
+	const TSharedRef<SDockTab> DockTab = SNew(SDockTab)
+		.ShouldAutosize(false)
+		.TabRole(ETabRole::PanelTab)
+		[
+			SAssignNew(TaskTableTreeView, STaskTableTreeView, TaskTable)
+		];
+
+	DockTab->SetOnTabClosed(SDockTab::FOnTabClosedCallback::CreateRaw(this, &FTaskGraphProfilerManager::OnTaskTableTreeViewTabClosed));
+
+	return DockTab;
+}
+
+////////////////////////////////////////////////////////////////////////////////////////////////////
+
+void FTaskGraphProfilerManager::OnTaskTableTreeViewTabClosed(TSharedRef<SDockTab> TabBeingClosed)
+{
+	TaskTableTreeView.Reset();
+}
+
+////////////////////////////////////////////////////////////////////////////////////////////////////
+
+} // namespace Insights
+
 #undef LOCTEXT_NAMESPACE
diff --git a/Engine/Source/Developer/TraceInsights/Private/Insights/TaskGraphProfiler/TaskGraphProfilerManager.h b/Engine/Source/Developer/TraceInsights/Private/Insights/TaskGraphProfiler/TaskGraphProfilerManager.h
--- a/Engine/Source/Developer/TraceInsights/Private/Insights/TaskGraphProfiler/TaskGraphProfilerManager.h
+++ b/Engine/Source/Developer/TraceInsights/Private/Insights/TaskGraphProfiler/TaskGraphProfilerManager.h
@@ -10,6 +10,17 @@
 #include "Insights/InsightsManager.h"
 #include "Insights/IUnrealInsightsModule.h"
 
+namespace Insights
+{
+
+class STaskTableTreeView;
+
+struct FTaskGraphProfilerTabs
+{
+	// Tab identifiers
+	static const FName TaskTableTreeViewTabID;
+};
+
 ////////////////////////////////////////////////////////////////////////////////////////////////////
 /**
  * This class manages the Task Graph Profiler state and settings.
@@ -44,6 +55,9 @@ public:
 	virtual void UnregisterMajorTabs() override;
 	virtual void OnWindowClosedEvent() override {}
 
+	TSharedRef<SDockTab> SpawnTab_TaskTableTreeView(const FSpawnTabArgs& Args);
+	void OnTaskTableTreeViewTabClosed(TSharedRef<SDockTab> TabBeingClosed);
+
 	//////////////////////////////////////////////////
 	bool GetIsAvailable() { return bIsAvailable; }
 
@@ -53,6 +67,8 @@ private:
 	/** Updates this manager, done through FCoreTicker. */
 	bool Tick(float DeltaTime);
 
+	void RegisterTimingProfilerLayoutExtensions(FInsightsMajorTabExtender& InOutExtender);
+
 private:
 	bool bIsInitialized;
 	bool bIsAvailable;
@@ -66,4 +82,11 @@ private:
 
 	/** A shared pointer to the global instance of the Task Graph Profiler manager. */
 	static TSharedPtr<FTaskGraphProfilerManager> Instance;
+
+	TWeakPtr<FTabManager> TimingTabManager;
+
+	TSharedPtr<Insights::STaskTableTreeView> TaskTableTreeView;
 };
+
+} // namespace Insights
+
diff --git a/Engine/Source/Developer/TraceInsights/Private/Insights/TaskGraphProfiler/ViewModels/TaskEntry.cpp b/Engine/Source/Developer/TraceInsights/Private/Insights/TaskGraphProfiler/ViewModels/TaskEntry.cpp
new file mode 100644
--- /dev/null
+++ b/Engine/Source/Developer/TraceInsights/Private/Insights/TaskGraphProfiler/ViewModels/TaskEntry.cpp
@@ -0,0 +1,34 @@
+// Copyright Epic Games, Inc. All Rights Reserved.
+
+#include "TaskEntry.h"
+
+namespace Insights
+{
+
+////////////////////////////////////////////////////////////////////////////////////////////////////
+
+FTaskEntry::FTaskEntry(const TraceServices::FTaskInfo& TaskInfo)
+	: Id(TaskInfo.Id)
+	, DebugName(TaskInfo.DebugName)
+	, bTracked(TaskInfo.bTracked)
+	, ThreadToExecuteOn(TaskInfo.ThreadToExecuteOn)
+	, CreatedTimestamp(TaskInfo.CreatedTimestamp)
+	, CreatedThreadId(TaskInfo.CreatedThreadId)
+	, LaunchedTimestamp(TaskInfo.LaunchedTimestamp)
+	, LaunchedThreadId(TaskInfo.LaunchedThreadId)
+	, ScheduledTimestamp(TaskInfo.ScheduledTimestamp)
+	, ScheduledThreadId(TaskInfo.ScheduledThreadId)
+	, StartedTimestamp(TaskInfo.StartedTimestamp)
+	, StartedThreadId(TaskInfo.StartedThreadId)
+	, FinishedTimestamp(TaskInfo.FinishedTimestamp)
+	, CompletedTimestamp(TaskInfo.CompletedTimestamp)
+	, CompletedThreadId(TaskInfo.CompletedThreadId)
+{
+	NumNested = TaskInfo.NestedTasks.Num();
+	NumSubsequents = TaskInfo.Subsequents.Num();
+	NumPrerequisites = TaskInfo.Prerequisites.Num();
+}
+
+////////////////////////////////////////////////////////////////////////////////////////////////////
+
+} // namespace Insights
diff --git a/Engine/Source/Developer/TraceInsights/Private/Insights/TaskGraphProfiler/ViewModels/TaskEntry.h b/Engine/Source/Developer/TraceInsights/Private/Insights/TaskGraphProfiler/ViewModels/TaskEntry.h
new file mode 100644
--- /dev/null
+++ b/Engine/Source/Developer/TraceInsights/Private/Insights/TaskGraphProfiler/ViewModels/TaskEntry.h
@@ -0,0 +1,74 @@
+// Copyright Epic Games, Inc. All Rights Reserved.
+
+#pragma once
+
+#include "CoreMinimal.h"
+
+#include "TraceServices/Model/TasksProfiler.h"
+
+namespace Insights
+{
+
+////////////////////////////////////////////////////////////////////////////////////////////////////
+
+class FTaskEntry
+{
+	friend class STaskTableTreeView;
+
+public:
+	FTaskEntry(const TraceServices::FTaskInfo &TaskInfo);
+	~FTaskEntry() {}
+
+	double GetCreatedTimestamp() const { return CreatedTimestamp; }
+	uint32 GetCreatedThreadId() const { return CreatedThreadId; }
+
+	double GetLaunchedTimestamp() const { return LaunchedTimestamp; }
+	uint32 GetLaunchedThreadId() const { return LaunchedThreadId; }
+
+	double GetScheduledTimestamp() const { return ScheduledTimestamp; }
+	uint32 GetScheduledThreadId() const { return ScheduledThreadId; }
+
+	double GetStartedTimestamp() const { return StartedTimestamp; }
+	uint32 GetStartedThreadId() const { return StartedThreadId; }
+
+	double GetFinishedTimestamp() const { return FinishedTimestamp; }
+
+	double GetCompletedTimestamp() const { return CompletedTimestamp; }
+	uint32 GetCompletedThreadId() const { return CompletedThreadId; }
+
+	uint32 GetNumNested() const { return NumNested; }
+	uint32 GetNumPrerequisites() const { return NumPrerequisites; }
+	uint32 GetNumSubsequents() const { return NumSubsequents; }
+
+private:
+	uint32 Id;
+
+	const TCHAR* DebugName; 
+	bool bTracked; 
+	int32 ThreadToExecuteOn;
+
+	double CreatedTimestamp;
+	uint32 CreatedThreadId;
+
+	double LaunchedTimestamp;
+	uint32 LaunchedThreadId;
+
+	double ScheduledTimestamp;
+	uint32 ScheduledThreadId;
+
+	double StartedTimestamp;
+	uint32 StartedThreadId;
+
+	double FinishedTimestamp;
+
+	double CompletedTimestamp;
+	uint32 CompletedThreadId;
+
+	uint32 NumNested;
+	uint32 NumSubsequents;
+	uint32 NumPrerequisites;
+};
+
+////////////////////////////////////////////////////////////////////////////////////////////////////
+
+} // namespace Insights
diff --git a/Engine/Source/Developer/TraceInsights/Private/Insights/TaskGraphProfiler/ViewModels/TaskNode.cpp b/Engine/Source/Developer/TraceInsights/Private/Insights/TaskGraphProfiler/ViewModels/TaskNode.cpp
new file mode 100644
--- /dev/null
+++ b/Engine/Source/Developer/TraceInsights/Private/Insights/TaskGraphProfiler/ViewModels/TaskNode.cpp
@@ -0,0 +1,14 @@
+// Copyright Epic Games, Inc. All Rights Reserved.
+
+#include "TaskNode.h"
+
+#define LOCTEXT_NAMESPACE "TaskNodeNode"
+
+namespace Insights
+{
+
+const FName FTaskNode::TypeName(TEXT("FTaskNode"));
+
+} // namespace Insights
+
+#undef LOCTEXT_NAMESPACE
diff --git a/Engine/Source/Developer/TraceInsights/Private/Insights/TaskGraphProfiler/ViewModels/TaskNode.h b/Engine/Source/Developer/TraceInsights/Private/Insights/TaskGraphProfiler/ViewModels/TaskNode.h
new file mode 100644
--- /dev/null
+++ b/Engine/Source/Developer/TraceInsights/Private/Insights/TaskGraphProfiler/ViewModels/TaskNode.h
@@ -0,0 +1,93 @@
+// Copyright Epic Games, Inc. All Rights Reserved.
+
+#pragma once
+
+#include "CoreMinimal.h"
+
+// Insights
+#include "Insights/TaskGraphProfiler/ViewModels/TaskTable.h"
+#include "Insights/TaskGraphProfiler/ViewModels/TaskEntry.h"
+#include "Insights/Table/ViewModels/TableTreeNode.h"
+
+////////////////////////////////////////////////////////////////////////////////////////////////////
+
+namespace Insights
+{
+
+enum class ETaskNodeType
+{
+	/** The TaskNode is an allocation node. */
+	Task,
+
+	/** The TaskNode is a group node. */
+	Group,
+
+	/** Invalid enum type, may be used as a number of enumerations. */
+	InvalidOrMax,
+};
+
+////////////////////////////////////////////////////////////////////////////////////////////////////
+
+class FTaskNode;
+
+/** Type definition for shared pointers to instances of FTaskNode. */
+typedef TSharedPtr<class FTaskNode> FTaskNodePtr;
+
+/** Type definition for shared references to instances of FTaskNode. */
+typedef TSharedRef<class FTaskNode> FTaskNodeRef;
+
+/** Type definition for shared references to const instances of FTaskNode. */
+typedef TSharedRef<const class FTaskNode> FTaskNodeRefConst;
+
+/** Type definition for weak references to instances of FTaskNode. */
+typedef TWeakPtr<class FTaskNode> FTaskNodeWeak;
+
+////////////////////////////////////////////////////////////////////////////////////////////////////
+/**
+ * Class used to store information about a task node (used in the STaskTreeView).
+ */
+class FTaskNode : public FTableTreeNode
+{
+public:
+	static const FName TypeName;
+
+public:
+	/** Initialization constructor for the Task node. */
+	explicit FTaskNode(const FName InName, TWeakPtr<FTaskTable> InParentTable, int32 InRowIndex)
+		: FTableTreeNode(InName, InParentTable, InRowIndex)
+		, Type(ETaskNodeType::Task)
+	{
+	}
+
+	/** Initialization constructor for the group node. */
+	explicit FTaskNode(const FName InGroupName, TWeakPtr<FTaskTable> InParentTable)
+		: FTableTreeNode(InGroupName, InParentTable)
+		, Type(ETaskNodeType::Group)
+	{
+	}
+
+	virtual const FName& GetTypeName() const override { return TypeName; }
+
+	/**
+	 * @return a type of this Task node or ETaskNodeType::Group for group nodes.
+	 */
+	ETaskNodeType GetType() const { return Type; }
+
+	FTaskTable& GetTaskTableChecked() const
+	{
+		const TSharedPtr<FTable>& TablePin = GetParentTable().Pin();
+		check(TablePin.IsValid());
+		return *StaticCastSharedPtr<FTaskTable>(TablePin);
+	}
+
+	bool IsValidTask() const { return GetTaskTableChecked().IsValidRowIndex(GetRowIndex()); }
+	const FTaskEntry* GetTask() const { return GetTaskTableChecked().GetTask(GetRowIndex()); }
+	const FTaskEntry& GetTaskChecked() const { return GetTaskTableChecked().GetTaskChecked(GetRowIndex()); }
+
+private:
+	const ETaskNodeType Type;
+};
+
+////////////////////////////////////////////////////////////////////////////////////////////////////
+
+} // namespace Insights
diff --git a/Engine/Source/Developer/TraceInsights/Private/Insights/TaskGraphProfiler/ViewModels/TaskTable.cpp b/Engine/Source/Developer/TraceInsights/Private/Insights/TaskGraphProfiler/ViewModels/TaskTable.cpp
new file mode 100644
--- /dev/null
+++ b/Engine/Source/Developer/TraceInsights/Private/Insights/TaskGraphProfiler/ViewModels/TaskTable.cpp
@@ -0,0 +1,579 @@
+// Copyright Epic Games, Inc. All Rights Reserved.
+
+#include "TaskTable.h"
+
+// Insights
+#include "Insights/TaskGraphProfiler/ViewModels/TaskNode.h"
+#include "Insights/TaskGraphProfiler/ViewModels/TaskTable.h"
+#include "Insights/Table/ViewModels/TableCellValueFormatter.h"
+#include "Insights/Table/ViewModels/TableCellValueGetter.h"
+#include "Insights/Table/ViewModels/TableCellValueSorter.h"
+#include "Insights/Table/ViewModels/TableColumn.h"
+
+#define LOCTEXT_NAMESPACE "Insights::FTaskTable"
+
+namespace Insights
+{
+
+////////////////////////////////////////////////////////////////////////////////////////////////////
+// Column identifiers
+
+const FName FTaskTableColumns::CreatedTimestampColumnId(TEXT("CreatedTimestamp"));
+const FName FTaskTableColumns::CreatedThreadIdColumnId(TEXT("CreatedThreadId"));
+const FName FTaskTableColumns::LaunchedTimestampColumnId(TEXT("LaunchedTimestamp"));
+const FName FTaskTableColumns::LaunchedThreadIdColumnId(TEXT("LaunchedThreadId"));
+const FName FTaskTableColumns::ScheduledTimestampColumnId(TEXT("ScheduledTimestamp"));
+const FName FTaskTableColumns::ScheduledThreadIdColumnId(TEXT("ScheduledThreadId"));
+const FName FTaskTableColumns::StartedTimestampColumnId(TEXT("StartedTimestamp"));
+const FName FTaskTableColumns::StartedThreadIdColumnId(TEXT("StartedThreadId"));
+const FName FTaskTableColumns::CompletedTimestampColumnId(TEXT("CompletedTimestamp"));
+const FName FTaskTableColumns::CompletedThreadIdColumnId(TEXT("CompletedThreadId"));
+const FName FTaskTableColumns::FinishedTimestampColumnId(TEXT("FinishedTimestamp"));
+const FName FTaskTableColumns::NumNestedColumnId(TEXT("NumNested"));
+const FName FTaskTableColumns::NumSubsequentsColumnId(TEXT("NumSubsequents"));
+const FName FTaskTableColumns::NumPrerequisitesColumnId(TEXT("NumPrerequisites"));
+
+////////////////////////////////////////////////////////////////////////////////////////////////////
+
+typedef FTableCellValue(*TaskFieldGetter)(const FTableColumn&, const FTaskEntry&);
+
+template<TaskFieldGetter Getter>
+class FTaskColumnValueGetter : public FTableCellValueGetter
+{
+public:
+	virtual const TOptional<FTableCellValue> GetValue(const FTableColumn& Column, const FBaseTreeNode& Node) const override
+	{
+		if (Node.IsGroup())
+		{
+			const FTableTreeNode& NodePtr = static_cast<const FTableTreeNode&>(Node);
+			if (NodePtr.HasAggregatedValue(Column.GetId()))
+			{
+				return NodePtr.GetAggregatedValue(Column.GetId());
+			}
+		}
+		else //if (Node->Is<FTaskNode>())
+		{
+			const FTaskNode& TaskNode = static_cast<const FTaskNode&>(Node);
+			const FTaskEntry* Task = TaskNode.GetTask();
+			if (Task)
+			{
+				return Getter(Column, *Task);
+			}
+		}
+
+		return TOptional<FTableCellValue>();
+	}
+};
+
+////////////////////////////////////////////////////////////////////////////////////////////////////
+
+struct DefaultTaskFieldGetterFuncts
+{
+	static FTableCellValue GetCreatedTimestamp(const FTableColumn& Column, const FTaskEntry& Task) { return FTableCellValue(Task.GetCreatedTimestamp());	}
+	static FTableCellValue GetCreatedThreadId(const FTableColumn& Column, const FTaskEntry& Task) { return FTableCellValue((int64)Task.GetCreatedThreadId());	}
+
+	static FTableCellValue GetLaunchedTimestamp(const FTableColumn& Column, const FTaskEntry& Task) { return FTableCellValue(Task.GetLaunchedTimestamp());	}
+	static FTableCellValue GetLaunchedThreadId(const FTableColumn& Column, const FTaskEntry& Task) { return FTableCellValue((int64)Task.GetLaunchedThreadId());	}
+
+	static FTableCellValue GetScheduledTimestamp(const FTableColumn& Column, const FTaskEntry& Task) { return FTableCellValue(Task.GetScheduledTimestamp()); }
+	static FTableCellValue GetScheduledThreadId(const FTableColumn& Column, const FTaskEntry& Task) { return FTableCellValue((int64)Task.GetScheduledThreadId()); }
+
+	static FTableCellValue GetStartedTimestamp(const FTableColumn& Column, const FTaskEntry& Task) { return FTableCellValue(Task.GetStartedTimestamp()); }
+	static FTableCellValue GetStartedThreadId(const FTableColumn& Column, const FTaskEntry& Task) { return FTableCellValue((int64)Task.GetStartedThreadId()); }
+
+	static FTableCellValue GetCompletedTimestamp(const FTableColumn& Column, const FTaskEntry& Task) { return FTableCellValue(Task.GetCompletedTimestamp()); }
+	static FTableCellValue GetCompletedThreadId(const FTableColumn& Column, const FTaskEntry& Task) { return FTableCellValue((int64)Task.GetCompletedThreadId()); }
+
+	static FTableCellValue GetFinishedTimestamp(const FTableColumn& Column, const FTaskEntry& Task) { return FTableCellValue(Task.GetFinishedTimestamp()); }
+
+	static FTableCellValue GetNumNested(const FTableColumn& Column, const FTaskEntry& Task) { return FTableCellValue((int64)Task.GetNumNested());	}
+	static FTableCellValue GetNumSubsequents(const FTableColumn& Column, const FTaskEntry& Task) { return FTableCellValue((int64)Task.GetNumSubsequents()); }
+	static FTableCellValue GetNumPrerequisites(const FTableColumn& Column, const FTaskEntry& Task) { return FTableCellValue((int64)Task.GetNumPrerequisites());	}
+};
+
+////////////////////////////////////////////////////////////////////////////////////////////////////
+// FTaskTable
+////////////////////////////////////////////////////////////////////////////////////////////////////
+
+FTaskTable::FTaskTable()
+{
+}
+
+////////////////////////////////////////////////////////////////////////////////////////////////////
+
+FTaskTable::~FTaskTable()
+{
+}
+
+////////////////////////////////////////////////////////////////////////////////////////////////////
+
+void FTaskTable::Reset()
+{
+	FTable::Reset();
+
+	AddDefaultColumns();
+}
+
+////////////////////////////////////////////////////////////////////////////////////////////////////
+
+void FTaskTable::AddDefaultColumns()
+{
+	//////////////////////////////////////////////////
+	// Hierarchy Column
+	{
+		const int32 HierarchyColumnIndex = -1;
+		const TCHAR* HierarchyColumnName = nullptr;
+		AddHierarchyColumn(HierarchyColumnIndex, HierarchyColumnName);
+
+		const TSharedRef<FTableColumn>& ColumnRef = GetColumns()[0];
+		ColumnRef->SetInitialWidth(200.0f);
+		ColumnRef->SetShortName(LOCTEXT("TaskColumnName", "Hierarchy"));
+		ColumnRef->SetTitleName(LOCTEXT("TaskColumnTitle", "Task Hierarchy"));
+		ColumnRef->SetDescription(LOCTEXT("TaskColumnDesc", "Hierarchy of the task's tree"));
+	}
+
+	int32 ColumnIndex = 0;
+
+	//////////////////////////////////////////////////
+	// Created Timestamp Column
+	{
+		TSharedRef<FTableColumn> ColumnRef = MakeShared<FTableColumn>(FTaskTableColumns::CreatedTimestampColumnId);
+		FTableColumn& Column = *ColumnRef;
+
+		Column.SetIndex(ColumnIndex++);
+
+		Column.SetShortName(LOCTEXT("CreatedTimestampColumnName", "Created"));
+		Column.SetTitleName(LOCTEXT("CreatedTimestampColumnTitle", "Created"));
+		Column.SetDescription(LOCTEXT("CreatedTimestampColumnDesc", "The time when the task was created."));
+
+		Column.SetFlags(ETableColumnFlags::ShouldBeVisible | ETableColumnFlags::CanBeHidden | ETableColumnFlags::CanBeFiltered);
+
+		Column.SetHorizontalAlignment(HAlign_Left);
+		Column.SetInitialWidth(100.0f);
+
+		Column.SetDataType(ETableCellDataType::Double);
+
+		TSharedRef<ITableCellValueGetter> Getter = MakeShared<FTaskColumnValueGetter<DefaultTaskFieldGetterFuncts::GetCreatedTimestamp>>();
+		Column.SetValueGetter(Getter);
+
+		TSharedRef<ITableCellValueFormatter> Formatter = MakeShared<FDoubleValueFormatterAsTimeAuto>();
+		Column.SetValueFormatter(Formatter);
+
+		TSharedRef<ITableCellValueSorter> Sorter = MakeShared<FSorterByDoubleValue>(ColumnRef);
+		Column.SetValueSorter(Sorter);
+
+		Column.SetAggregation(ETableColumnAggregation::Min);
+
+		AddColumn(ColumnRef);
+	}
+	//////////////////////////////////////////////////
+	// Created Thread Id Column
+	{
+		TSharedRef<FTableColumn> ColumnRef = MakeShared<FTableColumn>(FTaskTableColumns::CreatedThreadIdColumnId);
+		FTableColumn& Column = *ColumnRef;
+
+		Column.SetIndex(ColumnIndex++);
+
+		Column.SetShortName(LOCTEXT("CreatedThreadIdColumnName", "Created Thread Id"));
+		Column.SetTitleName(LOCTEXT("CreatedThreadIdColumnTitle", "Created Thread Id"));
+		Column.SetDescription(LOCTEXT("CreatedThreadIdColumnDesc", "The thread the task was created on."));
+
+		Column.SetFlags(ETableColumnFlags::CanBeHidden | ETableColumnFlags::CanBeFiltered);
+
+		Column.SetHorizontalAlignment(HAlign_Left);
+		Column.SetInitialWidth(100.0f);
+
+		Column.SetDataType(ETableCellDataType::Int64);
+
+		TSharedRef<ITableCellValueGetter> Getter = MakeShared<FTaskColumnValueGetter<DefaultTaskFieldGetterFuncts::GetCreatedThreadId>>();
+		Column.SetValueGetter(Getter);
+
+		TSharedRef<ITableCellValueFormatter> Formatter = MakeShared<FInt64ValueFormatterAsNumber>();
+		Column.SetValueFormatter(Formatter);
+
+		TSharedRef<ITableCellValueSorter> Sorter = MakeShared<FSorterByInt64Value>(ColumnRef);
+		Column.SetValueSorter(Sorter);
+
+		AddColumn(ColumnRef);
+	}
+
+	//////////////////////////////////////////////////
+	// Launched Timestamp Column
+	{
+		TSharedRef<FTableColumn> ColumnRef = MakeShared<FTableColumn>(FTaskTableColumns::LaunchedTimestampColumnId);
+		FTableColumn& Column = *ColumnRef;
+
+		Column.SetIndex(ColumnIndex++);
+
+		Column.SetShortName(LOCTEXT("LaunchedTimestampColumnName", "Launched"));
+		Column.SetTitleName(LOCTEXT("LaunchedTimestampColumnTitle", "Launched"));
+		Column.SetDescription(LOCTEXT("LaunchedTimestampColumnDesc", "The time when the task was Launched."));
+
+		Column.SetFlags(ETableColumnFlags::ShouldBeVisible | ETableColumnFlags::CanBeHidden | ETableColumnFlags::CanBeFiltered);
+
+		Column.SetHorizontalAlignment(HAlign_Left);
+		Column.SetInitialWidth(100.0f);
+
+		Column.SetDataType(ETableCellDataType::Double);
+
+		TSharedRef<ITableCellValueGetter> Getter = MakeShared<FTaskColumnValueGetter<DefaultTaskFieldGetterFuncts::GetLaunchedTimestamp>>();
+		Column.SetValueGetter(Getter);
+
+		TSharedRef<ITableCellValueFormatter> Formatter = MakeShared<FDoubleValueFormatterAsTimeAuto>();
+		Column.SetValueFormatter(Formatter);
+
+		TSharedRef<ITableCellValueSorter> Sorter = MakeShared<FSorterByDoubleValue>(ColumnRef);
+		Column.SetValueSorter(Sorter);
+
+		Column.SetAggregation(ETableColumnAggregation::Min);
+
+		AddColumn(ColumnRef);
+	}
+	//////////////////////////////////////////////////
+	// Launched Thread Id Column
+	{
+		TSharedRef<FTableColumn> ColumnRef = MakeShared<FTableColumn>(FTaskTableColumns::LaunchedThreadIdColumnId);
+		FTableColumn& Column = *ColumnRef;
+
+		Column.SetIndex(ColumnIndex++);
+
+		Column.SetShortName(LOCTEXT("LaunchedThreadIdColumnName", "Launched Thread Id"));
+		Column.SetTitleName(LOCTEXT("LaunchedThreadIdColumnTitle", "Launched Thread Id"));
+		Column.SetDescription(LOCTEXT("LaunchedThreadIdColumnDesc", "The thread the task was Launched on."));
+
+		Column.SetFlags(ETableColumnFlags::CanBeHidden | ETableColumnFlags::CanBeFiltered);
+
+		Column.SetHorizontalAlignment(HAlign_Left);
+		Column.SetInitialWidth(100.0f);
+
+		Column.SetDataType(ETableCellDataType::Int64);
+
+		TSharedRef<ITableCellValueGetter> Getter = MakeShared<FTaskColumnValueGetter<DefaultTaskFieldGetterFuncts::GetLaunchedThreadId>>();;
+		Column.SetValueGetter(Getter);
+
+		TSharedRef<ITableCellValueFormatter> Formatter = MakeShared<FInt64ValueFormatterAsNumber>();
+		Column.SetValueFormatter(Formatter);
+
+		TSharedRef<ITableCellValueSorter> Sorter = MakeShared<FSorterByInt64Value>(ColumnRef);
+		Column.SetValueSorter(Sorter);
+
+		AddColumn(ColumnRef);
+	}
+
+	//////////////////////////////////////////////////
+	// Scheduled Timestamp Column
+	{
+		TSharedRef<FTableColumn> ColumnRef = MakeShared<FTableColumn>(FTaskTableColumns::ScheduledTimestampColumnId);
+		FTableColumn& Column = *ColumnRef;
+
+		Column.SetIndex(ColumnIndex++);
+
+		Column.SetShortName(LOCTEXT("ScheduledTimestampColumnName", "Scheduled"));
+		Column.SetTitleName(LOCTEXT("ScheduledTimestampColumnTitle", "Scheduled"));
+		Column.SetDescription(LOCTEXT("ScheduledTimestampColumnDesc", "The time when the task was Scheduled."));
+
+		Column.SetFlags(ETableColumnFlags::ShouldBeVisible | ETableColumnFlags::CanBeHidden | ETableColumnFlags::CanBeFiltered);
+
+		Column.SetHorizontalAlignment(HAlign_Left);
+		Column.SetInitialWidth(100.0f);
+
+		Column.SetDataType(ETableCellDataType::Double);
+
+		TSharedRef<ITableCellValueGetter> Getter = MakeShared<FTaskColumnValueGetter<DefaultTaskFieldGetterFuncts::GetScheduledTimestamp>>();
+		Column.SetValueGetter(Getter);
+
+		TSharedRef<ITableCellValueFormatter> Formatter = MakeShared<FDoubleValueFormatterAsTimeAuto>();
+		Column.SetValueFormatter(Formatter);
+
+		TSharedRef<ITableCellValueSorter> Sorter = MakeShared<FSorterByDoubleValue>(ColumnRef);
+		Column.SetValueSorter(Sorter);
+
+		Column.SetAggregation(ETableColumnAggregation::Min);
+
+		AddColumn(ColumnRef);
+	}
+	//////////////////////////////////////////////////
+	// Scheduled Thread Id Column
+	{
+		TSharedRef<FTableColumn> ColumnRef = MakeShared<FTableColumn>(FTaskTableColumns::ScheduledThreadIdColumnId);
+		FTableColumn& Column = *ColumnRef;
+
+		Column.SetIndex(ColumnIndex++);
+
+		Column.SetShortName(LOCTEXT("ScheduledThreadIdColumnName", "Scheduled Thread Id"));
+		Column.SetTitleName(LOCTEXT("ScheduledThreadIdColumnTitle", "Scheduled Thread Id"));
+		Column.SetDescription(LOCTEXT("ScheduledThreadIdColumnDesc", "The thread the task was Scheduled on."));
+
+		Column.SetFlags(ETableColumnFlags::CanBeHidden | ETableColumnFlags::CanBeFiltered);
+
+		Column.SetHorizontalAlignment(HAlign_Left);
+		Column.SetInitialWidth(100.0f);
+
+		Column.SetDataType(ETableCellDataType::Int64);
+
+		TSharedRef<ITableCellValueGetter> Getter = MakeShared<FTaskColumnValueGetter<DefaultTaskFieldGetterFuncts::GetScheduledThreadId>>();
+		Column.SetValueGetter(Getter);
+
+		TSharedRef<ITableCellValueFormatter> Formatter = MakeShared<FInt64ValueFormatterAsNumber>();
+		Column.SetValueFormatter(Formatter);
+
+		TSharedRef<ITableCellValueSorter> Sorter = MakeShared<FSorterByInt64Value>(ColumnRef);
+		Column.SetValueSorter(Sorter);
+
+		AddColumn(ColumnRef);
+	}
+	//////////////////////////////////////////////////
+	// Started Timestamp Column
+	{
+		TSharedRef<FTableColumn> ColumnRef = MakeShared<FTableColumn>(FTaskTableColumns::StartedTimestampColumnId);
+		FTableColumn& Column = *ColumnRef;
+
+		Column.SetIndex(ColumnIndex++);
+
+		Column.SetShortName(LOCTEXT("StartedTimestampColumnName", "Started"));
+		Column.SetTitleName(LOCTEXT("StartedTimestampColumnTitle", "Started"));
+		Column.SetDescription(LOCTEXT("StartedTimestampColumnDesc", "The time when the task was Started."));
+
+		Column.SetFlags(ETableColumnFlags::ShouldBeVisible | ETableColumnFlags::CanBeHidden | ETableColumnFlags::CanBeFiltered);
+
+		Column.SetHorizontalAlignment(HAlign_Left);
+		Column.SetInitialWidth(100.0f);
+
+		Column.SetDataType(ETableCellDataType::Double);
+
+		TSharedRef<ITableCellValueGetter> Getter = MakeShared<FTaskColumnValueGetter<DefaultTaskFieldGetterFuncts::GetStartedTimestamp>>();
+		Column.SetValueGetter(Getter);
+
+		TSharedRef<ITableCellValueFormatter> Formatter = MakeShared<FDoubleValueFormatterAsTimeAuto>();
+		Column.SetValueFormatter(Formatter);
+
+		TSharedRef<ITableCellValueSorter> Sorter = MakeShared<FSorterByDoubleValue>(ColumnRef);
+		Column.SetValueSorter(Sorter);
+
+		Column.SetAggregation(ETableColumnAggregation::Min);
+
+		AddColumn(ColumnRef);
+	}
+	//////////////////////////////////////////////////
+	// Started Thread Id Column
+	{
+		TSharedRef<FTableColumn> ColumnRef = MakeShared<FTableColumn>(FTaskTableColumns::StartedThreadIdColumnId);
+		FTableColumn& Column = *ColumnRef;
+
+		Column.SetIndex(ColumnIndex++);
+
+		Column.SetShortName(LOCTEXT("StartedThreadIdColumnName", "Started Thread Id"));
+		Column.SetTitleName(LOCTEXT("StartedThreadIdColumnTitle", "Started Thread Id"));
+		Column.SetDescription(LOCTEXT("StartedThreadIdColumnDesc", "The thread the task was Started on."));
+
+		Column.SetFlags(ETableColumnFlags::CanBeHidden | ETableColumnFlags::CanBeFiltered);
+
+		Column.SetHorizontalAlignment(HAlign_Left);
+		Column.SetInitialWidth(100.0f);
+
+		Column.SetDataType(ETableCellDataType::Int64);
+
+		TSharedRef<ITableCellValueGetter> Getter = MakeShared<FTaskColumnValueGetter<DefaultTaskFieldGetterFuncts::GetStartedThreadId>>();
+		Column.SetValueGetter(Getter);
+
+		TSharedRef<ITableCellValueFormatter> Formatter = MakeShared<FInt64ValueFormatterAsNumber>();
+		Column.SetValueFormatter(Formatter);
+
+		TSharedRef<ITableCellValueSorter> Sorter = MakeShared<FSorterByInt64Value>(ColumnRef);
+		Column.SetValueSorter(Sorter);
+
+		AddColumn(ColumnRef);
+	}
+
+	//////////////////////////////////////////////////
+	// Completed Timestamp Column
+	{
+		TSharedRef<FTableColumn> ColumnRef = MakeShared<FTableColumn>(FTaskTableColumns::CompletedTimestampColumnId);
+		FTableColumn& Column = *ColumnRef;
+
+		Column.SetIndex(ColumnIndex++);
+
+		Column.SetShortName(LOCTEXT("CompletedTimestampColumnName", "Completed"));
+		Column.SetTitleName(LOCTEXT("CompletedTimestampColumnTitle", "Completed"));
+		Column.SetDescription(LOCTEXT("CompletedTimestampColumnDesc", "The time when the task was Completed."));
+
+		Column.SetFlags(ETableColumnFlags::ShouldBeVisible | ETableColumnFlags::CanBeHidden | ETableColumnFlags::CanBeFiltered);
+
+		Column.SetHorizontalAlignment(HAlign_Left);
+		Column.SetInitialWidth(100.0f);
+
+		Column.SetDataType(ETableCellDataType::Double);
+
+		TSharedRef<ITableCellValueGetter> Getter = MakeShared<FTaskColumnValueGetter<DefaultTaskFieldGetterFuncts::GetCompletedTimestamp>>();
+		Column.SetValueGetter(Getter);
+
+		TSharedRef<ITableCellValueFormatter> Formatter = MakeShared<FDoubleValueFormatterAsTimeAuto>();
+		Column.SetValueFormatter(Formatter);
+
+		TSharedRef<ITableCellValueSorter> Sorter = MakeShared<FSorterByDoubleValue>(ColumnRef);
+		Column.SetValueSorter(Sorter);
+
+		Column.SetAggregation(ETableColumnAggregation::Min);
+
+		AddColumn(ColumnRef);
+	}
+	//////////////////////////////////////////////////
+	// Completed Thread Id Column
+	{
+		TSharedRef<FTableColumn> ColumnRef = MakeShared<FTableColumn>(FTaskTableColumns::CompletedThreadIdColumnId);
+		FTableColumn& Column = *ColumnRef;
+
+		Column.SetIndex(ColumnIndex++);
+
+		Column.SetShortName(LOCTEXT("CompletedThreadIdColumnName", "Completed Thread Id"));
+		Column.SetTitleName(LOCTEXT("CompletedThreadIdColumnTitle", "Completed Thread Id"));
+		Column.SetDescription(LOCTEXT("CompletedThreadIdColumnDesc", "The thread the task was Completed on."));
+
+		Column.SetFlags(ETableColumnFlags::CanBeHidden | ETableColumnFlags::CanBeFiltered);
+
+		Column.SetHorizontalAlignment(HAlign_Left);
+		Column.SetInitialWidth(100.0f);
+
+		Column.SetDataType(ETableCellDataType::Int64);
+
+		TSharedRef<ITableCellValueGetter> Getter = MakeShared<FTaskColumnValueGetter<DefaultTaskFieldGetterFuncts::GetCompletedThreadId>>();
+		Column.SetValueGetter(Getter);
+
+		TSharedRef<ITableCellValueFormatter> Formatter = MakeShared<FInt64ValueFormatterAsNumber>();
+		Column.SetValueFormatter(Formatter);
+
+		TSharedRef<ITableCellValueSorter> Sorter = MakeShared<FSorterByInt64Value>(ColumnRef);
+		Column.SetValueSorter(Sorter);
+
+		AddColumn(ColumnRef);
+	}
+
+	//////////////////////////////////////////////////
+	// Finished Timestamp Column
+	{
+		TSharedRef<FTableColumn> ColumnRef = MakeShared<FTableColumn>(FTaskTableColumns::FinishedTimestampColumnId);
+		FTableColumn& Column = *ColumnRef;
+
+		Column.SetIndex(ColumnIndex++);
+
+		Column.SetShortName(LOCTEXT("FinishedTimestampColumnName", "Finished"));
+		Column.SetTitleName(LOCTEXT("FinishedTimestampColumnTitle", "Finished"));
+		Column.SetDescription(LOCTEXT("FinishedTimestampColumnDesc", "The time when the task was Finished."));
+
+		Column.SetFlags(ETableColumnFlags::ShouldBeVisible | ETableColumnFlags::CanBeHidden | ETableColumnFlags::CanBeFiltered);
+
+		Column.SetHorizontalAlignment(HAlign_Left);
+		Column.SetInitialWidth(100.0f);
+
+		Column.SetDataType(ETableCellDataType::Double);
+
+		TSharedRef<ITableCellValueGetter> Getter = MakeShared<FTaskColumnValueGetter<DefaultTaskFieldGetterFuncts::GetFinishedTimestamp>>();
+		Column.SetValueGetter(Getter);
+
+		TSharedRef<ITableCellValueFormatter> Formatter = MakeShared<FDoubleValueFormatterAsTimeAuto>();
+		Column.SetValueFormatter(Formatter);
+
+		TSharedRef<ITableCellValueSorter> Sorter = MakeShared<FSorterByDoubleValue>(ColumnRef);
+		Column.SetValueSorter(Sorter);
+
+		Column.SetAggregation(ETableColumnAggregation::Min);
+
+		AddColumn(ColumnRef);
+	}
+	//////////////////////////////////////////////////
+	// Num Nested Tasks Column
+	{
+		TSharedRef<FTableColumn> ColumnRef = MakeShared<FTableColumn>(FTaskTableColumns::NumNestedColumnId);
+		FTableColumn& Column = *ColumnRef;
+
+		Column.SetIndex(ColumnIndex++);
+
+		Column.SetShortName(LOCTEXT("NestedTasksColumnName", "Nested Tasks"));
+		Column.SetTitleName(LOCTEXT("NestedTasksColumnTitle", "Nested Tasks"));
+		Column.SetDescription(LOCTEXT("NestedTasksColumnDesc", "The number of nested tasks."));
+
+		Column.SetFlags(ETableColumnFlags::ShouldBeVisible | ETableColumnFlags::CanBeHidden | ETableColumnFlags::CanBeFiltered);
+
+		Column.SetHorizontalAlignment(HAlign_Left);
+		Column.SetInitialWidth(100.0f);
+
+		Column.SetDataType(ETableCellDataType::Int64);
+
+		TSharedRef<ITableCellValueGetter> Getter = MakeShared<FTaskColumnValueGetter<DefaultTaskFieldGetterFuncts::GetNumNested>>();
+		Column.SetValueGetter(Getter);
+
+		TSharedRef<ITableCellValueFormatter> Formatter = MakeShared<FInt64ValueFormatterAsNumber>();
+		Column.SetValueFormatter(Formatter);
+
+		TSharedRef<ITableCellValueSorter> Sorter = MakeShared<FSorterByInt64Value>(ColumnRef);
+		Column.SetValueSorter(Sorter);
+
+		AddColumn(ColumnRef);
+	}
+	//////////////////////////////////////////////////
+	// Num Subsequents Tasks Column
+	{
+		TSharedRef<FTableColumn> ColumnRef = MakeShared<FTableColumn>(FTaskTableColumns::NumSubsequentsColumnId);
+		FTableColumn& Column = *ColumnRef;
+
+		Column.SetIndex(ColumnIndex++);
+
+		Column.SetShortName(LOCTEXT("SubsequentTasksColumnName", "Subsequent Tasks"));
+		Column.SetTitleName(LOCTEXT("SubsequentTasksColumnTitle", "Subsequent Tasks"));
+		Column.SetDescription(LOCTEXT("SubsequentTasksColumnDesc", "The number of Subsequent tasks."));
+
+		Column.SetFlags(ETableColumnFlags::ShouldBeVisible | ETableColumnFlags::CanBeHidden | ETableColumnFlags::CanBeFiltered);
+
+		Column.SetHorizontalAlignment(HAlign_Left);
+		Column.SetInitialWidth(100.0f);
+
+		Column.SetDataType(ETableCellDataType::Int64);
+
+		TSharedRef<ITableCellValueGetter> Getter = MakeShared<FTaskColumnValueGetter<DefaultTaskFieldGetterFuncts::GetNumSubsequents>>();
+		Column.SetValueGetter(Getter);
+
+		TSharedRef<ITableCellValueFormatter> Formatter = MakeShared<FInt64ValueFormatterAsNumber>();
+		Column.SetValueFormatter(Formatter);
+
+		TSharedRef<ITableCellValueSorter> Sorter = MakeShared<FSorterByInt64Value>(ColumnRef);
+		Column.SetValueSorter(Sorter);
+
+		AddColumn(ColumnRef);
+	}
+	//////////////////////////////////////////////////
+	// Num Prerequisites Tasks Column
+	{
+		TSharedRef<FTableColumn> ColumnRef = MakeShared<FTableColumn>(FTaskTableColumns::NumPrerequisitesColumnId);
+		FTableColumn& Column = *ColumnRef;
+
+		Column.SetIndex(ColumnIndex++);
+
+		Column.SetShortName(LOCTEXT("PrerequisiteTasksColumnName", "Prerequisite Tasks"));
+		Column.SetTitleName(LOCTEXT("PrerequisiteTasksColumnTitle", "Prerequisite Tasks"));
+		Column.SetDescription(LOCTEXT("PrerequisiteTasksColumnDesc", "The number of Prerequisite tasks."));
+
+		Column.SetFlags(ETableColumnFlags::ShouldBeVisible | ETableColumnFlags::CanBeHidden | ETableColumnFlags::CanBeFiltered);
+
+		Column.SetHorizontalAlignment(HAlign_Left);
+		Column.SetInitialWidth(100.0f);
+
+		Column.SetDataType(ETableCellDataType::Int64);
+
+		TSharedRef<ITableCellValueGetter> Getter = MakeShared<FTaskColumnValueGetter<DefaultTaskFieldGetterFuncts::GetNumPrerequisites>>();
+		Column.SetValueGetter(Getter);
+
+		TSharedRef<ITableCellValueFormatter> Formatter = MakeShared<FInt64ValueFormatterAsNumber>();
+		Column.SetValueFormatter(Formatter);
+
+		TSharedRef<ITableCellValueSorter> Sorter = MakeShared<FSorterByInt64Value>(ColumnRef);
+		Column.SetValueSorter(Sorter);
+
+		AddColumn(ColumnRef);
+	}
+}
+
+////////////////////////////////////////////////////////////////////////////////////////////////////
+
+} // namespace Insights
+
+#undef LOCTEXT_NAMESPACE
diff --git a/Engine/Source/Developer/TraceInsights/Private/Insights/TaskGraphProfiler/ViewModels/TaskTable.h b/Engine/Source/Developer/TraceInsights/Private/Insights/TaskGraphProfiler/ViewModels/TaskTable.h
new file mode 100644
--- /dev/null
+++ b/Engine/Source/Developer/TraceInsights/Private/Insights/TaskGraphProfiler/ViewModels/TaskTable.h
@@ -0,0 +1,63 @@
+// Copyright Epic Games, Inc. All Rights Reserved.
+
+#pragma once
+
+#include "CoreMinimal.h"
+
+// Insights
+#include "Insights/TaskGraphProfiler/ViewModels/TaskEntry.h"
+#include "Insights/Table/ViewModels/Table.h"
+
+namespace Insights
+{
+
+class FTableColumn;
+
+////////////////////////////////////////////////////////////////////////////////////////////////////
+
+// Column identifiers
+struct FTaskTableColumns
+{
+	static const FName CreatedTimestampColumnId;
+	static const FName CreatedThreadIdColumnId;
+	static const FName LaunchedTimestampColumnId;
+	static const FName LaunchedThreadIdColumnId;
+	static const FName ScheduledTimestampColumnId;
+	static const FName ScheduledThreadIdColumnId;
+	static const FName StartedTimestampColumnId;
+	static const FName StartedThreadIdColumnId;
+	static const FName CompletedTimestampColumnId;
+	static const FName CompletedThreadIdColumnId;
+	static const FName FinishedTimestampColumnId;
+	static const FName NumNestedColumnId;
+	static const FName NumSubsequentsColumnId;
+	static const FName NumPrerequisitesColumnId;
+};
+
+////////////////////////////////////////////////////////////////////////////////////////////////////
+
+class FTaskTable : public FTable
+{
+public:
+	FTaskTable();
+	virtual ~FTaskTable();
+
+	virtual void Reset();
+
+	TArray<FTaskEntry>& GetTaskEntries() { return TaskEntries; }
+	const TArray<FTaskEntry>& GetTaskEntries() const { return TaskEntries; }
+
+	bool IsValidRowIndex(int32 InIndex) const { return InIndex >= 0 && InIndex < TaskEntries.Num(); }
+	const FTaskEntry* GetTask(int32 InIndex) const { return IsValidRowIndex(InIndex) ? &TaskEntries[InIndex] : nullptr; }
+	const FTaskEntry& GetTaskChecked(int32 InIndex) const { check(IsValidRowIndex(InIndex)); return TaskEntries[InIndex]; }
+
+private:
+	void AddDefaultColumns();
+
+private:
+	TArray<FTaskEntry> TaskEntries;
+};
+
+////////////////////////////////////////////////////////////////////////////////////////////////////
+
+} // namespace Insights
diff --git a/Engine/Source/Developer/TraceInsights/Private/Insights/TaskGraphProfiler/Widgets/STaskTableTreeView.cpp b/Engine/Source/Developer/TraceInsights/Private/Insights/TaskGraphProfiler/Widgets/STaskTableTreeView.cpp
new file mode 100644
--- /dev/null
+++ b/Engine/Source/Developer/TraceInsights/Private/Insights/TaskGraphProfiler/Widgets/STaskTableTreeView.cpp
@@ -0,0 +1,188 @@
+// Copyright Epic Games, Inc. All Rights Reserved.
+
+#include "STaskTableTreeView.h"
+
+#include "EditorStyleSet.h"
+#include "SlateOptMacros.h"
+#include "Styling/CoreStyle.h"
+#include "TraceServices/AnalysisService.h"
+#include "TraceServices/Model/TasksProfiler.h"
+
+// Insights
+#include "Insights/TaskGraphProfiler/TaskGraphProfilerManager.h"
+#include "Insights/TaskGraphProfiler/ViewModels/TaskEntry.h"
+#include "Insights/TaskGraphProfiler/ViewModels/TaskNode.h"
+#include "Insights/Table/ViewModels/TableColumn.h"
+#include "Insights/TimingProfilerManager.h"
+#include "Insights/ViewModels/FilterConfigurator.h"
+
+#include <limits>
+
+#define LOCTEXT_NAMESPACE "STaskTableTreeView"
+
+using namespace TraceServices;
+
+namespace Insights
+{
+
+////////////////////////////////////////////////////////////////////////////////////////////////////
+
+STaskTableTreeView::STaskTableTreeView()
+{
+	bRunInAsyncMode = true;
+}
+
+////////////////////////////////////////////////////////////////////////////////////////////////////
+
+STaskTableTreeView::~STaskTableTreeView()
+{
+}
+
+////////////////////////////////////////////////////////////////////////////////////////////////////
+
+void STaskTableTreeView::Construct(const FArguments& InArgs, TSharedPtr<Insights::FTaskTable> InTablePtr)
+{
+	ConstructWidget(InTablePtr);
+}
+
+////////////////////////////////////////////////////////////////////////////////////////////////////
+
+void STaskTableTreeView::Reset()
+{
+	STableTreeView::Reset();
+}
+
+////////////////////////////////////////////////////////////////////////////////////////////////////
+
+void STaskTableTreeView::Tick(const FGeometry& AllottedGeometry, const double InCurrentTime, const float InDeltaTime)
+{
+	STableTreeView::Tick(AllottedGeometry, InCurrentTime, InDeltaTime);
+
+	if (!bIsUpdateRunning)
+	{
+		RebuildTree(false);
+	}
+}
+
+////////////////////////////////////////////////////////////////////////////////////////////////////
+
+void STaskTableTreeView::RebuildTree(bool bResync)
+{
+	if (FTimingProfilerManager::Get()->GetSelectionStartTime() != QueryStartTime ||
+		FTimingProfilerManager::Get()->GetSelectionEndTime() != QueryEndTime)
+	{
+		QueryStartTime = FTimingProfilerManager::Get()->GetSelectionStartTime();
+		QueryEndTime = FTimingProfilerManager::Get()->GetSelectionEndTime();
+		TSharedPtr<FTaskTable> TaskTable = GetTaskTable();
+		TArray<FTaskEntry>& Allocs = TaskTable->GetTaskEntries();
+		Allocs.Empty();
+		TableTreeNodes.Empty();
+
+		if (QueryStartTime < QueryEndTime && Session.IsValid())
+		{
+			TraceServices::FAnalysisSessionReadScope SessionReadScope(*Session.Get());
+
+			const TraceServices::ITasksProvider* TasksProvider = TraceServices::ReadTasksProvider(*Session.Get());
+
+			if (TasksProvider)
+			{
+				FName BaseNodeName(TEXT("task"));
+
+				TArray<FTableTreeNodePtr>* Nodes = &TableTreeNodes;
+
+				TasksProvider->EnumerateTasks(QueryStartTime, QueryEndTime, [&Allocs, &TaskTable, &BaseNodeName, Nodes](const TraceServices::FTaskInfo& TaskInfo)
+					{
+						Allocs.Emplace(TaskInfo);
+						uint32 Index = Allocs.Num() - 1;
+						FName NodeName(BaseNodeName, Index);
+						FTaskNodePtr NodePtr = MakeShared<FTaskNode>(NodeName, TaskTable, Index);
+						Nodes->Add(NodePtr);
+						return TraceServices::ETaskEnumerationResult::Continue;
+					});
+			}
+		}
+
+		UpdateTree();
+		TreeView->RebuildList();
+	}
+}
+
+////////////////////////////////////////////////////////////////////////////////////////////////////
+
+bool STaskTableTreeView::IsRunning() const
+{
+	return STableTreeView::IsRunning();
+}
+
+////////////////////////////////////////////////////////////////////////////////////////////////////
+
+double STaskTableTreeView::GetAllOperationsDuration()
+{
+	return STableTreeView::GetAllOperationsDuration();
+}
+
+////////////////////////////////////////////////////////////////////////////////////////////////////
+
+FText STaskTableTreeView::GetCurrentOperationName() const
+{
+	return STableTreeView::GetCurrentOperationName();
+}
+
+////////////////////////////////////////////////////////////////////////////////////////////////////
+
+BEGIN_SLATE_FUNCTION_BUILD_OPTIMIZATION
+TSharedPtr<SWidget> STaskTableTreeView::ConstructToolbar()
+{
+	return nullptr;
+}
+END_SLATE_FUNCTION_BUILD_OPTIMIZATION
+
+////////////////////////////////////////////////////////////////////////////////////////////////////
+
+void STaskTableTreeView::ApplyColumnConfig(const TArrayView<FColumnConfig>& Preset)
+{
+	for (const TSharedRef<FTableColumn>& ColumnRef : Table->GetColumns())
+	{
+		FTableColumn& Column = ColumnRef.Get();
+		for (const FColumnConfig& Config : Preset)
+		{
+			if (Column.GetId() == Config.ColumnId)
+			{
+				if (Config.bIsVisible)
+				{
+					ShowColumn(Column);
+					if (Config.Width > 0.0f)
+					{
+						TreeViewHeaderRow->SetColumnWidth(Column.GetId(), Config.Width);
+					}
+				}
+				else
+				{
+					HideColumn(Column);
+				}
+			}
+		}
+	}
+}
+
+////////////////////////////////////////////////////////////////////////////////////////////////////
+
+BEGIN_SLATE_FUNCTION_BUILD_OPTIMIZATION
+TSharedPtr<SWidget> STaskTableTreeView::ConstructFooter()
+{
+	return nullptr;
+}
+END_SLATE_FUNCTION_BUILD_OPTIMIZATION
+
+////////////////////////////////////////////////////////////////////////////////////////////////////
+
+void STaskTableTreeView::InternalCreateGroupings()
+{
+	STableTreeView::InternalCreateGroupings();
+}
+
+////////////////////////////////////////////////////////////////////////////////////////////////////
+
+} // namespace Insights
+
+#undef LOCTEXT_NAMESPACE
diff --git a/Engine/Source/Developer/TraceInsights/Private/Insights/TaskGraphProfiler/Widgets/STaskTableTreeView.h b/Engine/Source/Developer/TraceInsights/Private/Insights/TaskGraphProfiler/Widgets/STaskTableTreeView.h
new file mode 100644
--- /dev/null
+++ b/Engine/Source/Developer/TraceInsights/Private/Insights/TaskGraphProfiler/Widgets/STaskTableTreeView.h
@@ -0,0 +1,90 @@
+// Copyright Epic Games, Inc. All Rights Reserved.
+
+#pragma once
+
+#include "CoreMinimal.h"
+#include "TraceServices/Model/TasksProfiler.h"
+
+// Insights
+#include "Insights/TaskGraphProfiler/ViewModels/TaskTable.h"
+#include "Insights/Table/Widgets/STableTreeView.h"
+
+namespace Insights
+{
+
+////////////////////////////////////////////////////////////////////////////////////////////////////
+
+class STaskTableTreeView : public STableTreeView
+{
+private:
+	struct FColumnConfig
+	{
+		const FName& ColumnId;
+		bool bIsVisible;
+		float Width;
+	};
+
+public:
+	/** Default constructor. */
+	STaskTableTreeView();
+
+	/** Virtual destructor. */
+	virtual ~STaskTableTreeView();
+
+	SLATE_BEGIN_ARGS(STaskTableTreeView) {}
+	SLATE_END_ARGS()
+
+	/**
+	 * Construct this widget
+	 * @param InArgs - The declaration data for this widget
+	 */
+	void Construct(const FArguments& InArgs, TSharedPtr<FTaskTable> InTablePtr);
+
+	virtual TSharedPtr<SWidget> ConstructToolbar() override;
+	virtual TSharedPtr<SWidget> ConstructFooter() override;
+
+	TSharedPtr<FTaskTable> GetTaskTable() { return StaticCastSharedPtr<FTaskTable>(GetTable()); }
+	const TSharedPtr<FTaskTable> GetTaskTable() const { return StaticCastSharedPtr<FTaskTable>(GetTable()); }
+
+	//void UpdateSourceTable(TSharedPtr<TraceServices::ITaskTable> SourceTable);
+
+	virtual void Reset();
+
+	/**
+	 * Ticks this widget.  Override in derived classes, but always call the parent implementation.
+	 *
+	 * @param  AllottedGeometry The space allotted for this widget
+	 * @param  InCurrentTime  Current absolute real time
+	 * @param  InDeltaTime  Real time passed since last tick
+	 */
+	virtual void Tick(const FGeometry& AllottedGeometry, const double InCurrentTime, const float InDeltaTime) override;
+
+	/**
+	 * Rebuilds the tree (if necessary).
+	 * @param bResync - If true, it forces a resync even if the list did not changed since last sync.
+	 */
+	virtual void RebuildTree(bool bResync);
+
+	////////////////////////////////////////////////////////////////////////////////////////////////////
+	// IAsyncOperationStatusProvider implementation
+
+	virtual bool IsRunning() const override;
+	virtual double GetAllOperationsDuration() override;
+	virtual FText GetCurrentOperationName() const override;
+
+	////////////////////////////////////////////////////////////////////////////////////////////////////
+
+protected:
+	virtual void InternalCreateGroupings() override;
+
+private:
+	void ApplyColumnConfig(const TArrayView<FColumnConfig>& Preset);
+
+private:
+	double QueryStartTime = 0.0f;
+	double QueryEndTime = 0.0f;
+};
+
+////////////////////////////////////////////////////////////////////////////////////////////////////
+
+} // namespace Insights
diff --git a/Engine/Source/Developer/TraceInsights/Private/Insights/TimingProfilerManager.cpp b/Engine/Source/Developer/TraceInsights/Private/Insights/TimingProfilerManager.cpp
--- a/Engine/Source/Developer/TraceInsights/Private/Insights/TimingProfilerManager.cpp
+++ b/Engine/Source/Developer/TraceInsights/Private/Insights/TimingProfilerManager.cpp
@@ -95,6 +95,8 @@ void FTimingProfilerManager::Initialize(IUnrealInsightsModule& InsightsModule)
 	FTimingProfilerCommands::Register();
 	BindCommands();
 
+	InsightsModule.OnRegisterMajorTabExtension(FInsightsManagerTabs::TimingProfilerTabId);
+
 	FInsightsManager::Get()->GetSessionChangedEvent().AddSP(this, &FTimingProfilerManager::OnSessionChanged);
 	OnSessionChanged();
 }
diff --git a/Engine/Source/Developer/TraceInsights/Private/Insights/TraceInsightsModule.cpp b/Engine/Source/Developer/TraceInsights/Private/Insights/TraceInsightsModule.cpp
--- a/Engine/Source/Developer/TraceInsights/Private/Insights/TraceInsightsModule.cpp
+++ b/Engine/Source/Developer/TraceInsights/Private/Insights/TraceInsightsModule.cpp
@@ -53,7 +53,7 @@ void FTraceInsightsModule::StartupModule()
 	RegisterComponent(FLoadingProfilerManager::CreateInstance());
 	RegisterComponent(FNetworkingProfilerManager::CreateInstance());
 	RegisterComponent(FMemoryProfilerManager::CreateInstance());
-	RegisterComponent(FTaskGraphProfilerManager::CreateInstance());
+	RegisterComponent(Insights::FTaskGraphProfilerManager::CreateInstance());
 
 	UnrealInsightsLayoutIni = FConfigCacheIni::GetDestIniFilename(TEXT("UnrealInsightsLayout"), nullptr, *FPaths::GeneratedConfigDir());
 }
@@ -365,7 +365,7 @@ const FInsightsMajorTabConfig& FTraceInsightsModule::FindMajorTabConfig(const FN
 
 ////////////////////////////////////////////////////////////////////////////////////////////////////
 
-const FOnRegisterMajorTabExtensions* FTraceInsightsModule::FindMajorTabLayoutExtension(const FName& InMajorTabId) const
+FOnRegisterMajorTabExtensions* FTraceInsightsModule::FindMajorTabLayoutExtension(const FName& InMajorTabId)
 {
 	return MajorTabExtensionDelegates.Find(InMajorTabId);
 }
diff --git a/Engine/Source/Developer/TraceInsights/Private/Insights/TraceInsightsModule.h b/Engine/Source/Developer/TraceInsights/Private/Insights/TraceInsightsModule.h
--- a/Engine/Source/Developer/TraceInsights/Private/Insights/TraceInsightsModule.h
+++ b/Engine/Source/Developer/TraceInsights/Private/Insights/TraceInsightsModule.h
@@ -63,7 +63,7 @@ public:
 
 	virtual const FInsightsMajorTabConfig& FindMajorTabConfig(const FName& InMajorTabId) const override;
 
-	const FOnRegisterMajorTabExtensions* FindMajorTabLayoutExtension(const FName& InMajorTabId) const;
+	FOnRegisterMajorTabExtensions* FindMajorTabLayoutExtension(const FName& InMajorTabId);
 
 	/** Retrieve ini path for saving persistent layout data. */
 	static const FString& GetUnrealInsightsLayoutIni();
diff --git a/Engine/Source/Developer/TraceInsights/Private/Insights/Widgets/STimingView.cpp b/Engine/Source/Developer/TraceInsights/Private/Insights/Widgets/STimingView.cpp
--- a/Engine/Source/Developer/TraceInsights/Private/Insights/Widgets/STimingView.cpp
+++ b/Engine/Source/Developer/TraceInsights/Private/Insights/Widgets/STimingView.cpp
@@ -4169,14 +4169,14 @@ void STimingView::ContextMenu_ShowTaskDependecies_Execute()
 
 bool STimingView::ContextMenu_ShowTaskDependecies_CanExecute()
 {
-	return FTaskGraphProfilerManager::Get()->GetIsAvailable();
+	return Insights::FTaskGraphProfilerManager::Get()->GetIsAvailable();
 }
 
 ////////////////////////////////////////////////////////////////////////////////////////////////////
 
 bool STimingView::ContextMenu_ShowTaskDependecies_IsChecked()
 {
-	return FTaskGraphProfilerManager::Get()->GetIsAvailable() && bShowEventRelations;
+	return Insights::FTaskGraphProfilerManager::Get()->GetIsAvailable() && bShowEventRelations;
 }
 
 ////////////////////////////////////////////////////////////////////////////////////////////////////
diff --git a/Engine/Source/Developer/TraceInsights/Public/Insights/IUnrealInsightsModule.h b/Engine/Source/Developer/TraceInsights/Public/Insights/IUnrealInsightsModule.h
--- a/Engine/Source/Developer/TraceInsights/Public/Insights/IUnrealInsightsModule.h
+++ b/Engine/Source/Developer/TraceInsights/Public/Insights/IUnrealInsightsModule.h
@@ -257,6 +257,8 @@ public:
 	/** Finds a major tab config for the specified id. */
 	virtual const FInsightsMajorTabConfig& FindMajorTabConfig(const FName& InMajorTabId) const = 0;
 
+	virtual FOnRegisterMajorTabExtensions* FindMajorTabLayoutExtension(const FName& InMajorTabId) = 0;
+
 	/** Sets the ini path for saving persistent layout data. */
 	virtual void SetUnrealInsightsLayoutIni(const FString& InIniPath) = 0;
 
diff --git a/Engine/Source/Developer/TraceServices/Private/Model/TasksProfiler.cpp b/Engine/Source/Developer/TraceServices/Private/Model/TasksProfiler.cpp
--- a/Engine/Source/Developer/TraceServices/Private/Model/TasksProfiler.cpp
+++ b/Engine/Source/Developer/TraceServices/Private/Model/TasksProfiler.cpp
@@ -436,4 +436,19 @@ namespace TraceServices
 	{
 		return Tasks.Num();
 	}
+
+	void FTasksProvider::EnumerateTasks(double StartTime, double EndTime, TaskCallback Callback) const
+	{
+		// A naive implementation for now
+		for (const FTaskInfo& Task : Tasks)
+		{
+			if (Task.CreatedTimestamp <= EndTime && Task.FinishedTimestamp >= StartTime)
+			{
+				if (Callback(Task) == ETaskEnumerationResult::Stop)
+				{
+					break;
+				}
+			}
+		}
+	}
 }
diff --git a/Engine/Source/Developer/TraceServices/Private/Model/TasksProfilerPrivate.h b/Engine/Source/Developer/TraceServices/Private/Model/TasksProfilerPrivate.h
--- a/Engine/Source/Developer/TraceServices/Private/Model/TasksProfilerPrivate.h
+++ b/Engine/Source/Developer/TraceServices/Private/Model/TasksProfilerPrivate.h
@@ -38,6 +38,7 @@ namespace TraceServices
 		virtual const FTaskInfo* TryGetTask(TaskTrace::FId TaskId) const;
 		virtual const FWaitingForTasks* TryGetWaiting(const TCHAR* TimerName, uint32 ThreadId, double Timestamp) const override;
 		virtual int64 GetNumTasks() const override;
+		virtual void EnumerateTasks(double StartTime, double EndTime, TaskCallback Callback) const override;
 
 	private:
 		// initializes conversion of task id to its index in the internal container, if it's not initialized already
diff --git a/Engine/Source/Developer/TraceServices/Public/TraceServices/Model/TasksProfiler.h b/Engine/Source/Developer/TraceServices/Public/TraceServices/Model/TasksProfiler.h
--- a/Engine/Source/Developer/TraceServices/Public/TraceServices/Model/TasksProfiler.h
+++ b/Engine/Source/Developer/TraceServices/Public/TraceServices/Model/TasksProfiler.h
@@ -72,10 +72,18 @@ namespace TraceServices
 		double FinishedTimestamp = FTaskInfo::InvalidTimestamp;
 	};
 
+	enum class ETaskEnumerationResult
+	{
+		Continue,
+		Stop,
+	};
+
 	// query interface to tasks info
 	class ITasksProvider : public IProvider
 	{
 	public:
+		typedef TFunctionRef<ETaskEnumerationResult(const FTaskInfo& Task)> TaskCallback;
+
 		virtual ~ITasksProvider() = default;
 
 		// returns a task for given thread and timestamp, if any, otherwise `nullptr`
@@ -88,6 +96,9 @@ namespace TraceServices
 
 		// returns the number of tasks stored in the provider
 		virtual int64 GetNumTasks() const = 0;
+
+		// Calls the callback for each task stored in the provider with CreatedTimestamp <= EndTime and FinishedTimestamp >= StartTime
+		virtual void EnumerateTasks(double StartTime, double EndTime, TaskCallback Callback) const = 0;
 	};
 
 	TRACESERVICES_API const ITasksProvider* ReadTasksProvider(const IAnalysisSession& Session);
