[Insights] Move Completed columns after Finished column in TaskTableTreeView

diff --git a/Engine/Source/Developer/TraceInsights/Private/Insights/TaskGraphProfiler/ViewModels/TaskTable.cpp b/Engine/Source/Developer/TraceInsights/Private/Insights/TaskGraphProfiler/ViewModels/TaskTable.cpp
--- a/Engine/Source/Developer/TraceInsights/Private/Insights/TaskGraphProfiler/ViewModels/TaskTable.cpp
+++ b/Engine/Source/Developer/TraceInsights/Private/Insights/TaskGraphProfiler/ViewModels/TaskTable.cpp
@@ -384,18 +384,17 @@
 
 		AddColumn(ColumnRef);
 	}
-
 	//////////////////////////////////////////////////
-	// Completed Timestamp Column
+	// Finished Timestamp Column
 	{
-		TSharedRef<FTableColumn> ColumnRef = MakeShared<FTableColumn>(FTaskTableColumns::CompletedTimestampColumnId);
+		TSharedRef<FTableColumn> ColumnRef = MakeShared<FTableColumn>(FTaskTableColumns::FinishedTimestampColumnId);
 		FTableColumn& Column = *ColumnRef;
 
 		Column.SetIndex(ColumnIndex++);
 
-		Column.SetShortName(LOCTEXT("CompletedTimestampColumnName", "Completed"));
-		Column.SetTitleName(LOCTEXT("CompletedTimestampColumnTitle", "Completed"));
-		Column.SetDescription(LOCTEXT("CompletedTimestampColumnDesc", "The time when the task was Completed."));
+		Column.SetShortName(LOCTEXT("FinishedTimestampColumnName", "Finished"));
+		Column.SetTitleName(LOCTEXT("FinishedTimestampColumnTitle", "Finished"));
+		Column.SetDescription(LOCTEXT("FinishedTimestampColumnDesc", "The time when the task was Finished."));
 
 		Column.SetFlags(ETableColumnFlags::ShouldBeVisible | ETableColumnFlags::CanBeHidden | ETableColumnFlags::CanBeFiltered);
 
@@ -404,7 +403,7 @@
 
 		Column.SetDataType(ETableCellDataType::Double);
 
-		TSharedRef<ITableCellValueGetter> Getter = MakeShared<FTaskColumnValueGetter<DefaultTaskFieldGetterFuncts::GetCompletedTimestamp>>();
+		TSharedRef<ITableCellValueGetter> Getter = MakeShared<FTaskColumnValueGetter<DefaultTaskFieldGetterFuncts::GetFinishedTimestamp>>();
 		Column.SetValueGetter(Getter);
 
 		TSharedRef<ITableCellValueFormatter> Formatter = MakeShared<FDoubleValueFormatterAsTimeAuto>();
@@ -418,6 +417,38 @@
 		AddColumn(ColumnRef);
 	}
 	//////////////////////////////////////////////////
+	// Completed Timestamp Column
+	{
+		TSharedRef<FTableColumn> ColumnRef = MakeShared<FTableColumn>(FTaskTableColumns::CompletedTimestampColumnId);
+		FTableColumn& Column = *ColumnRef;
+
+		Column.SetIndex(ColumnIndex++);
+
+		Column.SetShortName(LOCTEXT("CompletedTimestampColumnName", "Completed"));
+		Column.SetTitleName(LOCTEXT("CompletedTimestampColumnTitle", "Completed"));
+		Column.SetDescription(LOCTEXT("CompletedTimestampColumnDesc", "The time when the task was Completed."));
+
+		Column.SetFlags(ETableColumnFlags::ShouldBeVisible | ETableColumnFlags::CanBeHidden | ETableColumnFlags::CanBeFiltered);
+
+		Column.SetHorizontalAlignment(HAlign_Left);
+		Column.SetInitialWidth(100.0f);
+
+		Column.SetDataType(ETableCellDataType::Double);
+
+		TSharedRef<ITableCellValueGetter> Getter = MakeShared<FTaskColumnValueGetter<DefaultTaskFieldGetterFuncts::GetCompletedTimestamp>>();
+		Column.SetValueGetter(Getter);
+
+		TSharedRef<ITableCellValueFormatter> Formatter = MakeShared<FDoubleValueFormatterAsTimeAuto>();
+		Column.SetValueFormatter(Formatter);
+
+		TSharedRef<ITableCellValueSorter> Sorter = MakeShared<FSorterByDoubleValue>(ColumnRef);
+		Column.SetValueSorter(Sorter);
+
+		Column.SetAggregation(ETableColumnAggregation::Min);
+
+		AddColumn(ColumnRef);
+	}
+	//////////////////////////////////////////////////
 	// Completed Thread Id Column
 	{
 		TSharedRef<FTableColumn> ColumnRef = MakeShared<FTableColumn>(FTaskTableColumns::CompletedThreadIdColumnId);
@@ -447,37 +478,4 @@
 
 		AddColumn(ColumnRef);
 	}
-
 	//////////////////////////////////////////////////
-	// Finished Timestamp Column
-	{
-		TSharedRef<FTableColumn> ColumnRef = MakeShared<FTableColumn>(FTaskTableColumns::FinishedTimestampColumnId);
-		FTableColumn& Column = *ColumnRef;
-
-		Column.SetIndex(ColumnIndex++);
-
-		Column.SetShortName(LOCTEXT("FinishedTimestampColumnName", "Finished"));
-		Column.SetTitleName(LOCTEXT("FinishedTimestampColumnTitle", "Finished"));
-		Column.SetDescription(LOCTEXT("FinishedTimestampColumnDesc", "The time when the task was Finished."));
-
-		Column.SetFlags(ETableColumnFlags::ShouldBeVisible | ETableColumnFlags::CanBeHidden | ETableColumnFlags::CanBeFiltered);
-
-		Column.SetHorizontalAlignment(HAlign_Left);
-		Column.SetInitialWidth(100.0f);
-
-		Column.SetDataType(ETableCellDataType::Double);
-
-		TSharedRef<ITableCellValueGetter> Getter = MakeShared<FTaskColumnValueGetter<DefaultTaskFieldGetterFuncts::GetFinishedTimestamp>>();
-		Column.SetValueGetter(Getter);
-
-		TSharedRef<ITableCellValueFormatter> Formatter = MakeShared<FDoubleValueFormatterAsTimeAuto>();
-		Column.SetValueFormatter(Formatter);
-
-		TSharedRef<ITableCellValueSorter> Sorter = MakeShared<FSorterByDoubleValue>(ColumnRef);
-		Column.SetValueSorter(Sorter);
-
-		Column.SetAggregation(ETableColumnAggregation::Min);
-
-		AddColumn(ColumnRef);
-	}
-	//////////////////////////////////////////////////
