[Insights] Option to display relative timestamps in TaskTableTreeView

diff --git a/Engine/Source/Developer/TraceInsights/Private/Insights/TaskGraphProfiler/ViewModels/TaskTable.cpp b/Engine/Source/Developer/TraceInsights/Private/Insights/TaskGraphProfiler/ViewModels/TaskTable.cpp
--- a/Engine/Source/Developer/TraceInsights/Private/Insights/TaskGraphProfiler/ViewModels/TaskTable.cpp
+++ b/Engine/Source/Developer/TraceInsights/Private/Insights/TaskGraphProfiler/ViewModels/TaskTable.cpp
@@ -91,6 +91,52 @@ struct DefaultTaskFieldGetterFuncts
 	static FTableCellValue GetNumPrerequisites(const FTableColumn& Column, const FTaskEntry& Task) { return FTableCellValue((int64)Task.GetNumPrerequisites());	}
 };
 
+////////////////////////////////////////////////////////////////////////////////////////////////////
+
+struct RelativeToPreviousTaskFieldGetterFuncts
+{
+	static FTableCellValue GetCreatedTimestamp(const FTableColumn& Column, const FTaskEntry& Task) { return FTableCellValue(Task.GetCreatedTimestamp()); }
+	static FTableCellValue GetLaunchedTimestamp(const FTableColumn& Column, const FTaskEntry& Task) { return FTableCellValue(Task.GetLaunchedTimestamp() - Task.GetCreatedTimestamp()); }
+	static FTableCellValue GetScheduledTimestamp(const FTableColumn& Column, const FTaskEntry& Task) { return FTableCellValue(Task.GetScheduledTimestamp() - Task.GetLaunchedTimestamp()); }
+	static FTableCellValue GetStartedTimestamp(const FTableColumn& Column, const FTaskEntry& Task) { return FTableCellValue(Task.GetStartedTimestamp() - Task.GetScheduledTimestamp()); }
+	static FTableCellValue GetFinishedTimestamp(const FTableColumn& Column, const FTaskEntry& Task) { return FTableCellValue(Task.GetFinishedTimestamp() - Task.GetStartedTimestamp()); }
+	static FTableCellValue GetCompletedTimestamp(const FTableColumn& Column, const FTaskEntry& Task) 
+	{
+		// CompletedTimestamp can be 0, so don't show negative numbers in this case.
+		if (Task.GetCompletedTimestamp() != 0)
+		{
+			return FTableCellValue(Task.GetCompletedTimestamp() - Task.GetFinishedTimestamp());
+		}
+		else
+		{
+			return FTableCellValue(Task.GetCompletedTimestamp());
+		}
+	}
+};
+
+////////////////////////////////////////////////////////////////////////////////////////////////////
+
+struct RelativeToCreatedTaskFieldGetterFuncts
+{
+	static FTableCellValue GetCreatedTimestamp(const FTableColumn& Column, const FTaskEntry& Task) { return FTableCellValue(Task.GetCreatedTimestamp()); }
+	static FTableCellValue GetLaunchedTimestamp(const FTableColumn& Column, const FTaskEntry& Task) { return FTableCellValue(Task.GetLaunchedTimestamp() - Task.GetCreatedTimestamp()); }
+	static FTableCellValue GetScheduledTimestamp(const FTableColumn& Column, const FTaskEntry& Task) { return FTableCellValue(Task.GetScheduledTimestamp() - Task.GetCreatedTimestamp()); }
+	static FTableCellValue GetStartedTimestamp(const FTableColumn& Column, const FTaskEntry& Task) { return FTableCellValue(Task.GetStartedTimestamp() - Task.GetCreatedTimestamp()); }
+	static FTableCellValue GetFinishedTimestamp(const FTableColumn& Column, const FTaskEntry& Task) { return FTableCellValue(Task.GetFinishedTimestamp() - Task.GetCreatedTimestamp()); }
+	static FTableCellValue GetCompletedTimestamp(const FTableColumn& Column, const FTaskEntry& Task)
+	{
+		// CompletedTimestamp can be 0, so don't show negative numbers in this case.
+		if (Task.GetCompletedTimestamp() != 0)
+		{
+			return FTableCellValue(Task.GetCompletedTimestamp() - Task.GetCreatedTimestamp());
+		}
+		else
+		{
+			return FTableCellValue(Task.GetCompletedTimestamp());
+		}
+	}
+};
+
 ////////////////////////////////////////////////////////////////////////////////////////////////////
 // FTaskTable
 ////////////////////////////////////////////////////////////////////////////////////////////////////
@@ -572,6 +618,39 @@ void FTaskTable::AddDefaultColumns()
 
 ////////////////////////////////////////////////////////////////////////////////////////////////////
 
+void FTaskTable::SwitchToAbsoluteTimestamps()
+{
+	FindColumnChecked(FTaskTableColumns::LaunchedTimestampColumnId)->SetValueGetter(MakeShared<FTaskColumnValueGetter<DefaultTaskFieldGetterFuncts::GetLaunchedTimestamp>>());
+	FindColumnChecked(FTaskTableColumns::ScheduledTimestampColumnId)->SetValueGetter(MakeShared<FTaskColumnValueGetter<DefaultTaskFieldGetterFuncts::GetScheduledTimestamp>>());
+	FindColumnChecked(FTaskTableColumns::StartedTimestampColumnId)->SetValueGetter(MakeShared<FTaskColumnValueGetter<DefaultTaskFieldGetterFuncts::GetStartedTimestamp>>());
+	FindColumnChecked(FTaskTableColumns::FinishedTimestampColumnId)->SetValueGetter(MakeShared<FTaskColumnValueGetter<DefaultTaskFieldGetterFuncts::GetFinishedTimestamp>>());
+	FindColumnChecked(FTaskTableColumns::CompletedTimestampColumnId)->SetValueGetter(MakeShared<FTaskColumnValueGetter<DefaultTaskFieldGetterFuncts::GetCompletedTimestamp>>());
+}
+
+////////////////////////////////////////////////////////////////////////////////////////////////////
+
+void FTaskTable::SwitchToRelativeToPreviousTimestamps()
+{
+	FindColumnChecked(FTaskTableColumns::LaunchedTimestampColumnId)->SetValueGetter(MakeShared<FTaskColumnValueGetter<RelativeToPreviousTaskFieldGetterFuncts::GetLaunchedTimestamp>>());
+	FindColumnChecked(FTaskTableColumns::ScheduledTimestampColumnId)->SetValueGetter(MakeShared<FTaskColumnValueGetter<RelativeToPreviousTaskFieldGetterFuncts::GetScheduledTimestamp>>());
+	FindColumnChecked(FTaskTableColumns::StartedTimestampColumnId)->SetValueGetter(MakeShared<FTaskColumnValueGetter<RelativeToPreviousTaskFieldGetterFuncts::GetStartedTimestamp>>());
+	FindColumnChecked(FTaskTableColumns::FinishedTimestampColumnId)->SetValueGetter(MakeShared<FTaskColumnValueGetter<RelativeToPreviousTaskFieldGetterFuncts::GetFinishedTimestamp>>());
+	FindColumnChecked(FTaskTableColumns::CompletedTimestampColumnId)->SetValueGetter(MakeShared<FTaskColumnValueGetter<RelativeToPreviousTaskFieldGetterFuncts::GetCompletedTimestamp>>());
+}
+
+////////////////////////////////////////////////////////////////////////////////////////////////////
+
+void FTaskTable::SwitchToRelativeToCreatedTimestamps()
+{
+	FindColumnChecked(FTaskTableColumns::LaunchedTimestampColumnId)->SetValueGetter(MakeShared<FTaskColumnValueGetter<RelativeToCreatedTaskFieldGetterFuncts::GetLaunchedTimestamp>>());
+	FindColumnChecked(FTaskTableColumns::ScheduledTimestampColumnId)->SetValueGetter(MakeShared<FTaskColumnValueGetter<RelativeToCreatedTaskFieldGetterFuncts::GetScheduledTimestamp>>());
+	FindColumnChecked(FTaskTableColumns::StartedTimestampColumnId)->SetValueGetter(MakeShared<FTaskColumnValueGetter<RelativeToCreatedTaskFieldGetterFuncts::GetStartedTimestamp>>());
+	FindColumnChecked(FTaskTableColumns::FinishedTimestampColumnId)->SetValueGetter(MakeShared<FTaskColumnValueGetter<RelativeToCreatedTaskFieldGetterFuncts::GetFinishedTimestamp>>());
+	FindColumnChecked(FTaskTableColumns::CompletedTimestampColumnId)->SetValueGetter(MakeShared<FTaskColumnValueGetter<RelativeToCreatedTaskFieldGetterFuncts::GetCompletedTimestamp>>());
+}
+
+////////////////////////////////////////////////////////////////////////////////////////////////////
+
 } // namespace Insights
 
 #undef LOCTEXT_NAMESPACE
diff --git a/Engine/Source/Developer/TraceInsights/Private/Insights/TaskGraphProfiler/ViewModels/TaskTable.h b/Engine/Source/Developer/TraceInsights/Private/Insights/TaskGraphProfiler/ViewModels/TaskTable.h
--- a/Engine/Source/Developer/TraceInsights/Private/Insights/TaskGraphProfiler/ViewModels/TaskTable.h
+++ b/Engine/Source/Developer/TraceInsights/Private/Insights/TaskGraphProfiler/ViewModels/TaskTable.h
@@ -51,6 +51,10 @@ public:
 	const FTaskEntry* GetTask(int32 InIndex) const { return IsValidRowIndex(InIndex) ? &TaskEntries[InIndex] : nullptr; }
 	const FTaskEntry& GetTaskChecked(int32 InIndex) const { check(IsValidRowIndex(InIndex)); return TaskEntries[InIndex]; }
 
+	void SwitchToAbsoluteTimestamps();
+	void SwitchToRelativeToPreviousTimestamps();
+	void SwitchToRelativeToCreatedTimestamps();
+
 private:
 	void AddDefaultColumns();
 
diff --git a/Engine/Source/Developer/TraceInsights/Private/Insights/TaskGraphProfiler/Widgets/STaskTableTreeView.cpp b/Engine/Source/Developer/TraceInsights/Private/Insights/TaskGraphProfiler/Widgets/STaskTableTreeView.cpp
--- a/Engine/Source/Developer/TraceInsights/Private/Insights/TaskGraphProfiler/Widgets/STaskTableTreeView.cpp
+++ b/Engine/Source/Developer/TraceInsights/Private/Insights/TaskGraphProfiler/Widgets/STaskTableTreeView.cpp
@@ -7,6 +7,7 @@
 #include "Styling/CoreStyle.h"
 #include "TraceServices/AnalysisService.h"
 #include "TraceServices/Model/TasksProfiler.h"
+#include "Widgets/Input/SComboBox.h"
 
 // Insights
 #include "Insights/TaskGraphProfiler/TaskGraphProfilerManager.h"
@@ -25,6 +26,30 @@ using namespace TraceServices;
 namespace Insights
 {
 
+////////////////////////////////////////////////////////////////////////////////////////////////////
+// FTaskTableTreeViewCommands
+////////////////////////////////////////////////////////////////////////////////////////////////////
+
+class FTaskTableTreeViewCommands : public TCommands<FTaskTableTreeViewCommands>
+{
+public:
+	FTaskTableTreeViewCommands()
+		: TCommands<FTaskTableTreeViewCommands>(TEXT("FTaskTableTreeViewCommands"), NSLOCTEXT("FTaskTableTreeViewCommands", "Task Table Tree View Commands", "Task Table Tree View Commands"), NAME_None, FEditorStyle::Get().GetStyleSetName())
+	{
+	}
+
+	virtual ~FTaskTableTreeViewCommands()
+	{
+	}
+
+	// UI_COMMAND takes long for the compiler to optimize
+	PRAGMA_DISABLE_OPTIMIZATION
+	virtual void RegisterCommands() override
+	{
+	}
+	PRAGMA_ENABLE_OPTIMIZATION
+};
+
 ////////////////////////////////////////////////////////////////////////////////////////////////////
 
 STaskTableTreeView::STaskTableTreeView()
@@ -43,6 +68,18 @@ STaskTableTreeView::~STaskTableTreeView()
 void STaskTableTreeView::Construct(const FArguments& InArgs, TSharedPtr<Insights::FTaskTable> InTablePtr)
 {
 	ConstructWidget(InTablePtr);
+
+	AddCommmands();
+
+	// Make sure the default value is applied.
+	TimestampOptions_OnSelectionChanged(SelectedTimestampOption);
+}
+
+////////////////////////////////////////////////////////////////////////////////////////////////////
+
+void STaskTableTreeView::AddCommmands()
+{
+	FTaskTableTreeViewCommands::Register();
 }
 
 ////////////////////////////////////////////////////////////////////////////////////////////////////
@@ -132,9 +169,81 @@ FText STaskTableTreeView::GetCurrentOperationName() const
 
 BEGIN_SLATE_FUNCTION_BUILD_OPTIMIZATION
 TSharedPtr<SWidget> STaskTableTreeView::ConstructToolbar()
+{
+	return SNew(SHorizontalBox)
+		+ SHorizontalBox::Slot()
+		.AutoWidth()
+		.Padding(0.0f, 0.0f, 4.0f, 0.0f)
+		.VAlign(VAlign_Center)
+		[
+			SNew(STextBlock)
+			.Text(LOCTEXT("Timestamps", "Timestamps"))
+		]
+		
+		+ SHorizontalBox::Slot()
+		.AutoWidth()
+		.Padding(4.0f, 0.0f, 0.0f, 0.0f)
+		[
+			SNew(SBox)
+			.MinDesiredWidth(160.0f)
+			[
+				SNew(SComboBox<TSharedPtr<ETimestampOptions>>)
+				.OptionsSource(GetAvailableTimestampOptions())
+				.OnSelectionChanged(this, &STaskTableTreeView::TimestampOptions_OnSelectionChanged)
+				.OnGenerateWidget(this, &STaskTableTreeView::TimestampOptions_OnGenerateWidget)
+				.IsEnabled(this, &STaskTableTreeView::TimestampOptions_IsEnabled)
+				[
+					SNew(STextBlock)
+					.Text(this, &STaskTableTreeView::TimestampOptions_GetSelectionText)
+				]
+			]
+		];
+}
+
+////////////////////////////////////////////////////////////////////////////////////////////////////
+
+TSharedPtr<SWidget> STaskTableTreeView::ConstructFooter()
 {
 	return nullptr;
 }
+
+////////////////////////////////////////////////////////////////////////////////////////////////////
+
+TSharedRef<SWidget> STaskTableTreeView::TimestampOptions_OnGenerateWidget(TSharedPtr<ETimestampOptions> InOption)
+{
+	auto GetTooltipText = [](ETimestampOptions InOption)
+	{
+		switch (InOption)
+		{
+			case ETimestampOptions::Absolute:
+			{
+				return LOCTEXT("AbsoluteValueTooltip", "The timestamps for all columns will show absolute values.");
+			}
+			case ETimestampOptions::RelativeToPrevious:
+			{
+				return LOCTEXT("RelativeToPreviousTooltip", "The timestamps for all columns will show values relative to the previous stage. Ex: Scheduled will be relative to Launched.");
+			}
+			case ETimestampOptions::RelativeToCreated:
+			{
+				return LOCTEXT("RelativeToCreatedTooltip", "The timestamps for all columns will show values relative to the created time.");
+			}
+		}
+
+		return FText();
+	};
+
+	TSharedRef<SHorizontalBox> Widget = SNew(SHorizontalBox);
+	Widget->AddSlot()
+		.AutoWidth()
+		[
+			SNew(STextBlock)
+			.Text(TimestampOptions_GetText(*InOption))
+			.ToolTipText(GetTooltipText(*InOption))
+			.Margin(2.0f)
+		];
+
+	return Widget;
+}
 END_SLATE_FUNCTION_BUILD_OPTIMIZATION
 
 ////////////////////////////////////////////////////////////////////////////////////////////////////
@@ -167,18 +276,95 @@ void STaskTableTreeView::ApplyColumnConfig(const TArrayView<FColumnConfig>& Pres
 
 ////////////////////////////////////////////////////////////////////////////////////////////////////
 
-BEGIN_SLATE_FUNCTION_BUILD_OPTIMIZATION
-TSharedPtr<SWidget> STaskTableTreeView::ConstructFooter()
+void STaskTableTreeView::InternalCreateGroupings()
 {
-	return nullptr;
+	STableTreeView::InternalCreateGroupings();
 }
-END_SLATE_FUNCTION_BUILD_OPTIMIZATION
 
 ////////////////////////////////////////////////////////////////////////////////////////////////////
 
-void STaskTableTreeView::InternalCreateGroupings()
+const TArray<TSharedPtr<STaskTableTreeView::ETimestampOptions>>* STaskTableTreeView::GetAvailableTimestampOptions()
 {
-	STableTreeView::InternalCreateGroupings();
+	if (AvailableTimestampOptions.Num() == 0)
+	{
+		AvailableTimestampOptions.Add(MakeShared<ETimestampOptions>(ETimestampOptions::Absolute));
+		AvailableTimestampOptions.Add(MakeShared<ETimestampOptions>(ETimestampOptions::RelativeToPrevious));
+		AvailableTimestampOptions.Add(MakeShared<ETimestampOptions>(ETimestampOptions::RelativeToCreated));
+	}
+
+	return &AvailableTimestampOptions;
+}
+
+////////////////////////////////////////////////////////////////////////////////////////////////////
+
+void STaskTableTreeView::TimestampOptions_OnSelectionChanged(TSharedPtr<ETimestampOptions> InOption, ESelectInfo::Type SelectInfo)
+{
+	TimestampOptions_OnSelectionChanged(*InOption);
+}
+
+////////////////////////////////////////////////////////////////////////////////////////////////////
+
+void STaskTableTreeView::TimestampOptions_OnSelectionChanged(ETimestampOptions InOption)
+{
+	SelectedTimestampOption = InOption;
+
+	switch (SelectedTimestampOption)
+	{
+	case ETimestampOptions::Absolute:
+	{
+		GetTaskTable()->SwitchToAbsoluteTimestamps();
+		break;
+	}
+	case ETimestampOptions::RelativeToPrevious:
+	{
+		GetTaskTable()->SwitchToRelativeToPreviousTimestamps();
+		break;
+	}
+	case ETimestampOptions::RelativeToCreated:
+	{
+		GetTaskTable()->SwitchToRelativeToCreatedTimestamps();
+		break;
+	}
+	}
+
+	UpdateTree();
+}
+
+////////////////////////////////////////////////////////////////////////////////////////////////////
+
+FText STaskTableTreeView::TimestampOptions_GetSelectionText() const
+{
+	return TimestampOptions_GetText(SelectedTimestampOption);
+}
+
+////////////////////////////////////////////////////////////////////////////////////////////////////
+
+FText STaskTableTreeView::TimestampOptions_GetText(ETimestampOptions InOption) const
+{
+	switch (InOption)
+	{
+		case ETimestampOptions::Absolute:
+		{
+			return LOCTEXT("Absolute", "Absolute");
+		}
+		case ETimestampOptions::RelativeToPrevious:
+		{
+			return LOCTEXT("RelativeToPrevious", "Relative To Previous");
+		}
+		case ETimestampOptions::RelativeToCreated:
+		{
+			return LOCTEXT("RelativeToCreated", "Relative To Created");
+		}
+	}
+
+	return FText();
+}
+
+////////////////////////////////////////////////////////////////////////////////////////////////////
+
+bool STaskTableTreeView::TimestampOptions_IsEnabled() const
+{
+	return !bIsUpdateRunning;
 }
 
 ////////////////////////////////////////////////////////////////////////////////////////////////////
diff --git a/Engine/Source/Developer/TraceInsights/Private/Insights/TaskGraphProfiler/Widgets/STaskTableTreeView.h b/Engine/Source/Developer/TraceInsights/Private/Insights/TaskGraphProfiler/Widgets/STaskTableTreeView.h
--- a/Engine/Source/Developer/TraceInsights/Private/Insights/TaskGraphProfiler/Widgets/STaskTableTreeView.h
+++ b/Engine/Source/Developer/TraceInsights/Private/Insights/TaskGraphProfiler/Widgets/STaskTableTreeView.h
@@ -24,6 +24,13 @@ private:
 		float Width;
 	};
 
+	enum class ETimestampOptions : uint32
+	{
+		Absolute,
+		RelativeToPrevious,
+		RelativeToCreated,
+	};
+
 public:
 	/** Default constructor. */
 	STaskTableTreeView();
@@ -79,10 +86,22 @@ protected:
 
 private:
 	void ApplyColumnConfig(const TArrayView<FColumnConfig>& Preset);
+	void AddCommmands();
+
+	const TArray<TSharedPtr<ETimestampOptions>>* GetAvailableTimestampOptions();
+	TSharedRef<SWidget> TimestampOptions_OnGenerateWidget(TSharedPtr<ETimestampOptions> InOption);
+	void TimestampOptions_OnSelectionChanged(TSharedPtr<ETimestampOptions> InOption, ESelectInfo::Type SelectInfo);
+	void TimestampOptions_OnSelectionChanged(ETimestampOptions InOption);
+	FText TimestampOptions_GetSelectionText() const;
+	FText TimestampOptions_GetText(ETimestampOptions InOption) const;
+	bool TimestampOptions_IsEnabled() const;
 
 private:
 	double QueryStartTime = 0.0f;
 	double QueryEndTime = 0.0f;
+
+	ETimestampOptions SelectedTimestampOption = ETimestampOptions::RelativeToPrevious;
+	TArray<TSharedPtr<ETimestampOptions>> AvailableTimestampOptions;
 };
 
 ////////////////////////////////////////////////////////////////////////////////////////////////////
