TasksInsights: task names changes to be "task_{TaskId}" instead of "task_{Index}" in TaskTable. This allows to locate task in Timing View

diff --git a/Engine/Source/Developer/TraceInsights/Private/Insights/TaskGraphProfiler/Widgets/STaskTableTreeView.cpp b/Engine/Source/Developer/TraceInsights/Private/Insights/TaskGraphProfiler/Widgets/STaskTableTreeView.cpp
--- a/Engine/Source/Developer/TraceInsights/Private/Insights/TaskGraphProfiler/Widgets/STaskTableTreeView.cpp
+++ b/Engine/Source/Developer/TraceInsights/Private/Insights/TaskGraphProfiler/Widgets/STaskTableTreeView.cpp
@@ -131,7 +131,7 @@ void STaskTableTreeView::RebuildTree(bool bResync)
 					{
 						Allocs.Emplace(TaskInfo);
 						uint32 Index = Allocs.Num() - 1;
-						FName NodeName(BaseNodeName, Index);
+						FName NodeName(BaseNodeName, TaskInfo.Id + 1);
 						FTaskNodePtr NodePtr = MakeShared<FTaskNode>(NodeName, TaskTable, Index);
 						Nodes->Add(NodePtr);
 						return TraceServices::ETaskEnumerationResult::Continue;
