[Insights] FilterConfiguratorNode value convertor for memory values

diff --git a/Engine/Source/Developer/TraceInsights/Private/Insights/MemoryProfiler/ViewModels/MemAllocFilterValueConverter.cpp b/Engine/Source/Developer/TraceInsights/Private/Insights/MemoryProfiler/ViewModels/MemAllocFilterValueConverter.cpp
new file mode 100644
--- /dev/null
+++ b/Engine/Source/Developer/TraceInsights/Private/Insights/MemoryProfiler/ViewModels/MemAllocFilterValueConverter.cpp
@@ -0,0 +1,90 @@
+// Copyright Epic Games, Inc. All Rights Reserved.
+
+#include "MemAllocFilterValueConverter.h"
+
+#include "Misc/Char.h"
+
+#define LOCTEXT_NAMESPACE "Insights::FMemoryFilterValueConverter"
+
+namespace Insights
+{
+
+////////////////////////////////////////////////////////////////////////////////////////////////////
+// FMemoryFilterValueConverter
+////////////////////////////////////////////////////////////////////////////////////////////////////
+
+bool FMemoryFilterValueConverter::Convert(const FString& Input, int64& Output, FText& OutError) const
+{
+	FString TrimmedInput = Input.TrimStartAndEnd();
+
+	if (TrimmedInput.IsNumeric())
+	{
+		Output = FCString::Atoi64(*TrimmedInput);
+	}
+
+	static const TCHAR* Prefixes = TEXT("KMGTPE");
+	static const TCHAR* Suffix = TEXT("iB");
+	for (int Index = 0; Index < TrimmedInput.Len(); ++Index)
+	{
+		TCHAR Character = TrimmedInput[Index];
+		if (TChar<TCHAR>::IsDigit(Character) || TChar<TCHAR>::IsWhitespace(Character) || Character == '.')
+		{
+			continue;
+		}
+		else if (TChar<TCHAR>::IsAlpha(Character))
+		{
+			Character = TChar<TCHAR>::ToUpper(Character);
+			const TCHAR* PrefixPos = FCString::Strchr(Prefixes, Character);
+			if (PrefixPos == nullptr)
+			{
+				OutError = LOCTEXT("InvalidUnit", "Invalid unit. Expected one of KiB, MiB, GiB, TiB, PiB, EiB.");
+				return false;
+			}
+
+			int EndIndex = Index + 1;
+			while (EndIndex < TrimmedInput.Len() && TChar<TCHAR>::IsWhitespace(TrimmedInput[EndIndex]))
+			{
+				++EndIndex;
+			}
+
+			if (TrimmedInput.Mid(EndIndex).Compare(Suffix, ESearchCase::Type::IgnoreCase) != 0)
+			{
+				OutError = LOCTEXT("InvalidUnit", "Invalid unit. Expected one of KiB, MiB, GiB, TiB, PiB, EiB.");
+				return false;
+			}
+
+			// Equals to (PrefixPos - Prefixes + 1) * 2^10
+			Output = 1LL << (PrefixPos - Prefixes + 1) * 10;
+			double Value = FCString::Atod(*TrimmedInput.Left(Index));
+			Output *= Value;
+			break;
+		}
+		else
+		{
+			OutError = LOCTEXT("InvalidCharacters", "Invalid character.");
+			return false;
+		}
+	}
+
+	return true;
+}
+
+////////////////////////////////////////////////////////////////////////////////////////////////////
+
+FText FMemoryFilterValueConverter::GetTooltipText() const
+{
+	return LOCTEXT("MemoryTooltipText", "Enter a memory value, such as 100 or 4 KiB or 1 MiB");
+}
+
+////////////////////////////////////////////////////////////////////////////////////////////////////
+
+FText FMemoryFilterValueConverter::GetHintText() const
+{
+	return LOCTEXT("MemoryHintText", "Ex. 4 KiB");
+}
+
+////////////////////////////////////////////////////////////////////////////////////////////////////
+
+} // namespace Insights
+
+#undef LOCTEXT_NAMESPACE
diff --git a/Engine/Source/Developer/TraceInsights/Private/Insights/MemoryProfiler/ViewModels/MemAllocFilterValueConverter.h b/Engine/Source/Developer/TraceInsights/Private/Insights/MemoryProfiler/ViewModels/MemAllocFilterValueConverter.h
new file mode 100644
--- /dev/null
+++ b/Engine/Source/Developer/TraceInsights/Private/Insights/MemoryProfiler/ViewModels/MemAllocFilterValueConverter.h
@@ -0,0 +1,23 @@
+// Copyright Epic Games, Inc. All Rights Reserved.
+
+#pragma once
+
+// Insights
+#include "Insights/ViewModels/Filters.h"
+
+namespace Insights
+{
+
+////////////////////////////////////////////////////////////////////////////////////////////////////
+
+class FMemoryFilterValueConverter :public IFilterValueConvertor
+{
+public:
+	virtual bool Convert(const FString& Input, int64& Output, FText& OutError) const override;
+	virtual FText GetTooltipText() const override;
+	virtual FText GetHintText() const override;
+};
+
+////////////////////////////////////////////////////////////////////////////////////////////////////
+
+} // namespace Insights
diff --git a/Engine/Source/Developer/TraceInsights/Private/Insights/MemoryProfiler/ViewModels/MemAllocTable.cpp b/Engine/Source/Developer/TraceInsights/Private/Insights/MemoryProfiler/ViewModels/MemAllocTable.cpp
--- a/Engine/Source/Developer/TraceInsights/Private/Insights/MemoryProfiler/ViewModels/MemAllocTable.cpp
+++ b/Engine/Source/Developer/TraceInsights/Private/Insights/MemoryProfiler/ViewModels/MemAllocTable.cpp
@@ -2,16 +2,18 @@
 
 #include "MemAllocTable.h"
 
+#include "Widgets/SBoxPanel.h"
+#include "Widgets/SToolTip.h"
+#include "Widgets/Text/STextBlock.h"
+
 // Insights
+#include "Insights/MemoryProfiler/ViewModels/MemAllocFilterValueConverter.h"
 #include "Insights/MemoryProfiler/ViewModels/MemAllocNode.h"
 #include "Insights/MemoryProfiler/ViewModels/MemAllocTable.h"
 #include "Insights/Table/ViewModels/TableCellValueFormatter.h"
 #include "Insights/Table/ViewModels/TableCellValueGetter.h"
 #include "Insights/Table/ViewModels/TableCellValueSorter.h"
 #include "Insights/Table/ViewModels/TableColumn.h"
-#include "Widgets/SBoxPanel.h"
-#include "Widgets/SToolTip.h"
-#include "Widgets/Text/STextBlock.h"
 
 #define LOCTEXT_NAMESPACE "Insights::FMemAllocTable"
 
@@ -415,6 +417,9 @@ void FMemAllocTable::AddDefaultColumns()
 		TSharedRef<ITableCellValueSorter> Sorter = MakeShared<FSorterByInt64Value>(ColumnRef);
 		Column.SetValueSorter(Sorter);
 
+		TSharedRef<IFilterValueConvertor> Converter = MakeShared<FMemoryFilterValueConverter>();
+		Column.SetValueConverter(Converter);
+
 		Column.SetAggregation(ETableColumnAggregation::Sum);
 
 		AddColumn(ColumnRef);
diff --git a/Engine/Source/Developer/TraceInsights/Private/Insights/Table/ViewModels/TableColumn.h b/Engine/Source/Developer/TraceInsights/Private/Insights/Table/ViewModels/TableColumn.h
--- a/Engine/Source/Developer/TraceInsights/Private/Insights/Table/ViewModels/TableColumn.h
+++ b/Engine/Source/Developer/TraceInsights/Private/Insights/Table/ViewModels/TableColumn.h
@@ -12,6 +12,7 @@ namespace Insights
 {
 
 class FBaseTreeNode;
+class IFilterValueConvertor;
 class FTable;
 class ITableCellValueGetter;
 class ITableCellValueFormatter;
@@ -182,6 +183,12 @@ public:
 	/** Whether this column can be used for sorting. */
 	bool CanBeSorted() const { return ValueSorter.IsValid(); }
 
+	//////////////////////////////////////////////////
+	// Value Converter (can be nullptr)
+
+	TSharedPtr<IFilterValueConvertor> GetValueConverter() const { return ValueConverter; }
+	void SetValueConverter(TSharedPtr<IFilterValueConvertor> InValueConverter) { ValueConverter = InValueConverter; }
+
 	//////////////////////////////////////////////////
 
 	TWeakPtr<FTable> GetParentTable() const { return ParentTable; }
@@ -239,6 +246,9 @@ private:
 
 	/** Custom sorter for values displayed by this column. */
 	TSharedPtr<ITableCellValueSorter> ValueSorter;
+	
+	/** Used to convert in a custom way from string to column data type in FilterConfigurator */
+	TSharedPtr<IFilterValueConvertor> ValueConverter;
 
 	/* Parent table. Only one table instance can own this column. */
 	TWeakPtr<FTable> ParentTable;
diff --git a/Engine/Source/Developer/TraceInsights/Private/Insights/Table/Widgets/STableTreeView.cpp b/Engine/Source/Developer/TraceInsights/Private/Insights/Table/Widgets/STableTreeView.cpp
--- a/Engine/Source/Developer/TraceInsights/Private/Insights/Table/Widgets/STableTreeView.cpp
+++ b/Engine/Source/Developer/TraceInsights/Private/Insights/Table/Widgets/STableTreeView.cpp
@@ -2647,12 +2647,14 @@ FReply STableTreeView::OnAdvancedFiltersClicked()
 				case ETableCellDataType::Int64:
 				{
 					AvailableFilters->Add(MakeShared<FFilter>(Column->GetIndex(), Column->GetTitleName(), Column->GetDescription(), EFilterDataType::Int64, FFilterService::Get()->GetIntegerOperators()));
+					AvailableFilters->Last()->Convertor = Column->GetValueConverter();
 					Context.AddFilterData<int64>(Column->GetIndex(), 0);
 					break;
 				}
 				case ETableCellDataType::Double:
 				{
 					AvailableFilters->Add(MakeShared<FFilter>(Column->GetIndex(), Column->GetTitleName(), Column->GetDescription(), EFilterDataType::Double, FFilterService::Get()->GetDoubleOperators()));
+					AvailableFilters->Last()->Convertor = Column->GetValueConverter();
 					Context.AddFilterData<double>(Column->GetIndex(), 0.0);
 					break;
 				}
@@ -2663,6 +2665,7 @@ FReply STableTreeView::OnAdvancedFiltersClicked()
 					if (!Column->IsHierarchy())
 					{
 						AvailableFilters->Add(MakeShared<FFilter>(Column->GetIndex(), Column->GetTitleName(), Column->GetDescription(), EFilterDataType::String, FFilterService::Get()->GetStringOperators()));
+						AvailableFilters->Last()->Convertor = Column->GetValueConverter();
 						Context.AddFilterData<FString>(Column->GetIndex(), FString());
 					}
 					break;
diff --git a/Engine/Source/Developer/TraceInsights/Private/Insights/Tests/UITests.cpp b/Engine/Source/Developer/TraceInsights/Private/Insights/Tests/UITests.cpp
--- a/Engine/Source/Developer/TraceInsights/Private/Insights/Tests/UITests.cpp
+++ b/Engine/Source/Developer/TraceInsights/Private/Insights/Tests/UITests.cpp
@@ -6,6 +6,7 @@
 
 #include "TraceServices/Model/TimingProfiler.h"
 #include "Insights/InsightsManager.h"
+#include "Insights/MemoryProfiler/ViewModels/MemAllocFilterValueConverter.h"
 #include "Insights/TimingProfilerManager.h"
 #include "Insights/Widgets/STimingProfilerWindow.h"
 
@@ -40,6 +41,43 @@ bool FHideAndShowAllTimingViewTabs::RunTest(const FString& Parameters)
 	return true;
 }
 
+#endif // !WITH_EDITOR
+
 ////////////////////////////////////////////////////////////////////////////////////////////////////
 
-#endif // !WITH_EDITOR
+bool FMemoryFilterValueConverterTest::RunTest(const FString& Parameters)
+{
+	Insights::FMemoryFilterValueConverter Coverter;
+
+	FText Error;
+	int64 Value;
+
+	Coverter.Convert(TEXT("152485"), Value, Error);
+	TestEqual(TEXT("BasicValue"), 152485LL, Value);
+
+	Coverter.Convert(TEXT("125.56"), Value, Error);
+	TestEqual(TEXT("DoubleValue"), 125LL, Value);
+
+	Coverter.Convert(TEXT("3 KiB"), Value, Error);
+	TestEqual(TEXT("Kib"), 3072LL, Value);
+
+	Coverter.Convert(TEXT("7.14 KiB"), Value, Error);
+	TestEqual(TEXT("KibDouble"), 7311LL, Value);
+
+	Coverter.Convert(TEXT("5 MiB"), Value, Error);
+	TestEqual(TEXT("Mib"), 5242880LL, Value);
+
+	Coverter.Convert(TEXT("1 EiB"), Value, Error);
+	TestEqual(TEXT("Eib"), 1152921504606846976LL, Value);
+
+	Coverter.Convert(TEXT("2 kib"), Value, Error);
+	TestEqual(TEXT("CaseInsesitive"), 2048LL, Value);
+
+	TestFalse(TEXT("Fail1"), Coverter.Convert(TEXT("23test"), Value, Error));
+	TestFalse(TEXT("Fail2"), Coverter.Convert(TEXT("43 kOb"), Value, Error));
+	TestFalse(TEXT("FailInvalidChar"), Coverter.Convert(TEXT("45,"), Value, Error));
+
+	return !HasAnyErrors();
+}
+
+////////////////////////////////////////////////////////////////////////////////////////////////////
diff --git a/Engine/Source/Developer/TraceInsights/Private/Insights/Tests/UITests.h b/Engine/Source/Developer/TraceInsights/Private/Insights/Tests/UITests.h
--- a/Engine/Source/Developer/TraceInsights/Private/Insights/Tests/UITests.h
+++ b/Engine/Source/Developer/TraceInsights/Private/Insights/Tests/UITests.h
@@ -6,10 +6,13 @@
 #include "Logging/LogMacros.h"
 #include "Misc/AutomationTest.h"
 
-#if !WITH_EDITOR
-
 DECLARE_LOG_CATEGORY_EXTERN(UITests, Log, All);
 
+#if !WITH_EDITOR
+
 IMPLEMENT_SIMPLE_AUTOMATION_TEST(FHideAndShowAllTimingViewTabs, "Insights.HideAndShowAllTimingViewTabs", EAutomationTestFlags::ApplicationContextMask | EAutomationTestFlags::ProductFilter)
 
 #endif
+
+IMPLEMENT_SIMPLE_AUTOMATION_TEST(FMemoryFilterValueConverterTest, "Insights.FMemoryFilterValueConverterTest", EAutomationTestFlags::ApplicationContextMask | EAutomationTestFlags::ProductFilter)
+
diff --git a/Engine/Source/Developer/TraceInsights/Private/Insights/ViewModels/FilterConfiguratorNode.cpp b/Engine/Source/Developer/TraceInsights/Private/Insights/ViewModels/FilterConfiguratorNode.cpp
--- a/Engine/Source/Developer/TraceInsights/Private/Insights/ViewModels/FilterConfiguratorNode.cpp
+++ b/Engine/Source/Developer/TraceInsights/Private/Insights/ViewModels/FilterConfiguratorNode.cpp
@@ -158,18 +158,38 @@ void FFilterConfiguratorNode::ProcessFilter()
 		{
 		case EFilterDataType::Double:
 		{
-			FilterValue.Set<double>(FCString::Atod(*TextBoxValue));
+			if (SelectedFilter->Convertor.IsValid())
+			{
+				double Value = 0.0;
+				FText Errors;
+				bool Result = SelectedFilter->Convertor->Convert(TextBoxValue, Value, Errors);
+				FilterValue.Set<double>(Result ? Value : 0.0);
+			}
+			else
+			{
+				FilterValue.Set<double>(FCString::Atod(*TextBoxValue));
+			}
 			break;
 		}
 		case EFilterDataType::Int64:
 		{
-			if (TextBoxValue.Contains(TEXT("x")))
+			if (SelectedFilter->Convertor.IsValid())
 			{
-				FilterValue.Set<int64>((int64)FParse::HexNumber64(*TextBoxValue));
+				int64 Value = 0;
+				FText Errors;
+				bool Result = SelectedFilter->Convertor->Convert(TextBoxValue, Value, Errors);
+				FilterValue.Set<int64>(Result ? Value : 0);
 			}
 			else
 			{
-				FilterValue.Set<int64>(FCString::Atoi64(*TextBoxValue));
+				if (TextBoxValue.Contains(TEXT("x")))
+				{
+					FilterValue.Set<int64>((int64)FParse::HexNumber64(*TextBoxValue));
+				}
+				else
+				{
+					FilterValue.Set<int64>(FCString::Atoi64(*TextBoxValue));
+				}
 			}
 			break;
 		}
diff --git a/Engine/Source/Developer/TraceInsights/Private/Insights/ViewModels/Filters.h b/Engine/Source/Developer/TraceInsights/Private/Insights/ViewModels/Filters.h
--- a/Engine/Source/Developer/TraceInsights/Private/Insights/ViewModels/Filters.h
+++ b/Engine/Source/Developer/TraceInsights/Private/Insights/ViewModels/Filters.h
@@ -97,6 +97,17 @@ struct FFilterGroupOperator
 
 ////////////////////////////////////////////////////////////////////////////////////////////////////
 
+class IFilterValueConvertor
+{
+public:
+	virtual bool Convert(const FString& Input, int64& Output, FText& OutError) const { unimplemented(); return false; }
+	virtual bool Convert(const FString& Input, double& Output, FText& OutError) const { unimplemented(); return false; }
+	virtual FText GetTooltipText() const { return FText(); }
+	virtual FText GetHintText() const { return FText(); }
+};
+
+////////////////////////////////////////////////////////////////////////////////////////////////////
+
 struct FFilter
 {
 	typedef TSharedPtr<const TArray<TSharedPtr<IFilterOperator>>> SupportedOperatorsArrayPtr;
@@ -115,6 +126,7 @@ struct FFilter
 	FText Name;
 	FText Desc;
 	EFilterDataType DataType;
+	TSharedPtr<IFilterValueConvertor> Convertor;
 	SupportedOperatorsArrayPtr SupportedOperators;
 };
 
diff --git a/Engine/Source/Developer/TraceInsights/Private/Insights/Widgets/SFilterConfiguratorRow.cpp b/Engine/Source/Developer/TraceInsights/Private/Insights/Widgets/SFilterConfiguratorRow.cpp
--- a/Engine/Source/Developer/TraceInsights/Private/Insights/Widgets/SFilterConfiguratorRow.cpp
+++ b/Engine/Source/Developer/TraceInsights/Private/Insights/Widgets/SFilterConfiguratorRow.cpp
@@ -106,6 +106,9 @@ TSharedRef<SWidget> SFilterConfiguratorRow::GenerateWidgetForColumn(const FName&
 				.MinDesiredWidth(50.0f)
 				.OnTextCommitted(this, &SFilterConfiguratorRow::OnTextBoxValueCommitted)
 				.Text(this, &SFilterConfiguratorRow::GetTextBoxValue)
+				.ToolTipText(this, &SFilterConfiguratorRow::GetTextBoxTooltipText)
+				.HintText(this, &SFilterConfiguratorRow::GetTextBoxHintText)
+				.OnVerifyTextChanged(this, &SFilterConfiguratorRow::TextBox_OnVerifyTextChanged)
 			];
 			
 		GeneratedWidget->AddSlot()
@@ -405,6 +408,46 @@ void SFilterConfiguratorRow::OnTextBoxValueCommitted(const FText& InNewText, ETe
 
 ///////////////////////////////////////////////////////////////////////////////////////////////////
 
+FText SFilterConfiguratorRow::GetTextBoxTooltipText() const
+{
+	TSharedPtr<IFilterValueConvertor> Converter = FilterConfiguratorNodePtr->GetSelectedFilter()->Convertor;
+	if (Converter.IsValid())
+	{
+		return Converter->GetTooltipText();
+	}
+
+	return FText();
+}
+
+///////////////////////////////////////////////////////////////////////////////////////////////////
+
+FText SFilterConfiguratorRow::GetTextBoxHintText() const
+{
+	TSharedPtr<IFilterValueConvertor> Converter = FilterConfiguratorNodePtr->GetSelectedFilter()->Convertor;
+	if (Converter.IsValid())
+	{
+		return Converter->GetHintText();
+	}
+
+	return FText();
+}
+
+///////////////////////////////////////////////////////////////////////////////////////////////////
+
+bool SFilterConfiguratorRow::TextBox_OnVerifyTextChanged(const FText& InText, FText& OutErrorMessage)
+{
+	TSharedPtr<IFilterValueConvertor> Converter = FilterConfiguratorNodePtr->GetSelectedFilter()->Convertor;
+	if (Converter.IsValid())
+	{
+		int64 Value;
+		return Converter->Convert(InText.ToString(), Value, OutErrorMessage);
+	}
+
+	return true;
+}
+
+///////////////////////////////////////////////////////////////////////////////////////////////////
+
 } // namespace Insights
 
 #undef LOCTEXT_NAMESPACE
diff --git a/Engine/Source/Developer/TraceInsights/Private/Insights/Widgets/SFilterConfiguratorRow.h b/Engine/Source/Developer/TraceInsights/Private/Insights/Widgets/SFilterConfiguratorRow.h
--- a/Engine/Source/Developer/TraceInsights/Private/Insights/Widgets/SFilterConfiguratorRow.h
+++ b/Engine/Source/Developer/TraceInsights/Private/Insights/Widgets/SFilterConfiguratorRow.h
@@ -62,6 +62,9 @@ private:
 
 	FText GetTextBoxValue() const;
 	void OnTextBoxValueCommitted(const FText& InNewText, ETextCommit::Type InTextCommit);
+	FText GetTextBoxTooltipText() const;
+	FText GetTextBoxHintText() const;
+	bool TextBox_OnVerifyTextChanged(const FText& InText, FText& OutErrorMessage);
 
 private:
 	FFilterConfiguratorNodePtr FilterConfiguratorNodePtr;
