[Insights] Fix static analysis warning

diff --git a/Engine/Source/Developer/TraceInsights/Private/Insights/MemoryProfiler/ViewModels/MemAllocFilterValueConverter.cpp b/Engine/Source/Developer/TraceInsights/Private/Insights/MemoryProfiler/ViewModels/MemAllocFilterValueConverter.cpp
--- a/Engine/Source/Developer/TraceInsights/Private/Insights/MemoryProfiler/ViewModels/MemAllocFilterValueConverter.cpp
+++ b/Engine/Source/Developer/TraceInsights/Private/Insights/MemoryProfiler/ViewModels/MemAllocFilterValueConverter.cpp
@@ -54,7 +54,7 @@ bool FMemoryFilterValueConverter::Convert(const FString& Input, int64& Output, F
 			}
 
 			// Equals to (PrefixPos - Prefixes + 1) * 2^10
-			Output = 1LL << (PrefixPos - Prefixes + 1) * 10;
+			Output = 1LL << ((PrefixPos - Prefixes + 1) * 10);
 			double Value = FCString::Atod(*TrimmedInput.Left(Index));
 			Output *= Value;
 			break;
