[Insights]
  - Organize IO track so that concurrent (overlapped IO) reads don't share the same depth (by [REDACTED_PERSON], CL 15991490) + fixed searching for an I/O activity event

diff --git a/Engine/Source/Developer/TraceInsights/Private/Insights/ViewModels/FileActivityTimingTrack.cpp b/Engine/Source/Developer/TraceInsights/Private/Insights/ViewModels/FileActivityTimingTrack.cpp
--- a/Engine/Source/Developer/TraceInsights/Private/Insights/ViewModels/FileActivityTimingTrack.cpp
+++ b/Engine/Source/Developer/TraceInsights/Private/Insights/ViewModels/FileActivityTimingTrack.cpp
@@ -2,6 +2,7 @@
 
 #include "FileActivityTimingTrack.h"
 
+#include "Algo/BinarySearch.h"
 #include "Fonts/FontMeasure.h"
 #include "Framework/MultiBox/MultiBoxBuilder.h"
 #include "Styling/SlateBrush.h"
@@ -68,6 +69,8 @@ uint32 GetFileActivityTypeColor(TraceServices::EFileActivityType Type)
 // FFileActivitySharedState
 ////////////////////////////////////////////////////////////////////////////////////////////////////
 
+const uint32 FFileActivitySharedState::MaxLanes = 10000;
+
 void FFileActivitySharedState::OnBeginSession(Insights::ITimingViewSession& InSession)
 {
 	if (&InSession != TimingView)
@@ -80,7 +83,6 @@ void FFileActivitySharedState::OnBeginSession(Insights::ITimingViewSession& InSe
 
 	bShowHideAllIoTracks = false;
 	bForceIoEventsUpdate = false;
-	bMergeIoLanes = true;
 
 	FileActivities.Reset();
 	FileActivityMap.Reset();
@@ -101,7 +103,6 @@ void FFileActivitySharedState::OnEndSession(Insights::ITimingViewSession& InSess
 
 	bShowHideAllIoTracks = false;
 	bForceIoEventsUpdate = false;
-	bMergeIoLanes = true;
 
 	FileActivities.Reset();
 	FileActivityMap.Reset();
@@ -155,69 +156,71 @@ void FFileActivitySharedState::Tick(Insights::ITimingViewSession& InSession, con
 			const TraceServices::IFileActivityProvider& FileActivityProvider = *TraceServices::ReadFileActivityProvider(InAnalysisSession);
 			FileActivityProvider.EnumerateFileActivity([this](const TraceServices::FFileInfo& FileInfo, const TraceServices::IFileActivityProvider::Timeline& Timeline)
 			{
+				TSharedPtr<FIoFileActivity> Activity = MakeShared<FIoFileActivity>();
+
+				Activity->Id = FileInfo.Id;
+				Activity->Path = FileInfo.Path;
+				Activity->StartTime = +std::numeric_limits<double>::infinity();
+				Activity->EndTime = -std::numeric_limits<double>::infinity();
+				Activity->CloseStartTime = +std::numeric_limits<double>::infinity();
+				Activity->CloseEndTime = +std::numeric_limits<double>::infinity();
+				Activity->EventCount = 0;
+				Activity->Index = -1;
+				Activity->MaxConcurrentEvents = 0;
+				Activity->StartingDepth = 0;
+
+				const int32 ActivityIndex = FileActivities.Num();
+				FileActivities.Add(Activity);
+				FileActivityMap.Add(FileInfo.Id, Activity);
+
+				TArray<double> ConcurrentEvents;
 				Timeline.EnumerateEvents(-std::numeric_limits<double>::infinity(), +std::numeric_limits<double>::infinity(),
-					[this, &FileInfo, &Timeline](double EventStartTime, double EventEndTime, uint32 EventDepth, const TraceServices::FFileActivity* FileActivity)
+				[this, &Activity, ActivityIndex, &FileInfo, &Timeline, &ConcurrentEvents](double EventStartTime, double EventEndTime, uint32 EventDepth, const TraceServices::FFileActivity* FileActivity)
 				{
-					//if (EventEndTime == std::numeric_limits<double>::infinity())
-					//{
-					//	EventEndTime = EventStartTime;
-					//}
-
-					TSharedPtr<FIoFileActivity> Activity = FileActivityMap.FindRef(FileInfo.Id);
-					if (!Activity.IsValid())
+					if (FileActivity->ActivityType != TraceServices::FileActivityType_Close)
 					{
-						Activity = MakeShared<FIoFileActivity>();
+						// events should be ordered by start time, but Activity->StartTime may not be initialized
+						ensure(Activity->StartTime == +std::numeric_limits<double>::infinity() || EventStartTime >= Activity->StartTime);
+						if (EventStartTime < Activity->StartTime)
+						{
+							Activity->StartTime = EventStartTime;
+						}
 
-						Activity->Id = FileInfo.Id;
-						Activity->Path = FileInfo.Path;
-						Activity->StartTime = EventStartTime;
-						Activity->EndTime = EventEndTime;
+						if (EventEndTime > Activity->EndTime)
+						{
+							Activity->EndTime = EventEndTime;
+						}
+					}
+					else
+					{
+						// The time range for the Close event is stored separated;
+						// this allows us to insert lanes into the idle time between the last read from a file and when the file is actually closed
 						Activity->CloseStartTime = EventStartTime;
 						Activity->CloseEndTime = EventEndTime;
-						Activity->EventCount = 1;
-						Activity->Depth = -1;
-
-						FileActivities.Add(Activity);
-						FileActivityMap.Add(FileInfo.Id, Activity);
 					}
-					else
+
+					Activity->EventCount++;
+
+					uint32 LocalDepth = MAX_uint32;
+					for (int32 i = 0; i < ConcurrentEvents.Num(); ++i)
 					{
-						if (FileActivity->ActivityType != TraceServices::FileActivityType_Close)
-						{
-							ensure(EventStartTime >= Activity->StartTime);
-							if (EventStartTime < Activity->StartTime)
-							{
-								Activity->StartTime = EventStartTime;
-							}
-
-							if (EventEndTime > Activity->EndTime)
-							{
-								Activity->EndTime = EventEndTime;
-							}
-						}
-						else
+						if (EventStartTime > ConcurrentEvents[i])
 						{
-							// The time range for the Close event is stored separated;
-							// for the purpose of avoiding lane collisions (overlaps) between activities.
-							Activity->CloseStartTime = EventStartTime;
-							Activity->CloseEndTime = EventEndTime;
+							LocalDepth = i;
+							ConcurrentEvents[i] = EventEndTime;
+							break;
 						}
-
-						Activity->EventCount++;
 					}
 
-					if (bMergeIoLanes)
+					if (LocalDepth == MAX_uint32)
 					{
-						EventDepth = 0;
-					}
-					else
-					{
-						EventDepth = FileInfo.Id % 32; // simple layout
+						LocalDepth = ConcurrentEvents.Num();
+						ConcurrentEvents.Add(EventEndTime);
+						Activity->MaxConcurrentEvents = ConcurrentEvents.Num();
 					}
 
 					uint32 Type = ((uint32)FileActivity->ActivityType & 0x0F) | (FileActivity->Failed ? 0x80 : 0);
-
-					AllIoEvents.Add(FIoTimingEvent{ EventStartTime, EventEndTime, EventDepth, Type, FileActivity->Offset, FileActivity->Size, FileActivity->ActualSize, Activity });
+					AllIoEvents.Add(FIoTimingEvent{ EventStartTime, EventEndTime, LocalDepth, Type, FileActivity->Offset, FileActivity->Size, FileActivity->ActualSize, ActivityIndex });
 					return TraceServices::EEventEnumerate::Continue;
 				});
 
@@ -241,49 +244,76 @@ void FFileActivitySharedState::Tick(Insights::ITimingViewSession& InSession, con
 		Stopwatch.Stop();
 		UE_LOG(TimingProfiler, Log, TEXT("[IO] Sorted file activities and events in %s."), *TimeUtils::FormatTimeAuto(Stopwatch.GetAccumulatedTime()));
 
-		if (bMergeIoLanes)
+		if (FileActivities.Num() > 0)
 		{
 			//////////////////////////////////////////////////
 			// Compute depth for file activities (avoids overlaps).
 
 			Stopwatch.Restart();
 
-			TArray<TSharedPtr<FIoFileActivity>> ActivityLanes;
+			struct FLane
+			{
+				double EndTime = 0.0f;
+				double CloseStartTime;
+				double CloseEndTime;
+			};
+
+			TArray<FLane> Lanes; // one lane per event depth, a file activity occupies multiple lanes
 
-			for (TSharedPtr<FIoFileActivity> FileActivity : FileActivities)
+			for (const TSharedPtr<FIoFileActivity>& FileActivityPtr : FileActivities)
 			{
+				FIoFileActivity& Activity = *FileActivityPtr;
+
 				// Find lane (avoiding overlaps with other file activities).
-				for (int32 LaneIndex = 0; LaneIndex < ActivityLanes.Num(); ++LaneIndex)
+				int32 Depth = 0;
+				while (Depth < Lanes.Num())
 				{
-					TSharedPtr<FIoFileActivity> Lane = ActivityLanes[LaneIndex];
-
-					if (FileActivity->StartTime >= Lane->EndTime &&
-						(FileActivity->StartTime >= Lane->CloseEndTime || FileActivity->EndTime <= Lane->CloseStartTime)) // avoids overlaps with Close event
+					bool bOverlap = false;
+					for (int32 LocalDepth = 0; LocalDepth < Activity.MaxConcurrentEvents; ++LocalDepth)
+					{
+						if (Depth + LocalDepth >= Lanes.Num())
+						{
+							break;
+						}
+						const FLane& Lane = Lanes[Depth + LocalDepth];
+						if (Activity.StartTime < Lane.EndTime ||
+							(Activity.StartTime < Lane.CloseEndTime && Activity.EndTime > Lane.CloseStartTime)) // overlaps with a Close event
+						{
+							bOverlap = true;
+							Depth += LocalDepth;
+							break;
+						}
+					}
+					if (!bOverlap)
 					{
-						FileActivity->Depth = LaneIndex;
-						ActivityLanes[LaneIndex] = FileActivity;
 						break;
 					}
+					++Depth;
 				}
 
-				if (FileActivity->Depth < 0)
+				int32 NewLaneNum = Depth + Activity.MaxConcurrentEvents;
+
+				if (NewLaneNum > MaxLanes)
 				{
-					const int32 MaxLanes = 10000;
-					if (ActivityLanes.Num() < MaxLanes)
+					// Snap to the bottom; allows overlaps in this case.
+					Activity.StartingDepth = MaxLanes - Activity.MaxConcurrentEvents;
+				}
+				else
+				{
+					if (NewLaneNum > Lanes.Num())
 					{
-						// Add new lane.
-						FileActivity->Depth = ActivityLanes.Num();
-						ActivityLanes.Add(FileActivity);
+						Lanes.AddDefaulted(NewLaneNum - Lanes.Num());
 					}
-					else
+
+					Activity.StartingDepth = Depth;
+
+					// Set close event only for first lane of the activity.
+					Lanes[Depth].CloseStartTime = Activity.CloseStartTime;
+					Lanes[Depth].CloseEndTime = Activity.CloseEndTime;
+
+					for (int32 LocalDepth = 0; LocalDepth < Activity.MaxConcurrentEvents; ++LocalDepth)
 					{
-						int32 LaneIndex = ActivityLanes.Num() - 1;
-						FileActivity->Depth = LaneIndex;
-						TSharedPtr<FIoFileActivity> Lane = ActivityLanes[LaneIndex];
-						if (FileActivity->EndTime > Lane->EndTime)
-						{
-							ActivityLanes[LaneIndex] = FileActivity;
-						}
+						Lanes[Depth + LocalDepth].EndTime = Activity.EndTime;
 					}
 				}
 			}
@@ -297,7 +327,8 @@ void FFileActivitySharedState::Tick(Insights::ITimingViewSession& InSession, con
 
 			for (FIoTimingEvent& Event : AllIoEvents)
 			{
-				Event.Depth = static_cast<uint32>(Event.FileActivity->Depth);
+				Event.Depth += FileActivities[Event.FileActivityIndex]->StartingDepth;
+				ensure(Event.Depth < MaxLanes);
 			}
 
 			Stopwatch.Stop();
@@ -578,7 +609,7 @@ void FFileActivityTimingTrack::InitTooltip(FTooltipDrawState& InOutTooltip, cons
 			TypeLinearColor.B *= 2.0f;
 			InOutTooltip.AddTitle(TypeStr, TypeLinearColor);
 
-			InOutTooltip.AddTitle(InEvent.FileActivity->Path);
+			InOutTooltip.AddTitle(SharedState.FileActivities[InEvent.FileActivityIndex]->Path);
 
 			const double Duration = InEvent.EndTime - InEvent.StartTime;
 			InOutTooltip.AddNameValueTextLine(TEXT("Duration:"), TimeUtils::FormatTimeAuto(Duration));
@@ -615,48 +646,68 @@ bool FFileActivityTimingTrack::FindIoTimingEvent(const FTimingEventSearchParamet
 		// Search...
 		[this](TTimingEventSearch<FFileActivitySharedState::FIoTimingEvent>::FContext& InContext)
 		{
+			const TArray<FFileActivitySharedState::FIoTimingEvent>& Events = SharedState.GetAllEvents();
+
 			if (bIgnoreDuration)
 			{
-				for (const FFileActivitySharedState::FIoTimingEvent& Event : SharedState.GetAllEvents())
+				// Events are sorted by start time.
+				// Find the first event with StartTime >= searched StartTime.
+				int32 StartIndex = Algo::LowerBoundBy(Events, InContext.GetParameters().StartTime,
+					[](const FFileActivitySharedState::FIoTimingEvent& Event) { return Event.StartTime; });
+
+				for (int32 Index = StartIndex; Index < Events.Num(); ++Index)
 				{
+					const FFileActivitySharedState::FIoTimingEvent& Event = Events[Index];
+
 					if (bShowOnlyErrors && ((Event.Type & 0xF0) == 0))
 					{
 						continue;
 					}
 
-					if (Event.StartTime < InContext.GetParameters().StartTime)
-					{
-						continue;
-					}
+					ensure(Event.StartTime >= InContext.GetParameters().StartTime);
 
-					if (!InContext.ShouldContinueSearching() || Event.StartTime > InContext.GetParameters().EndTime)
+					if (Event.StartTime > InContext.GetParameters().EndTime)
 					{
 						break;
 					}
 
 					InContext.Check(Event.StartTime, Event.StartTime, bIgnoreEventDepth ? 0 : Event.Depth, Event);
+
+					if (!InContext.ShouldContinueSearching())
+					{
+						break;
+					}
 				}
 			}
 			else
 			{
-				for (const FFileActivitySharedState::FIoTimingEvent& Event : SharedState.GetAllEvents())
+				// Events are sorted by start time.
+				// Find the first event with StartTime >= searched EndTime.
+				int32 StartIndex = Algo::LowerBoundBy(Events, InContext.GetParameters().EndTime,
+					[](const FFileActivitySharedState::FIoTimingEvent& Event) { return Event.StartTime; });
+
+				// Start at the last event with StartTime < searched EndTime.
+				for (int32 Index = StartIndex - 1; Index >= 0; --Index)
 				{
+					const FFileActivitySharedState::FIoTimingEvent& Event = Events[Index];
+
 					if (bShowOnlyErrors && ((Event.Type & 0xF0) == 0))
 					{
 						continue;
 					}
 
-					if (!bIgnoreDuration && Event.EndTime <= InContext.GetParameters().StartTime)
+					if (Event.EndTime <= InContext.GetParameters().StartTime ||
+						Event.StartTime >= InContext.GetParameters().EndTime)
 					{
 						continue;
 					}
 
-					if (!InContext.ShouldContinueSearching() || Event.StartTime >= InContext.GetParameters().EndTime)
+					InContext.Check(Event.StartTime, Event.EndTime, bIgnoreEventDepth ? 0 : Event.Depth, Event);
+
+					if (!InContext.ShouldContinueSearching())
 					{
 						break;
 					}
-
-					InContext.Check(Event.StartTime, Event.EndTime, bIgnoreEventDepth ? 0 : Event.Depth, Event);
 				}
 			}
 		},
@@ -798,9 +849,9 @@ void FDetailedFileActivityTimingTrack::BuildDrawState(ITimingEventsTrackDrawStat
 				break;
 			}
 
-			ensure(Activity->Depth <= 10000);
+			ensure(Activity->StartingDepth < FFileActivitySharedState::MaxLanes);
 
-			Builder.AddEvent(Activity->StartTime, Activity->EndTime, Activity->Depth, 0x55333333,
+			Builder.AddEvent(Activity->StartTime, Activity->EndTime, Activity->StartingDepth, 0x55333333,
 				[&Activity](float Width)
 				{
 					FString EventName = Activity->Path;
@@ -828,7 +879,7 @@ void FDetailedFileActivityTimingTrack::BuildDrawState(ITimingEventsTrackDrawStat
 			break;
 		}
 
-		ensure(Event.Depth <= 10000);
+		ensure(Event.Depth < FFileActivitySharedState::MaxLanes);
 		const TraceServices::EFileActivityType ActivityType = static_cast<TraceServices::EFileActivityType>(Event.Type & 0x0F);
 
 		const bool bHasFailed = ((Event.Type & 0xF0) != 0);
@@ -845,7 +896,7 @@ void FDetailedFileActivityTimingTrack::BuildDrawState(ITimingEventsTrackDrawStat
 		}
 
 		Builder.AddEvent(Event.StartTime, Event.EndTime, Event.Depth, Color,
-			[&Event](float Width)
+			[&Event, FileActivity= SharedState.FileActivities[Event.FileActivityIndex]](float Width)
 			{
 				FString EventName;
 
@@ -866,7 +917,7 @@ void FDetailedFileActivityTimingTrack::BuildDrawState(ITimingEventsTrackDrawStat
 				if (ActivityType >= TraceServices::FileActivityType_Count)
 				{
 					EventName += " [";
-					EventName += Event.FileActivity->Path;
+					EventName += FileActivity->Path;
 					EventName += "]";
 				}
 
diff --git a/Engine/Source/Developer/TraceInsights/Private/Insights/ViewModels/FileActivityTimingTrack.h b/Engine/Source/Developer/TraceInsights/Private/Insights/ViewModels/FileActivityTimingTrack.h
--- a/Engine/Source/Developer/TraceInsights/Private/Insights/ViewModels/FileActivityTimingTrack.h
+++ b/Engine/Source/Developer/TraceInsights/Private/Insights/ViewModels/FileActivityTimingTrack.h
@@ -20,6 +20,7 @@ class FFileActivitySharedState : public Insights::ITimingViewExtender, public TS
 {
 	friend class FOverviewFileActivityTimingTrack;
 	friend class FDetailedFileActivityTimingTrack;
+	friend class FFileActivityTimingTrack;
 
 public:
 	struct FIoFileActivity
@@ -31,19 +32,21 @@ public:
 		double CloseStartTime;
 		double CloseEndTime;
 		int32 EventCount;
-		int32 Depth;
+		int32 Index;	// Different FIoFileActivity may have the same Index if their operations don't overlap in time
+		int32 MaxConcurrentEvents; // e.g. overlapped IO reads
+		uint32 StartingDepth; // Depth of first event on this file
 	};
 
 	struct FIoTimingEvent
 	{
 		double StartTime;
 		double EndTime;
-		uint32 Depth;
+		uint32 Depth; // During update, this is local within a track - then it's set to a global depth
 		uint32 Type; // TraceServices::EFileActivityType + "Failed" flag
 		uint64 Offset;
 		uint64 Size;
 		uint64 ActualSize;
-		TSharedPtr<FIoFileActivity> FileActivity;
+		int32 FileActivityIndex;
 	};
 
 public:
@@ -60,9 +63,6 @@ public:
 
 	void RequestUpdate() { bForceIoEventsUpdate = true; }
 
-	bool IsMergeLanesToggleOn() const { return bMergeIoLanes; }
-	void ToggleMergeLanes() { bMergeIoLanes = !bMergeIoLanes; RequestUpdate(); }
-
 	bool IsAllIoTracksToggleOn() const { return bShowHideAllIoTracks; }
 	void SetAllIoTracksToggle(bool bOnOff);
 	void ShowAllIoTracks() { SetAllIoTracksToggle(true); }
@@ -81,6 +81,8 @@ public:
 	bool AreBackgroundEventsVisible() const;
 	void ToggleBackgroundEvents();
 
+	static const uint32 MaxLanes;
+
 private:
 	void BuildSubMenu(FMenuBuilder& InOutMenuBuilder);
 
@@ -92,7 +94,6 @@ private:
 
 	bool bShowHideAllIoTracks;
 	bool bForceIoEventsUpdate;
-	bool bMergeIoLanes; // merge lanes of file activity events in a way that avoids duplication (for the Activity track)
 
 	TArray<TSharedPtr<FIoFileActivity>> FileActivities;
 	TMap<uint64, TSharedPtr<FIoFileActivity>> FileActivityMap;
