[Insights]
  - Memory/Timing: Added snap to bookmark and log messages when dragging a time marker with the mouse.
  - Timing view: Optimized the filtering for bookmarks in the Markers track.

diff --git a/Engine/Source/Developer/TraceInsights/Private/Insights/ViewModels/MarkersTimingTrack.cpp b/Engine/Source/Developer/TraceInsights/Private/Insights/ViewModels/MarkersTimingTrack.cpp
--- a/Engine/Source/Developer/TraceInsights/Private/Insights/ViewModels/MarkersTimingTrack.cpp
+++ b/Engine/Source/Developer/TraceInsights/Private/Insights/ViewModels/MarkersTimingTrack.cpp
@@ -15,6 +15,8 @@
 #include "Insights/TimingProfilerManager.h"
 #include "Insights/ViewModels/TimingTrackViewport.h"
 
+#include <limits>
+
 ////////////////////////////////////////////////////////////////////////////////////////////////////
 
 #define LOCTEXT_NAMESPACE "MarkersTimingTrack"
@@ -30,6 +32,7 @@ FMarkersTimingTrack::FMarkersTimingTrack()
 	//, TimeMarkerBoxes()
 	//, TimeMarkerTexts()
 	, bUseOnlyBookmarks(true)
+	, BookmarkCategory(nullptr)
 	, Header(*this)
 	, NumLogMessages(0)
 	, NumDrawBoxes(0)
@@ -57,6 +60,7 @@ void FMarkersTimingTrack::Reset()
 	TimeMarkerTexts.Reset();
 
 	bUseOnlyBookmarks = true;
+	BookmarkCategory = nullptr;
 
 	Header.Reset();
 	Header.SetIsInBackground(true);
@@ -91,6 +95,16 @@ void FMarkersTimingTrack::UpdateTrackNameAndHeight()
 
 ////////////////////////////////////////////////////////////////////////////////////////////////////
 
+void FMarkersTimingTrack::PreUpdate(const ITimingTrackUpdateContext& Context)
+{
+	if (!BookmarkCategory)
+	{
+		UpdateBookmarkCategory();
+	}
+}
+
+////////////////////////////////////////////////////////////////////////////////////////////////////
+
 void FMarkersTimingTrack::Update(const ITimingTrackUpdateContext& Context)
 {
 	Header.Update(Context);
@@ -282,6 +296,89 @@ void FMarkersTimingTrack::BuildContextMenu(FMenuBuilder& MenuBuilder)
 	MenuBuilder.EndSection();
 }
 
+////////////////////////////////////////////////////////////////////////////////////////////////////
+
+double FMarkersTimingTrack::Snap(double Time, const double SnapTolerance)
+{
+	if (bUseOnlyBookmarks && !BookmarkCategory)
+	{
+		return Time;
+	}
+
+	TSharedPtr<const TraceServices::IAnalysisSession> Session = FInsightsManager::Get()->GetSession();
+	if (Session)
+	{
+		TraceServices::FAnalysisSessionReadScope SessionReadScope(*Session.Get());
+		const TraceServices::ILogProvider& LogProvider = TraceServices::ReadLogProvider(*Session.Get());
+
+		double SnapTime = std::numeric_limits<double>::infinity();
+		double SnapDistance = std::numeric_limits<double>::infinity();
+
+		if (bUseOnlyBookmarks)
+		{
+			LogProvider.EnumerateMessages(
+				Time - SnapTolerance,
+				Time + SnapTolerance,
+				[&SnapTime, &SnapDistance, Time, this](const TraceServices::FLogMessageInfo& Message)
+				{
+					if (Message.Category == BookmarkCategory)
+					{
+						double Distance = FMath::Abs(Message.Time - Time);
+						if (Distance < SnapDistance)
+						{
+							SnapDistance = Distance;
+							SnapTime = Message.Time;
+						}
+					}
+				});
+		}
+		else
+		{
+			LogProvider.EnumerateMessages(
+				Time - SnapTolerance,
+				Time + SnapTolerance,
+				[&SnapTime, &SnapDistance, Time, this](const TraceServices::FLogMessageInfo& Message)
+				{
+					double Distance = FMath::Abs(Message.Time - Time);
+					if (Distance < SnapDistance)
+					{
+						SnapDistance = Distance;
+						SnapTime = Message.Time;
+					}
+				});
+		}
+
+		if (SnapDistance < SnapTolerance)
+		{
+			Time = SnapTime;
+		}
+	}
+
+	return Time;
+}
+
+////////////////////////////////////////////////////////////////////////////////////////////////////
+
+void FMarkersTimingTrack::UpdateBookmarkCategory()
+{
+	BookmarkCategory = nullptr;
+
+	TSharedPtr<const TraceServices::IAnalysisSession> Session = FInsightsManager::Get()->GetSession();
+	if (Session)
+	{
+		TraceServices::FAnalysisSessionReadScope SessionReadScope(*Session.Get());
+		const TraceServices::ILogProvider& LogProvider = TraceServices::ReadLogProvider(*Session.Get());
+
+		LogProvider.EnumerateCategories([this](const TraceServices::FLogCategoryInfo& Category)
+		{
+			if (FCString::Strcmp(Category.Name, TEXT("LogBookmark")) == 0)
+			{
+				BookmarkCategory = &Category;
+			}
+		});
+	}
+}
+
 ////////////////////////////////////////////////////////////////////////////////////////////////////
 // FTimeMarkerTrackBuilder
 ////////////////////////////////////////////////////////////////////////////////////////////////////
@@ -327,17 +424,15 @@ void FTimeMarkerTrackBuilder::AddLogMessage(const TraceServices::FLogMessageInfo
 	}
 
 	check(Message.Category != nullptr);
-	//check(Message.Category->Name != nullptr);
-
-	const TCHAR* CategoryName = Message.Category->Name != nullptr ? Message.Category->Name : TEXT("");
-
-	if (!Track.bUseOnlyBookmarks || FCString::Strcmp(CategoryName, TEXT("LogBookmark")) == 0)
+	if (!Track.bUseOnlyBookmarks || Message.Category == Track.BookmarkCategory)
 	{
 		float X = Viewport.TimeToSlateUnitsRounded(Message.Time);
 		if (X < 0.0f)
 		{
 			X = -1.0f;
 		}
+
+		const TCHAR* CategoryName = Message.Category->Name != nullptr ? Message.Category->Name : TEXT("");
 		AddTimeMarker(X, Message.Index, Message.Verbosity, CategoryName, Message.Message);
 	}
 }
diff --git a/Engine/Source/Developer/TraceInsights/Private/Insights/ViewModels/MarkersTimingTrack.h b/Engine/Source/Developer/TraceInsights/Private/Insights/ViewModels/MarkersTimingTrack.h
--- a/Engine/Source/Developer/TraceInsights/Private/Insights/ViewModels/MarkersTimingTrack.h
+++ b/Engine/Source/Developer/TraceInsights/Private/Insights/ViewModels/MarkersTimingTrack.h
@@ -11,6 +11,7 @@
 
 namespace TraceServices
 {
+	struct FLogCategoryInfo;
 	struct FLogMessageInfo;
 	class ILogProvider;
 }
@@ -73,6 +74,7 @@ public:
 	int32 GetNumTexts() const { return TimeMarkerTexts.Num(); }
 
 	// FBaseTimingTrack
+	virtual void PreUpdate(const ITimingTrackUpdateContext& Context) override;
 	virtual void Update(const ITimingTrackUpdateContext& Context) override;
 	virtual void PostUpdate(const ITimingTrackUpdateContext& Context) override;
 	virtual void Draw(const ITimingTrackDrawContext& Context) const override;
@@ -81,6 +83,8 @@ public:
 	virtual FReply OnMouseButtonDoubleClick(const FGeometry& MyGeometry, const FPointerEvent& MouseEvent) override;
 	virtual void BuildContextMenu(FMenuBuilder& MenuBuilder) override;
 
+	double Snap(double Time, double SnapTolerance);
+
 private:
 	void ResetCache()
 	{
@@ -91,11 +95,14 @@ private:
 	void UpdateTrackNameAndHeight();
 	void UpdateDrawState(const FTimingTrackViewport& Viewport);
 
+	void UpdateBookmarkCategory();
+
 private:
 	TArray<FTimeMarkerBoxInfo> TimeMarkerBoxes;
 	TArray<FTimeMarkerTextInfo> TimeMarkerTexts;
 
 	bool bUseOnlyBookmarks; // If true, uses only bookmarks; otherwise it uses all log messages.
+	const TraceServices::FLogCategoryInfo* BookmarkCategory;
 
 	FTrackHeader Header;
 
diff --git a/Engine/Source/Developer/TraceInsights/Private/Insights/Widgets/STimingView.cpp b/Engine/Source/Developer/TraceInsights/Private/Insights/Widgets/STimingView.cpp
--- a/Engine/Source/Developer/TraceInsights/Private/Insights/Widgets/STimingView.cpp
+++ b/Engine/Source/Developer/TraceInsights/Private/Insights/Widgets/STimingView.cpp
@@ -2244,9 +2244,18 @@ FReply STimingView::OnMouseMove(const FGeometry& MyGeometry, const FPointerEvent
 			if (HasMouseCapture())
 			{
 				bIsDragging = true;
+				double Time = Viewport.SlateUnitsToTime(MousePosition.X);
+
+				// Snap to markers.
+				const bool bSnapTimeMarkers = MarkersTrack->IsVisible();
+				if (bSnapTimeMarkers)
+				{
+					const double SnapTolerance = 5.0 / Viewport.GetScaleX(); // +/- 5 pixels
+					Time = MarkersTrack->Snap(Time, SnapTolerance);
+				}
 
 				TSharedRef<Insights::FTimeMarker> ScrubbingTimeMarker = TimeRulerTrack->GetScrubbingTimeMarker();
-				ScrubbingTimeMarker->SetTime(Viewport.SlateUnitsToTime(MousePosition.X));
+				ScrubbingTimeMarker->SetTime(Time);
 				RaiseTimeMarkerChanging(ScrubbingTimeMarker);
 			}
 		}
