[Insights] Focus on a Task in Timing View and draw it's relations when double clicking on it in TaskTableTreeView.

diff --git a/Engine/Source/Developer/TraceInsights/Private/Insights/Table/Widgets/STableTreeView.cpp b/Engine/Source/Developer/TraceInsights/Private/Insights/Table/Widgets/STableTreeView.cpp
--- a/Engine/Source/Developer/TraceInsights/Private/Insights/Table/Widgets/STableTreeView.cpp
+++ b/Engine/Source/Developer/TraceInsights/Private/Insights/Table/Widgets/STableTreeView.cpp
@@ -438,6 +438,8 @@ TSharedPtr<SWidget> STableTreeView::TreeView_GetMenuContent()
 
 	MenuBuilder.EndSection();
 
+	ExtendMenu(MenuBuilder);
+
 	MenuBuilder.BeginSection("Misc", LOCTEXT("ContextMenu_Header_Misc", "Miscellaneous"));
 	{
 		MenuBuilder.AddMenuEntry
diff --git a/Engine/Source/Developer/TraceInsights/Private/Insights/Table/Widgets/STableTreeView.h b/Engine/Source/Developer/TraceInsights/Private/Insights/Table/Widgets/STableTreeView.h
--- a/Engine/Source/Developer/TraceInsights/Private/Insights/Table/Widgets/STableTreeView.h
+++ b/Engine/Source/Developer/TraceInsights/Private/Insights/Table/Widgets/STableTreeView.h
@@ -197,7 +197,7 @@ protected:
 	void TreeView_OnSelectionChanged(FTableTreeNodePtr SelectedItem, ESelectInfo::Type SelectInfo);
 
 	/** Called by STreeView when a tree item is double clicked. */
-	void TreeView_OnMouseButtonDoubleClick(FTableTreeNodePtr TreeNode);
+	virtual void TreeView_OnMouseButtonDoubleClick(FTableTreeNodePtr TreeNode);
 
 	////////////////////////////////////////////////////////////////////////////////////////////////////
 	// Tree View - Table Row
@@ -389,6 +389,8 @@ protected:
 	void CountNumNodesPerDepthRec(FBaseTreeNode* InRoot, TArray<int32>& InOutNumNodesPerDepth, int32 InDepth, int32 InMaxDepth, int InMaxNodes) const;
 	void SetExpandValueForChildGroupsRec(FBaseTreeNode* InRoot, int32 InDepth, int32 InMaxDepth, bool InValue);
 
+	virtual void ExtendMenu(FMenuBuilder& Menu) {}
+
 protected:
 	/** Table view model. */
 	TSharedPtr<Insights::FTable> Table;
diff --git a/Engine/Source/Developer/TraceInsights/Private/Insights/TaskGraphProfiler/TaskGraphProfilerManager.cpp b/Engine/Source/Developer/TraceInsights/Private/Insights/TaskGraphProfiler/TaskGraphProfilerManager.cpp
--- a/Engine/Source/Developer/TraceInsights/Private/Insights/TaskGraphProfiler/TaskGraphProfilerManager.cpp
+++ b/Engine/Source/Developer/TraceInsights/Private/Insights/TaskGraphProfiler/TaskGraphProfilerManager.cpp
@@ -217,6 +217,117 @@ void FTaskGraphProfilerManager::OnTaskTableTreeViewTabClosed(TSharedRef<SDockTab
 
 ////////////////////////////////////////////////////////////////////////////////////////////////////
 
+void FTaskGraphProfilerManager::GetTaskRelations(const TraceServices::FTaskInfo* Task, const TraceServices::ITasksProvider* TasksProvider, AddRelationCallback Callback)
+{
+	const int32 MaxTasksToShow = 30;
+
+	if (Task == nullptr)
+	{
+		return;
+	}
+
+	if (Task->CreatedTimestamp != Task->LaunchedTimestamp || Task->CreatedThreadId != Task->LaunchedThreadId)
+	{
+		Callback(Task->CreatedTimestamp, Task->CreatedThreadId, Task->LaunchedTimestamp, Task->LaunchedThreadId, FTaskGraphRelation::ETaskGraphRelationType::Created);
+	}
+
+	Callback(Task->LaunchedTimestamp, Task->LaunchedThreadId, Task->ScheduledTimestamp, Task->ScheduledThreadId, FTaskGraphRelation::ETaskGraphRelationType::Launched);
+
+	int32 NumPrerequisitesToShow = FMath::Min(Task->Prerequisites.Num(), MaxTasksToShow);
+	for (int32 i = 0; i != NumPrerequisitesToShow; ++i)
+	{
+		const TraceServices::FTaskInfo* Prerequisite = TasksProvider->TryGetTask(Task->Prerequisites[i].RelativeId);
+		check(Prerequisite != nullptr);
+		Callback(Prerequisite->CompletedTimestamp, Prerequisite->CompletedThreadId, Task->ScheduledTimestamp, Task->ScheduledThreadId, FTaskGraphRelation::ETaskGraphRelationType::Prerequisite);
+	}
+
+	if (Task->LaunchedTimestamp != Task->ScheduledTimestamp || Task->LaunchedThreadId != Task->ScheduledThreadId)
+	{
+		Callback(Task->ScheduledTimestamp, Task->ScheduledThreadId, Task->StartedTimestamp, Task->StartedThreadId, FTaskGraphRelation::ETaskGraphRelationType::Scheduled);
+	}
+
+	int32 NumNestedToShow = FMath::Min(Task->NestedTasks.Num(), MaxTasksToShow);
+	for (int32 i = 0; i != NumNestedToShow; ++i)
+	{
+		const TraceServices::FTaskInfo::FRelationInfo& RelationInfo = Task->NestedTasks[i];
+		const TraceServices::FTaskInfo* NestedTask = TasksProvider->TryGetTask(RelationInfo.RelativeId);
+		check(NestedTask != nullptr);
+
+		Callback(RelationInfo.Timestamp, Task->StartedThreadId, NestedTask->StartedTimestamp, NestedTask->StartedThreadId, FTaskGraphRelation::ETaskGraphRelationType::AddedNested);
+
+		Callback(NestedTask->CompletedTimestamp, NestedTask->CompletedThreadId, NestedTask->CompletedTimestamp, Task->StartedThreadId, FTaskGraphRelation::ETaskGraphRelationType::NestedCompleted);
+	}
+
+	int32 NumSubsequentsToShow = FMath::Min(Task->NestedTasks.Num(), MaxTasksToShow);
+	for (int32 i = 0; i != NumSubsequentsToShow; ++i)
+	{
+		const TraceServices::FTaskInfo* Subsequent = TasksProvider->TryGetTask(Task->Subsequents[i].RelativeId);
+		check(Subsequent != nullptr);
+		Callback(Task->CompletedTimestamp, Task->CompletedThreadId, Subsequent->ScheduledTimestamp, Subsequent->ScheduledThreadId, FTaskGraphRelation::ETaskGraphRelationType::Subsequent);
+	}
+
+	if (Task->FinishedTimestamp != Task->CompletedTimestamp || Task->CompletedThreadId != Task->StartedThreadId)
+	{
+		Callback(Task->FinishedTimestamp, Task->StartedThreadId, Task->CompletedTimestamp, Task->StartedThreadId, FTaskGraphRelation::ETaskGraphRelationType::Completed);
+	}
+}
+
+////////////////////////////////////////////////////////////////////////////////////////////////////
+
+void FTaskGraphProfilerManager::GetTaskRelations(double Time, uint32 ThreadId, AddRelationCallback Callback)
+{
+	TSharedPtr<const TraceServices::IAnalysisSession> Session = FInsightsManager::Get()->GetSession();
+	if (!Session.IsValid())
+	{
+		return;
+	}
+
+	TraceServices::FAnalysisSessionReadScope SessionReadScope(*Session.Get());
+
+	const TraceServices::ITasksProvider* TasksProvider = TraceServices::ReadTasksProvider(*Session.Get());
+
+	if (TasksProvider == nullptr)
+	{
+		return;
+	}
+
+	const TraceServices::FTaskInfo* Task = TasksProvider->TryGetTask(ThreadId, Time);
+
+	if (Task != nullptr)
+	{
+		GetTaskRelations(Task, TasksProvider, Callback);
+	}
+}
+
+////////////////////////////////////////////////////////////////////////////////////////////////////
+
+void FTaskGraphProfilerManager::GetTaskRelations(uint32 TaskId, AddRelationCallback Callback)
+{
+	TSharedPtr<const TraceServices::IAnalysisSession> Session = FInsightsManager::Get()->GetSession();
+	if (!Session.IsValid())
+	{
+		return;
+	}
+
+	TraceServices::FAnalysisSessionReadScope SessionReadScope(*Session.Get());
+
+	const TraceServices::ITasksProvider* TasksProvider = TraceServices::ReadTasksProvider(*Session.Get());
+
+	if (TasksProvider == nullptr)
+	{
+		return;
+	}
+
+	const TraceServices::FTaskInfo* Task = TasksProvider->TryGetTask(TaskId);
+
+	if (Task != nullptr)
+	{
+		GetTaskRelations(Task, TasksProvider, Callback);
+	}
+}
+
+////////////////////////////////////////////////////////////////////////////////////////////////////
+
 } // namespace Insights
 
 #undef LOCTEXT_NAMESPACE
diff --git a/Engine/Source/Developer/TraceInsights/Private/Insights/TaskGraphProfiler/TaskGraphProfilerManager.h b/Engine/Source/Developer/TraceInsights/Private/Insights/TaskGraphProfiler/TaskGraphProfilerManager.h
--- a/Engine/Source/Developer/TraceInsights/Private/Insights/TaskGraphProfiler/TaskGraphProfilerManager.h
+++ b/Engine/Source/Developer/TraceInsights/Private/Insights/TaskGraphProfiler/TaskGraphProfilerManager.h
@@ -9,6 +9,13 @@
 // Insights
 #include "Insights/InsightsManager.h"
 #include "Insights/IUnrealInsightsModule.h"
+#include "Insights/ViewModels/TaskGraphRelation.h"
+
+namespace TraceServices
+{
+	struct FTaskInfo;
+	class ITasksProvider;
+}
 
 namespace Insights
 {
@@ -28,6 +35,8 @@ struct FTaskGraphProfilerTabs
 class FTaskGraphProfilerManager : public TSharedFromThis<FTaskGraphProfilerManager>, public IInsightsComponent
 {
 public:
+	typedef TFunction<void(double /*SourceTimestamp*/, uint32 /*SourceThreadId*/, double /*TargetTimestamp*/, uint32 /*TargetThreadId*/, FTaskGraphRelation::ETaskGraphRelationType /*Type*/)> AddRelationCallback;
+
 	/** Creates the Memory Profiler manager, only one instance can exist. */
 	FTaskGraphProfilerManager(TSharedRef<FUICommandList> InCommandList);
 
@@ -63,12 +72,17 @@ public:
 
 	void OnSessionChanged();
 
+	void GetTaskRelations(double Time, uint32 ThreadId, AddRelationCallback Callback);
+	void GetTaskRelations(uint32 TaskId, AddRelationCallback Callback);
+
 private:
 	/** Updates this manager, done through FCoreTicker. */
 	bool Tick(float DeltaTime);
 
 	void RegisterTimingProfilerLayoutExtensions(FInsightsMajorTabExtender& InOutExtender);
 
+	void GetTaskRelations(const TraceServices::FTaskInfo* Task, const TraceServices::ITasksProvider* TasksProvider, AddRelationCallback Callback);
+
 private:
 	bool bIsInitialized;
 	bool bIsAvailable;
diff --git a/Engine/Source/Developer/TraceInsights/Private/Insights/TaskGraphProfiler/ViewModels/TaskEntry.h b/Engine/Source/Developer/TraceInsights/Private/Insights/TaskGraphProfiler/ViewModels/TaskEntry.h
--- a/Engine/Source/Developer/TraceInsights/Private/Insights/TaskGraphProfiler/ViewModels/TaskEntry.h
+++ b/Engine/Source/Developer/TraceInsights/Private/Insights/TaskGraphProfiler/ViewModels/TaskEntry.h
@@ -40,6 +40,8 @@ public:
 	uint32 GetNumPrerequisites() const { return NumPrerequisites; }
 	uint32 GetNumSubsequents() const { return NumSubsequents; }
 
+	uint32 GetId() const { return Id; }
+
 private:
 	uint32 Id;
 
diff --git a/Engine/Source/Developer/TraceInsights/Private/Insights/TaskGraphProfiler/Widgets/STaskTableTreeView.cpp b/Engine/Source/Developer/TraceInsights/Private/Insights/TaskGraphProfiler/Widgets/STaskTableTreeView.cpp
--- a/Engine/Source/Developer/TraceInsights/Private/Insights/TaskGraphProfiler/Widgets/STaskTableTreeView.cpp
+++ b/Engine/Source/Developer/TraceInsights/Private/Insights/TaskGraphProfiler/Widgets/STaskTableTreeView.cpp
@@ -3,6 +3,7 @@
 #include "STaskTableTreeView.h"
 
 #include "EditorStyleSet.h"
+#include "Framework/MultiBox/MultiBoxBuilder.h"
 #include "SlateOptMacros.h"
 #include "Styling/CoreStyle.h"
 #include "TraceServices/AnalysisService.h"
@@ -16,6 +17,8 @@
 #include "Insights/Table/ViewModels/TableColumn.h"
 #include "Insights/TimingProfilerManager.h"
 #include "Insights/ViewModels/FilterConfigurator.h"
+#include "Insights/Widgets/STimingProfilerWindow.h"
+#include "Insights/Widgets/STimingView.h"
 
 #include <limits>
 
@@ -46,8 +49,11 @@ public:
 	PRAGMA_DISABLE_OPTIMIZATION
 	virtual void RegisterCommands() override
 	{
+		UI_COMMAND(Command_GoToTask, "Go To Task", "Pan and zoom to the task in Timing View.", EUserInterfaceActionType::Button, FInputChord());
 	}
 	PRAGMA_ENABLE_OPTIMIZATION
+
+	TSharedPtr<FUICommandInfo> Command_GoToTask;
 };
 
 ////////////////////////////////////////////////////////////////////////////////////////////////////
@@ -77,9 +83,29 @@ void STaskTableTreeView::Construct(const FArguments& InArgs, TSharedPtr<Insights
 
 ////////////////////////////////////////////////////////////////////////////////////////////////////
 
+void STaskTableTreeView::ExtendMenu(FMenuBuilder& MenuBuilder)
+{
+	MenuBuilder.BeginSection("Node", LOCTEXT("ContextMenu_Header_Task", "Task"));
+
+	{
+		MenuBuilder.AddMenuEntry
+		(
+			FTaskTableTreeViewCommands::Get().Command_GoToTask,
+			NAME_None,
+			TAttribute<FText>(),
+			TAttribute<FText>(),
+			FSlateIcon(FEditorStyle::GetStyleSetName(), "Profiler.EventGraph.ViewColumn")
+		);
+	}
+
+	MenuBuilder.EndSection();
+}
+
 void STaskTableTreeView::AddCommmands()
 {
 	FTaskTableTreeViewCommands::Register();
+
+	CommandList->MapAction(FTaskTableTreeViewCommands::Get().Command_GoToTask, FExecuteAction::CreateSP(this, &STaskTableTreeView::ContextMenu_GoToTask_Execute), FCanExecuteAction::CreateSP(this, &STaskTableTreeView::ContextMenu_GoToTask_CanExecute));
 }
 
 ////////////////////////////////////////////////////////////////////////////////////////////////////
@@ -369,6 +395,71 @@ bool STaskTableTreeView::TimestampOptions_IsEnabled() const
 
 ////////////////////////////////////////////////////////////////////////////////////////////////////
 
+bool STaskTableTreeView::ContextMenu_GoToTask_CanExecute() const
+{
+	return TreeView->GetNumItemsSelected() == 1;
+}
+
+////////////////////////////////////////////////////////////////////////////////////////////////////
+
+void STaskTableTreeView::ContextMenu_GoToTask_Execute()
+{
+	TArray<FTableTreeNodePtr> SelectedItems;
+	TreeView->GetSelectedItems(SelectedItems);
+
+	if (SelectedItems.Num() != 1)
+	{
+		return;
+	}
+
+	FTaskNodePtr SelectedTask = StaticCastSharedPtr<FTaskNode>(SelectedItems[0]);
+
+	TSharedPtr<STimingProfilerWindow> TimingWindow = FTimingProfilerManager::Get()->GetProfilerWindow();
+	if (!TimingWindow.IsValid())
+	{
+		return;
+	}
+
+	TSharedPtr<STimingView> TimingView = TimingWindow->GetTimingView();
+	if (!TimingView.IsValid())
+	{
+		return;
+	}
+
+	TimingView->ClearRelations();
+
+	auto AddRelation = [&TimingView](double SourceTimestamp, uint32 SourceThreadId, double TargetTimestamp, uint32 TargetThreadId, FTaskGraphRelation::ETaskGraphRelationType Type)
+	{
+		if (SourceTimestamp == TraceServices::FTaskInfo::InvalidTimestamp || TargetTimestamp == TraceServices::FTaskInfo::InvalidTimestamp)
+		{
+			return;
+		}
+
+		TUniquePtr<ITimingEventRelation> Relation = MakeUnique<FTaskGraphRelation>(SourceTimestamp, SourceThreadId, TargetTimestamp, TargetThreadId, Type);
+
+		TimingView->AddRelation(Relation);
+	};
+
+	FTaskGraphProfilerManager::Get()->GetTaskRelations(SelectedTask->GetTask()->GetId(), AddRelation);
+
+	double Duration = (SelectedTask->GetTask()->GetFinishedTimestamp() - SelectedTask->GetTask()->GetCreatedTimestamp()) * 1.5;
+	TimingView->ZoomOnTimeInterval(SelectedTask->GetTask()->GetCreatedTimestamp() - Duration * 0.15, Duration);
+}
+
+////////////////////////////////////////////////////////////////////////////////////////////////////
+
+void STaskTableTreeView::TreeView_OnMouseButtonDoubleClick(FTableTreeNodePtr TreeNode)
+{
+	if (!TreeNode->IsGroup())
+	{
+		ContextMenu_GoToTask_Execute();
+	}
+
+	STableTreeView::TreeView_OnMouseButtonDoubleClick(TreeNode);
+}
+
+////////////////////////////////////////////////////////////////////////////////////////////////////
+
 } // namespace Insights
 
 #undef LOCTEXT_NAMESPACE
diff --git a/Engine/Source/Developer/TraceInsights/Private/Insights/TaskGraphProfiler/Widgets/STaskTableTreeView.h b/Engine/Source/Developer/TraceInsights/Private/Insights/TaskGraphProfiler/Widgets/STaskTableTreeView.h
--- a/Engine/Source/Developer/TraceInsights/Private/Insights/TaskGraphProfiler/Widgets/STaskTableTreeView.h
+++ b/Engine/Source/Developer/TraceInsights/Private/Insights/TaskGraphProfiler/Widgets/STaskTableTreeView.h
@@ -9,6 +9,8 @@
 #include "Insights/TaskGraphProfiler/ViewModels/TaskTable.h"
 #include "Insights/Table/Widgets/STableTreeView.h"
 
+class FMenuBuilder;
+
 namespace Insights
 {
 
@@ -81,9 +83,13 @@ public:
 
 	////////////////////////////////////////////////////////////////////////////////////////////////////
 
+	virtual void TreeView_OnMouseButtonDoubleClick(FTableTreeNodePtr TreeNode) override;
+
 protected:
 	virtual void InternalCreateGroupings() override;
 
+	virtual void ExtendMenu(FMenuBuilder& MenuBuilder) override;
+
 private:
 	void ApplyColumnConfig(const TArrayView<FColumnConfig>& Preset);
 	void AddCommmands();
@@ -96,6 +102,9 @@ private:
 	FText TimestampOptions_GetText(ETimestampOptions InOption) const;
 	bool TimestampOptions_IsEnabled() const;
 
+	bool ContextMenu_GoToTask_CanExecute() const;
+	void ContextMenu_GoToTask_Execute();
+
 private:
 	double QueryStartTime = 0.0f;
 	double QueryEndTime = 0.0f;
diff --git a/Engine/Source/Developer/TraceInsights/Private/Insights/ViewModels/ThreadTimingTrack.cpp b/Engine/Source/Developer/TraceInsights/Private/Insights/ViewModels/ThreadTimingTrack.cpp
--- a/Engine/Source/Developer/TraceInsights/Private/Insights/ViewModels/ThreadTimingTrack.cpp
+++ b/Engine/Source/Developer/TraceInsights/Private/Insights/ViewModels/ThreadTimingTrack.cpp
@@ -16,6 +16,7 @@
 #include "Insights/Common/TimeUtils.h"
 #include "Insights/InsightsManager.h"
 #include "Insights/ITimingViewSession.h"
+#include "Insights/TaskGraphProfiler/TaskGraphProfilerManager.h"
 #include "Insights/TimingProfilerManager.h"
 #include "Insights/ViewModels/Filters.h"
 #include "Insights/ViewModels/FilterConfigurator.h"
@@ -1503,7 +1504,7 @@ bool FThreadTimingTrack::HasCustomFilter() const
 
 ////////////////////////////////////////////////////////////////////////////////////////////////////
 
-void FThreadTimingTrack::GetEventRelations(const ITimingEvent& InSelectedEvent, TArray<ITimingEventRelation*>& Relations) const
+void FThreadTimingTrack::GetEventRelations(const ITimingEvent& InSelectedEvent, TArray<TUniquePtr<ITimingEventRelation>>& Relations) const
 {
 	const int32 MaxTasksToShow = 30;
 	double StartTime = InSelectedEvent.GetStartTime();
@@ -1519,7 +1520,7 @@ void FThreadTimingTrack::GetEventRelations(const ITimingEvent& InSelectedEvent,
 				return;
 			}
 
-			FTaskGraphRelation* Relation = new FTaskGraphRelation(SourceTimestamp, SourceThreadId, TargetTimestamp, TargetThreadId, Type);
+			TUniquePtr<FTaskGraphRelation> Relation = MakeUnique<FTaskGraphRelation>(SourceTimestamp, SourceThreadId, TargetTimestamp, TargetThreadId, Type);
 
 			if (Relation->GetSourceThreadId() == ThreadId)
 			{
@@ -1536,57 +1537,10 @@ void FThreadTimingTrack::GetEventRelations(const ITimingEvent& InSelectedEvent,
 				Relation->SetIsSolved(true);
 			}
 
-			Relations.Add(Relation);
+			Relations.Add(MoveTemp(Relation));
 		};
 
-		const TraceServices::FTaskInfo* Task = TasksProvider->TryGetTask(ThreadId, StartTime);
-		if (Task)
-		{
-			if (Task->CreatedTimestamp != Task->LaunchedTimestamp || Task->CreatedThreadId != Task->LaunchedThreadId)
-			{
-				AddRelation(Task->CreatedTimestamp, Task->CreatedThreadId, Task->LaunchedTimestamp, Task->LaunchedThreadId, FTaskGraphRelation::ETaskGraphRelationType::Created);
-			}
-
-			AddRelation(Task->LaunchedTimestamp, Task->LaunchedThreadId, Task->ScheduledTimestamp, Task->ScheduledThreadId, FTaskGraphRelation::ETaskGraphRelationType::Launched);
-
-			int32 NumPrerequisitesToShow = FMath::Min(Task->Prerequisites.Num(), MaxTasksToShow);
-			for (int32 i = 0; i != NumPrerequisitesToShow; ++i)
-			{
-				const TraceServices::FTaskInfo* Prerequisite = TasksProvider->TryGetTask(Task->Prerequisites[i].RelativeId);
-				check(Prerequisite != nullptr);
-				AddRelation(Prerequisite->CompletedTimestamp, Prerequisite->CompletedThreadId, Task->StartedTimestamp, Task->StartedThreadId, FTaskGraphRelation::ETaskGraphRelationType::Prerequisite);
-			}
-
-			if (Task->LaunchedTimestamp != Task->ScheduledTimestamp || Task->LaunchedThreadId != Task->ScheduledThreadId)
-			{
-				AddRelation(Task->ScheduledTimestamp, Task->ScheduledThreadId, Task->StartedTimestamp, Task->StartedThreadId, FTaskGraphRelation::ETaskGraphRelationType::Scheduled);
-			}
-
-			int32 NumNestedToShow = FMath::Min(Task->NestedTasks.Num(), MaxTasksToShow);
-			for (int32 i = 0; i != NumNestedToShow; ++i)
-			{
-				const TraceServices::FTaskInfo::FRelationInfo& RelationInfo = Task->NestedTasks[i];
-				const TraceServices::FTaskInfo* NestedTask = TasksProvider->TryGetTask(RelationInfo.RelativeId);
-				check(NestedTask != nullptr);
-
-				AddRelation(RelationInfo.Timestamp, ThreadId, NestedTask->StartedTimestamp, NestedTask->StartedThreadId, FTaskGraphRelation::ETaskGraphRelationType::AddedNested);
-
-				AddRelation(NestedTask->CompletedTimestamp, NestedTask->CompletedThreadId, NestedTask->CompletedTimestamp, Task->StartedThreadId, FTaskGraphRelation::ETaskGraphRelationType::NestedCompleted);
-			}
-
-			int32 NumSubsequentsToShow = FMath::Min(Task->NestedTasks.Num(), MaxTasksToShow);
-			for (int32 i = 0; i != NumSubsequentsToShow; ++i)
-			{
-				const TraceServices::FTaskInfo* Subsequent = TasksProvider->TryGetTask(Task->Subsequents[i].RelativeId);
-				check(Subsequent != nullptr);
-				AddRelation(Task->CompletedTimestamp, ThreadId, Subsequent->StartedTimestamp, Subsequent->StartedThreadId, FTaskGraphRelation::ETaskGraphRelationType::Subsequent);
-			}
-
-			if (Task->FinishedTimestamp != Task->CompletedTimestamp || Task->CompletedThreadId != ThreadId)
-			{
-				AddRelation(Task->FinishedTimestamp, ThreadId, Task->CompletedTimestamp, ThreadId, FTaskGraphRelation::ETaskGraphRelationType::Completed);
-			}
-		}
+		FTaskGraphProfilerManager::Get()->GetTaskRelations(InSelectedEvent.GetStartTime(), ThreadId, AddRelation);
 
 		// if it's an event waiting for tasks completeness, add relations to these tasks
 		const TraceServices::ITimingProfilerProvider& TimingProfilerProvider = *TraceServices::ReadTimingProfilerProvider(*Session.Get());
@@ -1617,14 +1571,14 @@ void FThreadTimingTrack::GetEventRelations(const ITimingEvent& InSelectedEvent,
 
 ////////////////////////////////////////////////////////////////////////////////////////////////////
 
-void FThreadTimingTrack::SolveEventRelations(const TArray<ITimingEventRelation*>& Relations) const
+void FThreadTimingTrack::SolveEventRelations(const TArray<TUniquePtr<ITimingEventRelation>>& Relations) const
 {
-	for (ITimingEventRelation* Relation : Relations)
+	for (const TUniquePtr<ITimingEventRelation>& Relation : Relations)
 	{
 		if (!Relation->IsSolved())
 		{
 			// A check that we have the correct type will need to be added
-			FTaskGraphRelation* TaskRelation = static_cast<FTaskGraphRelation*>(Relation);
+			FTaskGraphRelation* TaskRelation = StaticCast<FTaskGraphRelation*>(Relation.Get());
 
 			if (TaskRelation->GetSourceThreadId() == ThreadId && !TaskRelation->GetSourceTrack())
 			{
diff --git a/Engine/Source/Developer/TraceInsights/Private/Insights/ViewModels/ThreadTimingTrack.h b/Engine/Source/Developer/TraceInsights/Private/Insights/ViewModels/ThreadTimingTrack.h
--- a/Engine/Source/Developer/TraceInsights/Private/Insights/ViewModels/ThreadTimingTrack.h
+++ b/Engine/Source/Developer/TraceInsights/Private/Insights/ViewModels/ThreadTimingTrack.h
@@ -160,8 +160,8 @@ public:
 	virtual void OnClipboardCopyEvent(const ITimingEvent& InSelectedEvent) const override;
 	virtual void BuildContextMenu(FMenuBuilder& MenuBuilder) override;
 
-	virtual void GetEventRelations(const ITimingEvent& InSelectedEvent, TArray<ITimingEventRelation*>& Relations) const override;
-	virtual void SolveEventRelations(const TArray<ITimingEventRelation*>& Relations) const override;
+	virtual void GetEventRelations(const ITimingEvent& InSelectedEvent, TArray<TUniquePtr<ITimingEventRelation>>& Relations) const override;
+	virtual void SolveEventRelations(const TArray<TUniquePtr<ITimingEventRelation>>& Relations) const override;
 
 protected:
 	virtual bool HasCustomFilter() const override;
diff --git a/Engine/Source/Developer/TraceInsights/Private/Insights/ViewModels/TimingViewDrawHelper.cpp b/Engine/Source/Developer/TraceInsights/Private/Insights/ViewModels/TimingViewDrawHelper.cpp
--- a/Engine/Source/Developer/TraceInsights/Private/Insights/ViewModels/TimingViewDrawHelper.cpp
+++ b/Engine/Source/Developer/TraceInsights/Private/Insights/ViewModels/TimingViewDrawHelper.cpp
@@ -656,9 +656,9 @@ FLinearColor FTimingViewDrawHelper::GetTrackNameTextColor(const FBaseTimingTrack
 
 ////////////////////////////////////////////////////////////////////////////////////////////////////
 
-void FTimingViewDrawHelper::DrawRelations(const TArray<ITimingEventRelation*> Relations, ITimingEventRelation::EDrawFilter Filter) const
+void FTimingViewDrawHelper::DrawRelations(const TArray<TUniquePtr<ITimingEventRelation>>& Relations, ITimingEventRelation::EDrawFilter Filter) const
 {
-	for (ITimingEventRelation* Relation : Relations)
+	for (const TUniquePtr<ITimingEventRelation> &Relation : Relations)
 	{
 		if (Relation->IsSolved())
 		{
diff --git a/Engine/Source/Developer/TraceInsights/Private/Insights/ViewModels/TimingViewDrawHelper.h b/Engine/Source/Developer/TraceInsights/Private/Insights/ViewModels/TimingViewDrawHelper.h
--- a/Engine/Source/Developer/TraceInsights/Private/Insights/ViewModels/TimingViewDrawHelper.h
+++ b/Engine/Source/Developer/TraceInsights/Private/Insights/ViewModels/TimingViewDrawHelper.h
@@ -196,7 +196,7 @@ public:
 
 	void EndDrawTracks() const;
 
-	void DrawRelations(const TArray<ITimingEventRelation*> Relations, ITimingEventRelation::EDrawFilter Filter) const;
+	void DrawRelations(const TArray<TUniquePtr<ITimingEventRelation>>& Relations, ITimingEventRelation::EDrawFilter Filter) const;
 
 	//////////////////////////////////////////////////
 
diff --git a/Engine/Source/Developer/TraceInsights/Private/Insights/Widgets/STimingView.cpp b/Engine/Source/Developer/TraceInsights/Private/Insights/Widgets/STimingView.cpp
--- a/Engine/Source/Developer/TraceInsights/Private/Insights/Widgets/STimingView.cpp
+++ b/Engine/Source/Developer/TraceInsights/Private/Insights/Widgets/STimingView.cpp
@@ -644,7 +644,7 @@ void STimingView::Tick(const FGeometry& AllottedGeometry, const double InCurrent
 		virtual const TSharedPtr<const ITimingEvent> GetHoveredEvent() const override { return TimingView->GetHoveredEvent(); }
 		virtual const TSharedPtr<const ITimingEvent> GetSelectedEvent() const override { return TimingView->GetSelectedEvent(); }
 		virtual const TSharedPtr<ITimingEventFilter> GetEventFilter() const override { return TimingView->GetEventFilter(); }
-		virtual const TArray<ITimingEventRelation*>& GetCurrentRelations() const override { return TimingView->GetCurrentRelations(); }
+		virtual const TArray<TUniquePtr<ITimingEventRelation>>& GetCurrentRelations() const override { return TimingView->GetCurrentRelations(); }
 		virtual double GetCurrentTime() const override { return CurrentTime; }
 		virtual float GetDeltaTime() const override { return DeltaTime; }
 
@@ -4192,12 +4192,7 @@ bool STimingView::ContextMenu_ShowTaskDependecies_IsChecked()
 
 void STimingView::UpdateEventRelations()
 {
-	for (ITimingEventRelation* Relation : CurrentRelations)
-	{
-		delete Relation;
-	}
-
-	CurrentRelations.Empty();
+	ClearRelations();
 
 	if (bShowEventRelations && SelectedEvent.IsValid())
 	{
@@ -4207,5 +4202,12 @@ void STimingView::UpdateEventRelations()
 
 ////////////////////////////////////////////////////////////////////////////////////////////////////
 
+void STimingView::ClearRelations()
+{
+	CurrentRelations.Empty();
+}
+
+////////////////////////////////////////////////////////////////////////////////////////////////////
+
 #undef ACTIVATE_BENCHMARK
 #undef LOCTEXT_NAMESPACE
diff --git a/Engine/Source/Developer/TraceInsights/Private/Insights/Widgets/STimingView.h b/Engine/Source/Developer/TraceInsights/Private/Insights/Widgets/STimingView.h
--- a/Engine/Source/Developer/TraceInsights/Private/Insights/Widgets/STimingView.h
+++ b/Engine/Source/Developer/TraceInsights/Private/Insights/Widgets/STimingView.h
@@ -344,7 +344,9 @@ public:
 
 	const TSharedPtr<FBaseTimingTrack> GetTrackAt(float InPosX, float InPosY) const;
 
-	const TArray<ITimingEventRelation*>& GetCurrentRelations() const { return CurrentRelations; }
+	const TArray<TUniquePtr<ITimingEventRelation>>& GetCurrentRelations() const { return CurrentRelations; }
+	void AddRelation(TUniquePtr<ITimingEventRelation>& Relation) { CurrentRelations.Add(MoveTemp(Relation)); }
+	void ClearRelations();
 
 protected:
 	virtual FVector2D ComputeDesiredSize(float) const override
@@ -628,6 +630,6 @@ protected:
 
 	TSharedPtr<FUICommandList> CommandList;
 
-	TArray<ITimingEventRelation*> CurrentRelations;
+	TArray<TUniquePtr<ITimingEventRelation>> CurrentRelations;
 	bool bShowEventRelations = true;
 };
\ No newline at end of file
diff --git a/Engine/Source/Developer/TraceInsights/Public/Insights/ViewModels/BaseTimingTrack.h b/Engine/Source/Developer/TraceInsights/Public/Insights/ViewModels/BaseTimingTrack.h
--- a/Engine/Source/Developer/TraceInsights/Public/Insights/ViewModels/BaseTimingTrack.h
+++ b/Engine/Source/Developer/TraceInsights/Public/Insights/ViewModels/BaseTimingTrack.h
@@ -94,7 +94,7 @@ public:
 	virtual const TSharedPtr<const ITimingEvent> GetHoveredEvent() const = 0;
 	virtual const TSharedPtr<const ITimingEvent> GetSelectedEvent() const = 0;
 	virtual const TSharedPtr<ITimingEventFilter> GetEventFilter() const = 0;
-	virtual const TArray<ITimingEventRelation*>& GetCurrentRelations() const = 0;
+	virtual const TArray<TUniquePtr<ITimingEventRelation>>& GetCurrentRelations() const = 0;
 	virtual double GetCurrentTime() const = 0;
 	virtual float GetDeltaTime() const = 0;
 };
@@ -254,9 +254,9 @@ public:
 	// Called back from the timing view when an event is copied to the clipboard with Ctrl+C.
 	virtual void OnClipboardCopyEvent(const ITimingEvent& InSelectedEvent) const {}
 
-	virtual void GetEventRelations(const ITimingEvent& InSelectedEvent, TArray<ITimingEventRelation*>& Dependencies) const {}
+	virtual void GetEventRelations(const ITimingEvent& InSelectedEvent, TArray<TUniquePtr<ITimingEventRelation>>& Dependencies) const {}
 
-	virtual void SolveEventRelations(const TArray<ITimingEventRelation*>& Dependencies) const {}
+	virtual void SolveEventRelations(const TArray<TUniquePtr<ITimingEventRelation>>& Dependencies) const {}
 
 	// Returns number of text lines needed to display the debug string.
 	//TODO: virtual int GetDebugStringLineCount() const { return 0; }
