[Insights] Clearing selected time interval no longer clears TaskTreeTableView

diff --git a/Engine/Source/Developer/TraceInsights/Private/Insights/TaskGraphProfiler/Widgets/STaskTableTreeView.cpp b/Engine/Source/Developer/TraceInsights/Private/Insights/TaskGraphProfiler/Widgets/STaskTableTreeView.cpp
--- a/Engine/Source/Developer/TraceInsights/Private/Insights/TaskGraphProfiler/Widgets/STaskTableTreeView.cpp
+++ b/Engine/Source/Developer/TraceInsights/Private/Insights/TaskGraphProfiler/Widgets/STaskTableTreeView.cpp
@@ -131,11 +131,18 @@ void STaskTableTreeView::Tick(const FGeometry& AllottedGeometry, const double In
 
 void STaskTableTreeView::RebuildTree(bool bResync)
 {
-	if (FTimingProfilerManager::Get()->GetSelectionStartTime() != QueryStartTime ||
-		FTimingProfilerManager::Get()->GetSelectionEndTime() != QueryEndTime)
+	double NewQueryStartTime = FTimingProfilerManager::Get()->GetSelectionStartTime();
+	double NewQueryEndTime = FTimingProfilerManager::Get()->GetSelectionEndTime();
+
+	if (NewQueryStartTime >= NewQueryEndTime)
+	{
+		return;
+	}
+
+	if (NewQueryStartTime != QueryStartTime || NewQueryEndTime != QueryEndTime)
 	{
-		QueryStartTime = FTimingProfilerManager::Get()->GetSelectionStartTime();
-		QueryEndTime = FTimingProfilerManager::Get()->GetSelectionEndTime();
+		QueryStartTime = NewQueryStartTime;
+		QueryEndTime = NewQueryEndTime;
 		TSharedPtr<FTaskTable> TaskTable = GetTaskTable();
 		TArray<FTaskEntry>& Allocs = TaskTable->GetTaskEntries();
 		Allocs.Empty();
