[Insights] Prevent the spawning of the TaskTableTreeView Tab if task data does not exist in the trace

diff --git a/Engine/Source/Developer/TraceInsights/Private/Insights/TaskGraphProfiler/TaskGraphProfilerManager.cpp b/Engine/Source/Developer/TraceInsights/Private/Insights/TaskGraphProfiler/TaskGraphProfilerManager.cpp
--- a/Engine/Source/Developer/TraceInsights/Private/Insights/TaskGraphProfiler/TaskGraphProfilerManager.cpp
+++ b/Engine/Source/Developer/TraceInsights/Private/Insights/TaskGraphProfiler/TaskGraphProfilerManager.cpp
@@ -136,8 +136,8 @@ bool FTaskGraphProfilerManager::Tick(float DeltaTime)
 			TSharedPtr<FTabManager> TabManagerShared = TimingTabManager.Pin();
 			if (TasksProvider && TasksProvider->GetNumTasks() > 0 && TabManagerShared.IsValid())
 			{
-				TabManagerShared->TryInvokeTab(FTaskGraphProfilerTabs::TaskTableTreeViewTabID);
 				bIsAvailable = true;
+				TabManagerShared->TryInvokeTab(FTaskGraphProfilerTabs::TaskTableTreeViewTabID);
 			}
 
 			if (Session->IsAnalysisComplete())
@@ -183,6 +183,7 @@ void FTaskGraphProfilerManager::RegisterTimingProfilerLayoutExtensions(FInsights
 	MinorTabConfig.TabTooltip = LOCTEXT("TaskTableTreeViewTabTitleTooltip", "Opens the Task Table Tree View tab, that allows Task Graph profilling.");
 	MinorTabConfig.TabIcon = FSlateIcon(FSlateIcon(FInsightsStyle::GetStyleSetName(), "TimersView.Icon.Small"));
 	MinorTabConfig.OnSpawnTab = FOnSpawnTab::CreateRaw(this, &FTaskGraphProfilerManager::SpawnTab_TaskTableTreeView);
+	MinorTabConfig.CanSpawnTab = FCanSpawnTab::CreateRaw(this, &FTaskGraphProfilerManager::CanSpawnTab_TaskTableTreeView);
 
 	InOutExtender.GetLayoutExtender().ExtendLayout(FTimingProfilerTabs::StatsCountersID
 		, ELayoutExtensionPosition::After
@@ -210,6 +211,13 @@ TSharedRef<SDockTab> FTaskGraphProfilerManager::SpawnTab_TaskTableTreeView(const
 
 ////////////////////////////////////////////////////////////////////////////////////////////////////
 
+bool FTaskGraphProfilerManager::CanSpawnTab_TaskTableTreeView(const FSpawnTabArgs& Args)
+{
+	return bIsAvailable;
+}
+
+////////////////////////////////////////////////////////////////////////////////////////////////////
+
 void FTaskGraphProfilerManager::OnTaskTableTreeViewTabClosed(TSharedRef<SDockTab> TabBeingClosed)
 {
 	TaskTableTreeView.Reset();
@@ -328,6 +336,22 @@ void FTaskGraphProfilerManager::GetTaskRelations(uint32 TaskId, AddRelationCallb
 
 ////////////////////////////////////////////////////////////////////////////////////////////////////
 
+void FTaskGraphProfilerManager::OnWindowClosedEvent()
+{
+	TSharedPtr<FTabManager> TimingTabManagerSharedPtr = TimingTabManager.Pin();
+
+	if (TimingTabManagerSharedPtr.IsValid())
+	{
+		TSharedPtr<SDockTab> Tab = TimingTabManagerSharedPtr->FindExistingLiveTab(FTaskGraphProfilerTabs::TaskTableTreeViewTabID);
+		if (Tab.IsValid())
+		{
+			Tab->RequestCloseTab();
+		}
+	}
+}
+
+////////////////////////////////////////////////////////////////////////////////////////////////////
+
 } // namespace Insights
 
 #undef LOCTEXT_NAMESPACE
diff --git a/Engine/Source/Developer/TraceInsights/Private/Insights/TaskGraphProfiler/TaskGraphProfilerManager.h b/Engine/Source/Developer/TraceInsights/Private/Insights/TaskGraphProfiler/TaskGraphProfilerManager.h
--- a/Engine/Source/Developer/TraceInsights/Private/Insights/TaskGraphProfiler/TaskGraphProfilerManager.h
+++ b/Engine/Source/Developer/TraceInsights/Private/Insights/TaskGraphProfiler/TaskGraphProfilerManager.h
@@ -62,9 +62,10 @@ public:
 	virtual void Shutdown() override;
 	virtual void RegisterMajorTabs(IUnrealInsightsModule& InsightsModule) override;
 	virtual void UnregisterMajorTabs() override;
-	virtual void OnWindowClosedEvent() override {}
+	virtual void OnWindowClosedEvent() override;
 
 	TSharedRef<SDockTab> SpawnTab_TaskTableTreeView(const FSpawnTabArgs& Args);
+	bool CanSpawnTab_TaskTableTreeView(const FSpawnTabArgs& Args);
 	void OnTaskTableTreeViewTabClosed(TSharedRef<SDockTab> TabBeingClosed);
 
 	//////////////////////////////////////////////////
diff --git a/Engine/Source/Developer/TraceInsights/Private/Insights/Widgets/STimingProfilerWindow.cpp b/Engine/Source/Developer/TraceInsights/Private/Insights/Widgets/STimingProfilerWindow.cpp
--- a/Engine/Source/Developer/TraceInsights/Private/Insights/Widgets/STimingProfilerWindow.cpp
+++ b/Engine/Source/Developer/TraceInsights/Private/Insights/Widgets/STimingProfilerWindow.cpp
@@ -438,7 +438,7 @@ void STimingProfilerWindow::Construct(const FArguments& InArgs, const TSharedRef
 	// Register any new minor tabs.
 	for (const FInsightsMinorTabConfig& MinorTabConfig : Extension->GetMinorTabs())
 	{
-		FTabSpawnerEntry& TabSpawnerEntry = TabManager->RegisterTabSpawner(MinorTabConfig.TabId, MinorTabConfig.OnSpawnTab);
+		FTabSpawnerEntry& TabSpawnerEntry = TabManager->RegisterTabSpawner(MinorTabConfig.TabId, MinorTabConfig.OnSpawnTab, MinorTabConfig.CanSpawnTab);
 
 		TabSpawnerEntry
 		.SetDisplayName(MinorTabConfig.TabLabel)
diff --git a/Engine/Source/Developer/TraceInsights/Public/Insights/IUnrealInsightsModule.h b/Engine/Source/Developer/TraceInsights/Public/Insights/IUnrealInsightsModule.h
--- a/Engine/Source/Developer/TraceInsights/Public/Insights/IUnrealInsightsModule.h
+++ b/Engine/Source/Developer/TraceInsights/Public/Insights/IUnrealInsightsModule.h
@@ -68,6 +68,8 @@ struct TRACEINSIGHTS_API FInsightsMinorTabConfig
 
 	FOnSpawnTab OnSpawnTab;
 
+	FCanSpawnTab CanSpawnTab;
+
 	FOnFindTabToReuse OnFindTabToReuse;
 
 	TSharedPtr<FWorkspaceItem> WorkspaceGroup;
