[Insights] Added "Widget Reflector" in the main menu of Unreal Insights (viewer mode). The option is available when -DebugTools is used in cmd line (ex.: if the "Enable Debug Tools" option is choosen in Browser Mode).

diff --git a/Engine/Source/Developer/TraceInsights/Private/Insights/Common/InsightsMenuBuilder.cpp b/Engine/Source/Developer/TraceInsights/Private/Insights/Common/InsightsMenuBuilder.cpp
--- a/Engine/Source/Developer/TraceInsights/Private/Insights/Common/InsightsMenuBuilder.cpp
+++ b/Engine/Source/Developer/TraceInsights/Private/Insights/Common/InsightsMenuBuilder.cpp
@@ -6,6 +6,8 @@
 #include "WorkspaceMenuStructure.h"
 #include "WorkspaceMenuStructureModule.h"
 
+#define LOCTEXT_NAMESPACE "InsightsMenuBuilder"
+
 ////////////////////////////////////////////////////////////////////////////////////////////////////
 
 FInsightsMenuBuilder::FInsightsMenuBuilder()
@@ -44,7 +46,22 @@ void FInsightsMenuBuilder::PopulateMenu(FMenuBuilder& MenuBuilder)
 {
 #if !WITH_EDITOR
 	FGlobalTabmanager::Get()->PopulateLocalTabSpawnerMenu(MenuBuilder);
+
+	static FName WidgetReflectorTabId("WidgetReflector");
+	bool bAllowDebugTools = FGlobalTabmanager::Get()->HasTabSpawner(WidgetReflectorTabId);
+	if (bAllowDebugTools)
+	{
+		MenuBuilder.BeginSection("WidgetTools");
+		MenuBuilder.AddMenuEntry(
+			LOCTEXT("OpenWidgetReflector", "Widget Reflector"),
+			LOCTEXT("OpenWidgetReflectorToolTip", "Opens the Widget Reflector, a handy tool for diagnosing problems with live widgets."),
+			FSlateIcon(FCoreStyle::Get().GetStyleSetName(), "WidgetReflector.Icon"),
+			FUIAction(FExecuteAction::CreateLambda([=] { FGlobalTabmanager::Get()->InvokeTab(WidgetReflectorTabId); })));
+		MenuBuilder.EndSection();
+	}
 #endif
 }
 
 ////////////////////////////////////////////////////////////////////////////////////////////////////
+
+#undef LOCTEXT_NAMESPACE
diff --git a/Engine/Source/Developer/TraceInsights/Private/Insights/TraceInsightsModule.cpp b/Engine/Source/Developer/TraceInsights/Private/Insights/TraceInsightsModule.cpp
--- a/Engine/Source/Developer/TraceInsights/Private/Insights/TraceInsightsModule.cpp
+++ b/Engine/Source/Developer/TraceInsights/Private/Insights/TraceInsightsModule.cpp
@@ -243,7 +243,7 @@ void FTraceInsightsModule::CreateSessionViewer(bool bAllowDebugTools)
 
 	AddAreaForSessionViewer(DefaultLayout);
 
-	AddAreaForWidgetReflector(DefaultLayout, bAllowDebugTools);
+	AddAreaForWidgetReflector(DefaultLayout, false);
 
 	// Load layout from ini file.
 	PersistentLayout = FLayoutSaveRestore::LoadFromConfig(UnrealInsightsLayoutIni, DefaultLayout);
@@ -300,18 +300,18 @@ void FTraceInsightsModule::AddAreaForSessionViewer(TSharedRef<FTabManager::FLayo
 
 ////////////////////////////////////////////////////////////////////////////////////////////////////
 
-void FTraceInsightsModule::AddAreaForWidgetReflector(TSharedRef<FTabManager::FLayout> Layout, bool bAllowDebugTools)
+void FTraceInsightsModule::AddAreaForWidgetReflector(TSharedRef<FTabManager::FLayout> Layout, bool bIsOpenedTab)
 {
 	const float DPIScaleFactor = FPlatformApplicationMisc::GetDPIScaleFactorAtPoint(10.0f, 10.0f);
 
 	// Create area and tab for Slate's WidgetReflector.
 	Layout->AddArea
 	(
-		FTabManager::NewArea(600.0f * DPIScaleFactor, 600.0f * DPIScaleFactor)
+		FTabManager::NewArea(800.0f * DPIScaleFactor, 400.0f * DPIScaleFactor)
 		->SetWindow(FVector2D(10.0f * DPIScaleFactor, 10.0f * DPIScaleFactor), false)
 		->Split
 		(
-			FTabManager::NewStack()->AddTab("WidgetReflector", bAllowDebugTools ? ETabState::OpenedTab : ETabState::ClosedTab)
+			FTabManager::NewStack()->AddTab(FTabId("WidgetReflector"), bIsOpenedTab ? ETabState::OpenedTab : ETabState::ClosedTab)
 		)
 	);
 }
