Fixed deprecation warning.

diff --git a/Engine/Source/Developer/TraceInsights/Private/Insights/Common/InsightsMenuBuilder.cpp b/Engine/Source/Developer/TraceInsights/Private/Insights/Common/InsightsMenuBuilder.cpp
--- a/Engine/Source/Developer/TraceInsights/Private/Insights/Common/InsightsMenuBuilder.cpp
+++ b/Engine/Source/Developer/TraceInsights/Private/Insights/Common/InsightsMenuBuilder.cpp
@@ -56,7 +56,7 @@ void FInsightsMenuBuilder::PopulateMenu(FMenuBuilder& MenuBuilder)
 			LOCTEXT("OpenWidgetReflector", "Widget Reflector"),
 			LOCTEXT("OpenWidgetReflectorToolTip", "Opens the Widget Reflector, a handy tool for diagnosing problems with live widgets."),
 			FSlateIcon(FCoreStyle::Get().GetStyleSetName(), "WidgetReflector.Icon"),
-			FUIAction(FExecuteAction::CreateLambda([=] { FGlobalTabmanager::Get()->InvokeTab(WidgetReflectorTabId); })));
+			FUIAction(FExecuteAction::CreateLambda([=] { FGlobalTabmanager::Get()->TryInvokeTab(WidgetReflectorTabId); })));
 		MenuBuilder.EndSection();
 	}
 #endif
