Add possibility to group by memory page of 4KB until we can use the real virtual memory regions
  - This is helpful in understanding allocation patterns that leads to many 4KB pages not being freed because they contain a single very small allocation

diff --git a/Engine/Source/Developer/TraceInsights/Private/Insights/MemoryProfiler/ViewModels/MemAllocTable.cpp b/Engine/Source/Developer/TraceInsights/Private/Insights/MemoryProfiler/ViewModels/MemAllocTable.cpp
--- a/Engine/Source/Developer/TraceInsights/Private/Insights/MemoryProfiler/ViewModels/MemAllocTable.cpp
+++ b/Engine/Source/Developer/TraceInsights/Private/Insights/MemoryProfiler/ViewModels/MemAllocTable.cpp
@@ -28,6 +28,7 @@ const FName FMemAllocTableColumns::StartTimeColumnId(TEXT("StartTime"));
 const FName FMemAllocTableColumns::EndTimeColumnId(TEXT("EndTime"));
 const FName FMemAllocTableColumns::DurationColumnId(TEXT("Duration"));
 const FName FMemAllocTableColumns::AddressColumnId(TEXT("Address"));
+const FName FMemAllocTableColumns::MemoryPageColumnId(TEXT("Memory Page"));
 const FName FMemAllocTableColumns::CountColumnId(TEXT("Count"));
 const FName FMemAllocTableColumns::SizeColumnId(TEXT("Size"));
 const FName FMemAllocTableColumns::TagColumnId(TEXT("Tag"));
@@ -307,6 +308,62 @@ void FMemAllocTable::AddDefaultColumns()
 		AddColumn(ColumnRef);
 	}
 	//////////////////////////////////////////////////
+	// Memory Page Column
+	{
+		TSharedRef<FTableColumn> ColumnRef = MakeShared<FTableColumn>(FMemAllocTableColumns::MemoryPageColumnId);
+		FTableColumn& Column = *ColumnRef;
+
+		Column.SetIndex(ColumnIndex++);
+
+		Column.SetShortName(LOCTEXT("MemoryPageColumnName", "Memory Page"));
+		Column.SetTitleName(LOCTEXT("MemoryPageColumnTitle", "Memory Page"));
+		Column.SetDescription(LOCTEXT("MemoryPageColumnDesc", "Memory Page of allocation"));
+
+		Column.SetFlags(ETableColumnFlags::CanBeHidden | ETableColumnFlags::CanBeFiltered);
+
+		Column.SetHorizontalAlignment(HAlign_Left);
+		Column.SetInitialWidth(120.0f);
+
+		Column.SetDataType(ETableCellDataType::Int64);
+
+		class FMemAllocPageValueGetter : public FTableCellValueGetter
+		{
+		public:
+			virtual const TOptional<FTableCellValue> GetValue(const FTableColumn& Column, const FBaseTreeNode& Node) const override
+			{
+				if (Node.IsGroup())
+				{
+					const FTableTreeNode& NodePtr = static_cast<const FTableTreeNode&>(Node);
+					if (NodePtr.HasAggregatedValue(Column.GetId()))
+					{
+						return NodePtr.GetAggregatedValue(Column.GetId());
+					}
+				}
+				else
+				{
+					const FMemAllocNode& MemAllocNode = static_cast<const FMemAllocNode&>(Node);
+					const FMemoryAlloc* Alloc = MemAllocNode.GetMemAlloc();
+					if (Alloc)
+					{
+						return FTableCellValue(static_cast<int64>(Alloc->GetPage()));
+					}
+				}
+
+				return TOptional<FTableCellValue>();
+			}
+		};
+		TSharedRef<ITableCellValueGetter> Getter = MakeShared<FMemAllocPageValueGetter>();
+		Column.SetValueGetter(Getter);
+
+		TSharedRef<ITableCellValueFormatter> Formatter = MakeShared<FInt64ValueFormatterAsHex64>();
+		Column.SetValueFormatter(Formatter);
+
+		TSharedRef<ITableCellValueSorter> Sorter = MakeShared<FSorterByInt64Value>(ColumnRef);
+		Column.SetValueSorter(Sorter);
+
+		AddColumn(ColumnRef);
+	}
+	//////////////////////////////////////////////////
 	// Alloc Count Column
 	{
 		TSharedRef<FTableColumn> ColumnRef = MakeShared<FTableColumn>(FMemAllocTableColumns::CountColumnId);
diff --git a/Engine/Source/Developer/TraceInsights/Private/Insights/MemoryProfiler/ViewModels/MemAllocTable.h b/Engine/Source/Developer/TraceInsights/Private/Insights/MemoryProfiler/ViewModels/MemAllocTable.h
--- a/Engine/Source/Developer/TraceInsights/Private/Insights/MemoryProfiler/ViewModels/MemAllocTable.h
+++ b/Engine/Source/Developer/TraceInsights/Private/Insights/MemoryProfiler/ViewModels/MemAllocTable.h
@@ -26,6 +26,7 @@ struct FMemAllocTableColumns
 	static const FName SizeColumnId;
 	static const FName TagColumnId;
 	static const FName FunctionColumnId;
+	static const FName MemoryPageColumnId;
 };
 
 ////////////////////////////////////////////////////////////////////////////////////////////////////
diff --git a/Engine/Source/Developer/TraceInsights/Private/Insights/MemoryProfiler/ViewModels/MemoryAlloc.h b/Engine/Source/Developer/TraceInsights/Private/Insights/MemoryProfiler/ViewModels/MemoryAlloc.h
--- a/Engine/Source/Developer/TraceInsights/Private/Insights/MemoryProfiler/ViewModels/MemoryAlloc.h
+++ b/Engine/Source/Developer/TraceInsights/Private/Insights/MemoryProfiler/ViewModels/MemoryAlloc.h
@@ -22,6 +22,7 @@ public:
 	double GetEndTime() const { return EndTime; }
 	double GetDuration() const { return EndTime - StartTime; }
 	uint64 GetAddress() const { return Address; }
+	uint64 GetPage() const { return Address & ~(4llu*1024-1); }
 	uint64 GetSize() const { return Size; }
 	const TCHAR* GetTag() const { return Tag; }
 	const TraceServices::FCallstack* GetCallstack() const { return Callstack; }
diff --git a/Engine/Source/Developer/TraceInsights/Private/Insights/MemoryProfiler/Widgets/SMemAllocTableTreeView.cpp b/Engine/Source/Developer/TraceInsights/Private/Insights/MemoryProfiler/Widgets/SMemAllocTableTreeView.cpp
--- a/Engine/Source/Developer/TraceInsights/Private/Insights/MemoryProfiler/Widgets/SMemAllocTableTreeView.cpp
+++ b/Engine/Source/Developer/TraceInsights/Private/Insights/MemoryProfiler/Widgets/SMemAllocTableTreeView.cpp
@@ -474,6 +474,16 @@ TSharedPtr<SWidget> SMemAllocTableTreeView::ConstructToolbar()
 			.OnClicked(this, &SMemAllocTableTreeView::OnCallstackViewClicked, true)
 		]
 
+		+ SHorizontalBox::Slot()
+		.AutoWidth()
+		.Padding(4.0f, 0.0f, 0.0f, 0.0f)
+		[
+			SNew(SButton)
+			.Text(LOCTEXT("ByMemoryPageBtn_Text", "By Memory Page"))
+			.ToolTipText(LOCTEXT("ByMemoryRegionBtn_Tooltip", "Memory Page Breakdown View\nConfigure the tree view to show a breakdown of allocations by memory page."))
+			.OnClicked(this, &SMemAllocTableTreeView::OnMemoryPageClicked)
+		]
+
 		+ SHorizontalBox::Slot()
 		.AutoWidth()
 		.Padding(4.0f, 0.0f, 0.0f, 0.0f)
@@ -562,15 +572,16 @@ FReply SMemAllocTableTreeView::OnDetailedViewClicked()
 
 	FColumnConfig Preset[] =
 	{
-		{ FTable::GetHierarchyColumnId(),           true,  200.0f },
-		{ FMemAllocTableColumns::StartTimeColumnId, true,  100.0f },
-		{ FMemAllocTableColumns::EndTimeColumnId,   true,  100.0f },
-		{ FMemAllocTableColumns::DurationColumnId,  true,  100.0f },
-		{ FMemAllocTableColumns::AddressColumnId,   true,  120.0f },
-		{ FMemAllocTableColumns::CountColumnId,     true,  100.0f },
-		{ FMemAllocTableColumns::SizeColumnId,      true,  100.0f },
-		{ FMemAllocTableColumns::TagColumnId,       true,  120.0f },
-		{ FMemAllocTableColumns::FunctionColumnId,  true,  550.0f },
+		{ FTable::GetHierarchyColumnId(),                true,  200.0f },
+		{ FMemAllocTableColumns::StartTimeColumnId,      true,  100.0f },
+		{ FMemAllocTableColumns::EndTimeColumnId,        true,  100.0f },
+		{ FMemAllocTableColumns::DurationColumnId,       true,  100.0f },
+		{ FMemAllocTableColumns::AddressColumnId,        true,  120.0f },
+		{ FMemAllocTableColumns::MemoryPageColumnId,     true,  120.0f },
+		{ FMemAllocTableColumns::CountColumnId,          true,  100.0f },
+		{ FMemAllocTableColumns::SizeColumnId,           true,  100.0f },
+		{ FMemAllocTableColumns::TagColumnId,            true,  120.0f },
+		{ FMemAllocTableColumns::FunctionColumnId,       true,  550.0f },
 	};
 	ApplyColumnConfig(TArrayView<FColumnConfig>(Preset, UE_ARRAY_COUNT(Preset)));
 
@@ -604,16 +615,64 @@ FReply SMemAllocTableTreeView::OnSizeViewClicked()
 
 	PostChangeGroupings();
 
+	FColumnConfig Preset[] =
+	{
+		{ FTable::GetHierarchyColumnId(),              true,  200.0f },
+		{ FMemAllocTableColumns::StartTimeColumnId,    false, 0.0f },
+		{ FMemAllocTableColumns::EndTimeColumnId,      false, 0.0f },
+		{ FMemAllocTableColumns::DurationColumnId,     false, 0.0f },
+		{ FMemAllocTableColumns::AddressColumnId,      true,  120.0f },
+		{ FMemAllocTableColumns::MemoryPageColumnId,   false, 0.0f },
+		{ FMemAllocTableColumns::CountColumnId,        true,  100.0f },
+		{ FMemAllocTableColumns::SizeColumnId,         true,  100.0f },
+		{ FMemAllocTableColumns::TagColumnId,          true,  120.0f },
+		{ FMemAllocTableColumns::FunctionColumnId,     true,  400.0f },
+	};
+	ApplyColumnConfig(TArrayView<FColumnConfig>(Preset, UE_ARRAY_COUNT(Preset)));
+
+	return FReply::Handled();
+}
+
+////////////////////////////////////////////////////////////////////////////////////////////////////
+
+
+FReply SMemAllocTableTreeView::OnMemoryPageClicked()
+{
+	ColumnBeingSorted = FTable::GetHierarchyColumnId();
+	ColumnSortMode = EColumnSortMode::Type::Ascending;
+	UpdateCurrentSortingByColumn();
+
+	PreChangeGroupings();
+
+	CurrentGroupings.Reset();
+
+	check(AvailableGroupings[0]->Is<FTreeNodeGroupingFlat>());
+	CurrentGroupings.Add(AvailableGroupings[0]);
+
+	TSharedPtr<FTreeNodeGrouping>* MemoryPageGrouping = AvailableGroupings.FindByPredicate(
+		[](TSharedPtr<FTreeNodeGrouping>& Grouping)
+	{
+		return Grouping->Is<FTreeNodeGroupingByUniqueValueInt64>() &&
+			Grouping->As<FTreeNodeGroupingByUniqueValueInt64>().GetColumnId() == FMemAllocTableColumns::MemoryPageColumnId;
+	});
+	if (MemoryPageGrouping)
+	{
+		CurrentGroupings.Add(*MemoryPageGrouping);
+	}
+
+	PostChangeGroupings();
+
 	FColumnConfig Preset[] =
 	{
 		{ FTable::GetHierarchyColumnId(),           true,  200.0f },
 		{ FMemAllocTableColumns::StartTimeColumnId, false, 0.0f },
 		{ FMemAllocTableColumns::EndTimeColumnId,   false, 0.0f },
 		{ FMemAllocTableColumns::DurationColumnId,  false, 0.0f },
-		{ FMemAllocTableColumns::AddressColumnId,   false, 0.0f },
+		{ FMemAllocTableColumns::AddressColumnId,   true,  120.0f },
+		{ FMemAllocTableColumns::MemoryPageColumnId,false, 0.0f },
 		{ FMemAllocTableColumns::CountColumnId,     true,  100.0f },
 		{ FMemAllocTableColumns::SizeColumnId,      true,  100.0f },
-		{ FMemAllocTableColumns::TagColumnId,       true,  120.0f },
+		{ FMemAllocTableColumns::TagColumnId,       false, 0.0f },
 		{ FMemAllocTableColumns::FunctionColumnId,  true,  400.0f },
 	};
 	ApplyColumnConfig(TArrayView<FColumnConfig>(Preset, UE_ARRAY_COUNT(Preset)));
@@ -656,6 +715,7 @@ FReply SMemAllocTableTreeView::OnTagViewClicked()
 		{ FMemAllocTableColumns::EndTimeColumnId,   false, 0.0f },
 		{ FMemAllocTableColumns::DurationColumnId,  false, 0.0f },
 		{ FMemAllocTableColumns::AddressColumnId,   false, 0.0f },
+		{ FMemAllocTableColumns::MemoryPageColumnId,false, 0.0f },
 		{ FMemAllocTableColumns::CountColumnId,     true,  100.0f },
 		{ FMemAllocTableColumns::SizeColumnId,      true,  100.0f },
 		{ FMemAllocTableColumns::TagColumnId,       false, 0.0f },
diff --git a/Engine/Source/Developer/TraceInsights/Private/Insights/MemoryProfiler/Widgets/SMemAllocTableTreeView.h b/Engine/Source/Developer/TraceInsights/Private/Insights/MemoryProfiler/Widgets/SMemAllocTableTreeView.h
--- a/Engine/Source/Developer/TraceInsights/Private/Insights/MemoryProfiler/Widgets/SMemAllocTableTreeView.h
+++ b/Engine/Source/Developer/TraceInsights/Private/Insights/MemoryProfiler/Widgets/SMemAllocTableTreeView.h
@@ -103,6 +103,7 @@ private:
 	FReply OnDetailedViewClicked();
 	FReply OnSizeViewClicked();
 	FReply OnTagViewClicked();
+	FReply OnMemoryPageClicked();
 	FReply OnCallstackViewClicked(bool bIsInverted);
 
 	FText GetSymbolResolutionStatus() const;
