[Insights
  - Browser: Added more logs to diagnose "Connect" functionality.
  - Browser: Added more trace metadata: Branch, BuildVersion and Changelist.
  - Session Info: Added info about Branch, BuildVersion and Changelist metadata and disabled Uri field.
  - Browser: Optimized the tooltips in Session list.
  - Added SLazyToolTip.

diff --git a/Engine/Source/Developer/TraceInsights/Private/Insights/StoreService/DiagnosticsSessionAnalyzer.cpp b/Engine/Source/Developer/TraceInsights/Private/Insights/StoreService/DiagnosticsSessionAnalyzer.cpp
--- a/Engine/Source/Developer/TraceInsights/Private/Insights/StoreService/DiagnosticsSessionAnalyzer.cpp
+++ b/Engine/Source/Developer/TraceInsights/Private/Insights/StoreService/DiagnosticsSessionAnalyzer.cpp
@@ -51,7 +51,9 @@ bool FDiagnosticsSessionAnalyzer::OnEvent(uint16 RouteId, EStyle, const FOnEvent
 		EventData.GetString("Platform", Platform);
 		EventData.GetString("AppName", AppName);
 		EventData.GetString("CommandLine", CommandLine);
-
+		EventData.GetString("Branch", Branch);
+		EventData.GetString("BuildVersion", BuildVersion);
+		Changelist = EventData.GetValue<uint32>("Changelist", 0);
 		ConfigurationType = (EBuildConfiguration) EventData.GetValue<uint8>("ConfigurationType");
 		TargetType = (EBuildTargetType) EventData.GetValue<uint8>("TargetType");
 
diff --git a/Engine/Source/Developer/TraceInsights/Private/Insights/StoreService/DiagnosticsSessionAnalyzer.h b/Engine/Source/Developer/TraceInsights/Private/Insights/StoreService/DiagnosticsSessionAnalyzer.h
--- a/Engine/Source/Developer/TraceInsights/Private/Insights/StoreService/DiagnosticsSessionAnalyzer.h
+++ b/Engine/Source/Developer/TraceInsights/Private/Insights/StoreService/DiagnosticsSessionAnalyzer.h
@@ -22,8 +22,11 @@ struct FDiagnosticsSessionAnalyzer : public UE::Trace::IAnalyzer
 	FString Platform;
 	FString AppName;
 	FString CommandLine;
-	EBuildConfiguration ConfigurationType;
-	EBuildTargetType TargetType;
+	FString Branch;
+	FString BuildVersion;
+	uint32 Changelist = 0;
+	EBuildConfiguration ConfigurationType = EBuildConfiguration::Unknown;
+	EBuildTargetType TargetType = EBuildTargetType::Unknown;
 };
 
 } // namespace Insights
diff --git a/Engine/Source/Developer/TraceInsights/Private/Insights/StoreService/StoreBrowser.cpp b/Engine/Source/Developer/TraceInsights/Private/Insights/StoreService/StoreBrowser.cpp
--- a/Engine/Source/Developer/TraceInsights/Private/Insights/StoreService/StoreBrowser.cpp
+++ b/Engine/Source/Developer/TraceInsights/Private/Insights/StoreService/StoreBrowser.cpp
@@ -372,6 +372,9 @@ void FStoreBrowser::UpdateMetadata(FStoreBrowserTraceInfo& Trace)
 			Trace.Platform = Analyzer.Platform;
 			Trace.AppName = Analyzer.AppName;
 			Trace.CommandLine = Analyzer.CommandLine;
+			Trace.Branch = Analyzer.Branch;
+			Trace.BuildVersion = Analyzer.BuildVersion;
+			Trace.Changelist = Analyzer.Changelist;
 			Trace.ConfigurationType = static_cast<EBuildConfiguration>(Analyzer.ConfigurationType);
 			Trace.TargetType = static_cast<EBuildTargetType>(Analyzer.TargetType);
 		}
diff --git a/Engine/Source/Developer/TraceInsights/Private/Insights/StoreService/StoreBrowser.h b/Engine/Source/Developer/TraceInsights/Private/Insights/StoreService/StoreBrowser.h
--- a/Engine/Source/Developer/TraceInsights/Private/Insights/StoreService/StoreBrowser.h
+++ b/Engine/Source/Developer/TraceInsights/Private/Insights/StoreService/StoreBrowser.h
@@ -43,6 +43,9 @@ struct FStoreBrowserTraceInfo
 	FString Platform;
 	FString AppName;
 	FString CommandLine;
+	FString Branch;
+	FString BuildVersion;
+	uint32 Changelist;
 	EBuildConfiguration ConfigurationType = EBuildConfiguration::Unknown;
 	EBuildTargetType TargetType = EBuildTargetType::Unknown;
 
diff --git a/Engine/Source/Developer/TraceInsights/Private/Insights/Widgets/SLazyToolTip.h b/Engine/Source/Developer/TraceInsights/Private/Insights/Widgets/SLazyToolTip.h
new file mode 100644
--- /dev/null
+++ b/Engine/Source/Developer/TraceInsights/Private/Insights/Widgets/SLazyToolTip.h
@@ -0,0 +1,68 @@
+// Copyright Epic Games, Inc. All Rights Reserved.
+
+#pragma once
+
+#include "CoreMinimal.h"
+#include "SlateFwd.h"
+#include "Widgets/DeclarativeSyntaxSupport.h"
+#include "Widgets/SCompoundWidget.h"
+
+class ILazyToolTipCreator
+{
+public:
+	virtual TSharedPtr<SToolTip> CreateTooltip() const = 0;
+};
+
+class SLazyToolTip : public SCompoundWidget, public IToolTip
+{
+	SLATE_BEGIN_ARGS(SLazyToolTip) {}
+	SLATE_END_ARGS()
+
+public:
+	void Construct(const FArguments& InArgs, TSharedPtr<const ILazyToolTipCreator> InCreator)
+	{
+		Creator = InCreator;
+	}
+
+	virtual TSharedRef<SWidget> AsWidget()
+	{
+		CreateToolTipWidget();
+		return ToolTipWidget.ToSharedRef();
+	}
+
+	virtual TSharedRef<SWidget> GetContentWidget()
+	{
+		CreateToolTipWidget();
+		return ToolTipWidget->GetContentWidget();
+	}
+
+	virtual void SetContentWidget(const TSharedRef<SWidget>& InContentWidget)
+	{
+		CreateToolTipWidget();
+		ToolTipWidget->SetContentWidget(InContentWidget);
+	}
+
+	void InvalidateWidget()
+	{
+		ToolTipWidget.Reset();
+	}
+
+	virtual bool IsEmpty() const { return false; }
+	virtual bool IsInteractive() const { return false; }
+	virtual void OnOpening() {}
+	virtual void OnClosed() {}
+
+private:
+	void CreateToolTipWidget()
+	{
+		if (!ToolTipWidget.IsValid())
+		{
+			ToolTipWidget = Creator->CreateTooltip();
+			check(ToolTipWidget.IsValid());
+		}
+	}
+
+private:
+	TSharedPtr<SToolTip> ToolTipWidget;
+	TSharedPtr<const ILazyToolTipCreator> Creator;
+};
diff --git a/Engine/Source/Developer/TraceInsights/Private/Insights/Widgets/SSessionInfoWindow.cpp b/Engine/Source/Developer/TraceInsights/Private/Insights/Widgets/SSessionInfoWindow.cpp
--- a/Engine/Source/Developer/TraceInsights/Private/Insights/Widgets/SSessionInfoWindow.cpp
+++ b/Engine/Source/Developer/TraceInsights/Private/Insights/Widgets/SSessionInfoWindow.cpp
@@ -229,13 +229,16 @@ TSharedRef<SDockTab> SSessionInfoWindow::SpawnTab_SessionInfo(const FSpawnTabArg
 	Image->SetImage(new FSlateColorBrush(FLinearColor(0.015f, 0.015f, 0.015f, 1.0f)));
 
 	AddInfoLine(VerticalBox, LOCTEXT("SessionName_HeaderText",	"Session Name:"),		TAttribute<FText>(this, &SSessionInfoWindow::GetSessionNameText));
-	AddInfoLine(VerticalBox, LOCTEXT("Uri_HeaderText",			"URI:"),				TAttribute<FText>(this, &SSessionInfoWindow::GetUriText));
+	//AddInfoLine(VerticalBox, LOCTEXT("Uri_HeaderText",			"URI:"),				TAttribute<FText>(this, &SSessionInfoWindow::GetUriText));
+	//AddInfoLine(VerticalBox, LOCTEXT("FileSize_HeaderText",		"File Size:"),			TAttribute<FText>(this, &SSessionInfoWindow::GetFileSizeText));
 	AddInfoLine(VerticalBox, LOCTEXT("Platform_HeaderText",		"Platform:"),			TAttribute<FText>(this, &SSessionInfoWindow::GetPlatformText));
 	AddInfoLine(VerticalBox, LOCTEXT("AppName_HeaderText",		"Application Name:"),	TAttribute<FText>(this, &SSessionInfoWindow::GetAppNameText));
+	AddInfoLine(VerticalBox, LOCTEXT("Branch_HeaderText",		"Branch:"),				TAttribute<FText>(this, &SSessionInfoWindow::GetBranchText));
+	AddInfoLine(VerticalBox, LOCTEXT("BuildVersion_HeaderText",	"Build Version:"),		TAttribute<FText>(this, &SSessionInfoWindow::GetBuildVersionText));
+	AddInfoLine(VerticalBox, LOCTEXT("Changelist_HeaderText",	"Changelist:"),			TAttribute<FText>(this, &SSessionInfoWindow::GetChangelistText));
 	AddInfoLine(VerticalBox, LOCTEXT("BuildConfig_HeaderText",	"Build Config:"),		TAttribute<FText>(this, &SSessionInfoWindow::GetBuildConfigText));
 	AddInfoLine(VerticalBox, LOCTEXT("BuildTarget_HeaderText",	"Build Target:"),		TAttribute<FText>(this, &SSessionInfoWindow::GetBuildTargetText));
 	AddInfoLine(VerticalBox, LOCTEXT("CommandLine_HeaderText",	"Command Line:"),		TAttribute<FText>(this, &SSessionInfoWindow::GetCommandLineText), true);
-	//AddInfoLine(VerticalBox, LOCTEXT("FileSize_HeaderText",		"File Size:"),			TAttribute<FText>(this, &SSessionInfoWindow::GetFileSizeText));
 	AddInfoLine(VerticalBox, LOCTEXT("Status_HeaderText",		"Analysis Status:"),	TAttribute<FText>(this, &SSessionInfoWindow::GetStatusText));
 	AddInfoLine(VerticalBox, LOCTEXT("Modules_HeaderText",		"Analysis Modules:"),	TAttribute<FText>(this, &SSessionInfoWindow::GetModulesText));
 
@@ -289,9 +292,12 @@ void SSessionInfoWindow::Tick(const FGeometry& AllottedGeometry, const double In
 			TraceServices::FSessionInfo SessionInfo = DiagnosticsProvider.GetSessionInfo();
 			PlatformText = FText::FromString(SessionInfo.Platform);
 			AppNameText = FText::FromString(SessionInfo.AppName);
-			CommandLineText = FText::FromString(SessionInfo.CommandLine);
+			BranchText = FText::FromString(SessionInfo.Branch);
+			BuildVersionText = FText::FromString(SessionInfo.BuildVersion);
+			ChangelistText = FText::AsNumber(SessionInfo.Changelist, &FNumberFormattingOptions::DefaultNoGrouping());
 			BuildConfigurationTypeText = FText::FromString(LexToString(SessionInfo.ConfigurationType));
 			BuildTargetTypeText = FText::FromString(LexToString(SessionInfo.TargetType));
+			CommandLineText = FText::FromString(SessionInfo.CommandLine);
 			bIsSessionInfoSet = true;
 		}
 	}
@@ -384,41 +390,6 @@ FText SSessionInfoWindow::GetUriText() const
 
 ////////////////////////////////////////////////////////////////////////////////////////////////////
 
-FText SSessionInfoWindow::GetPlatformText() const
-{
-	return PlatformText;
-}
-
-////////////////////////////////////////////////////////////////////////////////////////////////////
-
-FText SSessionInfoWindow::GetAppNameText() const
-{
-	return AppNameText;
-}
-
-////////////////////////////////////////////////////////////////////////////////////////////////////
-
-FText SSessionInfoWindow::GetBuildConfigText() const
-{
-	return BuildConfigurationTypeText;
-}
-
-////////////////////////////////////////////////////////////////////////////////////////////////////
-
-FText SSessionInfoWindow::GetBuildTargetText() const
-{
-	return BuildTargetTypeText;
-}
-
-////////////////////////////////////////////////////////////////////////////////////////////////////
-
-FText SSessionInfoWindow::GetCommandLineText() const
-{
-	return CommandLineText;
-}
-
-////////////////////////////////////////////////////////////////////////////////////////////////////
-
 FText SSessionInfoWindow::GetFileSizeText() const
 {
 	//TODO
diff --git a/Engine/Source/Developer/TraceInsights/Private/Insights/Widgets/SSessionInfoWindow.h b/Engine/Source/Developer/TraceInsights/Private/Insights/Widgets/SSessionInfoWindow.h
--- a/Engine/Source/Developer/TraceInsights/Private/Insights/Widgets/SSessionInfoWindow.h
+++ b/Engine/Source/Developer/TraceInsights/Private/Insights/Widgets/SSessionInfoWindow.h
@@ -118,7 +118,6 @@ private:
 	 */
 	static void FillMenu(FMenuBuilder& MenuBuilder, const TSharedPtr<FTabManager> TabManager);
 
-
 private:
 	void AddInfoLine(TSharedPtr<SVerticalBox> InVerticalBox, const FText& InHeader, const TAttribute<FText>& InValue, bool bMultiLine = false) const;
 
@@ -126,12 +125,17 @@ private:
 	void OnSessionInfoTabClosed(TSharedRef<SDockTab> TabBeingClosed);
 
 	FText GetSessionNameText() const;
+
+	FText GetPlatformText() const { return PlatformText; }
+	FText GetAppNameText() const { return AppNameText; }
+	FText GetBranchText() const { return BranchText; }
+	FText GetBuildVersionText() const { return BuildVersionText; }
+	FText GetChangelistText() const { return ChangelistText; }
+	FText GetBuildConfigText() const { return BuildConfigurationTypeText; }
+	FText GetBuildTargetText() const { return BuildTargetTypeText; }
+	FText GetCommandLineText() const { return CommandLineText; }
+
 	FText GetUriText() const;
-	FText GetPlatformText() const;
-	FText GetAppNameText() const;
-	FText GetBuildConfigText() const;
-	FText GetBuildTargetText() const;
-	FText GetCommandLineText() const;
 	FText GetFileSizeText() const;
 	FText GetStatusText() const;
 	FText GetModulesText() const;
@@ -149,9 +153,12 @@ private:
 
 	FText PlatformText;
 	FText AppNameText;
-	FText CommandLineText;
+	FText BranchText;
+	FText BuildVersionText;
+	FText ChangelistText;
 	FText BuildConfigurationTypeText;
 	FText BuildTargetTypeText;
+	FText CommandLineText;
 	bool bIsSessionInfoSet = false;
 
 	//FText Uri;
diff --git a/Engine/Source/Developer/TraceInsights/Private/Insights/Widgets/SStartPageWindow.cpp b/Engine/Source/Developer/TraceInsights/Private/Insights/Widgets/SStartPageWindow.cpp
--- a/Engine/Source/Developer/TraceInsights/Private/Insights/Widgets/SStartPageWindow.cpp
+++ b/Engine/Source/Developer/TraceInsights/Private/Insights/Widgets/SStartPageWindow.cpp
@@ -54,6 +54,7 @@
 #include "Insights/InsightsStyle.h"
 #include "Insights/Version.h"
 #include "Insights/Widgets/SInsightsSettings.h"
+#include "Insights/Widgets/SLazyToolTip.h"
 
 ////////////////////////////////////////////////////////////////////////////////////////////////////
 
@@ -63,7 +64,7 @@
 // STraceListRow
 ////////////////////////////////////////////////////////////////////////////////////////////////////
 
-class STraceListRow : public SMultiColumnTableRow<TSharedPtr<FTraceViewModel>>
+class STraceListRow : public SMultiColumnTableRow<TSharedPtr<FTraceViewModel>>, public ILazyToolTipCreator
 {
 	SLATE_BEGIN_ARGS(STraceListRow) {}
 	SLATE_END_ARGS()
@@ -84,7 +85,6 @@ public:
 		SMultiColumnTableRow<TSharedPtr<FTraceViewModel>>::Construct(FSuperRowType::FArguments(), InOwnerTableView);
 	}
 
-public:
 	virtual TSharedRef<SWidget> GenerateWidgetForColumn(const FName& ColumnName) override
 	{
 		if (ColumnName == FName(TEXT("Name")))
@@ -253,6 +253,58 @@ public:
 		}
 	}
 
+	FText GetTraceBranch() const
+	{
+		TSharedPtr<FTraceViewModel> TracePin = WeakTrace.Pin();
+		if (TracePin.IsValid())
+		{
+			return TracePin->Branch;
+		}
+		else
+		{
+			return FText::GetEmpty();
+		}
+	}
+
+	FText GetTraceBuildVersion() const
+	{
+		TSharedPtr<FTraceViewModel> TracePin = WeakTrace.Pin();
+		if (TracePin.IsValid())
+		{
+			return TracePin->BuildVersion;
+		}
+		else
+		{
+			return FText::GetEmpty();
+		}
+	}
+
+	FText GetTraceChangelist() const
+	{
+		TSharedPtr<FTraceViewModel> TracePin = WeakTrace.Pin();
+		if (TracePin.IsValid())
+		{
+			return FText::AsNumber(TracePin->Changelist, &FNumberFormattingOptions::DefaultNoGrouping());
+		}
+		else
+		{
+			return FText::GetEmpty();
+		}
+	}
+
+	EVisibility TraceChangelistVisibility() const
+	{
+		TSharedPtr<FTraceViewModel> TracePin = WeakTrace.Pin();
+		if (TracePin.IsValid())
+		{
+			return TracePin->Changelist != 0 ? EVisibility::Visible : EVisibility::Collapsed;
+		}
+		else
+		{
+			return EVisibility::Collapsed;
+		}
+	}
+
 	FText GetTraceBuildConfiguration() const
 	{
 		TSharedPtr<FTraceViewModel> TracePin = WeakTrace.Pin();
@@ -413,6 +465,12 @@ public:
 	}
 
 	TSharedPtr<IToolTip> GetTraceTooltip() const
+	{
+		return SNew(SLazyToolTip, SharedThis(this));
+	}
+
+	// ILazyToolTipCreator
+	virtual TSharedPtr<SToolTip> CreateTooltip() const override
 	{
 		TSharedPtr<FTraceViewModel> TracePin = WeakTrace.Pin();
 		if (TracePin.IsValid())
@@ -434,7 +492,7 @@ public:
 					[
 						SNew(STextBlock)
 						.Text(this, &STraceListRow::GetTraceName)
-						.Font(FCoreStyle::GetDefaultFontStyle("Bold", 9))
+						.Font(FCoreStyle::GetDefaultFontStyle("Bold", 10))
 						.ColorAndOpacity(FLinearColor::White)
 					]
 
@@ -468,6 +526,9 @@ public:
 			AddGridPanelRow(GridPanel, Row++, LOCTEXT("TraceTooltip_Platform", "Platform:"), &STraceListRow::GetTracePlatform);
 			AddGridPanelRow(GridPanel, Row++, LOCTEXT("TraceTooltip_AppName", "App Name:"), &STraceListRow::GetTraceAppName);
 			AddGridPanelRow(GridPanel, Row++, LOCTEXT("TraceTooltip_CommandLine", "Command Line:"), &STraceListRow::GetTraceCommandLine);
+			AddGridPanelRow(GridPanel, Row++, LOCTEXT("TraceTooltip_Branch", "Branch:"), &STraceListRow::GetTraceBranch);
+			AddGridPanelRow(GridPanel, Row++, LOCTEXT("TraceTooltip_BuildVersion", "Build Version:"), &STraceListRow::GetTraceBuildVersion);
+			AddGridPanelRow(GridPanel, Row++, LOCTEXT("TraceTooltip_Changelist", "Changelist:"), &STraceListRow::GetTraceChangelist, &STraceListRow::TraceChangelistVisibility);
 			AddGridPanelRow(GridPanel, Row++, LOCTEXT("TraceTooltip_BuildConfig", "Build Config:"), &STraceListRow::GetTraceBuildConfiguration);
 			AddGridPanelRow(GridPanel, Row++, LOCTEXT("TraceTooltip_BuildTarget", "Build Target:"), &STraceListRow::GetTraceBuildTarget);
 			AddGridPanelRow(GridPanel, Row++, LOCTEXT("TraceTooltip_Timestamp", "Timestamp:"), &STraceListRow::GetTraceTimestampForTooltip);
@@ -478,33 +539,54 @@ public:
 		}
 		else
 		{
-			return nullptr;
+			TSharedPtr<SToolTip> TraceTooltip =
+				SNew(SToolTip)
+				[
+					SNew(STextBlock)
+					.Text(LOCTEXT("TraceTooltip_NA", "N/A"))
+				];
+
+			return TraceTooltip;
 		}
 	}
 
-	void AddGridPanelRow(TSharedPtr<SGridPanel> Grid, int32 Row, const FText& Name,
-		typename TAttribute<FText>::FGetter::template TConstMethodPtr<STraceListRow> Value) const
+private:
+	void AddGridPanelRow(TSharedPtr<SGridPanel> Grid, int32 Row, const FText& InHeaderText,
+		typename TAttribute<FText>::FGetter::template TConstMethodPtr<STraceListRow> InValueTextFn,
+		typename TAttribute<EVisibility>::FGetter::template TConstMethodPtr<STraceListRow> InVisibilityFn = nullptr) const
 	{
-		Grid->AddSlot(0, Row)
+		SGridPanel::FSlot& Slot0 = Grid->AddSlot(0, Row)
 			.Padding(2.0f)
 			.HAlign(HAlign_Right)
 			[
 				SNew(STextBlock)
-				.Text(Name)
-				.Font(FCoreStyle::GetDefaultFontStyle("Bold", 8))
+				.Text(InHeaderText)
+				.Font(FCoreStyle::GetDefaultFontStyle("Bold", 10))
 				.ColorAndOpacity(FLinearColor::Gray)
 			];
 
-		Grid->AddSlot(1, Row)
+		SGridPanel::FSlot& Slot1 = Grid->AddSlot(1, Row)
 			.Padding(2.0f)
 			.HAlign(HAlign_Left)
 			[
 				SNew(STextBlock)
-				.Text(this, Value)
+				.Text(this, InValueTextFn)
 				.WrapTextAt(512.0f)
 				.WrappingPolicy(ETextWrappingPolicy::AllowPerCharacterWrapping)
 				.ColorAndOpacity(FLinearColor::White)
 			];
+
+		if (InVisibilityFn)
+		{
+			Slot0.GetWidget()->SetVisibility(MakeAttributeSP(this, InVisibilityFn));
+			Slot1.GetWidget()->SetVisibility(MakeAttributeSP(this, InVisibilityFn));
+		}
+		else
+		{
+			auto Fn = MakeAttributeSP(this, InValueTextFn);
+			Slot0.GetWidget()->SetVisibility(MakeAttributeLambda([this, Fn]() { return Fn.Get().IsEmpty() ? EVisibility::Collapsed : EVisibility::Visible; }));
+			Slot1.GetWidget()->SetVisibility(MakeAttributeLambda([this, Fn]() { return Fn.Get().IsEmpty() ? EVisibility::Collapsed : EVisibility::Visible; }));
+		}
 	}
 
 private:
@@ -1180,9 +1262,12 @@ FReply SStartPageWindow::Connect_OnClicked()
 		return FReply::Handled();
 	}
 
+	const FString& HostStr = HostText.ToString();
+	UE_LOG(TraceInsights, Log, TEXT("[StartPage] Try connecting to \"%s\"..."), *HostStr);
+
 	bool bConnectedSuccessfully = false;
 	UE::Trace::FControlClient ControlClient;
-	if (ControlClient.Connect(*HostText.ToString()))
+	if (ControlClient.Connect(*HostStr))
 	{
 		TSharedPtr<FInternetAddr> RecorderAddr;
 		if (ISocketSubsystem* Sockets = ISocketSubsystem::Get())
@@ -1191,7 +1276,9 @@ FReply SStartPageWindow::Connect_OnClicked()
 			RecorderAddr = Sockets->GetLocalHostAddr(*GLog, bCanBindAll);
 			if (RecorderAddr.IsValid())
 			{
-				ControlClient.SendSendTo(*RecorderAddr->ToString(false));
+				const FString RecorderAddrStr = RecorderAddr->ToString(false);
+				UE_LOG(TraceInsights, Log, TEXT("[StartPage] SendSendTo(\"%s\")..."), *RecorderAddrStr);
+				ControlClient.SendSendTo(*RecorderAddrStr);
 				bConnectedSuccessfully = true;
 			}
 		}
@@ -1199,6 +1286,7 @@ FReply SStartPageWindow::Connect_OnClicked()
 
 	if (bConnectedSuccessfully)
 	{
+		UE_LOG(TraceInsights, Log, TEXT("[StartPage] Successfully connected."));
 		FNotificationInfo NotificationInfo(FText::Format(LOCTEXT("ConnectSuccess", "Successfully connected to \"{0}\"!"), HostText));
 		NotificationInfo.bFireAndForget = false;
 		NotificationInfo.bUseLargeFont = false;
@@ -1211,6 +1299,7 @@ FReply SStartPageWindow::Connect_OnClicked()
 	}
 	else
 	{
+		UE_LOG(TraceInsights, Warning, TEXT("[StartPage] Failed to connect to \"%s\"!"), *HostStr);
 		FNotificationInfo NotificationInfo(FText::Format(LOCTEXT("ConnectFailed", "Failed to connect to \"{0}\"!"), HostText));
 		NotificationInfo.bFireAndForget = false;
 		NotificationInfo.bUseLargeFont = false;
@@ -1251,8 +1340,8 @@ void SStartPageWindow::RefreshTraceList()
 		{
 			TracesChangeSerial = NewChangeSerial;
 
-			UE_LOG(TraceInsights, Log, TEXT("[StartPage] Synching the trace list with StoreBrowser..."));
-	
+			//UE_LOG(TraceInsights, Log, TEXT("[StartPage] Synching the trace list with StoreBrowser..."));
+
 			const TArray<TSharedPtr<Insights::FStoreBrowserTraceInfo>>& InTraces = StoreBrowser->GetLockedTraces();
 			const TMap<uint32, TSharedPtr<Insights::FStoreBrowserTraceInfo>>& InTraceMap = StoreBrowser->GetLockedTraceMap();
 
@@ -1316,10 +1405,11 @@ void SStartPageWindow::RefreshTraceList()
 	}
 
 	StopwatchTotal.Stop();
-	if (UpdatedTraces > 0 || AddedTraces > 0 || RemovedTraces > 0)
+	const double Duration = StopwatchTotal.GetAccumulatedTime();
+	if ((Duration > 0.0001) && (UpdatedTraces > 0 || AddedTraces > 0 || RemovedTraces > 0))
 	{
-		UE_LOG(TraceInsights, Log, TEXT("[StartPage] The trace list refreshed in %llu ms (%d traces : %d updated, %d added, %d removed)."),
-			StopwatchTotal.GetAccumulatedTimeMs(), TraceViewModels.Num(), UpdatedTraces, AddedTraces, RemovedTraces);
+		UE_LOG(TraceInsights, Log, TEXT("[StartPage] The trace list refreshed in %.0f ms (%d traces : %d updated, %d added, %d removed)."),
+			Duration * 1000.0, TraceViewModels.Num(), UpdatedTraces, AddedTraces, RemovedTraces);
 	}
 }
 
@@ -1348,6 +1438,9 @@ void SStartPageWindow::UpdateTrace(FTraceViewModel& InOutTrace, const Insights::
 		InOutTrace.Platform = FText::FromString(InSourceTrace.Platform);
 		InOutTrace.AppName = FText::FromString(InSourceTrace.AppName);
 		InOutTrace.CommandLine = FText::FromString(InSourceTrace.CommandLine);
+		InOutTrace.Branch = FText::FromString(InSourceTrace.Branch);
+		InOutTrace.BuildVersion = FText::FromString(InSourceTrace.BuildVersion);
+		InOutTrace.Changelist = InSourceTrace.Changelist;
 		InOutTrace.ConfigurationType = InSourceTrace.ConfigurationType;
 		InOutTrace.TargetType = InSourceTrace.TargetType;
 	}
@@ -1652,7 +1745,7 @@ void SStartPageWindow::LoadTrace(uint32 InTraceId)
 
 	const uint32 StorePort = FInsightsManager::Get()->GetStoreClient()->GetStorePort();
 	FString CmdLine = FString::Printf(TEXT("-OpenTraceId=%d -StorePort=%d"), InTraceId, StorePort);
-		
+
 	FString ExtraCmdParams;
 	GetExtraCommandLineParams(ExtraCmdParams);
 	CmdLine += ExtraCmdParams;
diff --git a/Engine/Source/Developer/TraceInsights/Private/Insights/Widgets/SStartPageWindow.h b/Engine/Source/Developer/TraceInsights/Private/Insights/Widgets/SStartPageWindow.h
--- a/Engine/Source/Developer/TraceInsights/Private/Insights/Widgets/SStartPageWindow.h
+++ b/Engine/Source/Developer/TraceInsights/Private/Insights/Widgets/SStartPageWindow.h
@@ -66,6 +66,9 @@ struct FTraceViewModel
 	FText Platform;
 	FText AppName;
 	FText CommandLine;
+	FText Branch;
+	FText BuildVersion;
+	uint32 Changelist;
 	EBuildConfiguration ConfigurationType = EBuildConfiguration::Unknown;
 	EBuildTargetType TargetType = EBuildTargetType::Unknown;
 
diff --git a/Engine/Source/Developer/TraceServices/Private/Analyzers/DiagnosticsAnalysis.cpp b/Engine/Source/Developer/TraceServices/Private/Analyzers/DiagnosticsAnalysis.cpp
--- a/Engine/Source/Developer/TraceServices/Private/Analyzers/DiagnosticsAnalysis.cpp
+++ b/Engine/Source/Developer/TraceServices/Private/Analyzers/DiagnosticsAnalysis.cpp
@@ -73,7 +73,9 @@ bool FDiagnosticsAnalyzer::OnEvent(uint16 RouteId, EStyle Style, const FOnEventC
 		EventData.GetString("Platform", SessionInfo.Platform);
 		EventData.GetString("AppName", SessionInfo.AppName);
 		EventData.GetString("CommandLine", SessionInfo.CommandLine);
-
+		EventData.GetString("Branch", SessionInfo.Branch);
+		EventData.GetString("BuildVersion", SessionInfo.BuildVersion);
+		SessionInfo.Changelist = EventData.GetValue<uint32>("Changelist", 0);
 		SessionInfo.ConfigurationType = (EBuildConfiguration) EventData.GetValue<uint8>("ConfigurationType");
 		SessionInfo.TargetType = (EBuildTargetType) EventData.GetValue<uint8>("TargetType");
 
diff --git a/Engine/Source/Developer/TraceServices/Public/TraceServices/Model/Diagnostics.h b/Engine/Source/Developer/TraceServices/Public/TraceServices/Model/Diagnostics.h
--- a/Engine/Source/Developer/TraceServices/Public/TraceServices/Model/Diagnostics.h
+++ b/Engine/Source/Developer/TraceServices/Public/TraceServices/Model/Diagnostics.h
@@ -13,8 +13,11 @@ struct FSessionInfo
 	FString Platform;
 	FString AppName;
 	FString CommandLine;
-	EBuildConfiguration ConfigurationType;
-	EBuildTargetType TargetType;
+	FString Branch;
+	FString BuildVersion;
+	uint32 Changelist = 0;
+	EBuildConfiguration ConfigurationType = EBuildConfiguration::Unknown;
+	EBuildTargetType TargetType = EBuildTargetType::Unknown;
 };
 
 class IDiagnosticsProvider : public IProvider
