[Insights] Add extra check to prevent TimingView crash when TaskProvider is not registered

diff --git a/Engine/Source/Developer/TraceInsights/Private/Insights/ViewModels/ThreadTimingTrack.cpp b/Engine/Source/Developer/TraceInsights/Private/Insights/ViewModels/ThreadTimingTrack.cpp
--- a/Engine/Source/Developer/TraceInsights/Private/Insights/ViewModels/ThreadTimingTrack.cpp
+++ b/Engine/Source/Developer/TraceInsights/Private/Insights/ViewModels/ThreadTimingTrack.cpp
@@ -1031,120 +1031,122 @@ void FThreadTimingTrack::InitTooltip(FTooltipDrawState& InOutTooltip, const ITim
 		// tasks
 
 		const TraceServices::ITasksProvider* TasksProvider = TraceServices::ReadTasksProvider(*Session.Get());
-
-		auto AddTaskInfo = [&InOutTooltip, this](const TraceServices::FTaskInfo& Task)
+		if (TasksProvider != nullptr)
 		{
-			InOutTooltip.AddTextLine(FString::Printf(TEXT("-------- Task %d%s --------"), Task.Id, Task.bTracked ? TEXT("") : TEXT(" (not tracked)")), FLinearColor::Green);
-
-			if (Task.DebugName != nullptr)
+			auto AddTaskInfo = [&InOutTooltip, this](const TraceServices::FTaskInfo& Task)
 			{
-				InOutTooltip.AddTextLine(FString::Printf(TEXT("%s"), Task.DebugName), FLinearColor::Green);
-			}
+				InOutTooltip.AddTextLine(FString::Printf(TEXT("-------- Task %d%s --------"), Task.Id, Task.bTracked ? TEXT("") : TEXT(" (not tracked)")), FLinearColor::Green);
 
-			ENamedThreads::Type ThreadInfo = (ENamedThreads::Type)Task.ThreadToExecuteOn;
-			const TCHAR* NamedThreadsStr[] = {TEXT("Stats"), TEXT("RHI"), TEXT("Audio"), TEXT("Game"), TEXT("Rendering") };
-			ENamedThreads::Type ThreadIndex = ENamedThreads::GetThreadIndex(ThreadInfo);
+				if (Task.DebugName != nullptr)
+				{
+					InOutTooltip.AddTextLine(FString::Printf(TEXT("%s"), Task.DebugName), FLinearColor::Green);
+				}
 
-			if (ThreadIndex == ENamedThreads::AnyThread)
-			{
-				const TCHAR* TaskPri = ENamedThreads::GetTaskPriority(ThreadInfo) == ENamedThreads::NormalTaskPriority ? TEXT("Normal") : TEXT("High");
+				ENamedThreads::Type ThreadInfo = (ENamedThreads::Type)Task.ThreadToExecuteOn;
+				const TCHAR* NamedThreadsStr[] = { TEXT("Stats"), TEXT("RHI"), TEXT("Audio"), TEXT("Game"), TEXT("Rendering") };
+				ENamedThreads::Type ThreadIndex = ENamedThreads::GetThreadIndex(ThreadInfo);
 
-				int32 ThreadPriIndex = ENamedThreads::GetThreadPriorityIndex(ThreadInfo);
-				const TCHAR* ThreadPriStrs[] = { TEXT("Normal"), TEXT("High"), TEXT("Low") };
-				const TCHAR* ThreadPri = ThreadPriStrs[ThreadPriIndex];
+				if (ThreadIndex == ENamedThreads::AnyThread)
+				{
+					const TCHAR* TaskPri = ENamedThreads::GetTaskPriority(ThreadInfo) == ENamedThreads::NormalTaskPriority ? TEXT("Normal") : TEXT("High");
 
-				InOutTooltip.AddTextLine(FString::Printf(TEXT("%s Pri task on %s Pri worker"), TaskPri, ThreadPri), FLinearColor::Green);
-			}
-			else
-			{
-				const TCHAR* QueueStr = ENamedThreads::GetQueueIndex(ThreadInfo) == ENamedThreads::MainQueue ? TEXT("Main") : TEXT("Local");
-				InOutTooltip.AddTextLine(FString::Printf(TEXT("%s (%s queue)"), NamedThreadsStr[ThreadIndex], QueueStr), FLinearColor::Green);
-			}
+					int32 ThreadPriIndex = ENamedThreads::GetThreadPriorityIndex(ThreadInfo);
+					const TCHAR* ThreadPriStrs[] = { TEXT("Normal"), TEXT("High"), TEXT("Low") };
+					const TCHAR* ThreadPri = ThreadPriStrs[ThreadPriIndex];
 
-			{
-				TSharedPtr<FCpuTimingTrack> Track = SharedState.GetCpuTrack(Task.CreatedThreadId);
-				FString TrackName = Track.IsValid() ? Track->GetName() : TEXT("Unknown");
-				InOutTooltip.AddNameValueTextLine(TEXT("Created:"), FString::Printf(TEXT("%f on %s"), Task.CreatedTimestamp, *TrackName));
-			}
-			
-			{
-				TSharedPtr<FCpuTimingTrack> Track = SharedState.GetCpuTrack(Task.LaunchedThreadId);
-				FString TrackName = Track.IsValid() ? Track->GetName() : TEXT("Unknown");
-				InOutTooltip.AddNameValueTextLine(TEXT("Launched:"), FString::Printf(TEXT("%f (+%s) on %s"), Task.LaunchedTimestamp, *TimeUtils::FormatTimeAuto(Task.LaunchedTimestamp - Task.CreatedTimestamp), *TrackName));
-			}
-			{
-				TSharedPtr<FCpuTimingTrack> Track = SharedState.GetCpuTrack(Task.ScheduledThreadId);
-				FString TrackName = Track.IsValid() ? Track->GetName() : TEXT("Unknown");
-				InOutTooltip.AddNameValueTextLine(TEXT("Scheduled:"), FString::Printf(TEXT("%f (+%s) on %s"), Task.ScheduledTimestamp, *TimeUtils::FormatTimeAuto(Task.ScheduledTimestamp - Task.LaunchedTimestamp), *TrackName));
-			}
-			InOutTooltip.AddNameValueTextLine(TEXT("Started:"), FString::Printf(TEXT("%f (+%s)"), Task.StartedTimestamp, *TimeUtils::FormatTimeAuto(Task.StartedTimestamp - Task.ScheduledTimestamp)));
-			if (Task.FinishedTimestamp != TraceServices::FTaskInfo::InvalidTimestamp)
-			{
-				InOutTooltip.AddNameValueTextLine(TEXT("Finished:"), FString::Printf(TEXT("%f (+%s)"), Task.FinishedTimestamp, *TimeUtils::FormatTimeAuto(Task.FinishedTimestamp - Task.StartedTimestamp)));
+					InOutTooltip.AddTextLine(FString::Printf(TEXT("%s Pri task on %s Pri worker"), TaskPri, ThreadPri), FLinearColor::Green);
+				}
+				else
+				{
+					const TCHAR* QueueStr = ENamedThreads::GetQueueIndex(ThreadInfo) == ENamedThreads::MainQueue ? TEXT("Main") : TEXT("Local");
+					InOutTooltip.AddTextLine(FString::Printf(TEXT("%s (%s queue)"), NamedThreadsStr[ThreadIndex], QueueStr), FLinearColor::Green);
+				}
 
-				if (Task.CompletedTimestamp != TraceServices::FTaskInfo::InvalidTimestamp)
 				{
-					TSharedPtr<FCpuTimingTrack> Track = SharedState.GetCpuTrack(Task.CompletedThreadId);
+					TSharedPtr<FCpuTimingTrack> Track = SharedState.GetCpuTrack(Task.CreatedThreadId);
 					FString TrackName = Track.IsValid() ? Track->GetName() : TEXT("Unknown");
-					InOutTooltip.AddNameValueTextLine(TEXT("Completed:"), FString::Printf(TEXT("%f (+%s) on %s"), Task.FinishedTimestamp, *TimeUtils::FormatTimeAuto(Task.CompletedTimestamp - Task.FinishedTimestamp), *TrackName));
+					InOutTooltip.AddNameValueTextLine(TEXT("Created:"), FString::Printf(TEXT("%f on %s"), Task.CreatedTimestamp, *TrackName));
 				}
-			}
-			InOutTooltip.AddNameValueTextLine(TEXT("Prerequisite tasks:"), FString::Printf(TEXT("%d"), Task.Prerequisites.Num()));
-			InOutTooltip.AddNameValueTextLine(TEXT("Subsequent tasks:"), FString::Printf(TEXT("%d"), Task.Subsequents.Num()));
-			InOutTooltip.AddNameValueTextLine(TEXT("Nested tasks:"), FString::Printf(TEXT("%d"), Task.NestedTasks.Num()));
-		};
 
-		do // info about a task
-		{
-			const TraceServices::FTaskInfo* Task = TasksProvider->TryGetTask(ThreadId, TooltipEvent.GetStartTime());
-			if (Task == nullptr)
-			{
-				break;
-			}
+				{
+					TSharedPtr<FCpuTimingTrack> Track = SharedState.GetCpuTrack(Task.LaunchedThreadId);
+					FString TrackName = Track.IsValid() ? Track->GetName() : TEXT("Unknown");
+					InOutTooltip.AddNameValueTextLine(TEXT("Launched:"), FString::Printf(TEXT("%f (+%s) on %s"), Task.LaunchedTimestamp, *TimeUtils::FormatTimeAuto(Task.LaunchedTimestamp - Task.CreatedTimestamp), *TrackName));
+				}
+				{
+					TSharedPtr<FCpuTimingTrack> Track = SharedState.GetCpuTrack(Task.ScheduledThreadId);
+					FString TrackName = Track.IsValid() ? Track->GetName() : TEXT("Unknown");
+					InOutTooltip.AddNameValueTextLine(TEXT("Scheduled:"), FString::Printf(TEXT("%f (+%s) on %s"), Task.ScheduledTimestamp, *TimeUtils::FormatTimeAuto(Task.ScheduledTimestamp - Task.LaunchedTimestamp), *TrackName));
+				}
+				InOutTooltip.AddNameValueTextLine(TEXT("Started:"), FString::Printf(TEXT("%f (+%s)"), Task.StartedTimestamp, *TimeUtils::FormatTimeAuto(Task.StartedTimestamp - Task.ScheduledTimestamp)));
+				if (Task.FinishedTimestamp != TraceServices::FTaskInfo::InvalidTimestamp)
+				{
+					InOutTooltip.AddNameValueTextLine(TEXT("Finished:"), FString::Printf(TEXT("%f (+%s)"), Task.FinishedTimestamp, *TimeUtils::FormatTimeAuto(Task.FinishedTimestamp - Task.StartedTimestamp)));
+
+					if (Task.CompletedTimestamp != TraceServices::FTaskInfo::InvalidTimestamp)
+					{
+						TSharedPtr<FCpuTimingTrack> Track = SharedState.GetCpuTrack(Task.CompletedThreadId);
+						FString TrackName = Track.IsValid() ? Track->GetName() : TEXT("Unknown");
+						InOutTooltip.AddNameValueTextLine(TEXT("Completed:"), FString::Printf(TEXT("%f (+%s) on %s"), Task.FinishedTimestamp, *TimeUtils::FormatTimeAuto(Task.CompletedTimestamp - Task.FinishedTimestamp), *TrackName));
+					}
+				}
+				InOutTooltip.AddNameValueTextLine(TEXT("Prerequisite tasks:"), FString::Printf(TEXT("%d"), Task.Prerequisites.Num()));
+				InOutTooltip.AddNameValueTextLine(TEXT("Subsequent tasks:"), FString::Printf(TEXT("%d"), Task.Subsequents.Num()));
+				InOutTooltip.AddNameValueTextLine(TEXT("Nested tasks:"), FString::Printf(TEXT("%d"), Task.NestedTasks.Num()));
+			};
 
-			if (Task->FinishedTimestamp < TooltipEvent.GetEndTime())
+			do // info about a task
 			{
-				break;
-			}
+				const TraceServices::FTaskInfo* Task = TasksProvider->TryGetTask(ThreadId, TooltipEvent.GetStartTime());
+				if (Task == nullptr)
+				{
+					break;
+				}
+
+				if (Task->FinishedTimestamp < TooltipEvent.GetEndTime())
+				{
+					break;
+				}
 
-			AddTaskInfo(*Task);
-		} while (false);
+				AddTaskInfo(*Task);
+			} while (false);
 
-		do // info about blocking
-		{
-			const TraceServices::FWaitingForTasks* Waiting = TasksProvider->TryGetWaiting(*TimerName, ThreadId, TooltipEvent.GetStartTime());
-			if (Waiting == nullptr)
+			do // info about blocking
 			{
-				break;
-			}
+				const TraceServices::FWaitingForTasks* Waiting = TasksProvider->TryGetWaiting(*TimerName, ThreadId, TooltipEvent.GetStartTime());
+				if (Waiting == nullptr)
+				{
+					break;
+				}
 
-			InOutTooltip.AddTextLine(TEXT("-------- Wating for tasks --------"), FLinearColor::Red);
-			const int32 MaxWaitedTasksToList = 5;
-			FString TaskIdsStr = Waiting->Tasks.Num() <= MaxWaitedTasksToList ?
-				FString::JoinBy(Waiting->Tasks, TEXT(", "), [](TaskTrace::FId TaskId) { return FString::FromInt(TaskId); }) :
-				FString::Printf(TEXT("[%d]"), Waiting->Tasks.Num());
-			InOutTooltip.AddNameValueTextLine(TEXT("Tasks:"), FString::Printf(TEXT("%s"), *TaskIdsStr));
-			InOutTooltip.AddNameValueTextLine(TEXT("Started waiting:"), FString::Printf(TEXT("%f"), Waiting->StartedTimestamp));
-			InOutTooltip.AddNameValueTextLine(TEXT("Finished waiting:"), 
-				FString::Printf(TEXT("%s (+%s)"), 
-					Waiting->FinishedTimestamp == TraceServices::FTaskInfo::InvalidTimestamp ? TEXT("[not set]") : *FString::SanitizeFloat(Waiting->FinishedTimestamp),
-					*TimeUtils::FormatTimeAuto(Waiting->FinishedTimestamp - Waiting->StartedTimestamp)));
-
-			int32 NumTasksToList = FMath::Min(Waiting->Tasks.Num(), MaxWaitedTasksToList);
-			for (int32 i = 0; i != NumTasksToList; ++i)
-			{
-				const TraceServices::FTaskInfo* Task = TasksProvider->TryGetTask(Waiting->Tasks[i]);
-				if (Task != nullptr)
+				InOutTooltip.AddTextLine(TEXT("-------- Wating for tasks --------"), FLinearColor::Red);
+				const int32 MaxWaitedTasksToList = 5;
+				FString TaskIdsStr = Waiting->Tasks.Num() <= MaxWaitedTasksToList ?
+					FString::JoinBy(Waiting->Tasks, TEXT(", "), [](TaskTrace::FId TaskId) { return FString::FromInt(TaskId); }) :
+					FString::Printf(TEXT("[%d]"), Waiting->Tasks.Num());
+				InOutTooltip.AddNameValueTextLine(TEXT("Tasks:"), FString::Printf(TEXT("%s"), *TaskIdsStr));
+				InOutTooltip.AddNameValueTextLine(TEXT("Started waiting:"), FString::Printf(TEXT("%f"), Waiting->StartedTimestamp));
+				InOutTooltip.AddNameValueTextLine(TEXT("Finished waiting:"),
+					FString::Printf(TEXT("%s (+%s)"),
+						Waiting->FinishedTimestamp == TraceServices::FTaskInfo::InvalidTimestamp ? TEXT("[not set]") : *FString::SanitizeFloat(Waiting->FinishedTimestamp),
+						*TimeUtils::FormatTimeAuto(Waiting->FinishedTimestamp - Waiting->StartedTimestamp)));
+
+				int32 NumTasksToList = FMath::Min(Waiting->Tasks.Num(), MaxWaitedTasksToList);
+				for (int32 i = 0; i != NumTasksToList; ++i)
+				{
+					const TraceServices::FTaskInfo* Task = TasksProvider->TryGetTask(Waiting->Tasks[i]);
+					if (Task != nullptr)
+					{
+						AddTaskInfo(*Task);
+					}
+				}
+				if (NumTasksToList < Waiting->Tasks.Num())
 				{
-					AddTaskInfo(*Task);
+					InOutTooltip.AddTextLine(TEXT("[...]"), FLinearColor::Green);
 				}
-			}
-			if (NumTasksToList < Waiting->Tasks.Num())
-			{
-				InOutTooltip.AddTextLine(TEXT("[...]"), FLinearColor::Green);
-			}
 
-		} while (false);
+			} while (false);
+		}
 	}
 
 	InOutTooltip.UpdateLayout();
