[Insights] Task Track initial work

diff --git a/Engine/Source/Developer/TraceInsights/Private/Insights/TaskGraphProfiler/TaskGraphProfilerManager.cpp b/Engine/Source/Developer/TraceInsights/Private/Insights/TaskGraphProfiler/TaskGraphProfilerManager.cpp
--- a/Engine/Source/Developer/TraceInsights/Private/Insights/TaskGraphProfiler/TaskGraphProfilerManager.cpp
+++ b/Engine/Source/Developer/TraceInsights/Private/Insights/TaskGraphProfiler/TaskGraphProfilerManager.cpp
@@ -2,6 +2,7 @@
 
 #include "TaskGraphProfilerManager.h"
 
+#include "Features/IModularFeatures.h"
 #include "Modules/ModuleManager.h"
 #include "TraceServices/AnalysisService.h"
 #include "TraceServices/Model/TasksProfiler.h"
@@ -9,8 +10,10 @@
 
 #include "Insights/InsightsStyle.h"
 #include "Insights/TaskGraphProfiler/ViewModels/TaskTable.h"
+#include "Insights/TaskGraphProfiler/ViewModels/TaskTimingTrack.h"
 #include "Insights/TaskGraphProfiler/Widgets/STaskTableTreeView.h"
 #include "Insights/TimingProfilerManager.h"
+#include "Insights/ViewModels/TaskGraphRelation.h"
 #include "Insights/Widgets/STimingProfilerWindow.h"
 
 ////////////////////////////////////////////////////////////////////////////////////////////////////
@@ -67,6 +70,8 @@ void FTaskGraphProfilerManager::Initialize(IUnrealInsightsModule& InsightsModule
 	}
 	bIsInitialized = true;
 
+	InitializeColorCode();
+
 	// Register tick functions.
 	OnTick = FTickerDelegate::CreateSP(this, &FTaskGraphProfilerManager::Tick);
 	OnTickHandle = FTicker::GetCoreTicker().AddTicker(OnTick, 0.0f);
@@ -104,6 +109,11 @@ void FTaskGraphProfilerManager::Shutdown()
 FTaskGraphProfilerManager::~FTaskGraphProfilerManager()
 {
 	ensure(!bIsInitialized);
+
+	if (TaskTimingSharedState.IsValid())
+	{
+		IModularFeatures::Get().UnregisterModularFeature(Insights::TimingViewExtenderFeatureName, TaskTimingSharedState.Get());
+	}
 }
 
 ////////////////////////////////////////////////////////////////////////////////////////////////////
@@ -136,7 +146,25 @@ bool FTaskGraphProfilerManager::Tick(float DeltaTime)
 			TSharedPtr<FTabManager> TabManagerShared = TimingTabManager.Pin();
 			if (TasksProvider && TasksProvider->GetNumTasks() > 0 && TabManagerShared.IsValid())
 			{
+				TSharedPtr<STimingProfilerWindow> Window = FTimingProfilerManager::Get()->GetProfilerWindow();
+				if (!Window.IsValid())
+				{
+					return true;
+				}
+
+				TSharedPtr<STimingView> TimingView = Window->GetTimingView();
+				if (!TimingView.IsValid())
+				{
+					return true;
+				}
+
 				bIsAvailable = true;
+
+				if (!TaskTimingSharedState.IsValid())
+				{
+					TaskTimingSharedState = MakeShared<FTaskTimingSharedState>(TimingView.Get());
+					IModularFeatures::Get().RegisterModularFeature(Insights::TimingViewExtenderFeatureName, TaskTimingSharedState.Get());
+				}
 				TabManagerShared->TryInvokeTab(FTaskGraphProfilerTabs::TaskTableTreeViewTabID);
 			}
 
@@ -236,22 +264,22 @@ void FTaskGraphProfilerManager::GetTaskRelations(const TraceServices::FTaskInfo*
 
 	if (Task->CreatedTimestamp != Task->LaunchedTimestamp || Task->CreatedThreadId != Task->LaunchedThreadId)
 	{
-		Callback(Task->CreatedTimestamp, Task->CreatedThreadId, Task->LaunchedTimestamp, Task->LaunchedThreadId, FTaskGraphRelation::ETaskGraphRelationType::Created);
+		Callback(Task->CreatedTimestamp, Task->CreatedThreadId, Task->LaunchedTimestamp, Task->LaunchedThreadId, ETaskEventType::Created);
 	}
 
-	Callback(Task->LaunchedTimestamp, Task->LaunchedThreadId, Task->ScheduledTimestamp, Task->ScheduledThreadId, FTaskGraphRelation::ETaskGraphRelationType::Launched);
+	Callback(Task->LaunchedTimestamp, Task->LaunchedThreadId, Task->ScheduledTimestamp, Task->ScheduledThreadId, ETaskEventType::Launched);
 
 	int32 NumPrerequisitesToShow = FMath::Min(Task->Prerequisites.Num(), MaxTasksToShow);
 	for (int32 i = 0; i != NumPrerequisitesToShow; ++i)
 	{
 		const TraceServices::FTaskInfo* Prerequisite = TasksProvider->TryGetTask(Task->Prerequisites[i].RelativeId);
 		check(Prerequisite != nullptr);
-		Callback(Prerequisite->CompletedTimestamp, Prerequisite->CompletedThreadId, Task->ScheduledTimestamp, Task->ScheduledThreadId, FTaskGraphRelation::ETaskGraphRelationType::Prerequisite);
+		Callback(Prerequisite->CompletedTimestamp, Prerequisite->CompletedThreadId, Task->ScheduledTimestamp, Task->ScheduledThreadId, ETaskEventType::Prerequisite);
 	}
 
 	if (Task->LaunchedTimestamp != Task->ScheduledTimestamp || Task->LaunchedThreadId != Task->ScheduledThreadId)
 	{
-		Callback(Task->ScheduledTimestamp, Task->ScheduledThreadId, Task->StartedTimestamp, Task->StartedThreadId, FTaskGraphRelation::ETaskGraphRelationType::Scheduled);
+		Callback(Task->ScheduledTimestamp, Task->ScheduledThreadId, Task->StartedTimestamp, Task->StartedThreadId, ETaskEventType::Scheduled);
 	}
 
 	int32 NumNestedToShow = FMath::Min(Task->NestedTasks.Num(), MaxTasksToShow);
@@ -261,9 +289,9 @@ void FTaskGraphProfilerManager::GetTaskRelations(const TraceServices::FTaskInfo*
 		const TraceServices::FTaskInfo* NestedTask = TasksProvider->TryGetTask(RelationInfo.RelativeId);
 		check(NestedTask != nullptr);
 
-		Callback(RelationInfo.Timestamp, Task->StartedThreadId, NestedTask->StartedTimestamp, NestedTask->StartedThreadId, FTaskGraphRelation::ETaskGraphRelationType::AddedNested);
+		Callback(RelationInfo.Timestamp, Task->StartedThreadId, NestedTask->StartedTimestamp, NestedTask->StartedThreadId, ETaskEventType::AddedNested);
 
-		Callback(NestedTask->CompletedTimestamp, NestedTask->CompletedThreadId, NestedTask->CompletedTimestamp, Task->StartedThreadId, FTaskGraphRelation::ETaskGraphRelationType::NestedCompleted);
+		Callback(NestedTask->CompletedTimestamp, NestedTask->CompletedThreadId, NestedTask->CompletedTimestamp, Task->StartedThreadId, ETaskEventType::NestedCompleted);
 	}
 
 	int32 NumSubsequentsToShow = FMath::Min(Task->NestedTasks.Num(), MaxTasksToShow);
@@ -271,12 +299,12 @@ void FTaskGraphProfilerManager::GetTaskRelations(const TraceServices::FTaskInfo*
 	{
 		const TraceServices::FTaskInfo* Subsequent = TasksProvider->TryGetTask(Task->Subsequents[i].RelativeId);
 		check(Subsequent != nullptr);
-		Callback(Task->CompletedTimestamp, Task->CompletedThreadId, Subsequent->ScheduledTimestamp, Subsequent->ScheduledThreadId, FTaskGraphRelation::ETaskGraphRelationType::Subsequent);
+		Callback(Task->CompletedTimestamp, Task->CompletedThreadId, Subsequent->ScheduledTimestamp, Subsequent->ScheduledThreadId, ETaskEventType::Subsequent);
 	}
 
 	if (Task->FinishedTimestamp != Task->CompletedTimestamp || Task->CompletedThreadId != Task->StartedThreadId)
 	{
-		Callback(Task->FinishedTimestamp, Task->StartedThreadId, Task->CompletedTimestamp, Task->StartedThreadId, FTaskGraphRelation::ETaskGraphRelationType::Completed);
+		Callback(Task->FinishedTimestamp, Task->StartedThreadId, Task->CompletedTimestamp, Task->StartedThreadId, ETaskEventType::Completed);
 	}
 }
 
@@ -352,6 +380,29 @@ void FTaskGraphProfilerManager::OnWindowClosedEvent()
 
 ////////////////////////////////////////////////////////////////////////////////////////////////////
 
+void FTaskGraphProfilerManager::InitializeColorCode()
+{
+	ColorCode[static_cast<uint32>(ETaskEventType::Created)] = FLinearColor::Yellow;
+	ColorCode[static_cast<uint32>(ETaskEventType::Launched)] = FLinearColor::Green;
+	ColorCode[static_cast<uint32>(ETaskEventType::Prerequisite)] = FLinearColor::Red;
+	ColorCode[static_cast<uint32>(ETaskEventType::Scheduled)] = FLinearColor::Blue;
+	ColorCode[static_cast<uint32>(ETaskEventType::Started)] = FLinearColor::Red;
+	ColorCode[static_cast<uint32>(ETaskEventType::AddedNested)] = FLinearColor::Blue;
+	ColorCode[static_cast<uint32>(ETaskEventType::NestedCompleted)] = FLinearColor::Red;
+	ColorCode[static_cast<uint32>(ETaskEventType::Subsequent)] = FLinearColor::Red;
+	ColorCode[static_cast<uint32>(ETaskEventType::Completed)] = FLinearColor::Yellow;
+}
+
+////////////////////////////////////////////////////////////////////////////////////////////////////
+
+FLinearColor FTaskGraphProfilerManager::GetColorForTaskEvent(ETaskEventType InEvent)
+{
+	check(InEvent < ETaskEventType::NumTaskEventTypes);
+	return ColorCode[static_cast<uint32>(InEvent)];
+}
+
+////////////////////////////////////////////////////////////////////////////////////////////////////
+
 } // namespace Insights
 
 #undef LOCTEXT_NAMESPACE
diff --git a/Engine/Source/Developer/TraceInsights/Private/Insights/TaskGraphProfiler/TaskGraphProfilerManager.h b/Engine/Source/Developer/TraceInsights/Private/Insights/TaskGraphProfiler/TaskGraphProfilerManager.h
--- a/Engine/Source/Developer/TraceInsights/Private/Insights/TaskGraphProfiler/TaskGraphProfilerManager.h
+++ b/Engine/Source/Developer/TraceInsights/Private/Insights/TaskGraphProfiler/TaskGraphProfilerManager.h
@@ -9,7 +9,6 @@
 // Insights
 #include "Insights/InsightsManager.h"
 #include "Insights/IUnrealInsightsModule.h"
-#include "Insights/ViewModels/TaskGraphRelation.h"
 
 namespace TraceServices
 {
@@ -21,6 +20,7 @@ namespace Insights
 {
 
 class STaskTableTreeView;
+class FTaskTimingSharedState;
 
 struct FTaskGraphProfilerTabs
 {
@@ -28,6 +28,23 @@ struct FTaskGraphProfilerTabs
 	static const FName TaskTableTreeViewTabID;
 };
 
+////////////////////////////////////////////////////////////////////////////////////////////////////
+
+enum class ETaskEventType : uint32
+{
+	Created,
+	Launched,
+	Prerequisite,
+	Scheduled,
+	Started,
+	AddedNested,
+	NestedCompleted,
+	Subsequent,
+	Completed,
+	
+	NumTaskEventTypes,
+};
+
 ////////////////////////////////////////////////////////////////////////////////////////////////////
 /**
  * This class manages the Task Graph Profiler state and settings.
@@ -35,7 +52,7 @@ struct FTaskGraphProfilerTabs
 class FTaskGraphProfilerManager : public TSharedFromThis<FTaskGraphProfilerManager>, public IInsightsComponent
 {
 public:
-	typedef TFunction<void(double /*SourceTimestamp*/, uint32 /*SourceThreadId*/, double /*TargetTimestamp*/, uint32 /*TargetThreadId*/, FTaskGraphRelation::ETaskGraphRelationType /*Type*/)> AddRelationCallback;
+	typedef TFunction<void(double /*SourceTimestamp*/, uint32 /*SourceThreadId*/, double /*TargetTimestamp*/, uint32 /*TargetThreadId*/, ETaskEventType /*Type*/)> AddRelationCallback;
 
 	/** Creates the Memory Profiler manager, only one instance can exist. */
 	FTaskGraphProfilerManager(TSharedRef<FUICommandList> InCommandList);
@@ -75,6 +92,9 @@ public:
 
 	void GetTaskRelations(double Time, uint32 ThreadId, AddRelationCallback Callback);
 	void GetTaskRelations(uint32 TaskId, AddRelationCallback Callback);
+	FLinearColor GetColorForTaskEvent(ETaskEventType InEvent);
+
+	TSharedPtr<Insights::FTaskTimingSharedState> GetTaskTimingSharedState() { return TaskTimingSharedState;	}
 
 private:
 	/** Updates this manager, done through FCoreTicker. */
@@ -84,6 +104,8 @@ private:
 
 	void GetTaskRelations(const TraceServices::FTaskInfo* Task, const TraceServices::ITasksProvider* TasksProvider, AddRelationCallback Callback);
 
+	void InitializeColorCode();
+
 private:
 	bool bIsInitialized;
 	bool bIsAvailable;
@@ -98,9 +120,13 @@ private:
 	/** A shared pointer to the global instance of the Task Graph Profiler manager. */
 	static TSharedPtr<FTaskGraphProfilerManager> Instance;
 
+	/** Shared state for task tracks */
+	TSharedPtr<FTaskTimingSharedState> TaskTimingSharedState;
+
 	TWeakPtr<FTabManager> TimingTabManager;
 
 	TSharedPtr<Insights::STaskTableTreeView> TaskTableTreeView;
+	FLinearColor ColorCode[static_cast<uint32>(ETaskEventType::NumTaskEventTypes)];
 };
 
 } // namespace Insights
diff --git a/Engine/Source/Developer/TraceInsights/Private/Insights/TaskGraphProfiler/ViewModels/TaskTimingTrack.cpp b/Engine/Source/Developer/TraceInsights/Private/Insights/TaskGraphProfiler/ViewModels/TaskTimingTrack.cpp
new file mode 100644
--- /dev/null
+++ b/Engine/Source/Developer/TraceInsights/Private/Insights/TaskGraphProfiler/ViewModels/TaskTimingTrack.cpp
@@ -0,0 +1,276 @@
+// Copyright Epic Games, Inc. All Rights Reserved.
+
+#include "TaskTimingTrack.h"
+
+#include "TraceServices/Model/TasksProfiler.h"
+
+// Insights
+#include "Insights/ITimingViewSession.h"
+#include "Insights/TaskGraphProfiler/TaskGraphProfilerManager.h"
+#include "Insights/ViewModels/TaskGraphRelation.h"
+#include "Insights/ViewModels/ThreadTimingTrack.h"
+#include "Insights/ViewModels/ThreadTrackEvent.h"
+#include "Insights/ViewModels/TimingTrackViewport.h"
+#include "Insights/ViewModels/TimingViewDrawHelper.h"
+#include "Insights/ViewModels/TooltipDrawState.h"
+#include "Insights/Widgets/STimingView.h"
+
+#define LOCTEXT_NAMESPACE "TaskTimingTrack"
+
+namespace Insights
+{
+
+////////////////////////////////////////////////////////////////////////////////////////////////////
+// FTaskTimingSharedState
+////////////////////////////////////////////////////////////////////////////////////////////////////
+
+void FTaskTimingSharedState::OnBeginSession(Insights::ITimingViewSession& InSession)
+{
+	if (&InSession != TimingView)
+	{
+		return;
+	}
+
+	TaskTrack = nullptr;
+}
+
+////////////////////////////////////////////////////////////////////////////////////////////////////
+
+void FTaskTimingSharedState::OnEndSession(Insights::ITimingViewSession& InSession)
+{
+	if (&InSession != TimingView)
+	{
+		return;
+	}
+
+	TaskTrack = nullptr;
+}
+
+////////////////////////////////////////////////////////////////////////////////////////////////////
+
+void FTaskTimingSharedState::Tick(Insights::ITimingViewSession& InSession, const TraceServices::IAnalysisSession& InAnalysisSession)
+{
+	if (&InSession != TimingView)
+	{
+		return;
+	}
+
+	if (!TaskTrack.IsValid() && FTaskGraphProfilerManager::Get()->GetIsAvailable())
+	{
+		TaskTrack = MakeShared<FTaskTimingTrack>(*this, TEXT("Task Overview Track"), 0);
+		TaskTrack->SetVisibilityFlag(true);
+		TaskTrack->SetOrder(FTimingTrackOrder::Task);
+
+		TimingView->OnSelectedEventChanged().AddSP(TaskTrack.Get(), &FTaskTimingTrack::OnTimingEventSelected);
+
+		InSession.AddTopDockedTrack(TaskTrack);
+	}
+}
+
+////////////////////////////////////////////////////////////////////////////////////////////////////
+
+void FTaskTimingSharedState::ExtendFilterMenu(Insights::ITimingViewSession& InSession, FMenuBuilder& InOutMenuBuilder)
+{
+	if (&InSession != TimingView)
+	{
+		return;
+	}
+}
+
+////////////////////////////////////////////////////////////////////////////////////////////////////
+
+void FTaskTimingSharedState::SetTaskId(uint32 InTaskId)
+{
+	if (TaskTrack.IsValid())
+	{
+		TaskTrack->SetTaskId(InTaskId);
+	}
+}
+
+////////////////////////////////////////////////////////////////////////////////////////////////////
+// FTaskTimingTrack
+////////////////////////////////////////////////////////////////////////////////////////////////////
+
+INSIGHTS_IMPLEMENT_RTTI(FTaskTimingTrack)
+
+////////////////////////////////////////////////////////////////////////////////////////////////////
+
+const uint32 FTaskTimingTrack::InvalidTaskId = ~0;
+
+
+void FTaskTimingTrack::BuildDrawState(ITimingEventsTrackDrawStateBuilder& Builder, const ITimingTrackUpdateContext& Context)
+{
+	if (TaskId == InvalidTaskId)
+	{
+		return;
+	}
+
+	TSharedPtr<const TraceServices::IAnalysisSession> Session = FInsightsManager::Get()->GetSession();
+	if (!Session.IsValid())
+	{
+		return;
+	}
+
+	TraceServices::FAnalysisSessionReadScope SessionReadScope(*Session.Get());
+
+	const TraceServices::ITasksProvider* TasksProvider = TraceServices::ReadTasksProvider(*Session.Get());
+
+	if (TasksProvider == nullptr)
+	{
+		return;
+	}
+
+	const TraceServices::FTaskInfo* Task = TasksProvider->TryGetTask(TaskId);
+
+	if (Task == nullptr)
+	{
+		return;
+	}
+
+	Builder.AddEvent(Task->CreatedTimestamp, Task->LaunchedTimestamp, 0, TEXT("Created"), 0, FTaskGraphProfilerManager::Get()->GetColorForTaskEvent(ETaskEventType::Created).ToFColor(true).ToPackedARGB());
+	Builder.AddEvent(Task->LaunchedTimestamp, Task->ScheduledTimestamp, 0, TEXT("Launched"), 0, FTaskGraphProfilerManager::Get()->GetColorForTaskEvent(ETaskEventType::Launched).ToFColor(true).ToPackedARGB());
+	Builder.AddEvent(Task->ScheduledTimestamp, Task->StartedTimestamp, 0, TEXT("Scheduled"), 0,  FTaskGraphProfilerManager::Get()->GetColorForTaskEvent(ETaskEventType::Scheduled).ToFColor(true).ToPackedARGB());
+	Builder.AddEvent(Task->StartedTimestamp, Task->FinishedTimestamp, 0, TEXT("Started"), 0, FTaskGraphProfilerManager::Get()->GetColorForTaskEvent(ETaskEventType::NestedCompleted).ToFColor(true).ToPackedARGB());
+	if (Task->CompletedTimestamp > Task->FinishedTimestamp)
+	{
+		Builder.AddEvent(Task->FinishedTimestamp, Task->CompletedTimestamp, 0, TEXT("Completed"), 0, FTaskGraphProfilerManager::Get()->GetColorForTaskEvent(ETaskEventType::Completed).ToFColor(true).ToPackedARGB());
+	}
+}
+
+////////////////////////////////////////////////////////////////////////////////////////////////////
+
+void FTaskTimingTrack::BuildFilteredDrawState(ITimingEventsTrackDrawStateBuilder& Builder, const ITimingTrackUpdateContext& Context)
+{
+	BuildDrawState(Builder, Context);
+}
+
+////////////////////////////////////////////////////////////////////////////////////////////////////
+
+void FTaskTimingTrack::PostDraw(const ITimingTrackDrawContext& Context) const
+{
+}
+
+////////////////////////////////////////////////////////////////////////////////////////////////////
+
+void FTaskTimingTrack::OnTimingEventSelected(TSharedPtr<const ITimingEvent> InSelectedEvent)
+{
+	if (!InSelectedEvent.IsValid() || !InSelectedEvent->Is<FThreadTrackEvent>())
+	{
+		TaskId = InvalidTaskId;
+		SetDirtyFlag();
+		return;
+	}
+
+	TSharedPtr<const TraceServices::IAnalysisSession> Session = FInsightsManager::Get()->GetSession();
+	if (!Session.IsValid())
+	{
+		return;
+	}
+
+	TraceServices::FAnalysisSessionReadScope SessionReadScope(*Session.Get());
+
+	const TraceServices::ITasksProvider* TasksProvider = TraceServices::ReadTasksProvider(*Session.Get());
+
+	if (TasksProvider == nullptr)
+	{
+		return;
+	}
+
+	const FThreadTrackEvent &ThreadEvent = InSelectedEvent->As<FThreadTrackEvent>();
+	uint32 ThreadId = StaticCastSharedRef<const FThreadTimingTrack>(ThreadEvent.GetTrack())->GetThreadId();
+	const TraceServices::FTaskInfo* Task = TasksProvider->TryGetTask(ThreadId, ThreadEvent.GetStartTime());
+
+	if (Task != nullptr)
+	{
+		TaskId = Task->Id;
+	}
+	else
+	{
+		TaskId = InvalidTaskId;
+	}
+
+	SetDirtyFlag();
+}
+
+////////////////////////////////////////////////////////////////////////////////////////////////////
+
+const TSharedPtr<const ITimingEvent> FTaskTimingTrack::GetEvent(float InPosX, float InPosY, const FTimingTrackViewport& Viewport) const
+{
+	TSharedPtr<FTimingEvent> TimingEvent;
+
+	if (TaskId == InvalidTaskId)
+	{
+		return TimingEvent;
+	}
+
+	const FTimingViewLayout& Layout = Viewport.GetLayout();
+	const float TopLaneY = GetPosY() + 1.0f + Layout.TimelineDY; // +1.0f is for horizontal line between timelines
+	const float DY = InPosY - TopLaneY;
+
+	// If mouse is not above first sub-track or below last sub-track...
+	if (DY >= 0 && DY < GetHeight() - 1.0f - 2 * Layout.TimelineDY)
+	{
+		const double EventTime = Viewport.SlateUnitsToTime(InPosX);
+
+		TSharedPtr<const TraceServices::IAnalysisSession> Session = FInsightsManager::Get()->GetSession();
+		if (!Session.IsValid())
+		{
+			return TimingEvent;
+		}
+
+		TraceServices::FAnalysisSessionReadScope SessionReadScope(*Session.Get());
+
+		const TraceServices::ITasksProvider* TasksProvider = TraceServices::ReadTasksProvider(*Session.Get());
+
+		if (TasksProvider == nullptr)
+		{
+			return TimingEvent;
+		}
+
+		const TraceServices::FTaskInfo* Task = TasksProvider->TryGetTask(TaskId);
+
+		if (Task == nullptr)
+		{
+			return TimingEvent;
+		}
+
+		// This logic will be replaced
+		if (EventTime >= Task->CreatedTimestamp && EventTime < Task->LaunchedTimestamp)
+		{
+			TimingEvent = MakeShared<FTimingEvent>(SharedThis(this), Task->CreatedTimestamp, Task->LaunchedTimestamp, 0);
+		}
+		else if (EventTime >= Task->LaunchedTimestamp && EventTime < Task->ScheduledTimestamp)
+		{
+			TimingEvent = MakeShared<FTimingEvent>(SharedThis(this), Task->LaunchedTimestamp, Task->ScheduledTimestamp, 0);
+		}
+		else if (EventTime >= Task->ScheduledTimestamp && EventTime < Task->StartedTimestamp)
+		{
+			TimingEvent = MakeShared<FTimingEvent>(SharedThis(this), Task->ScheduledTimestamp, Task->StartedTimestamp, 0);
+		}
+		else if (EventTime >= Task->StartedTimestamp && EventTime < Task->FinishedTimestamp)
+		{
+			TimingEvent = MakeShared<FTimingEvent>(SharedThis(this), Task->StartedTimestamp, Task->FinishedTimestamp, 0);
+		}
+		else if (EventTime >= Task->FinishedTimestamp && EventTime < Task->CompletedTimestamp)
+		{
+			TimingEvent = MakeShared<FTimingEvent>(SharedThis(this), Task->FinishedTimestamp, Task->CompletedTimestamp, 0);
+		}
+	}
+
+	return TimingEvent;
+}
+
+////////////////////////////////////////////////////////////////////////////////////////////////////
+
+void FTaskTimingTrack::InitTooltip(FTooltipDrawState& InOutTooltip, const ITimingEvent& InTooltipEvent) const
+{
+	InOutTooltip.ResetContent();
+
+	// The rest of the implemention will follow
+}
+
+////////////////////////////////////////////////////////////////////////////////////////////////////
+
+} // namespace Insights
+
+#undef LOCTEXT_NAMESPACE
diff --git a/Engine/Source/Developer/TraceInsights/Private/Insights/TaskGraphProfiler/ViewModels/TaskTimingTrack.h b/Engine/Source/Developer/TraceInsights/Private/Insights/TaskGraphProfiler/ViewModels/TaskTimingTrack.h
new file mode 100644
--- /dev/null
+++ b/Engine/Source/Developer/TraceInsights/Private/Insights/TaskGraphProfiler/ViewModels/TaskTimingTrack.h
@@ -0,0 +1,107 @@
+// Copyright Epic Games, Inc. All Rights Reserved.
+
+#pragma once
+
+#include "CoreMinimal.h"
+
+// Insights
+#include "Insights/ITimingViewExtender.h"
+#include "Insights/ViewModels/TimingEventsTrack.h"
+
+class STimingView;
+
+////////////////////////////////////////////////////////////////////////////////////////////////////
+
+namespace Insights
+{
+class FTaskTimingTrack;
+
+class FTaskTimingSharedState : public Insights::ITimingViewExtender, public TSharedFromThis<FTaskTimingSharedState>
+{
+
+public:
+	explicit FTaskTimingSharedState(STimingView* InTimingView) : TimingView(InTimingView) {}
+	virtual ~FTaskTimingSharedState() = default;
+
+	TSharedPtr<FTaskTimingTrack> GetTaskTrack() { return TaskTrack; }
+
+	bool IsTaskTrackVisible() const;
+
+	//////////////////////////////////////////////////
+	// ITimingViewExtender interface
+
+	virtual void OnBeginSession(Insights::ITimingViewSession& InSession) override;
+	virtual void OnEndSession(Insights::ITimingViewSession& InSession) override;
+	virtual void Tick(Insights::ITimingViewSession& InSession, const TraceServices::IAnalysisSession& InAnalysisSession) override;
+	virtual void ExtendFilterMenu(Insights::ITimingViewSession& InSession, FMenuBuilder& InMenuBuilder) override;
+
+	//////////////////////////////////////////////////
+
+	bool IsTaskTrackToggleOn() const { return bShowHideTaskTrack; }
+	void SetTaskTrackToggle(bool bOnOff) { bShowHideTaskTrack = bOnOff; }
+	void ShowTaskTrack() { SetTaskTrackToggle(true); }
+	void HideTaskTrack() { SetTaskTrackToggle(false); }
+
+	void SetTaskId(uint32 InTaskId);
+
+private:
+	STimingView* TimingView;
+
+	bool bShowHideTaskTrack;
+
+	TSharedPtr<FTaskTimingTrack> TaskTrack;
+};
+
+////////////////////////////////////////////////////////////////////////////////////////////////////
+
+class FTaskTimingTrack : public FTimingEventsTrack
+{
+	INSIGHTS_DECLARE_RTTI(FTaskTimingTrack, FTimingEventsTrack)
+
+public:
+
+	struct FPendingEventInfo
+	{
+		double StartTime;
+		double EndTime;
+		uint32 Depth;
+		uint32 TimerIndex;
+	};
+
+	explicit FTaskTimingTrack(FTaskTimingSharedState& InSharedState, const FString& InName, uint32 InTimelineIndex)
+		: FTimingEventsTrack(InName)
+		, TimelineIndex(InTimelineIndex)
+		, SharedState(InSharedState)
+	{
+	}
+
+	virtual ~FTaskTimingTrack() {}
+
+	uint32 GetTimelineIndex() const { return TimelineIndex; }
+
+	virtual void BuildDrawState(ITimingEventsTrackDrawStateBuilder& Builder, const ITimingTrackUpdateContext& Context) override;
+	virtual void BuildFilteredDrawState(ITimingEventsTrackDrawStateBuilder& Builder, const ITimingTrackUpdateContext& Context) override;
+
+	virtual void PostDraw(const ITimingTrackDrawContext& Context) const override;
+
+	virtual void InitTooltip(FTooltipDrawState& InOutTooltip, const ITimingEvent& InTooltipEvent) const override;
+
+	virtual const TSharedPtr<const ITimingEvent> GetEvent(float InPosX, float InPosY, const FTimingTrackViewport& Viewport) const override;
+
+	void SetTaskId(uint32 InTaskId) { TaskId = InTaskId; SetDirtyFlag(); }
+	uint32 GetTaskId() { return TaskId; }
+	
+	void OnTimingEventSelected(TSharedPtr<const ITimingEvent> InSelectedEvent);
+
+private:
+	static const uint32 InvalidTaskId;
+	uint32 TimelineIndex;
+
+	FTaskTimingSharedState& SharedState;
+
+	uint32 TaskId = InvalidTaskId;
+};
+
+////////////////////////////////////////////////////////////////////////////////////////////////////
+
+} //namespace Insights
diff --git a/Engine/Source/Developer/TraceInsights/Private/Insights/TaskGraphProfiler/Widgets/STaskTableTreeView.cpp b/Engine/Source/Developer/TraceInsights/Private/Insights/TaskGraphProfiler/Widgets/STaskTableTreeView.cpp
--- a/Engine/Source/Developer/TraceInsights/Private/Insights/TaskGraphProfiler/Widgets/STaskTableTreeView.cpp
+++ b/Engine/Source/Developer/TraceInsights/Private/Insights/TaskGraphProfiler/Widgets/STaskTableTreeView.cpp
@@ -14,9 +14,11 @@
 #include "Insights/TaskGraphProfiler/TaskGraphProfilerManager.h"
 #include "Insights/TaskGraphProfiler/ViewModels/TaskEntry.h"
 #include "Insights/TaskGraphProfiler/ViewModels/TaskNode.h"
+#include "Insights/TaskGraphProfiler/ViewModels/TaskTimingTrack.h"
 #include "Insights/Table/ViewModels/TableColumn.h"
 #include "Insights/TimingProfilerManager.h"
 #include "Insights/ViewModels/FilterConfigurator.h"
+#include "Insights/ViewModels/TaskGraphRelation.h"
 #include "Insights/Widgets/STimingProfilerWindow.h"
 #include "Insights/Widgets/STimingView.h"
 
@@ -435,7 +437,7 @@ void STaskTableTreeView::ContextMenu_GoToTask_Execute()
 
 	TimingView->ClearRelations();
 
-	auto AddRelation = [&TimingView](double SourceTimestamp, uint32 SourceThreadId, double TargetTimestamp, uint32 TargetThreadId, FTaskGraphRelation::ETaskGraphRelationType Type)
+	auto AddRelation = [&TimingView](double SourceTimestamp, uint32 SourceThreadId, double TargetTimestamp, uint32 TargetThreadId, ETaskEventType Type)
 	{
 		if (SourceTimestamp == TraceServices::FTaskInfo::InvalidTimestamp || TargetTimestamp == TraceServices::FTaskInfo::InvalidTimestamp)
 		{
@@ -451,6 +453,13 @@ void STaskTableTreeView::ContextMenu_GoToTask_Execute()
 
 	double Duration = (SelectedTask->GetTask()->GetFinishedTimestamp() - SelectedTask->GetTask()->GetCreatedTimestamp()) * 1.5;
 	TimingView->ZoomOnTimeInterval(SelectedTask->GetTask()->GetCreatedTimestamp() - Duration * 0.15, Duration);
+
+	TSharedPtr<FTaskTimingSharedState> TaskSharedState = FTaskGraphProfilerManager::Get()->GetTaskTimingSharedState();
+
+	if (TaskSharedState.IsValid())
+	{
+		TaskSharedState->SetTaskId(SelectedTask->GetTask()->GetId());
+	}
 }
 
 ////////////////////////////////////////////////////////////////////////////////////////////////////
diff --git a/Engine/Source/Developer/TraceInsights/Private/Insights/ViewModels/TaskGraphRelation.cpp b/Engine/Source/Developer/TraceInsights/Private/Insights/ViewModels/TaskGraphRelation.cpp
--- a/Engine/Source/Developer/TraceInsights/Private/Insights/ViewModels/TaskGraphRelation.cpp
+++ b/Engine/Source/Developer/TraceInsights/Private/Insights/ViewModels/TaskGraphRelation.cpp
@@ -9,7 +9,7 @@
 namespace Insights
 {
 
-FTaskGraphRelation::FTaskGraphRelation(double InSourceTime, int32 InSourceThreadId, double InTargetTime, int32 InTargetThreadId, ETaskGraphRelationType InType)
+FTaskGraphRelation::FTaskGraphRelation(double InSourceTime, int32 InSourceThreadId, double InTargetTime, int32 InTargetThreadId, ETaskEventType InType)
 {
 	SourceTime = InSourceTime;
 	SourceThreadId = InSourceThreadId;
@@ -56,18 +56,7 @@ void FTaskGraphRelation::Draw(const FDrawContext& DrawContext, const FTimingTrac
 	constexpr float LineLenghtAtEnds = 20.0f;
 	FVector2D StartDir((X2 - X1) / DirectionFactor, 0.0f);
 
-	FLinearColor RelationColors[(int32)ETaskGraphRelationType::Count] = {
-		FLinearColor::Yellow, // Created
-		FLinearColor::Green, // Launched
-		FLinearColor::Red, // Prerequisite
-		FLinearColor::Blue, // Scheduled
-		FLinearColor::Blue, // AddedNested
-		FLinearColor::Red, // NestedCompleted
-		FLinearColor::Red, // Subsequent
-		FLinearColor::Yellow, // Completed
-	};
-	
-	FLinearColor Color = RelationColors[(int32)Type];
+	FLinearColor Color = FTaskGraphProfilerManager::Get()->GetColorForTaskEvent(Type);
 	TArray<FVector2D> ArrowPoints;
 
 	if (Distance > LineLenghtAtEnds)
diff --git a/Engine/Source/Developer/TraceInsights/Private/Insights/ViewModels/TaskGraphRelation.h b/Engine/Source/Developer/TraceInsights/Private/Insights/ViewModels/TaskGraphRelation.h
--- a/Engine/Source/Developer/TraceInsights/Private/Insights/ViewModels/TaskGraphRelation.h
+++ b/Engine/Source/Developer/TraceInsights/Private/Insights/ViewModels/TaskGraphRelation.h
@@ -4,6 +4,7 @@
 
 #include "CoreMinimal.h"
 
+#include "Insights/TaskGraphProfiler/TaskGraphProfilerManager.h"
 #include "Insights/ViewModels/ITimingEvent.h"
 
 namespace Insights
@@ -12,20 +13,7 @@ namespace Insights
 class FTaskGraphRelation : public ITimingEventRelation
 {
 public:
-	enum class ETaskGraphRelationType : int32
-	{
-		Created,
-		Launched,
-		Prerequisite,
-		Scheduled,
-		AddedNested,
-		NestedCompleted,
-		Subsequent,
-		Completed,
-		Count
-	};
-
-	FTaskGraphRelation(double InSourceTime, int32 InSourceThreadId, double InTargetTime, int32 InTargetThreadId, ETaskGraphRelationType InType);
+	FTaskGraphRelation(double InSourceTime, int32 InSourceThreadId, double InTargetTime, int32 InTargetThreadId, ETaskEventType InType);
 
 	virtual bool IsSolved() override { return bIsSolved; }
 
@@ -48,7 +36,7 @@ public:
 	int32 GetTargetDepth() { return TargetDepth; }
 
 	void SetIsSolved(bool InValue) { bIsSolved = InValue; }
-	ETaskGraphRelationType GetType() { return Type; }
+	ETaskEventType GetType() { return Type; }
 
 private:
 	double SourceTime;
@@ -57,7 +45,7 @@ private:
 	double TargetTime;
 	int32 TargetThreadId;
 	int32 TargetDepth = 0;
-	ETaskGraphRelationType Type;
+	ETaskEventType Type;
 
 	bool bIsSolved = false;
 
diff --git a/Engine/Source/Developer/TraceInsights/Private/Insights/ViewModels/ThreadTimingTrack.cpp b/Engine/Source/Developer/TraceInsights/Private/Insights/ViewModels/ThreadTimingTrack.cpp
--- a/Engine/Source/Developer/TraceInsights/Private/Insights/ViewModels/ThreadTimingTrack.cpp
+++ b/Engine/Source/Developer/TraceInsights/Private/Insights/ViewModels/ThreadTimingTrack.cpp
@@ -1515,7 +1515,7 @@ void FThreadTimingTrack::GetEventRelations(const ITimingEvent& InSelectedEvent,
 	const TraceServices::ITasksProvider* TasksProvider = TraceServices::ReadTasksProvider(*Session.Get());
 	if (TasksProvider)
 	{
-		auto AddRelation = [this, &InSelectedEvent, &Relations](double SourceTimestamp, uint32 SourceThreadId, double TargetTimestamp, uint32 TargetThreadId, FTaskGraphRelation::ETaskGraphRelationType Type)
+		auto AddRelation = [this, &InSelectedEvent, &Relations](double SourceTimestamp, uint32 SourceThreadId, double TargetTimestamp, uint32 TargetThreadId, ETaskEventType Type)
 		{
 			if (SourceTimestamp == TraceServices::FTaskInfo::InvalidTimestamp || TargetTimestamp == TraceServices::FTaskInfo::InvalidTimestamp)
 			{
@@ -1563,8 +1563,8 @@ void FThreadTimingTrack::GetEventRelations(const ITimingEvent& InSelectedEvent,
 				const TraceServices::FTaskInfo* WaitedTask = TasksProvider->TryGetTask(Waiting->Tasks[i]);
 				if (WaitedTask != nullptr)
 				{
-					AddRelation(StartTime, ThreadId, WaitedTask->StartedTimestamp, WaitedTask->StartedThreadId, FTaskGraphRelation::ETaskGraphRelationType::AddedNested);
-					AddRelation(WaitedTask->CompletedTimestamp, WaitedTask->CompletedThreadId, WaitedTask->CompletedTimestamp, ThreadId, FTaskGraphRelation::ETaskGraphRelationType::NestedCompleted);
+					AddRelation(StartTime, ThreadId, WaitedTask->StartedTimestamp, WaitedTask->StartedThreadId, ETaskEventType::AddedNested);
+					AddRelation(WaitedTask->CompletedTimestamp, WaitedTask->CompletedThreadId, WaitedTask->CompletedTimestamp, ThreadId, ETaskEventType::NestedCompleted);
 				}
 			}
 		}
diff --git a/Engine/Source/Developer/TraceInsights/Private/Insights/Widgets/STimingView.cpp b/Engine/Source/Developer/TraceInsights/Private/Insights/Widgets/STimingView.cpp
--- a/Engine/Source/Developer/TraceInsights/Private/Insights/Widgets/STimingView.cpp
+++ b/Engine/Source/Developer/TraceInsights/Private/Insights/Widgets/STimingView.cpp
@@ -4178,14 +4178,14 @@ void STimingView::ContextMenu_ShowTaskDependecies_Execute()
 
 bool STimingView::ContextMenu_ShowTaskDependecies_CanExecute()
 {
-	return Insights::FTaskGraphProfilerManager::Get()->GetIsAvailable();
+	return Insights::FTaskGraphProfilerManager::Get().IsValid() && Insights::FTaskGraphProfilerManager::Get()->GetIsAvailable();
 }
 
 ////////////////////////////////////////////////////////////////////////////////////////////////////
 
 bool STimingView::ContextMenu_ShowTaskDependecies_IsChecked()
 {
-	return Insights::FTaskGraphProfilerManager::Get()->GetIsAvailable() && bShowEventRelations;
+	return Insights::FTaskGraphProfilerManager::Get().IsValid() && Insights::FTaskGraphProfilerManager::Get()->GetIsAvailable() && bShowEventRelations;
 }
 
 ////////////////////////////////////////////////////////////////////////////////////////////////////
diff --git a/Engine/Source/Developer/TraceInsights/Public/Insights/ViewModels/BaseTimingTrack.h b/Engine/Source/Developer/TraceInsights/Public/Insights/ViewModels/BaseTimingTrack.h
--- a/Engine/Source/Developer/TraceInsights/Public/Insights/ViewModels/BaseTimingTrack.h
+++ b/Engine/Source/Developer/TraceInsights/Public/Insights/ViewModels/BaseTimingTrack.h
@@ -48,7 +48,8 @@ struct TRACEINSIGHTS_API FTimingTrackOrder
 	static constexpr int32 Memory     = 1 * GroupRange;
 	static constexpr int32 Gpu        = 2 * GroupRange;
 	static constexpr int32 Cpu        = 3 * GroupRange;
-	static constexpr int32 Last       = 4 * GroupRange;
+	static constexpr int32 Task       = 4 * GroupRange;
+	static constexpr int32 Last       = 5 * GroupRange;
 };
 
 ////////////////////////////////////////////////////////////////////////////////////////////////////
