* tuned task stages names in TaskTimingTrack
* fixed task relations visualisation (subsequents weren't visualised)
* it was almost impossible to follow tasks links. When a task relations are visualised, the task has arrows to its subsequents, to the point of where subsequents were scheduled, but there's no way to follow to subsequent's execution. To help with this, when a task is selected, now prerequisites, nested tasks and subsequents (with all their relations) are also visualised. This widens the visualisation range and is still fast

diff --git a/Engine/Source/Developer/TraceInsights/Private/Insights/TaskGraphProfiler/TaskGraphProfilerManager.cpp b/Engine/Source/Developer/TraceInsights/Private/Insights/TaskGraphProfiler/TaskGraphProfilerManager.cpp
--- a/Engine/Source/Developer/TraceInsights/Private/Insights/TaskGraphProfiler/TaskGraphProfilerManager.cpp
+++ b/Engine/Source/Developer/TraceInsights/Private/Insights/TaskGraphProfiler/TaskGraphProfilerManager.cpp
@@ -253,14 +253,32 @@ void FTaskGraphProfilerManager::OnTaskTableTreeViewTabClosed(TSharedRef<SDockTab
 
 ////////////////////////////////////////////////////////////////////////////////////////////////////
 
-void FTaskGraphProfilerManager::GetTaskRelations(const TraceServices::FTaskInfo* Task, const TraceServices::ITasksProvider* TasksProvider, AddRelationCallback Callback)
+void FTaskGraphProfilerManager::GetTaskRelations(const TraceServices::FTaskInfo* Task, const TraceServices::ITasksProvider* TasksProvider, FAddRelationCallback Callback)
 {
-	const int32 MaxTasksToShow = 30;
+	check(Task != nullptr);
 
-	if (Task == nullptr)
+	auto GetSingleTaskRelationsForAll = [this, TasksProvider, &Callback](const TArray< TraceServices::FTaskInfo::FRelationInfo>& Collection)
 	{
-		return;
-	}
+		for (const TraceServices::FTaskInfo::FRelationInfo& Relation : Collection)
+		{
+			const TraceServices::FTaskInfo* Task = TasksProvider->TryGetTask(Relation.RelativeId);
+			if (Task != nullptr)
+			{
+				GetSingleTaskRelations(Task, TasksProvider, Callback);
+			}
+		}
+	};
+
+
+	GetSingleTaskRelationsForAll(Task->Prerequisites);
+	GetSingleTaskRelations(Task, TasksProvider, Callback);
+	GetSingleTaskRelationsForAll(Task->NestedTasks);
+	GetSingleTaskRelationsForAll(Task->Subsequents);
+}
+
+void FTaskGraphProfilerManager::GetSingleTaskRelations(const TraceServices::FTaskInfo* Task, const TraceServices::ITasksProvider* TasksProvider, FAddRelationCallback Callback)
+{
+	const int32 MaxTasksToShow = 30;
 
 	if (Task->CreatedTimestamp != Task->LaunchedTimestamp || Task->CreatedThreadId != Task->LaunchedThreadId)
 	{
@@ -294,12 +312,15 @@ void FTaskGraphProfilerManager::GetTaskRelations(const TraceServices::FTaskInfo*
 		Callback(NestedTask->CompletedTimestamp, NestedTask->CompletedThreadId, NestedTask->CompletedTimestamp, Task->StartedThreadId, ETaskEventType::NestedCompleted);
 	}
 
-	int32 NumSubsequentsToShow = FMath::Min(Task->NestedTasks.Num(), MaxTasksToShow);
+	int32 NumSubsequentsToShow = FMath::Min(Task->Subsequents.Num(), MaxTasksToShow);
 	for (int32 i = 0; i != NumSubsequentsToShow; ++i)
 	{
 		const TraceServices::FTaskInfo* Subsequent = TasksProvider->TryGetTask(Task->Subsequents[i].RelativeId);
 		check(Subsequent != nullptr);
-		Callback(Task->CompletedTimestamp, Task->CompletedThreadId, Subsequent->ScheduledTimestamp, Subsequent->ScheduledThreadId, ETaskEventType::Subsequent);
+		if (Task->CompletedTimestamp < Subsequent->ScheduledTimestamp)
+		{
+			Callback(Task->CompletedTimestamp, Task->CompletedThreadId, Subsequent->ScheduledTimestamp, Subsequent->ScheduledThreadId, ETaskEventType::Subsequent);
+		}
 	}
 
 	if (Task->FinishedTimestamp != Task->CompletedTimestamp || Task->CompletedThreadId != Task->StartedThreadId)
@@ -310,7 +331,7 @@ void FTaskGraphProfilerManager::GetTaskRelations(const TraceServices::FTaskInfo*
 
 ////////////////////////////////////////////////////////////////////////////////////////////////////
 
-void FTaskGraphProfilerManager::GetTaskRelations(double Time, uint32 ThreadId, AddRelationCallback Callback)
+void FTaskGraphProfilerManager::GetTaskRelations(double Time, uint32 ThreadId, FAddRelationCallback Callback)
 {
 	TSharedPtr<const TraceServices::IAnalysisSession> Session = FInsightsManager::Get()->GetSession();
 	if (!Session.IsValid())
@@ -337,7 +358,7 @@ void FTaskGraphProfilerManager::GetTaskRelations(double Time, uint32 ThreadId, A
 
 ////////////////////////////////////////////////////////////////////////////////////////////////////
 
-void FTaskGraphProfilerManager::GetTaskRelations(uint32 TaskId, AddRelationCallback Callback)
+void FTaskGraphProfilerManager::GetTaskRelations(uint32 TaskId, FAddRelationCallback Callback)
 {
 	TSharedPtr<const TraceServices::IAnalysisSession> Session = FInsightsManager::Get()->GetSession();
 	if (!Session.IsValid())
diff --git a/Engine/Source/Developer/TraceInsights/Private/Insights/TaskGraphProfiler/TaskGraphProfilerManager.h b/Engine/Source/Developer/TraceInsights/Private/Insights/TaskGraphProfiler/TaskGraphProfilerManager.h
--- a/Engine/Source/Developer/TraceInsights/Private/Insights/TaskGraphProfiler/TaskGraphProfilerManager.h
+++ b/Engine/Source/Developer/TraceInsights/Private/Insights/TaskGraphProfiler/TaskGraphProfilerManager.h
@@ -52,7 +52,7 @@ enum class ETaskEventType : uint32
 class FTaskGraphProfilerManager : public TSharedFromThis<FTaskGraphProfilerManager>, public IInsightsComponent
 {
 public:
-	typedef TFunction<void(double /*SourceTimestamp*/, uint32 /*SourceThreadId*/, double /*TargetTimestamp*/, uint32 /*TargetThreadId*/, ETaskEventType /*Type*/)> AddRelationCallback;
+	typedef TFunctionRef<void(double /*SourceTimestamp*/, uint32 /*SourceThreadId*/, double /*TargetTimestamp*/, uint32 /*TargetThreadId*/, ETaskEventType /*Type*/)> FAddRelationCallback;
 
 	/** Creates the Memory Profiler manager, only one instance can exist. */
 	FTaskGraphProfilerManager(TSharedRef<FUICommandList> InCommandList);
@@ -90,8 +90,8 @@ public:
 
 	void OnSessionChanged();
 
-	void GetTaskRelations(double Time, uint32 ThreadId, AddRelationCallback Callback);
-	void GetTaskRelations(uint32 TaskId, AddRelationCallback Callback);
+	void GetTaskRelations(double Time, uint32 ThreadId, FAddRelationCallback Callback);
+	void GetTaskRelations(uint32 TaskId, FAddRelationCallback Callback);
 	FLinearColor GetColorForTaskEvent(ETaskEventType InEvent);
 
 	TSharedPtr<Insights::FTaskTimingSharedState> GetTaskTimingSharedState() { return TaskTimingSharedState;	}
@@ -102,7 +102,8 @@ private:
 
 	void RegisterTimingProfilerLayoutExtensions(FInsightsMajorTabExtender& InOutExtender);
 
-	void GetTaskRelations(const TraceServices::FTaskInfo* Task, const TraceServices::ITasksProvider* TasksProvider, AddRelationCallback Callback);
+	void GetTaskRelations(const TraceServices::FTaskInfo* Task, const TraceServices::ITasksProvider* TasksProvider, FAddRelationCallback Callback);
+	void GetSingleTaskRelations(const TraceServices::FTaskInfo* Task, const TraceServices::ITasksProvider* TasksProvider, FAddRelationCallback Callback);
 
 	void InitializeColorCode();
 
diff --git a/Engine/Source/Developer/TraceInsights/Private/Insights/TaskGraphProfiler/ViewModels/TaskTimingTrack.cpp b/Engine/Source/Developer/TraceInsights/Private/Insights/TaskGraphProfiler/ViewModels/TaskTimingTrack.cpp
--- a/Engine/Source/Developer/TraceInsights/Private/Insights/TaskGraphProfiler/ViewModels/TaskTimingTrack.cpp
+++ b/Engine/Source/Developer/TraceInsights/Private/Insights/TaskGraphProfiler/ViewModels/TaskTimingTrack.cpp
@@ -127,10 +127,10 @@ void FTaskTimingTrack::BuildDrawState(ITimingEventsTrackDrawStateBuilder& Builde
 		return;
 	}
 
-	Builder.AddEvent(Task->CreatedTimestamp, Task->LaunchedTimestamp, 0, TEXT("Created"), 0, FTaskGraphProfilerManager::Get()->GetColorForTaskEvent(ETaskEventType::Created).ToFColor(true).ToPackedARGB());
-	Builder.AddEvent(Task->LaunchedTimestamp, Task->ScheduledTimestamp, 0, TEXT("Launched"), 0, FTaskGraphProfilerManager::Get()->GetColorForTaskEvent(ETaskEventType::Launched).ToFColor(true).ToPackedARGB());
+	Builder.AddEvent(Task->CreatedTimestamp, Task->LaunchedTimestamp, 0, TEXT("Launched"), 0, FTaskGraphProfilerManager::Get()->GetColorForTaskEvent(ETaskEventType::Created).ToFColor(true).ToPackedARGB());
+	Builder.AddEvent(Task->LaunchedTimestamp, Task->ScheduledTimestamp, 0, TEXT("Dispatched"), 0, FTaskGraphProfilerManager::Get()->GetColorForTaskEvent(ETaskEventType::Launched).ToFColor(true).ToPackedARGB());
 	Builder.AddEvent(Task->ScheduledTimestamp, Task->StartedTimestamp, 0, TEXT("Scheduled"), 0,  FTaskGraphProfilerManager::Get()->GetColorForTaskEvent(ETaskEventType::Scheduled).ToFColor(true).ToPackedARGB());
-	Builder.AddEvent(Task->StartedTimestamp, Task->FinishedTimestamp, 0, TEXT("Started"), 0, FTaskGraphProfilerManager::Get()->GetColorForTaskEvent(ETaskEventType::NestedCompleted).ToFColor(true).ToPackedARGB());
+	Builder.AddEvent(Task->StartedTimestamp, Task->FinishedTimestamp, 0, TEXT("Executed"), 0, FTaskGraphProfilerManager::Get()->GetColorForTaskEvent(ETaskEventType::NestedCompleted).ToFColor(true).ToPackedARGB());
 	if (Task->CompletedTimestamp > Task->FinishedTimestamp)
 	{
 		Builder.AddEvent(Task->FinishedTimestamp, Task->CompletedTimestamp, 0, TEXT("Completed"), 0, FTaskGraphProfilerManager::Get()->GetColorForTaskEvent(ETaskEventType::Completed).ToFColor(true).ToPackedARGB());
