[Insights] Fixed crash (in Unreal Insights, when opening a trace file) becasue of a log message spec event being received with an unknown category (category name was empty). Also named these categores as "N/A" instead of being an emty string.

diff --git a/Engine/Source/Developer/TraceInsights/Private/Insights/ViewModels/MarkersTimingTrack.cpp b/Engine/Source/Developer/TraceInsights/Private/Insights/ViewModels/MarkersTimingTrack.cpp
--- a/Engine/Source/Developer/TraceInsights/Private/Insights/ViewModels/MarkersTimingTrack.cpp
+++ b/Engine/Source/Developer/TraceInsights/Private/Insights/ViewModels/MarkersTimingTrack.cpp
@@ -371,7 +371,7 @@ void FMarkersTimingTrack::UpdateBookmarkCategory()
 
 		LogProvider.EnumerateCategories([this](const TraceServices::FLogCategoryInfo& Category)
 		{
-			if (FCString::Strcmp(Category.Name, TEXT("LogBookmark")) == 0)
+			if (Category.Name && FCString::Strcmp(Category.Name, TEXT("LogBookmark")) == 0)
 			{
 				BookmarkCategory = &Category;
 			}
diff --git a/Engine/Source/Developer/TraceServices/Private/Model/Log.cpp b/Engine/Source/Developer/TraceServices/Private/Model/Log.cpp
--- a/Engine/Source/Developer/TraceServices/Private/Model/Log.cpp
+++ b/Engine/Source/Developer/TraceServices/Private/Model/Log.cpp
@@ -52,6 +52,8 @@ FLogCategoryInfo& FLogProvider::GetCategory(uint64 CategoryPointer)
 	else
 	{
 		FLogCategoryInfo& Category = Categories.PushBack();
+		Category.Name = TEXT("N/A");
+		Category.DefaultVerbosity = ELogVerbosity::All;
 		CategoryMap.Add(CategoryPointer, &Category);
 		return Category;
 	}
