[Insights] Refactor of TaskGraphProfiler systems. The functionality is much more self contained now.
+ Fix for task relations and the task timeline show/hide functionality being inconsistent
+ Removed the step of "Solving" relations

diff --git a/Engine/Source/Developer/TraceInsights/Private/Insights/TaskGraphProfiler/TaskGraphProfilerManager.cpp b/Engine/Source/Developer/TraceInsights/Private/Insights/TaskGraphProfiler/TaskGraphProfilerManager.cpp
--- a/Engine/Source/Developer/TraceInsights/Private/Insights/TaskGraphProfiler/TaskGraphProfilerManager.cpp
+++ b/Engine/Source/Developer/TraceInsights/Private/Insights/TaskGraphProfiler/TaskGraphProfilerManager.cpp
@@ -14,7 +14,10 @@
 #include "Insights/TaskGraphProfiler/Widgets/STaskTableTreeView.h"
 #include "Insights/TimingProfilerManager.h"
 #include "Insights/ViewModels/TaskGraphRelation.h"
+#include "Insights/ViewModels/ThreadTimingTrack.h"
+#include "Insights/ViewModels/ThreadTrackEvent.h"
 #include "Insights/Widgets/STimingProfilerWindow.h"
+#include "Insights/Widgets/STimingView.h"
 
 ////////////////////////////////////////////////////////////////////////////////////////////////////
 
@@ -253,51 +256,55 @@ void FTaskGraphProfilerManager::OnTaskTableTreeViewTabClosed(TSharedRef<SDockTab
 
 ////////////////////////////////////////////////////////////////////////////////////////////////////
 
-void FTaskGraphProfilerManager::GetTaskRelations(const TraceServices::FTaskInfo* Task, const TraceServices::ITasksProvider* TasksProvider, FAddRelationCallback Callback)
+void FTaskGraphProfilerManager::ShowTaskRelations(const TraceServices::FTaskInfo* Task, const TraceServices::ITasksProvider* TasksProvider, const FThreadTrackEvent* InSelectedEvent)
 {
 	check(Task != nullptr);
 
-	auto GetSingleTaskRelationsForAll = [this, TasksProvider, &Callback](const TArray< TraceServices::FTaskInfo::FRelationInfo>& Collection)
+	if (!bShowRelations)
+	{
+		return;
+	}
+
+	auto GetSingleTaskRelationsForAll = [this, TasksProvider, InSelectedEvent](const TArray< TraceServices::FTaskInfo::FRelationInfo>& Collection)
 	{
 		for (const TraceServices::FTaskInfo::FRelationInfo& Relation : Collection)
 		{
 			const TraceServices::FTaskInfo* Task = TasksProvider->TryGetTask(Relation.RelativeId);
 			if (Task != nullptr)
 			{
-				GetSingleTaskRelations(Task, TasksProvider, Callback);
+				GetSingleTaskRelations(Task, TasksProvider, InSelectedEvent);
 			}
 		}
 	};
 
-
 	GetSingleTaskRelationsForAll(Task->Prerequisites);
-	GetSingleTaskRelations(Task, TasksProvider, Callback);
+	GetSingleTaskRelations(Task, TasksProvider, InSelectedEvent);
 	GetSingleTaskRelationsForAll(Task->NestedTasks);
 	GetSingleTaskRelationsForAll(Task->Subsequents);
 }
 
-void FTaskGraphProfilerManager::GetSingleTaskRelations(const TraceServices::FTaskInfo* Task, const TraceServices::ITasksProvider* TasksProvider, FAddRelationCallback Callback)
+void FTaskGraphProfilerManager::GetSingleTaskRelations(const TraceServices::FTaskInfo* Task, const TraceServices::ITasksProvider* TasksProvider, const FThreadTrackEvent* InSelectedEvent)
 {
 	const int32 MaxTasksToShow = 30;
 
 	if (Task->CreatedTimestamp != Task->LaunchedTimestamp || Task->CreatedThreadId != Task->LaunchedThreadId)
 	{
-		Callback(Task->CreatedTimestamp, Task->CreatedThreadId, Task->LaunchedTimestamp, Task->LaunchedThreadId, ETaskEventType::Created);
+		AddRelation(InSelectedEvent, Task->CreatedTimestamp, Task->CreatedThreadId, Task->LaunchedTimestamp, Task->LaunchedThreadId, ETaskEventType::Created);
 	}
 
-	Callback(Task->LaunchedTimestamp, Task->LaunchedThreadId, Task->ScheduledTimestamp, Task->ScheduledThreadId, ETaskEventType::Launched);
+	AddRelation(InSelectedEvent, Task->LaunchedTimestamp, Task->LaunchedThreadId, Task->ScheduledTimestamp, Task->ScheduledThreadId, ETaskEventType::Launched);
 
 	int32 NumPrerequisitesToShow = FMath::Min(Task->Prerequisites.Num(), MaxTasksToShow);
 	for (int32 i = 0; i != NumPrerequisitesToShow; ++i)
 	{
 		const TraceServices::FTaskInfo* Prerequisite = TasksProvider->TryGetTask(Task->Prerequisites[i].RelativeId);
 		check(Prerequisite != nullptr);
-		Callback(Prerequisite->CompletedTimestamp, Prerequisite->CompletedThreadId, Task->ScheduledTimestamp, Task->ScheduledThreadId, ETaskEventType::Prerequisite);
+		AddRelation(InSelectedEvent, Prerequisite->CompletedTimestamp, Prerequisite->CompletedThreadId, Task->ScheduledTimestamp, Task->ScheduledThreadId, ETaskEventType::Prerequisite);
 	}
 
 	if (Task->LaunchedTimestamp != Task->ScheduledTimestamp || Task->LaunchedThreadId != Task->ScheduledThreadId)
 	{
-		Callback(Task->ScheduledTimestamp, Task->ScheduledThreadId, Task->StartedTimestamp, Task->StartedThreadId, ETaskEventType::Scheduled);
+		AddRelation(InSelectedEvent, Task->ScheduledTimestamp, Task->ScheduledThreadId, Task->StartedTimestamp, Task->StartedThreadId, ETaskEventType::Scheduled);
 	}
 
 	int32 NumNestedToShow = FMath::Min(Task->NestedTasks.Num(), MaxTasksToShow);
@@ -307,9 +314,9 @@ void FTaskGraphProfilerManager::GetSingleTaskRelations(const TraceServices::FTas
 		const TraceServices::FTaskInfo* NestedTask = TasksProvider->TryGetTask(RelationInfo.RelativeId);
 		check(NestedTask != nullptr);
 
-		Callback(RelationInfo.Timestamp, Task->StartedThreadId, NestedTask->StartedTimestamp, NestedTask->StartedThreadId, ETaskEventType::AddedNested);
+		AddRelation(InSelectedEvent, RelationInfo.Timestamp, Task->StartedThreadId, NestedTask->StartedTimestamp, NestedTask->StartedThreadId, ETaskEventType::AddedNested);
 
-		Callback(NestedTask->CompletedTimestamp, NestedTask->CompletedThreadId, NestedTask->CompletedTimestamp, Task->StartedThreadId, ETaskEventType::NestedCompleted);
+		AddRelation(InSelectedEvent, NestedTask->CompletedTimestamp, NestedTask->CompletedThreadId, NestedTask->CompletedTimestamp, Task->StartedThreadId, ETaskEventType::NestedCompleted);
 	}
 
 	int32 NumSubsequentsToShow = FMath::Min(Task->Subsequents.Num(), MaxTasksToShow);
@@ -319,19 +326,19 @@ void FTaskGraphProfilerManager::GetSingleTaskRelations(const TraceServices::FTas
 		check(Subsequent != nullptr);
 		if (Task->CompletedTimestamp < Subsequent->ScheduledTimestamp)
 		{
-			Callback(Task->CompletedTimestamp, Task->CompletedThreadId, Subsequent->ScheduledTimestamp, Subsequent->ScheduledThreadId, ETaskEventType::Subsequent);
+			AddRelation(InSelectedEvent, Task->CompletedTimestamp, Task->CompletedThreadId, Subsequent->ScheduledTimestamp, Subsequent->ScheduledThreadId, ETaskEventType::Subsequent);
 		}
 	}
 
 	if (Task->FinishedTimestamp != Task->CompletedTimestamp || Task->CompletedThreadId != Task->StartedThreadId)
 	{
-		Callback(Task->FinishedTimestamp, Task->StartedThreadId, Task->CompletedTimestamp, Task->StartedThreadId, ETaskEventType::Completed);
+		AddRelation(InSelectedEvent, Task->FinishedTimestamp, Task->StartedThreadId, Task->CompletedTimestamp, Task->StartedThreadId, ETaskEventType::Completed);
 	}
 }
 
 ////////////////////////////////////////////////////////////////////////////////////////////////////
 
-void FTaskGraphProfilerManager::GetTaskRelations(double Time, uint32 ThreadId, FAddRelationCallback Callback)
+void FTaskGraphProfilerManager::ShowTaskRelations(const FThreadTrackEvent* InSelectedEvent, uint32 ThreadId)
 {
 	TSharedPtr<const TraceServices::IAnalysisSession> Session = FInsightsManager::Get()->GetSession();
 	if (!Session.IsValid())
@@ -348,17 +355,18 @@ void FTaskGraphProfilerManager::GetTaskRelations(double Time, uint32 ThreadId, F
 		return;
 	}
 
-	const TraceServices::FTaskInfo* Task = TasksProvider->TryGetTask(ThreadId, Time);
+	const TraceServices::FTaskInfo* Task = TasksProvider->TryGetTask(ThreadId, InSelectedEvent->GetStartTime());
+	ClearTaskRelations();
 
 	if (Task != nullptr)
 	{
-		GetTaskRelations(Task, TasksProvider, Callback);
+		ShowTaskRelations(Task, TasksProvider, InSelectedEvent);
 	}
 }
 
 ////////////////////////////////////////////////////////////////////////////////////////////////////
 
-void FTaskGraphProfilerManager::GetTaskRelations(uint32 TaskId, FAddRelationCallback Callback)
+void FTaskGraphProfilerManager::ShowTaskRelations(uint32 TaskId)
 {
 	TSharedPtr<const TraceServices::IAnalysisSession> Session = FInsightsManager::Get()->GetSession();
 	if (!Session.IsValid())
@@ -376,10 +384,11 @@ void FTaskGraphProfilerManager::GetTaskRelations(uint32 TaskId, FAddRelationCall
 	}
 
 	const TraceServices::FTaskInfo* Task = TasksProvider->TryGetTask(TaskId);
+	ClearTaskRelations();
 
 	if (Task != nullptr)
 	{
-		GetTaskRelations(Task, TasksProvider, Callback);
+		ShowTaskRelations(Task, TasksProvider, nullptr);
 	}
 }
 
@@ -424,6 +433,113 @@ FLinearColor FTaskGraphProfilerManager::GetColorForTaskEvent(ETaskEventType InEv
 
 ////////////////////////////////////////////////////////////////////////////////////////////////////
 
+void FTaskGraphProfilerManager::AddRelation(const FThreadTrackEvent* InSelectedEvent, double SourceTimestamp, uint32 SourceThreadId, double TargetTimestamp, uint32 TargetThreadId, ETaskEventType Type)
+{
+	if (SourceTimestamp == TraceServices::FTaskInfo::InvalidTimestamp || TargetTimestamp == TraceServices::FTaskInfo::InvalidTimestamp)
+	{
+		return;
+	}
+
+	TSharedPtr<class STimingProfilerWindow> TimingWindow = FTimingProfilerManager::Get()->GetProfilerWindow();
+	if (!TimingWindow.IsValid())
+	{
+		return;
+	}
+
+	TSharedPtr<STimingView> TimingView = TimingWindow->GetTimingView();
+	if (!TimingView.IsValid())
+	{
+		return;
+	}
+
+	TSharedPtr<FThreadTimingSharedState> ThreadSharedState = TimingView->GetThreadTimingSharedState();
+
+	TUniquePtr<ITimingEventRelation> Relation = MakeUnique<FTaskGraphRelation>(SourceTimestamp, SourceThreadId, TargetTimestamp, TargetThreadId, Type);
+	FTaskGraphRelation* TaskRelationPtr = StaticCast<FTaskGraphRelation*>(Relation.Get());
+
+	// If we have a valid event, we can try getting the tracks and depths using this faster approach.
+	if (InSelectedEvent)
+	{
+		TSharedPtr<const FThreadTimingTrack> EventTrack = StaticCastSharedRef<const FThreadTimingTrack>(InSelectedEvent->GetTrack());
+		uint32 ThreadId = EventTrack->GetThreadId();
+
+		if (TaskRelationPtr->GetSourceThreadId() == ThreadId)
+		{
+			TaskRelationPtr->SetSourceTrack(EventTrack);
+			TaskRelationPtr->SetSourceDepth(InSelectedEvent->GetDepth());
+		}
+
+		if (TaskRelationPtr->GetTargetThreadId() == ThreadId)
+		{
+			TaskRelationPtr->SetTargetTrack(EventTrack);
+			TaskRelationPtr->SetTargetDepth(InSelectedEvent->GetDepth());
+		}
+	}
+
+	if(!TaskRelationPtr->GetSourceTrack().IsValid())
+	{
+		TSharedPtr<FCpuTimingTrack> Track = ThreadSharedState->GetCpuTrack(TaskRelationPtr->GetSourceThreadId());
+		if (Track.IsValid())
+		{
+			TaskRelationPtr->SetSourceTrack(Track);
+			TaskRelationPtr->SetSourceDepth(Track->GetDepthAt(TaskRelationPtr->GetSourceTime()) - 1);
+		}
+	}
+
+	if (!TaskRelationPtr->GetTargetTrack().IsValid())
+
+	{
+		TSharedPtr<FCpuTimingTrack> Track = ThreadSharedState->GetCpuTrack(TaskRelationPtr->GetTargetThreadId());
+		if (Track.IsValid())
+		{
+			TaskRelationPtr->SetTargetTrack(Track);
+			TaskRelationPtr->SetTargetDepth(Track->GetDepthAt(TaskRelationPtr->GetTargetTime()) - 1);
+		}
+	}
+
+	if (TaskRelationPtr->GetSourceTrack().IsValid() && TaskRelationPtr->GetTargetTrack().IsValid())
+	{
+		TimingView->AddRelation(Relation);
+	}
+}
+
+////////////////////////////////////////////////////////////////////////////////////////////////////
+
+void FTaskGraphProfilerManager::ClearTaskRelations()
+{
+	TSharedPtr<class STimingProfilerWindow> TimingWindow = FTimingProfilerManager::Get()->GetProfilerWindow();
+	if (!TimingWindow.IsValid())
+	{
+		return;
+	}
+
+	TSharedPtr<STimingView> TimingView = TimingWindow->GetTimingView();
+	if (!TimingView.IsValid())
+	{
+		return;
+	}
+
+	TArray<TUniquePtr<ITimingEventRelation>>& Relations = TimingView->EditCurrentRelations();
+	Relations.RemoveAll([](TUniquePtr<ITimingEventRelation>& Relations)
+		{
+			return Relations->Is<FTaskGraphRelation>();
+		});
+}
+
+////////////////////////////////////////////////////////////////////////////////////////////////////
+
+void FTaskGraphProfilerManager::SetShowRelations(bool bInValue)
+{
+	bShowRelations = bInValue;
+	if (!bShowRelations)
+	{
+		ClearTaskRelations();
+		TaskTimingSharedState->SetTaskId(FTaskTimingTrack::InvalidTaskId);
+	}
+}
+
+////////////////////////////////////////////////////////////////////////////////////////////////////
+
 } // namespace Insights
 
 #undef LOCTEXT_NAMESPACE
diff --git a/Engine/Source/Developer/TraceInsights/Private/Insights/TaskGraphProfiler/TaskGraphProfilerManager.h b/Engine/Source/Developer/TraceInsights/Private/Insights/TaskGraphProfiler/TaskGraphProfilerManager.h
--- a/Engine/Source/Developer/TraceInsights/Private/Insights/TaskGraphProfiler/TaskGraphProfilerManager.h
+++ b/Engine/Source/Developer/TraceInsights/Private/Insights/TaskGraphProfiler/TaskGraphProfilerManager.h
@@ -16,6 +16,8 @@ namespace TraceServices
 	class ITasksProvider;
 }
 
+class FThreadTrackEvent;
+
 namespace Insights
 {
 
@@ -90,20 +92,25 @@ public:
 
 	void OnSessionChanged();
 
-	void GetTaskRelations(double Time, uint32 ThreadId, FAddRelationCallback Callback);
-	void GetTaskRelations(uint32 TaskId, FAddRelationCallback Callback);
+	void ShowTaskRelations(const FThreadTrackEvent* InSelectedEvent, uint32 ThreadId);
+	void ShowTaskRelations(uint32 TaskId);
+	void AddRelation(const FThreadTrackEvent* InSelectedEvent, double SourceTimestamp, uint32 SourceThreadId, double TargetTimestamp, uint32 TargetThreadId, ETaskEventType Type);
+	void ClearTaskRelations();
 	FLinearColor GetColorForTaskEvent(ETaskEventType InEvent);
 
 	TSharedPtr<Insights::FTaskTimingSharedState> GetTaskTimingSharedState() { return TaskTimingSharedState;	}
 
+	bool GetShowRelations() const { return bShowRelations; }
+	void SetShowRelations(bool bInValue);
+
 private:
 	/** Updates this manager, done through FCoreTicker. */
 	bool Tick(float DeltaTime);
 
 	void RegisterTimingProfilerLayoutExtensions(FInsightsMajorTabExtender& InOutExtender);
 
-	void GetTaskRelations(const TraceServices::FTaskInfo* Task, const TraceServices::ITasksProvider* TasksProvider, FAddRelationCallback Callback);
-	void GetSingleTaskRelations(const TraceServices::FTaskInfo* Task, const TraceServices::ITasksProvider* TasksProvider, FAddRelationCallback Callback);
+	void ShowTaskRelations(const TraceServices::FTaskInfo* Task, const TraceServices::ITasksProvider* TasksProvider, const FThreadTrackEvent* InSelectedEvent);
+	void GetSingleTaskRelations(const TraceServices::FTaskInfo* Task, const TraceServices::ITasksProvider* TasksProvider, const FThreadTrackEvent* InSelectedEvent);
 
 	void InitializeColorCode();
 
@@ -128,6 +135,7 @@ private:
 
 	TSharedPtr<Insights::STaskTableTreeView> TaskTableTreeView;
 	FLinearColor ColorCode[static_cast<uint32>(ETaskEventType::NumTaskEventTypes)];
+	bool bShowRelations = true;
 };
 
 } // namespace Insights
diff --git a/Engine/Source/Developer/TraceInsights/Private/Insights/TaskGraphProfiler/ViewModels/TaskTimingTrack.cpp b/Engine/Source/Developer/TraceInsights/Private/Insights/TaskGraphProfiler/ViewModels/TaskTimingTrack.cpp
--- a/Engine/Source/Developer/TraceInsights/Private/Insights/TaskGraphProfiler/ViewModels/TaskTimingTrack.cpp
+++ b/Engine/Source/Developer/TraceInsights/Private/Insights/TaskGraphProfiler/ViewModels/TaskTimingTrack.cpp
@@ -65,6 +65,17 @@ void FTaskTimingSharedState::Tick(Insights::ITimingViewSession& InSession, const
 
 		InSession.AddTopDockedTrack(TaskTrack);
 	}
+
+	if (bResetOnNextTick)
+	{
+		bResetOnNextTick = false;
+		if (!TimingView->GetSelectedEvent().IsValid() && 
+			(!TimingView->GetSelectedTrack().IsValid() || TimingView->GetSelectedTrack().Get() != TaskTrack.Get()))
+		{
+			SetTaskId(FTaskTimingTrack::InvalidTaskId);
+			FTaskGraphProfilerManager::Get()->ClearTaskRelations();
+		}
+	}
 }
 
 ////////////////////////////////////////////////////////////////////////////////////////////////////
@@ -154,10 +165,26 @@ void FTaskTimingTrack::PostDraw(const ITimingTrackDrawContext& Context) const
 
 void FTaskTimingTrack::OnTimingEventSelected(TSharedPtr<const ITimingEvent> InSelectedEvent)
 {
+	if (!FTaskGraphProfilerManager::Get()->GetShowRelations())
+	{
+		return;
+	}
+
+	if ((InSelectedEvent.IsValid() && InSelectedEvent->GetTrack()->Is<FTaskTimingTrack>())
+		|| IsSelected())
+	{
+		// The user has selected a Task Event. Do nothing.
+		return;
+	}
+
 	if (!InSelectedEvent.IsValid() || !InSelectedEvent->Is<FThreadTrackEvent>())
 	{
-		TaskId = InvalidTaskId;
-		SetDirtyFlag();
+		if (TaskId != InvalidTaskId)
+		{
+			TaskId = InvalidTaskId;
+			FTaskGraphProfilerManager::Get()->ClearTaskRelations();
+			SetDirtyFlag();
+		}
 		return;
 	}
 
@@ -177,6 +204,8 @@ void FTaskTimingTrack::OnTimingEventSelected(TSharedPtr<const ITimingEvent> InSe
 	}
 
 	const FThreadTrackEvent &ThreadEvent = InSelectedEvent->As<FThreadTrackEvent>();
+	GetEventRelations(ThreadEvent);
+
 	uint32 ThreadId = StaticCastSharedRef<const FThreadTimingTrack>(ThreadEvent.GetTrack())->GetThreadId();
 	const TraceServices::FTaskInfo* Task = TasksProvider->TryGetTask(ThreadId, ThreadEvent.GetStartTime());
 
@@ -271,6 +300,59 @@ void FTaskTimingTrack::InitTooltip(FTooltipDrawState& InOutTooltip, const ITimin
 
 ////////////////////////////////////////////////////////////////////////////////////////////////////
 
+void FTaskTimingTrack::GetEventRelations(const FThreadTrackEvent& InSelectedEvent)
+{
+	const int32 MaxTasksToShow = 30;
+	double StartTime = InSelectedEvent.GetStartTime();
+
+	TSharedPtr<const TraceServices::IAnalysisSession> Session = FInsightsManager::Get()->GetSession();
+	const TraceServices::ITasksProvider* TasksProvider = TraceServices::ReadTasksProvider(*Session.Get());
+	if (TasksProvider)
+	{
+		STimingView* TimingView = SharedState.GetTimingView();
+		TSharedRef<const FThreadTimingTrack> EventTrack = StaticCastSharedRef<const FThreadTimingTrack>(InSelectedEvent.GetTrack());
+		uint32 ThreadId = EventTrack->GetThreadId();
+		
+		FTaskGraphProfilerManager::Get()->ShowTaskRelations(&InSelectedEvent, ThreadId);
+
+		// if it's an event waiting for tasks completeness, add relations to these tasks
+		const TraceServices::ITimingProfilerProvider& TimingProfilerProvider = *TraceServices::ReadTimingProfilerProvider(*Session.Get());
+		const TraceServices::ITimingProfilerTimerReader* TimerReader;
+		TraceServices::FAnalysisSessionReadScope SessionReadScope(*Session.Get());
+		TimingProfilerProvider.ReadTimers([&TimerReader](const TraceServices::ITimingProfilerTimerReader& Out) { TimerReader = &Out; });
+
+		const FThreadTrackEvent& ThreadTrackEvent = *static_cast<const FThreadTrackEvent*>(&InSelectedEvent);
+		const TraceServices::FTimingProfilerTimer* Timer = TimerReader->GetTimer(ThreadTrackEvent.GetTimerIndex());
+		check(Timer != nullptr);
+
+		const TraceServices::FWaitingForTasks* Waiting = TasksProvider->TryGetWaiting(Timer->Name, ThreadId, StartTime);
+		if (Waiting != nullptr)
+		{
+			int32 NumWaitedTasksToShow = FMath::Min(Waiting->Tasks.Num(), MaxTasksToShow);
+			for (int32 TaskIndex = 0; TaskIndex != NumWaitedTasksToShow; ++TaskIndex)
+			{
+				const TraceServices::FTaskInfo* WaitedTask = TasksProvider->TryGetTask(Waiting->Tasks[TaskIndex]);
+				if (WaitedTask != nullptr)
+				{
+					FTaskGraphProfilerManager::Get()->AddRelation(&InSelectedEvent, StartTime, ThreadId, WaitedTask->StartedTimestamp, WaitedTask->StartedThreadId, ETaskEventType::AddedNested);
+					FTaskGraphProfilerManager::Get()->AddRelation(&InSelectedEvent, WaitedTask->CompletedTimestamp, WaitedTask->CompletedThreadId, WaitedTask->CompletedTimestamp, ThreadId, ETaskEventType::NestedCompleted);
+				}
+			}
+		}
+	}
+}
+
+////////////////////////////////////////////////////////////////////////////////////////////////////
+
+FReply FTaskTimingTrack::OnMouseButtonDown(const FGeometry& MyGeometry, const FPointerEvent& MouseEvent)
+{
+	SharedState.SetResetOnNextTick(true);
+
+	return FReply::Unhandled();
+}
+
+////////////////////////////////////////////////////////////////////////////////////////////////////
+
 } // namespace Insights
 
 #undef LOCTEXT_NAMESPACE
diff --git a/Engine/Source/Developer/TraceInsights/Private/Insights/TaskGraphProfiler/ViewModels/TaskTimingTrack.h b/Engine/Source/Developer/TraceInsights/Private/Insights/TaskGraphProfiler/ViewModels/TaskTimingTrack.h
--- a/Engine/Source/Developer/TraceInsights/Private/Insights/TaskGraphProfiler/ViewModels/TaskTimingTrack.h
+++ b/Engine/Source/Developer/TraceInsights/Private/Insights/TaskGraphProfiler/ViewModels/TaskTimingTrack.h
@@ -8,6 +8,7 @@
 #include "Insights/ITimingViewExtender.h"
 #include "Insights/ViewModels/TimingEventsTrack.h"
 
+class FThreadTrackEvent;
 class STimingView;
 
 ////////////////////////////////////////////////////////////////////////////////////////////////////
@@ -44,10 +45,15 @@ public:
 
 	void SetTaskId(uint32 InTaskId);
 
+	void SetResetOnNextTick(bool bInValue) { bResetOnNextTick = bInValue; }
+
+	STimingView* GetTimingView() { return TimingView; }
+
 private:
 	STimingView* TimingView;
 
 	bool bShowHideTaskTrack;
+	bool bResetOnNextTick = false;
 
 	TSharedPtr<FTaskTimingTrack> TaskTrack;
 };
@@ -59,6 +65,7 @@ class FTaskTimingTrack : public FTimingEventsTrack
 	INSIGHTS_DECLARE_RTTI(FTaskTimingTrack, FTimingEventsTrack)
 
 public:
+	static const uint32 InvalidTaskId;
 
 	struct FPendingEventInfo
 	{
@@ -81,6 +88,7 @@ public:
 
 	virtual void BuildDrawState(ITimingEventsTrackDrawStateBuilder& Builder, const ITimingTrackUpdateContext& Context) override;
 	virtual void BuildFilteredDrawState(ITimingEventsTrackDrawStateBuilder& Builder, const ITimingTrackUpdateContext& Context) override;
+	virtual FReply OnMouseButtonDown(const FGeometry& MyGeometry, const FPointerEvent& MouseEvent) override;
 
 	virtual void PostDraw(const ITimingTrackDrawContext& Context) const override;
 
@@ -92,9 +100,9 @@ public:
 	uint32 GetTaskId() { return TaskId; }
 	
 	void OnTimingEventSelected(TSharedPtr<const ITimingEvent> InSelectedEvent);
+	void GetEventRelations(const FThreadTrackEvent& InSelectedEvent);
 
 private:
-	static const uint32 InvalidTaskId;
 	uint32 TimelineIndex;
 
 	FTaskTimingSharedState& SharedState;
diff --git a/Engine/Source/Developer/TraceInsights/Private/Insights/TaskGraphProfiler/Widgets/STaskTableTreeView.cpp b/Engine/Source/Developer/TraceInsights/Private/Insights/TaskGraphProfiler/Widgets/STaskTableTreeView.cpp
--- a/Engine/Source/Developer/TraceInsights/Private/Insights/TaskGraphProfiler/Widgets/STaskTableTreeView.cpp
+++ b/Engine/Source/Developer/TraceInsights/Private/Insights/TaskGraphProfiler/Widgets/STaskTableTreeView.cpp
@@ -435,28 +435,14 @@ void STaskTableTreeView::ContextMenu_GoToTask_Execute()
 		return;
 	}
 
-	TimingView->ClearRelations();
-
-	auto AddRelation = [&TimingView](double SourceTimestamp, uint32 SourceThreadId, double TargetTimestamp, uint32 TargetThreadId, ETaskEventType Type)
-	{
-		if (SourceTimestamp == TraceServices::FTaskInfo::InvalidTimestamp || TargetTimestamp == TraceServices::FTaskInfo::InvalidTimestamp)
-		{
-			return;
-		}
-
-		TUniquePtr<ITimingEventRelation> Relation = MakeUnique<FTaskGraphRelation>(SourceTimestamp, SourceThreadId, TargetTimestamp, TargetThreadId, Type);
-
-		TimingView->AddRelation(Relation);
-	};
-
-	FTaskGraphProfilerManager::Get()->GetTaskRelations(SelectedTask->GetTask()->GetId(), AddRelation);
+	FTaskGraphProfilerManager::Get()->ShowTaskRelations(SelectedTask->GetTask()->GetId());
 
 	double Duration = (SelectedTask->GetTask()->GetFinishedTimestamp() - SelectedTask->GetTask()->GetCreatedTimestamp()) * 1.5;
 	TimingView->ZoomOnTimeInterval(SelectedTask->GetTask()->GetCreatedTimestamp() - Duration * 0.15, Duration);
 
 	TSharedPtr<FTaskTimingSharedState> TaskSharedState = FTaskGraphProfilerManager::Get()->GetTaskTimingSharedState();
 
-	if (TaskSharedState.IsValid())
+	if (TaskSharedState.IsValid() && FTaskGraphProfilerManager::Get()->GetShowRelations())
 	{
 		TaskSharedState->SetTaskId(SelectedTask->GetTask()->GetId());
 	}
diff --git a/Engine/Source/Developer/TraceInsights/Private/Insights/ViewModels/TaskGraphRelation.cpp b/Engine/Source/Developer/TraceInsights/Private/Insights/ViewModels/TaskGraphRelation.cpp
--- a/Engine/Source/Developer/TraceInsights/Private/Insights/ViewModels/TaskGraphRelation.cpp
+++ b/Engine/Source/Developer/TraceInsights/Private/Insights/ViewModels/TaskGraphRelation.cpp
@@ -8,6 +8,7 @@
 
 namespace Insights
 {
+INSIGHTS_IMPLEMENT_RTTI(FTaskGraphRelation)
 
 FTaskGraphRelation::FTaskGraphRelation(double InSourceTime, int32 InSourceThreadId, double InTargetTime, int32 InTargetThreadId, ETaskEventType InType)
 {
diff --git a/Engine/Source/Developer/TraceInsights/Private/Insights/ViewModels/TaskGraphRelation.h b/Engine/Source/Developer/TraceInsights/Private/Insights/ViewModels/TaskGraphRelation.h
--- a/Engine/Source/Developer/TraceInsights/Private/Insights/ViewModels/TaskGraphRelation.h
+++ b/Engine/Source/Developer/TraceInsights/Private/Insights/ViewModels/TaskGraphRelation.h
@@ -12,11 +12,11 @@ namespace Insights
 
 class FTaskGraphRelation : public ITimingEventRelation
 {
+	INSIGHTS_DECLARE_RTTI(FTaskGraphRelation, ITimingEventRelation)
+
 public:
 	FTaskGraphRelation(double InSourceTime, int32 InSourceThreadId, double InTargetTime, int32 InTargetThreadId, ETaskEventType InType);
 
-	virtual bool IsSolved() override { return bIsSolved; }
-
 	virtual void Draw(const FDrawContext& DrawContext, const FTimingTrackViewport& Viewport, const ITimingViewDrawHelper& Helper, const ITimingEventRelation::EDrawFilter Filter) override;
 
 	void SetSourceTrack(TSharedPtr<const FBaseTimingTrack> InSourceTrack) { SourceTrack = InSourceTrack; }
@@ -35,7 +35,6 @@ public:
 	void SetTargetDepth(int32 InDepth) { TargetDepth = InDepth; }
 	int32 GetTargetDepth() { return TargetDepth; }
 
-	void SetIsSolved(bool InValue) { bIsSolved = InValue; }
 	ETaskEventType GetType() { return Type; }
 
 private:
@@ -47,8 +46,6 @@ private:
 	int32 TargetDepth = 0;
 	ETaskEventType Type;
 
-	bool bIsSolved = false;
-
 	TWeakPtr<const FBaseTimingTrack> SourceTrack;
 	TWeakPtr<const FBaseTimingTrack> TargetTrack;
 };
diff --git a/Engine/Source/Developer/TraceInsights/Private/Insights/ViewModels/ThreadTimingTrack.cpp b/Engine/Source/Developer/TraceInsights/Private/Insights/ViewModels/ThreadTimingTrack.cpp
--- a/Engine/Source/Developer/TraceInsights/Private/Insights/ViewModels/ThreadTimingTrack.cpp
+++ b/Engine/Source/Developer/TraceInsights/Private/Insights/ViewModels/ThreadTimingTrack.cpp
@@ -16,11 +16,9 @@
 #include "Insights/Common/TimeUtils.h"
 #include "Insights/InsightsManager.h"
 #include "Insights/ITimingViewSession.h"
-#include "Insights/TaskGraphProfiler/TaskGraphProfilerManager.h"
 #include "Insights/TimingProfilerManager.h"
 #include "Insights/ViewModels/Filters.h"
 #include "Insights/ViewModels/FilterConfigurator.h"
-#include "Insights/ViewModels/TaskGraphRelation.h"
 #include "Insights/ViewModels/TimerNode.h"
 #include "Insights/ViewModels/ThreadTrackEvent.h"
 #include "Insights/ViewModels/TimingEventSearch.h"
@@ -1506,104 +1504,6 @@ bool FThreadTimingTrack::HasCustomFilter() const
 
 ////////////////////////////////////////////////////////////////////////////////////////////////////
 
-void FThreadTimingTrack::GetEventRelations(const ITimingEvent& InSelectedEvent, TArray<TUniquePtr<ITimingEventRelation>>& Relations) const
-{
-	const int32 MaxTasksToShow = 30;
-	double StartTime = InSelectedEvent.GetStartTime();
-
-	TSharedPtr<const TraceServices::IAnalysisSession> Session = FInsightsManager::Get()->GetSession();
-	const TraceServices::ITasksProvider* TasksProvider = TraceServices::ReadTasksProvider(*Session.Get());
-	if (TasksProvider)
-	{
-		auto AddRelation = [this, &InSelectedEvent, &Relations](double SourceTimestamp, uint32 SourceThreadId, double TargetTimestamp, uint32 TargetThreadId, ETaskEventType Type)
-		{
-			if (SourceTimestamp == TraceServices::FTaskInfo::InvalidTimestamp || TargetTimestamp == TraceServices::FTaskInfo::InvalidTimestamp)
-			{
-				return;
-			}
-
-			TUniquePtr<FTaskGraphRelation> Relation = MakeUnique<FTaskGraphRelation>(SourceTimestamp, SourceThreadId, TargetTimestamp, TargetThreadId, Type);
-
-			if (Relation->GetSourceThreadId() == ThreadId)
-			{
-				Relation->SetSourceTrack(SharedThis(this));
-				Relation->SetSourceDepth(InSelectedEvent.GetDepth());
-			}
-			if (Relation->GetTargetThreadId() == ThreadId)
-			{
-				Relation->SetTargetTrack(SharedThis(this));
-				Relation->SetTargetDepth(InSelectedEvent.GetDepth());
-			}
-			if (Relation->GetSourceThreadId() == ThreadId && Relation->GetTargetThreadId() == ThreadId)
-			{
-				Relation->SetIsSolved(true);
-			}
-
-			Relations.Add(MoveTemp(Relation));
-		};
-
-		FTaskGraphProfilerManager::Get()->GetTaskRelations(InSelectedEvent.GetStartTime(), ThreadId, AddRelation);
-
-		// if it's an event waiting for tasks completeness, add relations to these tasks
-		const TraceServices::ITimingProfilerProvider& TimingProfilerProvider = *TraceServices::ReadTimingProfilerProvider(*Session.Get());
-		const TraceServices::ITimingProfilerTimerReader* TimerReader;
-		TraceServices::FAnalysisSessionReadScope SessionReadScope(*Session.Get());
-		TimingProfilerProvider.ReadTimers([&TimerReader](const TraceServices::ITimingProfilerTimerReader& Out) { TimerReader = &Out; });
-
-		const FThreadTrackEvent& ThreadTrackEvent = *static_cast<const FThreadTrackEvent*>(&InSelectedEvent);
-		const TraceServices::FTimingProfilerTimer* Timer = TimerReader->GetTimer(ThreadTrackEvent.GetTimerIndex());
-		check(Timer != nullptr);
-
-		const TraceServices::FWaitingForTasks* Waiting = TasksProvider->TryGetWaiting(Timer->Name, ThreadId, StartTime);
-		if (Waiting != nullptr)
-		{
-			int32 NumWaitedTasksToShow = FMath::Min(Waiting->Tasks.Num(), MaxTasksToShow);
-			for (int32 i = 0; i != NumWaitedTasksToShow; ++i)
-			{
-				const TraceServices::FTaskInfo* WaitedTask = TasksProvider->TryGetTask(Waiting->Tasks[i]);
-				if (WaitedTask != nullptr)
-				{
-					AddRelation(StartTime, ThreadId, WaitedTask->StartedTimestamp, WaitedTask->StartedThreadId, ETaskEventType::AddedNested);
-					AddRelation(WaitedTask->CompletedTimestamp, WaitedTask->CompletedThreadId, WaitedTask->CompletedTimestamp, ThreadId, ETaskEventType::NestedCompleted);
-				}
-			}
-		}
-	}
-}
-
-////////////////////////////////////////////////////////////////////////////////////////////////////
-
-void FThreadTimingTrack::SolveEventRelations(const TArray<TUniquePtr<ITimingEventRelation>>& Relations) const
-{
-	for (const TUniquePtr<ITimingEventRelation>& Relation : Relations)
-	{
-		if (!Relation->IsSolved())
-		{
-			// A check that we have the correct type will need to be added
-			FTaskGraphRelation* TaskRelation = StaticCast<FTaskGraphRelation*>(Relation.Get());
-
-			if (TaskRelation->GetSourceThreadId() == ThreadId && !TaskRelation->GetSourceTrack())
-			{
-				TaskRelation->SetSourceTrack(StaticCastSharedRef<const FBaseTimingTrack>(SharedThis(this)));
-				TaskRelation->SetSourceDepth(GetDepthAt(TaskRelation->GetSourceTime()) - 1);
-			}
-
-			if (TaskRelation->GetTargetThreadId() == ThreadId && !TaskRelation->GetTargetTrack())
-			{
-				TaskRelation->SetTargetTrack(StaticCastSharedRef<const FBaseTimingTrack>(SharedThis(this)));
-				TaskRelation->SetTargetDepth(GetDepthAt(TaskRelation->GetTargetTime()) - 1);
-			}
-
-			if (TaskRelation->GetSourceTrack() && TaskRelation->GetTargetTrack())
-			{
-				TaskRelation->SetIsSolved(true);
-			}
-		}
-	}
-}
-
-////////////////////////////////////////////////////////////////////////////////////////////////////
-
 int32 FThreadTimingTrack::GetDepthAt(double Time) const
 {
 	int32 Depth = 0;
diff --git a/Engine/Source/Developer/TraceInsights/Private/Insights/ViewModels/ThreadTimingTrack.h b/Engine/Source/Developer/TraceInsights/Private/Insights/ViewModels/ThreadTimingTrack.h
--- a/Engine/Source/Developer/TraceInsights/Private/Insights/ViewModels/ThreadTimingTrack.h
+++ b/Engine/Source/Developer/TraceInsights/Private/Insights/ViewModels/ThreadTimingTrack.h
@@ -160,12 +160,10 @@ public:
 	virtual void OnClipboardCopyEvent(const ITimingEvent& InSelectedEvent) const override;
 	virtual void BuildContextMenu(FMenuBuilder& MenuBuilder) override;
 
-	virtual void GetEventRelations(const ITimingEvent& InSelectedEvent, TArray<TUniquePtr<ITimingEventRelation>>& Relations) const override;
-	virtual void SolveEventRelations(const TArray<TUniquePtr<ITimingEventRelation>>& Relations) const override;
+	int32 GetDepthAt(double Time) const;
 
 protected:
 	virtual bool HasCustomFilter() const override;
-	int32 GetDepthAt(double Time) const;
 
 private:
 	void DrawSelectedEventInfo(const FThreadTrackEvent& SelectedEvent, const FTimingTrackViewport& Viewport, const FDrawContext& DrawContext, const FSlateBrush* WhiteBrush, const FSlateFontInfo& Font) const;
diff --git a/Engine/Source/Developer/TraceInsights/Private/Insights/ViewModels/TimingEvent.cpp b/Engine/Source/Developer/TraceInsights/Private/Insights/ViewModels/TimingEvent.cpp
--- a/Engine/Source/Developer/TraceInsights/Private/Insights/ViewModels/TimingEvent.cpp
+++ b/Engine/Source/Developer/TraceInsights/Private/Insights/ViewModels/TimingEvent.cpp
@@ -7,11 +7,12 @@
 #define LOCTEXT_NAMESPACE "TimingEvent"
 
 ////////////////////////////////////////////////////////////////////////////////////////////////////
-// ITimingEvent, FTimingEvent
+// ITimingEvent, FTimingEvent, ITimingEventRelation
 ////////////////////////////////////////////////////////////////////////////////////////////////////
 
 INSIGHTS_IMPLEMENT_RTTI(ITimingEvent)
 INSIGHTS_IMPLEMENT_RTTI(FTimingEvent)
+INSIGHTS_IMPLEMENT_RTTI(ITimingEventRelation)
 
 ////////////////////////////////////////////////////////////////////////////////////////////////////
 
diff --git a/Engine/Source/Developer/TraceInsights/Private/Insights/ViewModels/TimingEventsTrack.cpp b/Engine/Source/Developer/TraceInsights/Private/Insights/ViewModels/TimingEventsTrack.cpp
--- a/Engine/Source/Developer/TraceInsights/Private/Insights/ViewModels/TimingEventsTrack.cpp
+++ b/Engine/Source/Developer/TraceInsights/Private/Insights/ViewModels/TimingEventsTrack.cpp
@@ -148,8 +148,6 @@ void FTimingEventsTrack::PreUpdate(const ITimingTrackUpdateContext& Context)
 		SetNumLanes(MaxDepth + 1);
 	}
 
-	SolveEventRelations(Context.GetCurrentRelations());
-
 	UpdateTrackHeight(Context);
 }
 
diff --git a/Engine/Source/Developer/TraceInsights/Private/Insights/ViewModels/TimingViewDrawHelper.cpp b/Engine/Source/Developer/TraceInsights/Private/Insights/ViewModels/TimingViewDrawHelper.cpp
--- a/Engine/Source/Developer/TraceInsights/Private/Insights/ViewModels/TimingViewDrawHelper.cpp
+++ b/Engine/Source/Developer/TraceInsights/Private/Insights/ViewModels/TimingViewDrawHelper.cpp
@@ -660,10 +660,7 @@ void FTimingViewDrawHelper::DrawRelations(const TArray<TUniquePtr<ITimingEventRe
 {
 	for (const TUniquePtr<ITimingEventRelation> &Relation : Relations)
 	{
-		if (Relation->IsSolved())
-		{
-			Relation->Draw(DrawContext, Viewport, *this, Filter);
-		}
+		Relation->Draw(DrawContext, Viewport, *this, Filter);
 	}
 }
 
diff --git a/Engine/Source/Developer/TraceInsights/Private/Insights/Widgets/STimingView.cpp b/Engine/Source/Developer/TraceInsights/Private/Insights/Widgets/STimingView.cpp
--- a/Engine/Source/Developer/TraceInsights/Private/Insights/Widgets/STimingView.cpp
+++ b/Engine/Source/Developer/TraceInsights/Private/Insights/Widgets/STimingView.cpp
@@ -117,8 +117,6 @@ STimingView::~STimingView()
 	ForegroundTracks.Reset();
 
 	SelectedEvent.Reset();
-	UpdateEventRelations();
-
 	for (Insights::ITimingViewExtender* Extender : GetExtenders())
 	{
 		Extender->OnEndSession(*this);
@@ -3521,7 +3519,6 @@ void STimingView::OnSelectedTimingEventChanged()
 		}
 	}
 
-	UpdateEventRelations();
 	OnSelectedEventChangedDelegate.Broadcast(SelectedEvent);
 }
 
@@ -4170,8 +4167,11 @@ TArray<Insights::ITimingViewExtender*> STimingView::GetExtenders() const
 
 void STimingView::ContextMenu_ShowTaskDependecies_Execute()
 {
-	bShowEventRelations = !bShowEventRelations;
-	UpdateEventRelations();
+	TSharedPtr<Insights::FTaskGraphProfilerManager> TaskGraphManager = Insights::FTaskGraphProfilerManager::Get();
+	if (TaskGraphManager.IsValid() && TaskGraphManager->GetIsAvailable())
+	{
+		TaskGraphManager->SetShowRelations(!TaskGraphManager->GetShowRelations());
+	}
 }
 
 ////////////////////////////////////////////////////////////////////////////////////////////////////
@@ -4185,19 +4185,8 @@ bool STimingView::ContextMenu_ShowTaskDependecies_CanExecute()
 
 bool STimingView::ContextMenu_ShowTaskDependecies_IsChecked()
 {
-	return Insights::FTaskGraphProfilerManager::Get().IsValid() && Insights::FTaskGraphProfilerManager::Get()->GetIsAvailable() && bShowEventRelations;
-}
-
-////////////////////////////////////////////////////////////////////////////////////////////////////
-
-void STimingView::UpdateEventRelations()
-{
-	ClearRelations();
-
-	if (bShowEventRelations && SelectedEvent.IsValid())
-	{
-		SelectedEvent->GetTrack()->GetEventRelations(*SelectedEvent, CurrentRelations);
-	}
+	TSharedPtr<Insights::FTaskGraphProfilerManager> TaskGraphManager = Insights::FTaskGraphProfilerManager::Get();
+	return TaskGraphManager.IsValid() && TaskGraphManager->GetIsAvailable() && TaskGraphManager->GetShowRelations();
 }
 
 ////////////////////////////////////////////////////////////////////////////////////////////////////
diff --git a/Engine/Source/Developer/TraceInsights/Private/Insights/Widgets/STimingView.h b/Engine/Source/Developer/TraceInsights/Private/Insights/Widgets/STimingView.h
--- a/Engine/Source/Developer/TraceInsights/Private/Insights/Widgets/STimingView.h
+++ b/Engine/Source/Developer/TraceInsights/Private/Insights/Widgets/STimingView.h
@@ -345,6 +345,7 @@ public:
 	const TSharedPtr<FBaseTimingTrack> GetTrackAt(float InPosX, float InPosY) const;
 
 	const TArray<TUniquePtr<ITimingEventRelation>>& GetCurrentRelations() const { return CurrentRelations; }
+	TArray<TUniquePtr<ITimingEventRelation>>& EditCurrentRelations() { return CurrentRelations; }
 	void AddRelation(TUniquePtr<ITimingEventRelation>& Relation) { CurrentRelations.Add(MoveTemp(Relation)); }
 	void ClearRelations();
 
@@ -436,8 +437,6 @@ protected:
 	bool ContextMenu_ShowTaskDependecies_CanExecute();
 	bool ContextMenu_ShowTaskDependecies_IsChecked();
 
-	void UpdateEventRelations();
-
 protected:
 	/** The track's viewport. Encapsulates info about position and scale. */
 	FTimingTrackViewport Viewport;
@@ -631,5 +630,4 @@ protected:
 	TSharedPtr<FUICommandList> CommandList;
 
 	TArray<TUniquePtr<ITimingEventRelation>> CurrentRelations;
-	bool bShowEventRelations = true;
 };
\ No newline at end of file
diff --git a/Engine/Source/Developer/TraceInsights/Public/Insights/ViewModels/BaseTimingTrack.h b/Engine/Source/Developer/TraceInsights/Public/Insights/ViewModels/BaseTimingTrack.h
--- a/Engine/Source/Developer/TraceInsights/Public/Insights/ViewModels/BaseTimingTrack.h
+++ b/Engine/Source/Developer/TraceInsights/Public/Insights/ViewModels/BaseTimingTrack.h
@@ -255,10 +255,6 @@ public:
 	// Called back from the timing view when an event is copied to the clipboard with Ctrl+C.
 	virtual void OnClipboardCopyEvent(const ITimingEvent& InSelectedEvent) const {}
 
-	virtual void GetEventRelations(const ITimingEvent& InSelectedEvent, TArray<TUniquePtr<ITimingEventRelation>>& Dependencies) const {}
-
-	virtual void SolveEventRelations(const TArray<TUniquePtr<ITimingEventRelation>>& Dependencies) const {}
-
 	// Returns number of text lines needed to display the debug string.
 	//TODO: virtual int GetDebugStringLineCount() const { return 0; }
 
diff --git a/Engine/Source/Developer/TraceInsights/Public/Insights/ViewModels/ITimingEvent.h b/Engine/Source/Developer/TraceInsights/Public/Insights/ViewModels/ITimingEvent.h
--- a/Engine/Source/Developer/TraceInsights/Public/Insights/ViewModels/ITimingEvent.h
+++ b/Engine/Source/Developer/TraceInsights/Public/Insights/ViewModels/ITimingEvent.h
@@ -69,6 +69,8 @@ public:
 
 class ITimingEventRelation
 {
+	INSIGHTS_DECLARE_RTTI_BASE(ITimingEventRelation)
+
 public:
 	enum class EDrawFilter
 	{
@@ -78,7 +80,6 @@ public:
 
 	ITimingEventRelation() {}
 	virtual ~ITimingEventRelation() {}
-	virtual bool IsSolved() = 0;
 	virtual void Draw(const FDrawContext& DrawContext, const FTimingTrackViewport& Viewport, const ITimingViewDrawHelper& Helper, const EDrawFilter Filter) = 0;
 };
 
