[Insights] Event tooltips for task timeline events.

diff --git a/Engine/Source/Developer/TraceInsights/Private/Insights/TaskGraphProfiler/ViewModels/TaskTimingTrack.cpp b/Engine/Source/Developer/TraceInsights/Private/Insights/TaskGraphProfiler/ViewModels/TaskTimingTrack.cpp
--- a/Engine/Source/Developer/TraceInsights/Private/Insights/TaskGraphProfiler/ViewModels/TaskTimingTrack.cpp
+++ b/Engine/Source/Developer/TraceInsights/Private/Insights/TaskGraphProfiler/ViewModels/TaskTimingTrack.cpp
@@ -5,8 +5,10 @@
 #include "TraceServices/Model/TasksProfiler.h"
 
 // Insights
+#include "Insights/Common/TimeUtils.h"
 #include "Insights/ITimingViewSession.h"
 #include "Insights/TaskGraphProfiler/TaskGraphProfilerManager.h"
+#include "Insights/TaskGraphProfiler/ViewModels/TaskTrackEvent.h"
 #include "Insights/ViewModels/TaskGraphRelation.h"
 #include "Insights/ViewModels/ThreadTimingTrack.h"
 #include "Insights/ViewModels/ThreadTrackEvent.h"
@@ -225,7 +227,7 @@ void FTaskTimingTrack::OnTimingEventSelected(TSharedPtr<const ITimingEvent> InSe
 
 const TSharedPtr<const ITimingEvent> FTaskTimingTrack::GetEvent(float InPosX, float InPosY, const FTimingTrackViewport& Viewport) const
 {
-	TSharedPtr<FTimingEvent> TimingEvent;
+	TSharedPtr<FTaskTrackEvent> TimingEvent;
 
 	if (TaskId == InvalidTaskId)
 	{
@@ -266,23 +268,28 @@ const TSharedPtr<const ITimingEvent> FTaskTimingTrack::GetEvent(float InPosX, fl
 		// This logic will be replaced
 		if (EventTime >= Task->CreatedTimestamp && EventTime < Task->LaunchedTimestamp)
 		{
-			TimingEvent = MakeShared<FTimingEvent>(SharedThis(this), Task->CreatedTimestamp, Task->LaunchedTimestamp, 0);
+			TimingEvent = MakeShared<FTaskTrackEvent>(SharedThis(this), Task->CreatedTimestamp, Task->LaunchedTimestamp, 0, ETaskTrackEventType::Launched);
 		}
 		else if (EventTime >= Task->LaunchedTimestamp && EventTime < Task->ScheduledTimestamp)
 		{
-			TimingEvent = MakeShared<FTimingEvent>(SharedThis(this), Task->LaunchedTimestamp, Task->ScheduledTimestamp, 0);
+			TimingEvent = MakeShared<FTaskTrackEvent>(SharedThis(this), Task->LaunchedTimestamp, Task->ScheduledTimestamp, 0, ETaskTrackEventType::Dispatched);
 		}
 		else if (EventTime >= Task->ScheduledTimestamp && EventTime < Task->StartedTimestamp)
 		{
-			TimingEvent = MakeShared<FTimingEvent>(SharedThis(this), Task->ScheduledTimestamp, Task->StartedTimestamp, 0);
+			TimingEvent = MakeShared<FTaskTrackEvent>(SharedThis(this), Task->ScheduledTimestamp, Task->StartedTimestamp, 0, ETaskTrackEventType::Scheduled);
 		}
 		else if (EventTime >= Task->StartedTimestamp && EventTime < Task->FinishedTimestamp)
 		{
-			TimingEvent = MakeShared<FTimingEvent>(SharedThis(this), Task->StartedTimestamp, Task->FinishedTimestamp, 0);
+			TimingEvent = MakeShared<FTaskTrackEvent>(SharedThis(this), Task->StartedTimestamp, Task->FinishedTimestamp, 0, ETaskTrackEventType::Executed);
 		}
 		else if (EventTime >= Task->FinishedTimestamp && EventTime < Task->CompletedTimestamp)
 		{
-			TimingEvent = MakeShared<FTimingEvent>(SharedThis(this), Task->FinishedTimestamp, Task->CompletedTimestamp, 0);
+			TimingEvent = MakeShared<FTaskTrackEvent>(SharedThis(this), Task->FinishedTimestamp, Task->CompletedTimestamp, 0, ETaskTrackEventType::Completed);
+		}
+
+		if (TimingEvent.IsValid())
+		{
+			TimingEvent->SetTaskId(Task->Id);
 		}
 	}
 
@@ -295,7 +302,65 @@ void FTaskTimingTrack::InitTooltip(FTooltipDrawState& InOutTooltip, const ITimin
 {
 	InOutTooltip.ResetContent();
 
-	// The rest of the implemention will follow
+	if (!InTooltipEvent.CheckTrack(this) || !InTooltipEvent.Is<FTaskTrackEvent>())
+	{
+		return;
+	}
+
+	const FTaskTrackEvent& TaskTrackEvent = InTooltipEvent.As<FTaskTrackEvent>();
+	InOutTooltip.AddTitle(TaskTrackEvent.GetEventName());
+
+	InOutTooltip.AddNameValueTextLine(TaskTrackEvent.GetStartLabel(), TimeUtils::FormatTimeAuto(TaskTrackEvent.GetStartTime(), 6));
+	InOutTooltip.AddNameValueTextLine(TaskTrackEvent.GetEndLabel(), TimeUtils::FormatTimeAuto(TaskTrackEvent.GetEndTime(), 6));
+	InOutTooltip.AddNameValueTextLine(TEXT("Duration:"), TimeUtils::FormatTimeAuto(TaskTrackEvent.GetEndTime() - TaskTrackEvent.GetStartTime(), 6));
+
+	TSharedPtr<const TraceServices::IAnalysisSession> Session = FInsightsManager::Get()->GetSession();
+	if (!Session.IsValid())
+	{
+		return;
+	}
+
+	TraceServices::FAnalysisSessionReadScope SessionReadScope(*Session.Get());
+
+	const TraceServices::ITasksProvider* TasksProvider = TraceServices::ReadTasksProvider(*Session.Get());
+
+	if (TasksProvider == nullptr)
+	{
+		return;
+	}
+
+	const TraceServices::FTaskInfo* Task = TasksProvider->TryGetTask(TaskTrackEvent.GetTaskId());
+
+	if (Task == nullptr)
+	{
+		return;
+	}
+
+	InOutTooltip.AddNameValueTextLine(TEXT("Task Id:"), FString::Printf(TEXT("%d"), Task->Id));
+
+	switch (TaskTrackEvent.GetTaskEventType())
+	{
+	case ETaskTrackEventType::Launched:
+		break;
+	case ETaskTrackEventType::Dispatched:
+	{
+		InOutTooltip.AddNameValueTextLine(TEXT("Prerequisite tasks:"), FString::Printf(TEXT("%d"), Task->Prerequisites.Num()));
+		break;
+	}
+	case ETaskTrackEventType::Scheduled:
+		break;
+	case ETaskTrackEventType::Executed:
+		InOutTooltip.AddNameValueTextLine(TEXT("Nested tasks:"), FString::Printf(TEXT("%d"), Task->NestedTasks.Num()));
+		break;
+	case ETaskTrackEventType::Completed:
+		InOutTooltip.AddNameValueTextLine(TEXT("Subsequent tasks:"), FString::Printf(TEXT("%d"), Task->Subsequents.Num()));
+		break;
+	default:
+		checkf(false, TEXT("Unknown task event type"));
+		break;
+	}
+
+	InOutTooltip.UpdateLayout();
 }
 
 ////////////////////////////////////////////////////////////////////////////////////////////////////
diff --git a/Engine/Source/Developer/TraceInsights/Private/Insights/TaskGraphProfiler/ViewModels/TaskTrackEvent.cpp b/Engine/Source/Developer/TraceInsights/Private/Insights/TaskGraphProfiler/ViewModels/TaskTrackEvent.cpp
new file mode 100644
--- /dev/null
+++ b/Engine/Source/Developer/TraceInsights/Private/Insights/TaskGraphProfiler/ViewModels/TaskTrackEvent.cpp
@@ -0,0 +1,89 @@
+// Copyright Epic Games, Inc. All Rights Reserved.
+
+#include "Insights/TaskGraphProfiler/ViewModels/TaskTrackEvent.h"
+
+////////////////////////////////////////////////////////////////////////////////////////////////////
+
+INSIGHTS_IMPLEMENT_RTTI(FTaskTrackEvent)
+
+////////////////////////////////////////////////////////////////////////////////////////////////////
+
+FTaskTrackEvent::FTaskTrackEvent(const TSharedRef<const FBaseTimingTrack> InTrack, double InStartTime, double InEndTime, uint32 InDepth, ETaskTrackEventType InType)
+	: FTimingEvent(InTrack, InStartTime, InEndTime, InDepth)
+	, TaskEventType(InType)
+{}
+
+////////////////////////////////////////////////////////////////////////////////////////////////////
+
+FString FTaskTrackEvent::GetStartLabel() const
+{
+	switch (TaskEventType)
+	{
+	case ETaskTrackEventType::Launched:
+		return TEXT("Created Time:");
+	case ETaskTrackEventType::Dispatched:
+		return TEXT("Launched Time");
+	case ETaskTrackEventType::Scheduled:
+		return TEXT("Scheduled Time:");
+	case ETaskTrackEventType::Executed:
+		return TEXT("Started Time:");
+	case ETaskTrackEventType::Completed:
+		return TEXT("Finished Time:");
+	default:
+		checkf(false, TEXT("Unknown task event type"));
+		break;
+	}
+
+	return TEXT("");
+}
+
+////////////////////////////////////////////////////////////////////////////////////////////////////
+
+FString FTaskTrackEvent::GetEndLabel() const
+{
+	switch (TaskEventType)
+	{
+	case ETaskTrackEventType::Launched:
+		return TEXT("Launched Time:");
+	case ETaskTrackEventType::Dispatched:
+		return TEXT("Scheduled Time");
+	case ETaskTrackEventType::Scheduled:
+		return TEXT("Started Time:");
+	case ETaskTrackEventType::Executed:
+		return TEXT("Finished Time:");
+	case ETaskTrackEventType::Completed:
+		return TEXT("Completed Time:");
+	default:
+		checkf(false, TEXT("Unknown task event type"));
+		break;
+	}
+
+	return TEXT("");
+}
+
+////////////////////////////////////////////////////////////////////////////////////////////////////
+
+FString FTaskTrackEvent::GetEventName() const
+{
+	switch (TaskEventType)
+	{
+	case ETaskTrackEventType::Launched:
+		return TEXT("Launched");
+	case ETaskTrackEventType::Dispatched:
+		return TEXT("Dispatched");
+	case ETaskTrackEventType::Scheduled:
+		return TEXT("Scheduled");
+	case ETaskTrackEventType::Executed:
+		return TEXT("Executed");
+	case ETaskTrackEventType::Completed:
+		return TEXT("Completed");
+	default:
+		checkf(false, TEXT("Unknown task event type"));
+		break;
+	}
+
+	return TEXT("");
+}
+
+////////////////////////////////////////////////////////////////////////////////////////////////////
+////////////////////////////////////////////////////////////////////////////////////////////////////
diff --git a/Engine/Source/Developer/TraceInsights/Private/Insights/TaskGraphProfiler/ViewModels/TaskTrackEvent.h b/Engine/Source/Developer/TraceInsights/Private/Insights/TaskGraphProfiler/ViewModels/TaskTrackEvent.h
new file mode 100644
--- /dev/null
+++ b/Engine/Source/Developer/TraceInsights/Private/Insights/TaskGraphProfiler/ViewModels/TaskTrackEvent.h
@@ -0,0 +1,46 @@
+// Copyright Epic Games, Inc. All Rights Reserved.
+
+#pragma once
+
+#include "CoreMinimal.h"
+
+// Insights
+#include "Insights/ViewModels/TimingEvent.h"
+
+////////////////////////////////////////////////////////////////////////////////////////////////////
+
+enum class ETaskTrackEventType : uint32
+{
+	Launched,
+	Dispatched,
+	Scheduled,
+	Executed,
+	Completed,
+};
+
+////////////////////////////////////////////////////////////////////////////////////////////////////
+
+class FTaskTrackEvent : public FTimingEvent
+{
+	INSIGHTS_DECLARE_RTTI(FTaskTrackEvent, FTimingEvent)
+
+public:
+	FTaskTrackEvent(const TSharedRef<const FBaseTimingTrack> InTrack, double InStartTime, double InEndTime, uint32 InDepth, ETaskTrackEventType InTaskEventType);
+	virtual ~FTaskTrackEvent() {}
+
+	FString GetStartLabel() const;
+	FString GetEndLabel() const;
+
+	FString GetEventName() const;
+
+	ETaskTrackEventType GetTaskEventType() const { return TaskEventType; }
+
+	uint32 GetTaskId() const { return TaskId; }
+	void SetTaskId(uint32 InTaskId) { TaskId = InTaskId; }
+
+private:
+	ETaskTrackEventType TaskEventType;
+	uint32 TaskId = ~0;
+};
+
+////////////////////////////////////////////////////////////////////////////////////////////////////
