[Insights] Fix for relations of a task selected in Task Table Tree View dissapearing when panning in TimingView

diff --git a/Engine/Source/Developer/TraceInsights/Private/Insights/TaskGraphProfiler/ViewModels/TaskTimingTrack.cpp b/Engine/Source/Developer/TraceInsights/Private/Insights/TaskGraphProfiler/ViewModels/TaskTimingTrack.cpp
--- a/Engine/Source/Developer/TraceInsights/Private/Insights/TaskGraphProfiler/ViewModels/TaskTimingTrack.cpp
+++ b/Engine/Source/Developer/TraceInsights/Private/Insights/TaskGraphProfiler/ViewModels/TaskTimingTrack.cpp
@@ -411,7 +411,24 @@ void FTaskTimingTrack::GetEventRelations(const FThreadTrackEvent& InSelectedEven
 
 FReply FTaskTimingTrack::OnMouseButtonDown(const FGeometry& MyGeometry, const FPointerEvent& MouseEvent)
 {
-	SharedState.SetResetOnNextTick(true);
+	if (MouseEvent.GetEffectingButton() == EKeys::LeftMouseButton)
+	{
+		MousePositionOnButtonDown = MyGeometry.AbsoluteToLocal(MouseEvent.GetScreenSpacePosition());
+	}
+
+	return FReply::Unhandled();
+}
+
+////////////////////////////////////////////////////////////////////////////////////////////////////
+
+FReply FTaskTimingTrack::OnMouseButtonUp(const FGeometry& MyGeometry, const FPointerEvent& MouseEvent)
+{
+	FVector2D MousePositionOnButtonUp = MyGeometry.AbsoluteToLocal(MouseEvent.GetScreenSpacePosition());
+
+	if (MouseEvent.GetEffectingButton() == EKeys::LeftMouseButton && MousePositionOnButtonUp.Equals(MousePositionOnButtonDown, 2.0f))
+	{
+		SharedState.SetResetOnNextTick(true);
+	}
 
 	return FReply::Unhandled();
 }
diff --git a/Engine/Source/Developer/TraceInsights/Private/Insights/TaskGraphProfiler/ViewModels/TaskTimingTrack.h b/Engine/Source/Developer/TraceInsights/Private/Insights/TaskGraphProfiler/ViewModels/TaskTimingTrack.h
--- a/Engine/Source/Developer/TraceInsights/Private/Insights/TaskGraphProfiler/ViewModels/TaskTimingTrack.h
+++ b/Engine/Source/Developer/TraceInsights/Private/Insights/TaskGraphProfiler/ViewModels/TaskTimingTrack.h
@@ -89,6 +89,7 @@ public:
 	virtual void BuildDrawState(ITimingEventsTrackDrawStateBuilder& Builder, const ITimingTrackUpdateContext& Context) override;
 	virtual void BuildFilteredDrawState(ITimingEventsTrackDrawStateBuilder& Builder, const ITimingTrackUpdateContext& Context) override;
 	virtual FReply OnMouseButtonDown(const FGeometry& MyGeometry, const FPointerEvent& MouseEvent) override;
+	virtual FReply OnMouseButtonUp(const FGeometry& MyGeometry, const FPointerEvent& MouseEvent) override;
 
 	virtual void PostDraw(const ITimingTrackDrawContext& Context) const override;
 
@@ -108,6 +109,8 @@ private:
 	FTaskTimingSharedState& SharedState;
 
 	uint32 TaskId = InvalidTaskId;
+
+	FVector2D MousePositionOnButtonDown;
 };
 
 ////////////////////////////////////////////////////////////////////////////////////////////////////
