[Insights] Fix for Insights menu not apearing on MAC

diff --git a/Engine/Source/Developer/TraceInsights/Private/Insights/TraceInsightsModule.cpp b/Engine/Source/Developer/TraceInsights/Private/Insights/TraceInsightsModule.cpp
--- a/Engine/Source/Developer/TraceInsights/Private/Insights/TraceInsightsModule.cpp
+++ b/Engine/Source/Developer/TraceInsights/Private/Insights/TraceInsightsModule.cpp
@@ -249,7 +249,11 @@ void FTraceInsightsModule::CreateSessionViewer(bool bAllowDebugTools)
 	PersistentLayout = FLayoutSaveRestore::LoadFromConfig(UnrealInsightsLayoutIni, DefaultLayout);
 
 	// Restore application layout.
+#if PLATFORM_MAC
+	const bool bEmbedTitleAreaContent = true;
+#else
 	const bool bEmbedTitleAreaContent = false;
+#endif
 	const EOutputCanBeNullptr OutputCanBeNullptr = EOutputCanBeNullptr::Never;
 	TSharedPtr<SWidget> Content = FGlobalTabmanager::Get()->RestoreFrom(PersistentLayout.ToSharedRef(), RootWindow, bEmbedTitleAreaContent, OutputCanBeNullptr);
 
