[Insights]
  - Memory: Changed the "symbols resolved" message in status bar of Allocs Table to show also the number pending symbols to be resolved.
  - Timing view: Hidden the "Exclusive Time percent in tooltip when event duration is 0 (instead of showing NaN).

diff --git a/Engine/Source/Developer/TraceInsights/Private/Insights/MemoryProfiler/Widgets/SMemAllocTableTreeView.cpp b/Engine/Source/Developer/TraceInsights/Private/Insights/MemoryProfiler/Widgets/SMemAllocTableTreeView.cpp
--- a/Engine/Source/Developer/TraceInsights/Private/Insights/MemoryProfiler/Widgets/SMemAllocTableTreeView.cpp
+++ b/Engine/Source/Developer/TraceInsights/Private/Insights/MemoryProfiler/Widgets/SMemAllocTableTreeView.cpp
@@ -337,6 +337,7 @@ void SMemAllocTableTreeView::UpdateQuery(TraceServices::IAllocationsProvider::EQ
 					Alloc.Size = Allocation->GetSize();
 					Alloc.Tag = Provider.GetTagName(Allocation->GetTag());
 					Alloc.Callstack = Allocation->GetCallstack();
+					check(Alloc.Callstack != nullptr);
 				}
 
 				PageStopwatch.Stop();
@@ -836,7 +837,16 @@ FText SMemAllocTableTreeView::GetSymbolResolutionStatus() const
 	if (ModuleProvider)
 	{
 		ModuleProvider->GetStats(&Stats);
-		return FText::Format(LOCTEXT("SymbolsResolved", "{0} / {1} symbols resolved. {2} failed."), Stats.SymbolsResolved, Stats.SymbolsDiscovered, Stats.SymbolsFailed);
+		check(Stats.SymbolsDiscovered >= Stats.SymbolsResolved + Stats.SymbolsFailed);
+		const uint32 SymbolsPending = Stats.SymbolsDiscovered - Stats.SymbolsResolved - Stats.SymbolsFailed;
+		if (SymbolsPending != 0)
+		{
+			return FText::Format(LOCTEXT("SymbolsResolved1", "{0} symbols ({1} resolved, {2} failed, {3} pending)"), Stats.SymbolsDiscovered, Stats.SymbolsResolved, Stats.SymbolsFailed, SymbolsPending);
+		}
+		else
+		{
+			return FText::Format(LOCTEXT("SymbolsResolved2", "{0} symbols ({1} resolved, {2} failed)"), Stats.SymbolsDiscovered, Stats.SymbolsResolved, Stats.SymbolsFailed);
+		}
 	}
 	else
 	{
diff --git a/Engine/Source/Developer/TraceInsights/Private/Insights/ViewModels/ThreadTimingTrack.cpp b/Engine/Source/Developer/TraceInsights/Private/Insights/ViewModels/ThreadTimingTrack.cpp
--- a/Engine/Source/Developer/TraceInsights/Private/Insights/ViewModels/ThreadTimingTrack.cpp
+++ b/Engine/Source/Developer/TraceInsights/Private/Insights/ViewModels/ThreadTimingTrack.cpp
@@ -1011,10 +1011,12 @@ void FThreadTimingTrack::InitTooltip(FTooltipDrawState& InOutTooltip, const ITim
 
 		InOutTooltip.AddNameValueTextLine(TEXT("Inclusive Time:"), TimeUtils::FormatTimeAuto(TooltipEvent.GetDuration()));
 
+		if (TooltipEvent.GetDuration() > 0.0)
 		{
+			const double ExclusiveTimePercent = TooltipEvent.GetExclusiveTime() / TooltipEvent.GetDuration();
 			FNumberFormattingOptions FormattingOptions;
 			FormattingOptions.MaximumFractionalDigits = 2;
-			const FString ExclStr = FString::Printf(TEXT("%s (%s)"), *TimeUtils::FormatTimeAuto(TooltipEvent.GetExclusiveTime()), *FText::AsPercent(TooltipEvent.GetExclusiveTime() / TooltipEvent.GetDuration(), &FormattingOptions).ToString());
+			const FString ExclStr = FString::Printf(TEXT("%s (%s)"), *TimeUtils::FormatTimeAuto(TooltipEvent.GetExclusiveTime()), *FText::AsPercent(ExclusiveTimePercent, &FormattingOptions).ToString());
 			InOutTooltip.AddNameValueTextLine(TEXT("Exclusive Time:"), ExclStr);
 		}
 
