[Insights] Updated again the text with resolved symbols to minimize the text flickering.

diff --git a/Engine/Source/Developer/TraceInsights/Private/Insights/MemoryProfiler/Widgets/SMemAllocTableTreeView.cpp b/Engine/Source/Developer/TraceInsights/Private/Insights/MemoryProfiler/Widgets/SMemAllocTableTreeView.cpp
--- a/Engine/Source/Developer/TraceInsights/Private/Insights/MemoryProfiler/Widgets/SMemAllocTableTreeView.cpp
+++ b/Engine/Source/Developer/TraceInsights/Private/Insights/MemoryProfiler/Widgets/SMemAllocTableTreeView.cpp
@@ -841,7 +841,7 @@ FText SMemAllocTableTreeView::GetSymbolResolutionStatus() const
 		const uint32 SymbolsPending = Stats.SymbolsDiscovered - Stats.SymbolsResolved - Stats.SymbolsFailed;
 		if (SymbolsPending != 0)
 		{
-			return FText::Format(LOCTEXT("SymbolsResolved1", "{0} symbols ({1} resolved, {2} failed, {3} pending)"), Stats.SymbolsDiscovered, Stats.SymbolsResolved, Stats.SymbolsFailed, SymbolsPending);
+			return FText::Format(LOCTEXT("SymbolsResolved1", "Resolving {0} / {1} symbols ({2} resolved, {3} failed)"), SymbolsPending, Stats.SymbolsDiscovered, Stats.SymbolsResolved, Stats.SymbolsFailed);
 		}
 		else
 		{
