[Insights] - format memory graph unit amounts with 2decimal places

diff --git a/Engine/Source/Developer/TraceInsights/Private/Insights/MemoryProfiler/ViewModels/MemoryGraphTrack.cpp b/Engine/Source/Developer/TraceInsights/Private/Insights/MemoryProfiler/ViewModels/MemoryGraphTrack.cpp
--- a/Engine/Source/Developer/TraceInsights/Private/Insights/MemoryProfiler/ViewModels/MemoryGraphTrack.cpp
+++ b/Engine/Source/Developer/TraceInsights/Private/Insights/MemoryProfiler/ViewModels/MemoryGraphTrack.cpp
@@ -46,12 +46,12 @@ FString FMemoryGraphSeries::FormatValue(double Value) const
 
 		if (ValueInt64 < 0)
 		{
-			FString Auto = FMemoryGraphTrack::FormatValue(-Value, UnitValue, UnitText, 0);
+			FString Auto = FMemoryGraphTrack::FormatValue(-Value, UnitValue, UnitText, 2);
 			return FString::Printf(TEXT("-%s (%s bytes)"), *Auto, *FText::AsNumber(ValueInt64).ToString());
 		}
 		else
 		{
-			FString Auto = FMemoryGraphTrack::FormatValue(Value, UnitValue, UnitText, 0);
+			FString Auto = FMemoryGraphTrack::FormatValue(Value, UnitValue, UnitText, 2);
 			return FString::Printf(TEXT("%s (%s bytes)"), *Auto, *FText::AsNumber(ValueInt64).ToString());
 		}
 	}
