Changing default LLM tags in Memory Insights to include the common tags that account for the majority of memory.

- Animation, Lumen, Nanite, ProgramSize, RenderTargets, SceneRender, UObject.

diff --git a/Engine/Source/Developer/TraceInsights/Private/Insights/MemoryProfiler/ViewModels/MemorySharedState.cpp b/Engine/Source/Developer/TraceInsights/Private/Insights/MemoryProfiler/ViewModels/MemorySharedState.cpp
--- a/Engine/Source/Developer/TraceInsights/Private/Insights/MemoryProfiler/ViewModels/MemorySharedState.cpp
+++ b/Engine/Source/Developer/TraceInsights/Private/Insights/MemoryProfiler/ViewModels/MemorySharedState.cpp
@@ -227,6 +227,13 @@ void FMemorySharedState::CreateDefaultTracks()
 		TEXT("Textures"),
 		TEXT("Physics"),
 		TEXT("Audio"),
+		TEXT("Animation"),
+		TEXT("Lumen"),
+		TEXT("Nanite"),
+		TEXT("ProgramSize"),
+		TEXT("RenderTargets"),
+		TEXT("SceneRender"),
+		TEXT("UObject")
 	};
 	constexpr int32 DefaultTagCount = UE_ARRAY_COUNT(DefaultTags);
 
