[Insights] - Format memory graph vertex axis units to two decimal places.

diff --git a/Engine/Source/Developer/TraceInsights/Private/Insights/MemoryProfiler/ViewModels/MemoryGraphTrack.cpp b/Engine/Source/Developer/TraceInsights/Private/Insights/MemoryProfiler/ViewModels/MemoryGraphTrack.cpp
--- a/Engine/Source/Developer/TraceInsights/Private/Insights/MemoryProfiler/ViewModels/MemoryGraphTrack.cpp
+++ b/Engine/Source/Developer/TraceInsights/Private/Insights/MemoryProfiler/ViewModels/MemoryGraphTrack.cpp
@@ -30,6 +30,8 @@
 // FMemoryGraphSeries
 ////////////////////////////////////////////////////////////////////////////////////////////////////
 
+constexpr int32 DefaultDecimalDigitCount = 2;
+
 FString FMemoryGraphSeries::FormatValue(double Value) const
 {
 	const int64 ValueInt64 = static_cast<int64>(Value);
@@ -46,12 +48,12 @@ FString FMemoryGraphSeries::FormatValue(double Value) const
 
 		if (ValueInt64 < 0)
 		{
-			FString Auto = FMemoryGraphTrack::FormatValue(-Value, UnitValue, UnitText, 2);
+			FString Auto = FMemoryGraphTrack::FormatValue(-Value, UnitValue, UnitText, DefaultDecimalDigitCount);
 			return FString::Printf(TEXT("-%s (%s bytes)"), *Auto, *FText::AsNumber(ValueInt64).ToString());
 		}
 		else
 		{
-			FString Auto = FMemoryGraphTrack::FormatValue(Value, UnitValue, UnitText, 2);
+			FString Auto = FMemoryGraphTrack::FormatValue(Value, UnitValue, UnitText, DefaultDecimalDigitCount);
 			return FString::Printf(TEXT("%s (%s bytes)"), *Auto, *FText::AsNumber(ValueInt64).ToString());
 		}
 	}
@@ -73,7 +75,7 @@ FMemoryGraphTrack::FMemoryGraphTrack(FMemorySharedState& InSharedState)
 	: FGraphTrack()
 	, SharedState(InSharedState)
 	, LabelUnit(EGraphTrackLabelUnit::Auto)
-	, LabelDecimalDigitCount(-1)
+	, LabelDecimalDigitCount(DefaultDecimalDigitCount)
 	, DefaultMinValue(+std::numeric_limits<double>::infinity())
 	, DefaultMaxValue(-std::numeric_limits<double>::infinity())
 	, AllSeriesMinValue(0.0)
