UnrealInsights: Add Max-Min column to the counters table

diff --git a/Engine/Source/Developer/TraceInsights/Private/Insights/ViewModels/StatsNode.cpp b/Engine/Source/Developer/TraceInsights/Private/Insights/ViewModels/StatsNode.cpp
--- a/Engine/Source/Developer/TraceInsights/Private/Insights/ViewModels/StatsNode.cpp
+++ b/Engine/Source/Developer/TraceInsights/Private/Insights/ViewModels/StatsNode.cpp
@@ -231,4 +231,11 @@ const FText FStatsNode::GetTextForAggregatedStatsUpperQuartile(bool bForTooltip)
 
 ////////////////////////////////////////////////////////////////////////////////////////////////////
 
+const FText FStatsNode::GetTextForAggregatedStatsDiff(bool bForTooltip) const
+{
+	return FormatAggregatedStatsValue(AggregatedStats.DoubleStats.Max - AggregatedStats.DoubleStats.Min, AggregatedStats.Int64Stats.Max - AggregatedStats.Int64Stats.Min, bForTooltip);
+}
+
+////////////////////////////////////////////////////////////////////////////////////////////////////
+
 #undef LOCTEXT_NAMESPACE
diff --git a/Engine/Source/Developer/TraceInsights/Private/Insights/ViewModels/StatsNode.h b/Engine/Source/Developer/TraceInsights/Private/Insights/ViewModels/StatsNode.h
--- a/Engine/Source/Developer/TraceInsights/Private/Insights/ViewModels/StatsNode.h
+++ b/Engine/Source/Developer/TraceInsights/Private/Insights/ViewModels/StatsNode.h
@@ -245,6 +245,7 @@ public:
 	const FText GetTextForAggregatedStatsMedian(bool bForTooltip = false) const;
 	const FText GetTextForAggregatedStatsLowerQuartile(bool bForTooltip = false) const;
 	const FText GetTextForAggregatedStatsUpperQuartile(bool bForTooltip = false) const;
+	const FText GetTextForAggregatedStatsDiff(bool bForTooltip = false) const;
 
 	/** Sorts children using the specified class instance. */
 	template<typename TSortingClass>
diff --git a/Engine/Source/Developer/TraceInsights/Private/Insights/ViewModels/StatsViewColumnFactory.cpp b/Engine/Source/Developer/TraceInsights/Private/Insights/ViewModels/StatsViewColumnFactory.cpp
--- a/Engine/Source/Developer/TraceInsights/Private/Insights/ViewModels/StatsViewColumnFactory.cpp
+++ b/Engine/Source/Developer/TraceInsights/Private/Insights/ViewModels/StatsViewColumnFactory.cpp
@@ -27,6 +27,7 @@ const FName FStatsViewColumns::AverageColumnID(TEXT("Average"));
 const FName FStatsViewColumns::MedianColumnID(TEXT("Median"));
 const FName FStatsViewColumns::LowerQuartileColumnID(TEXT("LowerQuartile"));
 const FName FStatsViewColumns::MinColumnID(TEXT("Min"));
+const FName FStatsViewColumns::DiffColumnID(TEXT("Diff"));
 
 ////////////////////////////////////////////////////////////////////////////////////////////////////
 
@@ -46,6 +47,7 @@ void FStatsViewColumnFactory::CreateStatsViewColumns(TArray<TSharedRef<Insights:
 	Columns.Add(CreateMedianColumn());
 	//Columns.Add(CreateLowerQuartileColumn());
 	Columns.Add(CreateMinColumn());
+	Columns.Add(CreateDiffColumn());
 }
 
 ////////////////////////////////////////////////////////////////////////////////////////////////////
@@ -662,4 +664,57 @@ TSharedRef<Insights::FTableColumn> FStatsViewColumnFactory::CreateMinColumn()
 
 ////////////////////////////////////////////////////////////////////////////////////////////////////
 
+TSharedRef<Insights::FTableColumn> FStatsViewColumnFactory::CreateDiffColumn()
+{
+	using namespace Insights;
+
+	TSharedRef<FTableColumn> ColumnRef = MakeShared<FTableColumn>(FStatsViewColumns::DiffColumnID);
+	FTableColumn& Column = *ColumnRef;
+
+	Column.SetShortName(LOCTEXT("Diff_ColumnName", "Max - Min"));
+	Column.SetTitleName(LOCTEXT("Diff_ColumnTitle", "Maximum - Minimum"));
+	Column.SetDescription(LOCTEXT("Diff_ColumnDesc", "Maximum - Minimum of selected values"));
+
+	Column.SetFlags(ETableColumnFlags::CanBeHidden |
+		ETableColumnFlags::CanBeFiltered);
+
+	Column.SetHorizontalAlignment(HAlign_Right);
+	Column.SetInitialWidth(AggregatedStatsColumnInitialWidth);
+
+	Column.SetDataType(ETableCellDataType::Double);
+
+	class FDiffValueGetter : public FTableCellValueGetter
+	{
+	public:
+		virtual const TOptional<FTableCellValue> GetValue(const FTableColumn& Column, const FBaseTreeNode& Node) const
+		{
+			ensure(Column.GetId() == FStatsViewColumns::DiffColumnID);
+			const FStatsNode& StatsNode = static_cast<const FStatsNode&>(Node);
+			return TOptional<FTableCellValue>(FTableCellValue(StatsNode.GetAggregatedStats().DoubleStats.Max - StatsNode.GetAggregatedStats().DoubleStats.Min));
+		}
+	};
+	TSharedRef<ITableCellValueGetter> Getter = MakeShared<FDiffValueGetter>();
+	Column.SetValueGetter(Getter);
+
+	class FDiffValueFormatter : public FTableCellValueFormatter
+	{
+	public:
+		virtual FText FormatValue(const FTableColumn& Column, const FBaseTreeNode& Node) const override
+		{
+			ensure(Column.GetId() == FStatsViewColumns::DiffColumnID);
+			const FStatsNode& StatsNode = static_cast<const FStatsNode&>(Node);
+			return StatsNode.GetTextForAggregatedStatsDiff();
+		}
+	};
+	TSharedRef<ITableCellValueFormatter> Formatter = MakeShared<FDiffValueFormatter>();
+	Column.SetValueFormatter(Formatter);
+
+	TSharedRef<ITableCellValueSorter> Sorter = MakeShared<FSorterByDoubleValue>(ColumnRef);
+	Column.SetValueSorter(Sorter);
+
+	return ColumnRef;
+}
+
+////////////////////////////////////////////////////////////////////////////////////////////////////
+
 #undef LOCTEXT_NAMESPACE
diff --git a/Engine/Source/Developer/TraceInsights/Private/Insights/ViewModels/StatsViewColumnFactory.h b/Engine/Source/Developer/TraceInsights/Private/Insights/ViewModels/StatsViewColumnFactory.h
--- a/Engine/Source/Developer/TraceInsights/Private/Insights/ViewModels/StatsViewColumnFactory.h
+++ b/Engine/Source/Developer/TraceInsights/Private/Insights/ViewModels/StatsViewColumnFactory.h
@@ -21,6 +21,7 @@ struct FStatsViewColumns
 	static const FName MedianColumnID;
 	static const FName LowerQuartileColumnID;
 	static const FName MinColumnID;
+	static const FName DiffColumnID;
 };
 
 ////////////////////////////////////////////////////////////////////////////////////////////////////
@@ -49,6 +50,7 @@ public:
 	static TSharedRef<Insights::FTableColumn> CreateMedianColumn();
 	static TSharedRef<Insights::FTableColumn> CreateLowerQuartileColumn();
 	static TSharedRef<Insights::FTableColumn> CreateMinColumn();
+	static TSharedRef<Insights::FTableColumn> CreateDiffColumn();
 
 private:
 	static constexpr float AggregatedStatsColumnInitialWidth = 80.0f;
diff --git a/Engine/Source/Developer/TraceInsights/Private/Insights/Widgets/SStatsView.cpp b/Engine/Source/Developer/TraceInsights/Private/Insights/Widgets/SStatsView.cpp
--- a/Engine/Source/Developer/TraceInsights/Private/Insights/Widgets/SStatsView.cpp
+++ b/Engine/Source/Developer/TraceInsights/Private/Insights/Widgets/SStatsView.cpp
@@ -1727,6 +1727,7 @@ void SStatsView::ContextMenu_ShowMinMaxMedColumns_Execute()
 		FStatsViewColumns::MedianColumnID,
 		FStatsViewColumns::LowerQuartileColumnID,
 		FStatsViewColumns::MinColumnID,
+		FStatsViewColumns::DiffColumnID,
 	};
 
 	ColumnBeingSorted = FStatsViewColumns::CountColumnID;
