[Insights] Allow timing view extenders to add options to the  global context menu
+ Move TaskDependencies menu option to TaskTimingSharedState
+ Added 3 new menu options to control the visibility of prerequisite, subsequent and nested task relations

diff --git a/Engine/Source/Developer/TraceInsights/Private/Insights/TaskGraphProfiler/TaskGraphProfilerManager.cpp b/Engine/Source/Developer/TraceInsights/Private/Insights/TaskGraphProfiler/TaskGraphProfilerManager.cpp
--- a/Engine/Source/Developer/TraceInsights/Private/Insights/TaskGraphProfiler/TaskGraphProfilerManager.cpp
+++ b/Engine/Source/Developer/TraceInsights/Private/Insights/TaskGraphProfiler/TaskGraphProfilerManager.cpp
@@ -277,10 +277,19 @@ void FTaskGraphProfilerManager::ShowTaskRelations(const TraceServices::FTaskInfo
 		}
 	};
 
-	GetSingleTaskRelationsForAll(Task->Prerequisites);
+	if (bShowPrerequisites)
+	{
+		GetSingleTaskRelationsForAll(Task->Prerequisites);
+	}
 	GetSingleTaskRelations(Task, TasksProvider, InSelectedEvent);
-	GetSingleTaskRelationsForAll(Task->NestedTasks);
-	GetSingleTaskRelationsForAll(Task->Subsequents);
+	if (bShowNestedTasks)
+	{
+		GetSingleTaskRelationsForAll(Task->NestedTasks);
+	}
+	if (bShowSubsequents)
+	{
+		GetSingleTaskRelationsForAll(Task->Subsequents);
+	}
 }
 
 void FTaskGraphProfilerManager::GetSingleTaskRelations(const TraceServices::FTaskInfo* Task, const TraceServices::ITasksProvider* TasksProvider, const FThreadTrackEvent* InSelectedEvent)
@@ -533,7 +542,6 @@ void FTaskGraphProfilerManager::SetShowRelations(bool bInValue)
 	bShowRelations = bInValue;
 	if (!bShowRelations)
 	{
-		ClearTaskRelations();
 		TaskTimingSharedState->SetTaskId(FTaskTimingTrack::InvalidTaskId);
 	}
 }
diff --git a/Engine/Source/Developer/TraceInsights/Private/Insights/TaskGraphProfiler/TaskGraphProfilerManager.h b/Engine/Source/Developer/TraceInsights/Private/Insights/TaskGraphProfiler/TaskGraphProfilerManager.h
--- a/Engine/Source/Developer/TraceInsights/Private/Insights/TaskGraphProfiler/TaskGraphProfilerManager.h
+++ b/Engine/Source/Developer/TraceInsights/Private/Insights/TaskGraphProfiler/TaskGraphProfilerManager.h
@@ -103,6 +103,15 @@ public:
 	bool GetShowRelations() const { return bShowRelations; }
 	void SetShowRelations(bool bInValue);
 
+	bool GetShowPrerequisites() const { return bShowPrerequisites; }
+	void SetShowPrerequisites(bool bInValue) { bShowPrerequisites = bInValue; }
+
+	bool GetShowSubsequents() const { return bShowSubsequents; }
+	void SetShowSubsequents(bool bInValue) { bShowSubsequents = bInValue; }
+
+	bool GetShowNestedTasks() const { return bShowNestedTasks; }
+	void SetShowNestedTasks(bool bInValue) { bShowNestedTasks = bInValue; }
+
 private:
 	/** Updates this manager, done through FCoreTicker. */
 	bool Tick(float DeltaTime);
@@ -136,6 +145,9 @@ private:
 	TSharedPtr<Insights::STaskTableTreeView> TaskTableTreeView;
 	FLinearColor ColorCode[static_cast<uint32>(ETaskEventType::NumTaskEventTypes)];
 	bool bShowRelations = true;
+	bool bShowPrerequisites = false;
+	bool bShowSubsequents = false;
+	bool bShowNestedTasks = false;
 };
 
 } // namespace Insights
diff --git a/Engine/Source/Developer/TraceInsights/Private/Insights/TaskGraphProfiler/ViewModels/TaskTimingTrack.cpp b/Engine/Source/Developer/TraceInsights/Private/Insights/TaskGraphProfiler/ViewModels/TaskTimingTrack.cpp
--- a/Engine/Source/Developer/TraceInsights/Private/Insights/TaskGraphProfiler/ViewModels/TaskTimingTrack.cpp
+++ b/Engine/Source/Developer/TraceInsights/Private/Insights/TaskGraphProfiler/ViewModels/TaskTimingTrack.cpp
@@ -2,6 +2,10 @@
 
 #include "TaskTimingTrack.h"
 
+#include "EditorStyleSet.h"
+#include "Framework/Commands/Commands.h"
+#include "Framework/Commands/UICommandList.h"
+#include "Framework/MultiBox/MultiBoxBuilder.h"
 #include "TraceServices/Model/TasksProfiler.h"
 
 // Insights
@@ -21,11 +25,50 @@
 
 namespace Insights
 {
+////////////////////////////////////////////////////////////////////////////////////////////////////
+// FTaskTimingStateCommands
+////////////////////////////////////////////////////////////////////////////////////////////////////
+
+class FTaskTimingStateCommands : public TCommands<FTaskTimingStateCommands>
+{
+public:
+	FTaskTimingStateCommands()
+		: TCommands<FTaskTimingStateCommands>(TEXT("FTaskTimingStateCommands"), NSLOCTEXT("FTaskTimingStateCommands", "Task Timing State Commands", "Task Table Tree View Commands"), NAME_None, FEditorStyle::Get().GetStyleSetName())
+	{
+	}
+
+	virtual ~FTaskTimingStateCommands()
+	{
+	}
+
+	// UI_COMMAND takes long for the compiler to optimize
+	PRAGMA_DISABLE_OPTIMIZATION
+	virtual void RegisterCommands() override
+	{
+		UI_COMMAND(Command_ShowTaskDependencies, "Show Task Dependencies ", "Show/hide dependencies of the current task (for a selected cpu timing event)", EUserInterfaceActionType::ToggleButton, FInputChord(EKeys::T));
+		UI_COMMAND(Command_ShowTaskPrerequisites, "Show dependencies for prerequisites ", "Show/hide dependecies of the current task's prerequisites", EUserInterfaceActionType::ToggleButton, FInputChord(EKeys::P));
+		UI_COMMAND(Command_ShowTaskSubsequents, "Show dependencies for subsequents ", "Show/hide dependecies of the current task's subsequents", EUserInterfaceActionType::ToggleButton, FInputChord(EKeys::S));
+		UI_COMMAND(Command_ShowNestedTasks, "Show dependencies for nested tasks ", "Show/hide dependencies of the current task's nested tasks", EUserInterfaceActionType::ToggleButton, FInputChord(EKeys::N));
+	}
+	PRAGMA_ENABLE_OPTIMIZATION
+
+	TSharedPtr<FUICommandInfo> Command_ShowTaskDependencies;
+	TSharedPtr<FUICommandInfo> Command_ShowTaskPrerequisites;
+	TSharedPtr<FUICommandInfo> Command_ShowTaskSubsequents;
+	TSharedPtr<FUICommandInfo> Command_ShowNestedTasks;
+};
 
 ////////////////////////////////////////////////////////////////////////////////////////////////////
 // FTaskTimingSharedState
 ////////////////////////////////////////////////////////////////////////////////////////////////////
 
+FTaskTimingSharedState::FTaskTimingSharedState(STimingView* InTimingView) 
+	: TimingView(InTimingView) 
+{
+}
+
+////////////////////////////////////////////////////////////////////////////////////////////////////
+
 void FTaskTimingSharedState::OnBeginSession(Insights::ITimingViewSession& InSession)
 {
 	if (&InSession != TimingView)
@@ -59,6 +102,8 @@ void FTaskTimingSharedState::Tick(Insights::ITimingViewSession& InSession, const
 
 	if (!TaskTrack.IsValid() && FTaskGraphProfilerManager::Get()->GetIsAvailable())
 	{
+		InitCommandList();
+
 		TaskTrack = MakeShared<FTaskTimingTrack>(*this, TEXT("Task Overview Track"), 0);
 		TaskTrack->SetVisibilityFlag(true);
 		TaskTrack->SetOrder(FTimingTrackOrder::Task);
@@ -100,6 +145,255 @@ void FTaskTimingSharedState::SetTaskId(uint32 InTaskId)
 	}
 }
 
+////////////////////////////////////////////////////////////////////////////////////////////////////
+
+bool FTaskTimingSharedState::ExtendGlobalContextMenu(ITimingViewSession& InSession, FMenuBuilder& InMenuBuilder)
+{
+	if (!CommandList.IsValid())
+	{
+		return false;
+	}
+
+	if (&InSession != TimingView)
+	{
+		return false;
+	}
+
+	InMenuBuilder.PushCommandList(CommandList.ToSharedRef());
+
+	InMenuBuilder.BeginSection(TEXT("Event"), LOCTEXT("Event", "Event"));
+
+	InMenuBuilder.AddSubMenu(LOCTEXT("Tasks", "Tasks"), LOCTEXT("Task", "Task Graph Insights settings"), FNewMenuDelegate::CreateSP(this, &FTaskTimingSharedState::BuildTasksSubMenu));
+
+	InMenuBuilder.EndSection();
+
+	InMenuBuilder.PopCommandList();
+
+	return true;
+}
+
+////////////////////////////////////////////////////////////////////////////////////////////////////
+
+void FTaskTimingSharedState::BuildTasksSubMenu(FMenuBuilder& MenuBuilder)
+{
+	MenuBuilder.AddMenuEntry
+	(
+		FTaskTimingStateCommands::Get().Command_ShowTaskDependencies,
+		NAME_None,
+		TAttribute<FText>(),
+		TAttribute<FText>(),
+		FSlateIcon(FEditorStyle::GetStyleSetName(), "Profiler.Type.Calls")
+	);
+
+	MenuBuilder.AddSeparator();
+
+	MenuBuilder.AddMenuEntry
+	(
+		FTaskTimingStateCommands::Get().Command_ShowTaskPrerequisites,
+		NAME_None,
+		TAttribute<FText>(),
+		TAttribute<FText>(),
+		FSlateIcon(FEditorStyle::GetStyleSetName(), "Profiler.Type.Calls")
+	);
+
+	MenuBuilder.AddMenuEntry
+	(
+		FTaskTimingStateCommands::Get().Command_ShowTaskSubsequents,
+		NAME_None,
+		TAttribute<FText>(),
+		TAttribute<FText>(),
+		FSlateIcon(FEditorStyle::GetStyleSetName(), "Profiler.Type.Calls")
+	);
+
+	MenuBuilder.AddMenuEntry
+	(
+		FTaskTimingStateCommands::Get().Command_ShowNestedTasks,
+		NAME_None,
+		TAttribute<FText>(),
+		TAttribute<FText>(),
+		FSlateIcon(FEditorStyle::GetStyleSetName(), "Profiler.Type.Calls")
+	);
+}
+
+////////////////////////////////////////////////////////////////////////////////////////////////////
+
+void FTaskTimingSharedState::InitCommandList()
+{
+	if (CommandList.IsValid())
+	{
+		return;
+	}
+
+	FTaskTimingStateCommands::Register();
+
+	CommandList = MakeShared<FUICommandList>();
+
+	CommandList->MapAction(
+		FTaskTimingStateCommands::Get().Command_ShowTaskDependencies,
+		FExecuteAction::CreateSP(this, &FTaskTimingSharedState::ContextMenu_ShowTaskDependecies_Execute),
+		FCanExecuteAction::CreateSP(this, &FTaskTimingSharedState::ContextMenu_ShowTaskDependecies_CanExecute),
+		FIsActionChecked::CreateSP(this, &FTaskTimingSharedState::ContextMenu_ShowTaskDependecies_IsChecked));
+
+	CommandList->MapAction(
+		FTaskTimingStateCommands::Get().Command_ShowTaskPrerequisites,
+		FExecuteAction::CreateSP(this, &FTaskTimingSharedState::ContextMenu_ShowTaskPrerequisites_Execute),
+		FCanExecuteAction::CreateSP(this, &FTaskTimingSharedState::ContextMenu_ShowTaskPrerequisites_CanExecute),
+		FIsActionChecked::CreateSP(this, &FTaskTimingSharedState::ContextMenu_ShowTaskPrerequisites_IsChecked));
+
+	CommandList->MapAction(
+		FTaskTimingStateCommands::Get().Command_ShowTaskSubsequents,
+		FExecuteAction::CreateSP(this, &FTaskTimingSharedState::ContextMenu_ShowTaskSubsequents_Execute),
+		FCanExecuteAction::CreateSP(this, &FTaskTimingSharedState::ContextMenu_ShowTaskSubsequents_CanExecute),
+		FIsActionChecked::CreateSP(this, &FTaskTimingSharedState::ContextMenu_ShowTaskSubsequents_IsChecked));
+
+	CommandList->MapAction(
+		FTaskTimingStateCommands::Get().Command_ShowNestedTasks,
+		FExecuteAction::CreateSP(this, &FTaskTimingSharedState::ContextMenu_ShowNestedTasks_Execute),
+		FCanExecuteAction::CreateSP(this, &FTaskTimingSharedState::ContextMenu_ShowNestedTasks_CanExecute),
+		FIsActionChecked::CreateSP(this, &FTaskTimingSharedState::ContextMenu_ShowNestedTasks_IsChecked));
+}
+
+////////////////////////////////////////////////////////////////////////////////////////////////////
+// Dependencies
+
+void FTaskTimingSharedState::ContextMenu_ShowTaskDependecies_Execute()
+{
+	TSharedPtr<Insights::FTaskGraphProfilerManager> TaskGraphManager = Insights::FTaskGraphProfilerManager::Get();
+	if (TaskGraphManager.IsValid() && TaskGraphManager->GetIsAvailable())
+	{
+		TaskGraphManager->SetShowRelations(!TaskGraphManager->GetShowRelations());
+		OnTaskSettingsChanged();
+	}
+}
+
+////////////////////////////////////////////////////////////////////////////////////////////////////
+
+bool FTaskTimingSharedState::ContextMenu_ShowTaskDependecies_CanExecute()
+{
+	return Insights::FTaskGraphProfilerManager::Get().IsValid() && Insights::FTaskGraphProfilerManager::Get()->GetIsAvailable();
+}
+
+////////////////////////////////////////////////////////////////////////////////////////////////////
+
+bool FTaskTimingSharedState::ContextMenu_ShowTaskDependecies_IsChecked()
+{
+	TSharedPtr<Insights::FTaskGraphProfilerManager> TaskGraphManager = Insights::FTaskGraphProfilerManager::Get();
+	return TaskGraphManager.IsValid() && TaskGraphManager->GetIsAvailable() && TaskGraphManager->GetShowRelations();
+}
+
+////////////////////////////////////////////////////////////////////////////////////////////////////
+// Prerequisites
+
+bool FTaskTimingSharedState::ContextMenu_ShowTaskPrerequisites_CanExecute()
+{
+	return Insights::FTaskGraphProfilerManager::Get().IsValid() && Insights::FTaskGraphProfilerManager::Get()->GetIsAvailable();
+}
+
+////////////////////////////////////////////////////////////////////////////////////////////////////
+
+bool FTaskTimingSharedState::ContextMenu_ShowTaskPrerequisites_IsChecked()
+{
+	TSharedPtr<Insights::FTaskGraphProfilerManager> TaskGraphManager = Insights::FTaskGraphProfilerManager::Get();
+	return TaskGraphManager.IsValid() && TaskGraphManager->GetIsAvailable() && TaskGraphManager->GetShowPrerequisites();
+}
+
+////////////////////////////////////////////////////////////////////////////////////////////////////
+
+void FTaskTimingSharedState::ContextMenu_ShowTaskPrerequisites_Execute()
+{
+	TSharedPtr<Insights::FTaskGraphProfilerManager> TaskGraphManager = Insights::FTaskGraphProfilerManager::Get();
+	if (TaskGraphManager.IsValid() && TaskGraphManager->GetIsAvailable())
+	{
+		TaskGraphManager->SetShowPrerequisites(!TaskGraphManager->GetShowPrerequisites());
+		OnTaskSettingsChanged();
+	}
+}
+
+////////////////////////////////////////////////////////////////////////////////////////////////////
+// Subsequents
+
+bool FTaskTimingSharedState::ContextMenu_ShowTaskSubsequents_CanExecute()
+{
+	return Insights::FTaskGraphProfilerManager::Get().IsValid() && Insights::FTaskGraphProfilerManager::Get()->GetIsAvailable();
+}
+
+////////////////////////////////////////////////////////////////////////////////////////////////////
+
+bool FTaskTimingSharedState::ContextMenu_ShowTaskSubsequents_IsChecked()
+{
+	TSharedPtr<Insights::FTaskGraphProfilerManager> TaskGraphManager = Insights::FTaskGraphProfilerManager::Get();
+	return TaskGraphManager.IsValid() && TaskGraphManager->GetIsAvailable() && TaskGraphManager->GetShowSubsequents();
+}
+
+////////////////////////////////////////////////////////////////////////////////////////////////////
+
+void FTaskTimingSharedState::ContextMenu_ShowTaskSubsequents_Execute()
+{
+	TSharedPtr<Insights::FTaskGraphProfilerManager> TaskGraphManager = Insights::FTaskGraphProfilerManager::Get();
+	if (TaskGraphManager.IsValid() && TaskGraphManager->GetIsAvailable())
+	{
+		TaskGraphManager->SetShowSubsequents(!TaskGraphManager->GetShowSubsequents());
+		OnTaskSettingsChanged();
+	}
+}
+
+////////////////////////////////////////////////////////////////////////////////////////////////////
+// NestedTasks
+
+bool FTaskTimingSharedState::ContextMenu_ShowNestedTasks_CanExecute()
+{
+	return Insights::FTaskGraphProfilerManager::Get().IsValid() && Insights::FTaskGraphProfilerManager::Get()->GetIsAvailable();
+}
+
+////////////////////////////////////////////////////////////////////////////////////////////////////
+
+bool FTaskTimingSharedState::ContextMenu_ShowNestedTasks_IsChecked()
+{
+	TSharedPtr<Insights::FTaskGraphProfilerManager> TaskGraphManager = Insights::FTaskGraphProfilerManager::Get();
+	return TaskGraphManager.IsValid() && TaskGraphManager->GetIsAvailable() && TaskGraphManager->GetShowNestedTasks();
+}
+
+////////////////////////////////////////////////////////////////////////////////////////////////////
+
+void FTaskTimingSharedState::ContextMenu_ShowNestedTasks_Execute()
+{
+	TSharedPtr<Insights::FTaskGraphProfilerManager> TaskGraphManager = Insights::FTaskGraphProfilerManager::Get();
+	if (TaskGraphManager.IsValid() && TaskGraphManager->GetIsAvailable())
+	{
+		TaskGraphManager->SetShowNestedTasks(!TaskGraphManager->GetShowNestedTasks());
+		OnTaskSettingsChanged();
+	}
+}
+
+////////////////////////////////////////////////////////////////////////////////////////////////////
+
+void FTaskTimingSharedState::OnTaskSettingsChanged()
+{
+	if (!TaskTrack.IsValid())
+	{
+		return;
+	}
+
+	if (TaskTrack->GetTaskId() == FTaskTimingTrack::InvalidTaskId)
+	{
+		FTaskGraphProfilerManager::Get()->ClearTaskRelations();
+	}
+
+	TSharedPtr<Insights::FTaskGraphProfilerManager> TaskGraphManager = Insights::FTaskGraphProfilerManager::Get();
+
+	TSharedPtr<const ITimingEvent> SelectedEvent = TimingView->GetSelectedEvent();
+	if (SelectedEvent.IsValid() && SelectedEvent->Is<const FThreadTrackEvent>())
+	{
+		TaskTrack->OnTimingEventSelected(SelectedEvent);
+		return;
+	}
+
+	if(TaskTrack->GetTaskId() != FTaskTimingTrack::InvalidTaskId)
+	{
+		FTaskGraphProfilerManager::Get()->ShowTaskRelations(TaskTrack->GetTaskId());
+	}
+}
+
 ////////////////////////////////////////////////////////////////////////////////////////////////////
 // FTaskTimingTrack
 ////////////////////////////////////////////////////////////////////////////////////////////////////
diff --git a/Engine/Source/Developer/TraceInsights/Private/Insights/TaskGraphProfiler/ViewModels/TaskTimingTrack.h b/Engine/Source/Developer/TraceInsights/Private/Insights/TaskGraphProfiler/ViewModels/TaskTimingTrack.h
--- a/Engine/Source/Developer/TraceInsights/Private/Insights/TaskGraphProfiler/ViewModels/TaskTimingTrack.h
+++ b/Engine/Source/Developer/TraceInsights/Private/Insights/TaskGraphProfiler/ViewModels/TaskTimingTrack.h
@@ -21,7 +21,7 @@ class FTaskTimingSharedState : public Insights::ITimingViewExtender, public TSha
 {
 
 public:
-	explicit FTaskTimingSharedState(STimingView* InTimingView) : TimingView(InTimingView) {}
+	FTaskTimingSharedState(STimingView* InTimingView);
 	virtual ~FTaskTimingSharedState() = default;
 
 	TSharedPtr<FTaskTimingTrack> GetTaskTrack() { return TaskTrack; }
@@ -35,6 +35,7 @@ public:
 	virtual void OnEndSession(Insights::ITimingViewSession& InSession) override;
 	virtual void Tick(Insights::ITimingViewSession& InSession, const TraceServices::IAnalysisSession& InAnalysisSession) override;
 	virtual void ExtendFilterMenu(Insights::ITimingViewSession& InSession, FMenuBuilder& InMenuBuilder) override;
+	virtual bool ExtendGlobalContextMenu(ITimingViewSession& InSession, FMenuBuilder& InMenuBuilder) override;
 
 	//////////////////////////////////////////////////
 
@@ -49,6 +50,29 @@ public:
 
 	STimingView* GetTimingView() { return TimingView; }
 
+private:
+	void InitCommandList();
+
+	void BuildTasksSubMenu(FMenuBuilder& MenuBuilder);
+
+	void ContextMenu_ShowTaskDependecies_Execute();
+	bool ContextMenu_ShowTaskDependecies_CanExecute();
+	bool ContextMenu_ShowTaskDependecies_IsChecked();
+
+	void ContextMenu_ShowTaskPrerequisites_Execute();
+	bool ContextMenu_ShowTaskPrerequisites_CanExecute();
+	bool ContextMenu_ShowTaskPrerequisites_IsChecked();
+
+	void ContextMenu_ShowTaskSubsequents_Execute();
+	bool ContextMenu_ShowTaskSubsequents_CanExecute();
+	bool ContextMenu_ShowTaskSubsequents_IsChecked();
+
+	void ContextMenu_ShowNestedTasks_Execute();
+	bool ContextMenu_ShowNestedTasks_CanExecute();
+	bool ContextMenu_ShowNestedTasks_IsChecked();
+
+	void OnTaskSettingsChanged();
+
 private:
 	STimingView* TimingView;
 
@@ -56,6 +80,7 @@ private:
 	bool bResetOnNextTick = false;
 
 	TSharedPtr<FTaskTimingTrack> TaskTrack;
+	TSharedPtr<FUICommandList> CommandList;
 };
 
 ////////////////////////////////////////////////////////////////////////////////////////////////////
diff --git a/Engine/Source/Developer/TraceInsights/Private/Insights/TimingProfilerCommands.cpp b/Engine/Source/Developer/TraceInsights/Private/Insights/TimingProfilerCommands.cpp
--- a/Engine/Source/Developer/TraceInsights/Private/Insights/TimingProfilerCommands.cpp
+++ b/Engine/Source/Developer/TraceInsights/Private/Insights/TimingProfilerCommands.cpp
@@ -102,12 +102,6 @@ void FTimingViewCommands::RegisterCommands()
 		EUserInterfaceActionType::ToggleButton,
 		FInputChord(EKeys::G));
 
-	UI_COMMAND(ShowTaskDependencies,
-		"Task Dependencies",
-		"Show/hide all dependencies of the current task (for a selected cpu timing event).",
-		EUserInterfaceActionType::ToggleButton,
-		FInputChord(EKeys::T));
-
 	//UI_COMMAND(ShowAllGpuTracks, "GPU Track(s)", "Show/hide the GPU track(s).", EUserInterfaceActionType::ToggleButton, FInputChord(EKeys::Y));
 	//UI_COMMAND(ShowAllCpuTracks, "CPU Thread Tracks", "Show/hide all CPU tracks (and all CPU thread groups).", EUserInterfaceActionType::ToggleButton, FInputChord(EKeys::U));
 }
diff --git a/Engine/Source/Developer/TraceInsights/Private/Insights/TimingProfilerCommands.h b/Engine/Source/Developer/TraceInsights/Private/Insights/TimingProfilerCommands.h
--- a/Engine/Source/Developer/TraceInsights/Private/Insights/TimingProfilerCommands.h
+++ b/Engine/Source/Developer/TraceInsights/Private/Insights/TimingProfilerCommands.h
@@ -75,9 +75,6 @@ public:
 	/** Toggles visibility for Main Graph track. */
 	TSharedPtr<FUICommandInfo> ShowMainGraphTrack;
 
-	/** Toggles visibility for task graph dependency relations */
-	TSharedPtr<FUICommandInfo> ShowTaskDependencies;
-
 	/** Toggles visibility for GPU thread track. */
 	//TSharedPtr<FUICommandInfo> ShowAllGpuTracks;
 
diff --git a/Engine/Source/Developer/TraceInsights/Private/Insights/Widgets/STimingView.cpp b/Engine/Source/Developer/TraceInsights/Private/Insights/Widgets/STimingView.cpp
--- a/Engine/Source/Developer/TraceInsights/Private/Insights/Widgets/STimingView.cpp
+++ b/Engine/Source/Developer/TraceInsights/Private/Insights/Widgets/STimingView.cpp
@@ -2710,6 +2710,7 @@ void STimingView::ShowContextMenu(const FPointerEvent& MouseEvent)
 	const bool bShouldCloseWindowAfterMenuSelection = true;
 
 	FMenuBuilder MenuBuilder(bShouldCloseWindowAfterMenuSelection, CommandList);
+	bool bHasAnyActions = false;
 
 	if (HoveredTrack.IsValid())
 	{
@@ -2718,8 +2719,15 @@ void STimingView::ShowContextMenu(const FPointerEvent& MouseEvent)
 		MenuBuilder.EndSection();
 
 		HoveredTrack->BuildContextMenu(MenuBuilder);
+		bHasAnyActions = true;
 	}
-	else
+
+	for (Insights::ITimingViewExtender* Extender : GetExtenders())
+	{
+		bHasAnyActions |= Extender->ExtendGlobalContextMenu(*this, MenuBuilder);
+	}
+	
+	if(!bHasAnyActions)
 	{
 		MenuBuilder.BeginSection(TEXT("Empty"));
 		{
@@ -2735,22 +2743,6 @@ void STimingView::ShowContextMenu(const FPointerEvent& MouseEvent)
 		MenuBuilder.EndSection();
 	}
 
-	if (SelectedTrack.IsValid() && SelectedTrack->Is<FTimingEventsTrack>())
-	{
-		MenuBuilder.BeginSection(TEXT("Event"), LOCTEXT("Event", "Event"));
-
-		MenuBuilder.AddMenuEntry
-		(
-			Commands.ShowTaskDependencies,
-			NAME_None,
-			TAttribute<FText>(),
-			TAttribute<FText>(),
-			FSlateIcon(FEditorStyle::GetStyleSetName(), "Profiler.Type.Calls")
-		);
-
-		MenuBuilder.EndSection();
-	}
-
 	TSharedRef<SWidget> MenuWidget = MenuBuilder.MakeWidget();
 
 	FWidgetPath EventPath = MouseEvent.GetEventPath() != nullptr ? *MouseEvent.GetEventPath() : FWidgetPath();
@@ -2955,12 +2947,6 @@ void STimingView::BindCommands()
 		FCanExecuteAction::CreateLambda([] { return true; }),
 		FIsActionChecked::CreateSP(this, &STimingView::ShowHideGraphTrack_IsChecked));
 
-	CommandList->MapAction(
-		Commands.ShowTaskDependencies,
-		FExecuteAction::CreateSP(this, &STimingView::ContextMenu_ShowTaskDependecies_Execute),
-		FCanExecuteAction::CreateSP(this, &STimingView::ContextMenu_ShowTaskDependecies_CanExecute),
-		FIsActionChecked::CreateSP(this, &STimingView::ContextMenu_ShowTaskDependecies_IsChecked));
-
 	//TODO: let extenders to bind other commands
 
 	//CommandList->MapAction(
@@ -3923,7 +3909,6 @@ TSharedRef<SWidget> STimingView::MakeTracksFilterMenu()
 		const FTimingViewCommands& Commands = FTimingViewCommands::Get();
 
 		MenuBuilder.AddMenuEntry(Commands.ShowMainGraphTrack);
-		MenuBuilder.AddMenuEntry(Commands.ShowTaskDependencies);
 		MenuBuilder.AddMenuEntry(Commands.ToggleCompactMode);
 		MenuBuilder.AddMenuEntry(Commands.AutoHideEmptyTracks);
 	}
@@ -4165,32 +4150,6 @@ TArray<Insights::ITimingViewExtender*> STimingView::GetExtenders() const
 
 ////////////////////////////////////////////////////////////////////////////////////////////////////
 
-void STimingView::ContextMenu_ShowTaskDependecies_Execute()
-{
-	TSharedPtr<Insights::FTaskGraphProfilerManager> TaskGraphManager = Insights::FTaskGraphProfilerManager::Get();
-	if (TaskGraphManager.IsValid() && TaskGraphManager->GetIsAvailable())
-	{
-		TaskGraphManager->SetShowRelations(!TaskGraphManager->GetShowRelations());
-	}
-}
-
-////////////////////////////////////////////////////////////////////////////////////////////////////
-
-bool STimingView::ContextMenu_ShowTaskDependecies_CanExecute()
-{
-	return Insights::FTaskGraphProfilerManager::Get().IsValid() && Insights::FTaskGraphProfilerManager::Get()->GetIsAvailable();
-}
-
-////////////////////////////////////////////////////////////////////////////////////////////////////
-
-bool STimingView::ContextMenu_ShowTaskDependecies_IsChecked()
-{
-	TSharedPtr<Insights::FTaskGraphProfilerManager> TaskGraphManager = Insights::FTaskGraphProfilerManager::Get();
-	return TaskGraphManager.IsValid() && TaskGraphManager->GetIsAvailable() && TaskGraphManager->GetShowRelations();
-}
-
-////////////////////////////////////////////////////////////////////////////////////////////////////
-
 void STimingView::ClearRelations()
 {
 	CurrentRelations.Empty();
diff --git a/Engine/Source/Developer/TraceInsights/Public/Insights/ITimingViewExtender.h b/Engine/Source/Developer/TraceInsights/Public/Insights/ITimingViewExtender.h
--- a/Engine/Source/Developer/TraceInsights/Public/Insights/ITimingViewExtender.h
+++ b/Engine/Source/Developer/TraceInsights/Public/Insights/ITimingViewExtender.h
@@ -30,6 +30,10 @@ public:
 
 	/** Extension hook for the 'quick filter' menu */
 	virtual void ExtendFilterMenu(ITimingViewSession& InSession, FMenuBuilder& InMenuBuilder) = 0;
+
+	/** Extension hook for the context menu for all tracks
+	@return True if any menu option was added and False if no option was added */ 
+	virtual bool ExtendGlobalContextMenu(ITimingViewSession& InSession, FMenuBuilder& InMenuBuilder) { return false; }
 };
 
 } // namespace Insights
