[Insights]
  - Timing view: Added experimental support for multiple coloring modes: by name (default), by id and by duration. Added shortcut key 0 to toggle between these three modes.
  - Timing view: Added experimental support for limiting the depth of timing event tracks (to 4 lanes). Added shortcut key X to toggle between 4 lanes limit and unlimitted.

diff --git a/Engine/Source/Developer/TraceInsights/Private/Insights/TimingProfilerManager.h b/Engine/Source/Developer/TraceInsights/Private/Insights/TimingProfilerManager.h
--- a/Engine/Source/Developer/TraceInsights/Private/Insights/TimingProfilerManager.h
+++ b/Engine/Source/Developer/TraceInsights/Private/Insights/TimingProfilerManager.h
@@ -15,6 +15,15 @@
 namespace Insights
 {
 	class FTimerButterflyAggregator;
+
+	enum class ETimingEventsColoringMode : uint32
+	{
+		ByName,
+		ById,
+		ByDuration,
+
+		Count
+	};
 }
 
 class STimingProfilerWindow;
@@ -145,6 +154,12 @@ public:
 
 	const FName& GetLogListingName() const { return LogListingName; }
 
+	Insights::ETimingEventsColoringMode GetColoringMode() const { return ColoringMode; }
+	void SetColoringMode(Insights::ETimingEventsColoringMode InColoringMode) { ColoringMode = InColoringMode; }
+
+	uint32 GetEventDepthLimit() const { return EventDepthLimit; }
+	void SetEventDepthLimit(uint32 InEventDepthLimit) { EventDepthLimit = InEventDepthLimit; }
+
 private:
 	/** Binds our UI commands to delegates. */
 	void BindCommands();
@@ -216,6 +231,9 @@ private:
 	/** The name of the Timing Insights log listing. */
 	FName LogListingName;
 
+	Insights::ETimingEventsColoringMode ColoringMode = Insights::ETimingEventsColoringMode::ByName;
+	uint32 EventDepthLimit = 1000;
+
 	/** A shared pointer to the global instance of the Timing Profiler manager. */
 	static TSharedPtr<FTimingProfilerManager> Instance;
 };
diff --git a/Engine/Source/Developer/TraceInsights/Private/Insights/ViewModels/ThreadTimingTrack.cpp b/Engine/Source/Developer/TraceInsights/Private/Insights/ViewModels/ThreadTimingTrack.cpp
--- a/Engine/Source/Developer/TraceInsights/Private/Insights/ViewModels/ThreadTimingTrack.cpp
+++ b/Engine/Source/Developer/TraceInsights/Private/Insights/ViewModels/ThreadTimingTrack.cpp
@@ -204,8 +204,31 @@ static void AddTimingEventToBuilder(ITimingEventsTrackDrawStateBuilder& Builder,
 									double EventStartTime, double EventEndTime, uint32 EventDepth,
 									uint32 TimerIndex, const TraceServices::FTimingProfilerTimer* Timer)
 {
-	//const uint32 EventColor = FTimingEvent::ComputeEventColor(Timer->Id);
-	const uint32 EventColor = FTimingEvent::ComputeEventColor(Timer->Name);
+	if (EventDepth >= FTimingProfilerManager::Get()->GetEventDepthLimit())
+	{
+		return;
+	}
+
+	uint32 EventColor;
+	switch (FTimingProfilerManager::Get()->GetColoringMode())
+	{
+		case Insights::ETimingEventsColoringMode::ByName:
+			EventColor = FTimingEvent::ComputeEventColor(Timer->Name);
+			break;
+		case Insights::ETimingEventsColoringMode::ById:
+			EventColor = FTimingEvent::ComputeEventColor(Timer->Id);
+			break;
+		case Insights::ETimingEventsColoringMode::ByDuration:
+		{
+			const double EventDuration = EventEndTime - EventStartTime;
+			EventColor = (EventDuration > 0.01) ? 0xFF883333 :
+						 (EventDuration > 0.001) ? 0xFF998833 :
+						 (EventDuration > 0.0001) ? 0xFF338833 : 0xFF333388;
+			break;
+		}
+		default:
+			EventColor = 0xFF000000;
+	}
 
 	Builder.AddEvent(EventStartTime, EventEndTime, EventDepth, EventColor,
 		[TimerIndex, Timer, EventStartTime, EventEndTime](float Width)
@@ -1073,11 +1096,13 @@ void FThreadTimingTrack::InitTooltip(FTooltipDrawState& InOutTooltip, const ITim
 					FString TrackName = Track.IsValid() ? Track->GetName() : TEXT("Unknown");
 					InOutTooltip.AddNameValueTextLine(TEXT("Launched:"), FString::Printf(TEXT("%f (+%s) on %s"), Task.LaunchedTimestamp, *TimeUtils::FormatTimeAuto(Task.LaunchedTimestamp - Task.CreatedTimestamp), *TrackName));
 				}
+
 				{
 					TSharedPtr<FCpuTimingTrack> Track = SharedState.GetCpuTrack(Task.ScheduledThreadId);
 					FString TrackName = Track.IsValid() ? Track->GetName() : TEXT("Unknown");
 					InOutTooltip.AddNameValueTextLine(TEXT("Scheduled:"), FString::Printf(TEXT("%f (+%s) on %s"), Task.ScheduledTimestamp, *TimeUtils::FormatTimeAuto(Task.ScheduledTimestamp - Task.LaunchedTimestamp), *TrackName));
 				}
+
 				InOutTooltip.AddNameValueTextLine(TEXT("Started:"), FString::Printf(TEXT("%f (+%s)"), Task.StartedTimestamp, *TimeUtils::FormatTimeAuto(Task.StartedTimestamp - Task.ScheduledTimestamp)));
 				if (Task.FinishedTimestamp != TraceServices::FTaskInfo::InvalidTimestamp)
 				{
diff --git a/Engine/Source/Developer/TraceInsights/Private/Insights/Widgets/STimingView.cpp b/Engine/Source/Developer/TraceInsights/Private/Insights/Widgets/STimingView.cpp
--- a/Engine/Source/Developer/TraceInsights/Private/Insights/Widgets/STimingView.cpp
+++ b/Engine/Source/Developer/TraceInsights/Private/Insights/Widgets/STimingView.cpp
@@ -2653,21 +2653,41 @@ FReply STimingView::OnKeyDown(const FGeometry& MyGeometry, const FKeyEvent& InKe
 	else if (InKeyEvent.GetKey() == EKeys::One)
 	{
 		LoadingSharedState->SetColorSchema(0);
+		Viewport.AddDirtyFlags(ETimingTrackViewportDirtyFlags::HInvalidated);
 		return FReply::Handled();
 	}
 	else if (InKeyEvent.GetKey() == EKeys::Two)
 	{
 		LoadingSharedState->SetColorSchema(1);
+		Viewport.AddDirtyFlags(ETimingTrackViewportDirtyFlags::HInvalidated);
 		return FReply::Handled();
 	}
 	else if (InKeyEvent.GetKey() == EKeys::Three)
 	{
 		LoadingSharedState->SetColorSchema(2);
+		Viewport.AddDirtyFlags(ETimingTrackViewportDirtyFlags::HInvalidated);
 		return FReply::Handled();
 	}
 	else if (InKeyEvent.GetKey() == EKeys::Four)
 	{
 		LoadingSharedState->SetColorSchema(3);
+		Viewport.AddDirtyFlags(ETimingTrackViewportDirtyFlags::HInvalidated);
+		return FReply::Handled();
+	}
+	else if (InKeyEvent.GetKey() == EKeys::Zero)
+	{
+		uint32 Mode = (uint32)FTimingProfilerManager::Get()->GetColoringMode();
+		Mode = (Mode + 1) % (uint32)Insights::ETimingEventsColoringMode::Count;
+		FTimingProfilerManager::Get()->SetColoringMode((Insights::ETimingEventsColoringMode)Mode);
+		Viewport.AddDirtyFlags(ETimingTrackViewportDirtyFlags::HInvalidated);
+		return FReply::Handled();
+	}
+	else if (InKeyEvent.GetKey() == EKeys::X)
+	{
+		uint32 DepthLimit = FTimingProfilerManager::Get()->GetEventDepthLimit();
+		DepthLimit = (DepthLimit != 4) ? 4 : 1000;
+		FTimingProfilerManager::Get()->SetEventDepthLimit(DepthLimit);
+		Viewport.AddDirtyFlags(ETimingTrackViewportDirtyFlags::HInvalidated);
 		return FReply::Handled();
 	}
 #if ACTIVATE_BENCHMARK
