[Insights] Improvements to how the task relation depth is calculated.

diff --git a/Engine/Source/Developer/TraceInsights/Private/Insights/TaskGraphProfiler/TaskGraphProfilerManager.cpp b/Engine/Source/Developer/TraceInsights/Private/Insights/TaskGraphProfiler/TaskGraphProfilerManager.cpp
--- a/Engine/Source/Developer/TraceInsights/Private/Insights/TaskGraphProfiler/TaskGraphProfilerManager.cpp
+++ b/Engine/Source/Developer/TraceInsights/Private/Insights/TaskGraphProfiler/TaskGraphProfilerManager.cpp
@@ -135,7 +135,7 @@ void FTaskGraphProfilerManager::UnregisterMajorTabs()
 
 bool FTaskGraphProfilerManager::Tick(float DeltaTime)
 {
-	// Check if session has Memory events (to spawn the tab), but not too often.
+	// Check if session has task events (to spawn the tab), but not too often.
 	if (!bIsAvailable && AvailabilityCheck.Tick())
 	{
 		bool bShouldBeAvailable = false;
@@ -301,8 +301,10 @@ void FTaskGraphProfilerManager::GetSingleTaskRelations(const TraceServices::FTas
 		AddRelation(InSelectedEvent, Task->CreatedTimestamp, Task->CreatedThreadId, Task->LaunchedTimestamp, Task->LaunchedThreadId, ETaskEventType::Created);
 	}
 
-	AddRelation(InSelectedEvent, Task->LaunchedTimestamp, Task->LaunchedThreadId, Task->ScheduledTimestamp, Task->ScheduledThreadId, ETaskEventType::Launched);
-
+	if (Task->LaunchedTimestamp != Task->ScheduledTimestamp || Task->LaunchedThreadId != Task->ScheduledThreadId)
+	{
+		AddRelation(InSelectedEvent, Task->LaunchedTimestamp, Task->LaunchedThreadId, Task->ScheduledTimestamp, Task->ScheduledThreadId, ETaskEventType::Launched);
+	}
 	int32 NumPrerequisitesToShow = FMath::Min(Task->Prerequisites.Num(), MaxTasksToShow);
 	for (int32 i = 0; i != NumPrerequisitesToShow; ++i)
 	{
@@ -311,10 +313,8 @@ void FTaskGraphProfilerManager::GetSingleTaskRelations(const TraceServices::FTas
 		AddRelation(InSelectedEvent, Prerequisite->CompletedTimestamp, Prerequisite->CompletedThreadId, Task->ScheduledTimestamp, Task->ScheduledThreadId, ETaskEventType::Prerequisite);
 	}
 
-	if (Task->LaunchedTimestamp != Task->ScheduledTimestamp || Task->LaunchedThreadId != Task->ScheduledThreadId)
-	{
-		AddRelation(InSelectedEvent, Task->ScheduledTimestamp, Task->ScheduledThreadId, Task->StartedTimestamp, Task->StartedThreadId, ETaskEventType::Scheduled);
-	}
+	int32 ExecutionStartedDepth = GetDepthOfTaskExecution(Task->StartedTimestamp, Task->FinishedTimestamp, Task->StartedThreadId);
+	AddRelation(InSelectedEvent, Task->ScheduledTimestamp, Task->ScheduledThreadId, -1, Task->StartedTimestamp, Task->StartedThreadId, ExecutionStartedDepth, ETaskEventType::Scheduled);
 
 	int32 NumNestedToShow = FMath::Min(Task->NestedTasks.Num(), MaxTasksToShow);
 	for (int32 i = 0; i != NumNestedToShow; ++i)
@@ -323,7 +323,9 @@ void FTaskGraphProfilerManager::GetSingleTaskRelations(const TraceServices::FTas
 		const TraceServices::FTaskInfo* NestedTask = TasksProvider->TryGetTask(RelationInfo.RelativeId);
 		check(NestedTask != nullptr);
 
-		AddRelation(InSelectedEvent, RelationInfo.Timestamp, Task->StartedThreadId, NestedTask->StartedTimestamp, NestedTask->StartedThreadId, ETaskEventType::AddedNested);
+		int32 NestedExecutionStartedDepth = GetDepthOfTaskExecution(Task->StartedTimestamp, Task->FinishedTimestamp, Task->StartedThreadId);
+
+		AddRelation(InSelectedEvent, RelationInfo.Timestamp, Task->StartedThreadId, -1, NestedTask->StartedTimestamp,  NestedTask->StartedThreadId, NestedExecutionStartedDepth,  ETaskEventType::AddedNested);
 
 		AddRelation(InSelectedEvent, NestedTask->CompletedTimestamp, NestedTask->CompletedThreadId, NestedTask->CompletedTimestamp, Task->StartedThreadId, ETaskEventType::NestedCompleted);
 	}
@@ -341,7 +343,7 @@ void FTaskGraphProfilerManager::GetSingleTaskRelations(const TraceServices::FTas
 
 	if (Task->FinishedTimestamp != Task->CompletedTimestamp || Task->CompletedThreadId != Task->StartedThreadId)
 	{
-		AddRelation(InSelectedEvent, Task->FinishedTimestamp, Task->StartedThreadId, Task->CompletedTimestamp, Task->StartedThreadId, ETaskEventType::Completed);
+		AddRelation(InSelectedEvent, Task->FinishedTimestamp, Task->StartedThreadId, ExecutionStartedDepth, Task->CompletedTimestamp, Task->StartedThreadId, -1,  ETaskEventType::Completed);
 	}
 }
 
@@ -443,6 +445,13 @@ FLinearColor FTaskGraphProfilerManager::GetColorForTaskEvent(ETaskEventType InEv
 ////////////////////////////////////////////////////////////////////////////////////////////////////
 
 void FTaskGraphProfilerManager::AddRelation(const FThreadTrackEvent* InSelectedEvent, double SourceTimestamp, uint32 SourceThreadId, double TargetTimestamp, uint32 TargetThreadId, ETaskEventType Type)
+{
+	AddRelation(InSelectedEvent, SourceTimestamp, SourceThreadId, -1, TargetTimestamp, TargetThreadId, -1, Type);
+}
+
+////////////////////////////////////////////////////////////////////////////////////////////////////
+
+void FTaskGraphProfilerManager::AddRelation(const FThreadTrackEvent* InSelectedEvent, double SourceTimestamp, uint32 SourceThreadId, int32 SourceDepth, double TargetTimestamp, uint32 TargetThreadId, int32 TargetDepth, ETaskEventType Type)
 {
 	if (SourceTimestamp == TraceServices::FTaskInfo::InvalidTimestamp || TargetTimestamp == TraceServices::FTaskInfo::InvalidTimestamp)
 	{
@@ -475,13 +484,13 @@ void FTaskGraphProfilerManager::AddRelation(const FThreadTrackEvent* InSelectedE
 		if (TaskRelationPtr->GetSourceThreadId() == ThreadId)
 		{
 			TaskRelationPtr->SetSourceTrack(EventTrack);
-			TaskRelationPtr->SetSourceDepth(InSelectedEvent->GetDepth());
+			TaskRelationPtr->SetSourceDepth(SourceDepth > 0 ? SourceDepth : EventTrack->GetDepthAt(TaskRelationPtr->GetSourceTime()) - 1);
 		}
 
 		if (TaskRelationPtr->GetTargetThreadId() == ThreadId)
 		{
 			TaskRelationPtr->SetTargetTrack(EventTrack);
-			TaskRelationPtr->SetTargetDepth(InSelectedEvent->GetDepth());
+			TaskRelationPtr->SetTargetDepth(TargetDepth > 0 ? TargetDepth : EventTrack->GetDepthAt(TaskRelationPtr->GetTargetTime()) - 1);
 		}
 	}
 
@@ -491,7 +500,7 @@ void FTaskGraphProfilerManager::AddRelation(const FThreadTrackEvent* InSelectedE
 		if (Track.IsValid())
 		{
 			TaskRelationPtr->SetSourceTrack(Track);
-			TaskRelationPtr->SetSourceDepth(Track->GetDepthAt(TaskRelationPtr->GetSourceTime()) - 1);
+			TaskRelationPtr->SetSourceDepth(SourceDepth > 0 ? SourceDepth : Track->GetDepthAt(TaskRelationPtr->GetSourceTime()) - 1);
 		}
 	}
 
@@ -502,7 +511,7 @@ void FTaskGraphProfilerManager::AddRelation(const FThreadTrackEvent* InSelectedE
 		if (Track.IsValid())
 		{
 			TaskRelationPtr->SetTargetTrack(Track);
-			TaskRelationPtr->SetTargetDepth(Track->GetDepthAt(TaskRelationPtr->GetTargetTime()) - 1);
+			TaskRelationPtr->SetTargetDepth(TargetDepth > 0 ? TargetDepth : Track->GetDepthAt(TaskRelationPtr->GetTargetTime()) - 1);
 		}
 	}
 
@@ -548,6 +557,64 @@ void FTaskGraphProfilerManager::SetShowRelations(bool bInValue)
 
 ////////////////////////////////////////////////////////////////////////////////////////////////////
 
+int32 FTaskGraphProfilerManager::GetDepthOfTaskExecution(double TaskStartedTime, double TaskFinishedTime, uint32 ThreadId)
+{
+	int32 Depth = -1;
+	TSharedPtr<class STimingProfilerWindow> TimingWindow = FTimingProfilerManager::Get()->GetProfilerWindow();
+	if (!TimingWindow.IsValid())
+	{
+		return Depth;
+	}
+
+	TSharedPtr<STimingView> TimingView = TimingWindow->GetTimingView();
+	if (!TimingView.IsValid())
+	{
+		return Depth;
+	}
+
+	TSharedPtr<FThreadTimingSharedState> ThreadSharedState = TimingView->GetThreadTimingSharedState();
+
+	TSharedPtr<FCpuTimingTrack> Track = ThreadSharedState->GetCpuTrack(ThreadId);
+
+	if (!Track.IsValid())
+	{
+		return Depth;
+	}
+
+	TSharedPtr<const TraceServices::IAnalysisSession> Session = FInsightsManager::Get()->GetSession();
+	if (Session.IsValid() && TraceServices::ReadTimingProfilerProvider(*Session.Get()))
+	{
+		TraceServices::FAnalysisSessionReadScope SessionReadScope(*Session.Get());
+
+		const TraceServices::ITimingProfilerProvider& TimingProfilerProvider = *TraceServices::ReadTimingProfilerProvider(*Session.Get());
+
+		TimingProfilerProvider.ReadTimeline(Track->GetTimelineIndex(),
+			[TaskStartedTime, TaskFinishedTime, &Depth](const TraceServices::ITimingProfilerProvider::Timeline& Timeline)
+			{
+				Timeline.EnumerateEventsDownSampled(TaskStartedTime, TaskFinishedTime, 0, [TaskStartedTime, &Depth](bool IsEnter, double Time, const TraceServices::FTimingProfilerEvent& Event)
+					{
+						if (Time < TaskStartedTime)
+						{
+							check(IsEnter);
+							++Depth;
+							return TraceServices::EEventEnumerate::Continue;
+						}
+
+						if (IsEnter)
+						{
+							++Depth;
+						}
+
+						return TraceServices::EEventEnumerate::Stop;
+					});
+			});
+	}
+
+	return Depth;
+}
+
+////////////////////////////////////////////////////////////////////////////////////////////////////
+
 } // namespace Insights
 
 #undef LOCTEXT_NAMESPACE
diff --git a/Engine/Source/Developer/TraceInsights/Private/Insights/TaskGraphProfiler/TaskGraphProfilerManager.h b/Engine/Source/Developer/TraceInsights/Private/Insights/TaskGraphProfiler/TaskGraphProfilerManager.h
--- a/Engine/Source/Developer/TraceInsights/Private/Insights/TaskGraphProfiler/TaskGraphProfilerManager.h
+++ b/Engine/Source/Developer/TraceInsights/Private/Insights/TaskGraphProfiler/TaskGraphProfilerManager.h
@@ -56,13 +56,13 @@ class FTaskGraphProfilerManager : public TSharedFromThis<FTaskGraphProfilerManag
 public:
 	typedef TFunctionRef<void(double /*SourceTimestamp*/, uint32 /*SourceThreadId*/, double /*TargetTimestamp*/, uint32 /*TargetThreadId*/, ETaskEventType /*Type*/)> FAddRelationCallback;
 
-	/** Creates the Memory Profiler manager, only one instance can exist. */
+	/** Creates the Task Graph Profiler manager, only one instance can exist. */
 	FTaskGraphProfilerManager(TSharedRef<FUICommandList> InCommandList);
 
 	/** Virtual destructor. */
 	virtual ~FTaskGraphProfilerManager();
 
-	/** Creates an instance of the Memory Profiler manager. */
+	/** Creates an instance of the Task Graph Profiler manager. */
 	static TSharedPtr<FTaskGraphProfilerManager> CreateInstance();
 
 	/**
@@ -95,6 +95,7 @@ public:
 	void ShowTaskRelations(const FThreadTrackEvent* InSelectedEvent, uint32 ThreadId);
 	void ShowTaskRelations(uint32 TaskId);
 	void AddRelation(const FThreadTrackEvent* InSelectedEvent, double SourceTimestamp, uint32 SourceThreadId, double TargetTimestamp, uint32 TargetThreadId, ETaskEventType Type);
+	void AddRelation(const FThreadTrackEvent* InSelectedEvent, double SourceTimestamp, uint32 SourceThreadId, int32 SourceDepth, double TargetTimestamp, uint32 TargetThreadId, int32 TargetDepth, ETaskEventType Type);
 	void ClearTaskRelations();
 	FLinearColor GetColorForTaskEvent(ETaskEventType InEvent);
 
@@ -112,6 +113,8 @@ public:
 	bool GetShowNestedTasks() const { return bShowNestedTasks; }
 	void SetShowNestedTasks(bool bInValue) { bShowNestedTasks = bInValue; }
 
+	int32 GetDepthOfTaskExecution(double TaskStartedTime, double TaskFinishedTime, uint32 ThreadId);
+
 private:
 	/** Updates this manager, done through FCoreTicker. */
 	bool Tick(float DeltaTime);
diff --git a/Engine/Source/Developer/TraceInsights/Private/Insights/TaskGraphProfiler/ViewModels/TaskTimingTrack.cpp b/Engine/Source/Developer/TraceInsights/Private/Insights/TaskGraphProfiler/ViewModels/TaskTimingTrack.cpp
--- a/Engine/Source/Developer/TraceInsights/Private/Insights/TaskGraphProfiler/ViewModels/TaskTimingTrack.cpp
+++ b/Engine/Source/Developer/TraceInsights/Private/Insights/TaskGraphProfiler/ViewModels/TaskTimingTrack.cpp
@@ -693,7 +693,9 @@ void FTaskTimingTrack::GetEventRelations(const FThreadTrackEvent& InSelectedEven
 				const TraceServices::FTaskInfo* WaitedTask = TasksProvider->TryGetTask(Waiting->Tasks[TaskIndex]);
 				if (WaitedTask != nullptr)
 				{
-					FTaskGraphProfilerManager::Get()->AddRelation(&InSelectedEvent, StartTime, ThreadId, WaitedTask->StartedTimestamp, WaitedTask->StartedThreadId, ETaskEventType::AddedNested);
+					int32 WaitingTaskExecutionDepth = FTaskGraphProfilerManager::Get()->GetDepthOfTaskExecution(WaitedTask->StartedTimestamp, WaitedTask->FinishedTimestamp, WaitedTask->StartedThreadId);
+
+					FTaskGraphProfilerManager::Get()->AddRelation(&InSelectedEvent, StartTime, ThreadId, -1, WaitedTask->StartedTimestamp, WaitedTask->StartedThreadId, WaitingTaskExecutionDepth, ETaskEventType::AddedNested);
 					FTaskGraphProfilerManager::Get()->AddRelation(&InSelectedEvent, WaitedTask->CompletedTimestamp, WaitedTask->CompletedThreadId, WaitedTask->CompletedTimestamp, ThreadId, ETaskEventType::NestedCompleted);
 				}
 			}
diff --git a/Engine/Source/Developer/TraceInsights/Private/Insights/TaskGraphProfiler/ViewModels/TaskTimingTrack.h b/Engine/Source/Developer/TraceInsights/Private/Insights/TaskGraphProfiler/ViewModels/TaskTimingTrack.h
--- a/Engine/Source/Developer/TraceInsights/Private/Insights/TaskGraphProfiler/ViewModels/TaskTimingTrack.h
+++ b/Engine/Source/Developer/TraceInsights/Private/Insights/TaskGraphProfiler/ViewModels/TaskTimingTrack.h
@@ -94,14 +94,6 @@ class FTaskTimingTrack : public FTimingEventsTrack
 public:
 	static const uint32 InvalidTaskId;
 
-	struct FPendingEventInfo
-	{
-		double StartTime;
-		double EndTime;
-		uint32 Depth;
-		uint32 TimerIndex;
-	};
-
 	explicit FTaskTimingTrack(FTaskTimingSharedState& InSharedState, const FString& InName, uint32 InTimelineIndex)
 		: FTimingEventsTrack(InName)
 		, TimelineIndex(InTimelineIndex)
