[Insights] Fixes for timing view relations arrows

diff --git a/Engine/Source/Developer/TraceInsights/Private/Insights/ViewModels/TaskGraphRelation.cpp b/Engine/Source/Developer/TraceInsights/Private/Insights/ViewModels/TaskGraphRelation.cpp
--- a/Engine/Source/Developer/TraceInsights/Private/Insights/ViewModels/TaskGraphRelation.cpp
+++ b/Engine/Source/Developer/TraceInsights/Private/Insights/ViewModels/TaskGraphRelation.cpp
@@ -60,6 +60,10 @@ void FTaskGraphRelation::Draw(const FDrawContext& DrawContext, const FTimingTrac
 	FLinearColor Color = FTaskGraphProfilerManager::Get()->GetColorForTaskEvent(Type);
 	TArray<FVector2D> ArrowPoints;
 
+	constexpr float ArrowDirectionLen = 15.0f;
+	constexpr float ArrowRotationAngle = 20.0f;
+	FVector2D ArrowDirection(-ArrowDirectionLen, 0.0f);
+
 	if (Distance > LineLenghtAtEnds)
 	{
 		FVector2D SplineStart(StartPoint.X + LineLenghtAtEnds, StartPoint.Y);
@@ -83,21 +87,25 @@ void FTaskGraphRelation::Draw(const FDrawContext& DrawContext, const FTimingTrac
 		ArrowPoints.Add(StartPoint);
 		ArrowPoints.Add(EndPoint);
 		DrawContext.DrawLines(LayerId, 0.0f, 0.0f, ArrowPoints, ESlateDrawEffect::None, Color, /*bAntialias=*/ true, /*Thickness=*/ 2.0f);
+
+		ArrowDirection = StartPoint - EndPoint;
+		ArrowDirection.Normalize();
+		ArrowDirection *= ArrowDirectionLen;
 	}
 
 	FVector2D ArrowOrigin = EndPoint;
 
-	constexpr float ArrowDirectionLen = -15.0f;
-	constexpr float ArrowRotationAngle = 20.0f;
-	FVector2D ArrowDirection(ArrowDirectionLen, 0.0f);
-
 	ArrowPoints.Empty();
-	ArrowPoints.Add(ArrowOrigin + ArrowDirection.GetRotated(ArrowRotationAngle));
 	ArrowPoints.Add(ArrowOrigin);
 	ArrowPoints.Add(ArrowOrigin + ArrowDirection.GetRotated(-ArrowRotationAngle));
 
-	constexpr float YOffset = 1.0f; // Needed to align the arrow with the line. 
-	DrawContext.DrawLines(LayerId, 0.0f, YOffset, ArrowPoints, ESlateDrawEffect::None, Color, /*bAntialias=*/ true, /*Thickness=*/ 2.0f);
+	DrawContext.DrawLines(LayerId, 0.0f, 0.0f, ArrowPoints, ESlateDrawEffect::None, Color, /*bAntialias=*/ true, /*Thickness=*/ 2.0f);
+
+	ArrowPoints.Empty();
+	ArrowPoints.Add(ArrowOrigin);
+	ArrowPoints.Add(ArrowOrigin + ArrowDirection.GetRotated(ArrowRotationAngle));
+
+	DrawContext.DrawLines(LayerId, 0.0f, 0.0f, ArrowPoints, ESlateDrawEffect::None, Color, /*bAntialias=*/ true, /*Thickness=*/ 2.0f);
 }
 
 } // namespace Insights
\ No newline at end of file
