[Insights] Fix crash when starting automation window

diff --git a/Engine/Source/Developer/TraceInsights/Private/Insights/Tests/InsightsTestRunner.cpp b/Engine/Source/Developer/TraceInsights/Private/Insights/Tests/InsightsTestRunner.cpp
--- a/Engine/Source/Developer/TraceInsights/Private/Insights/Tests/InsightsTestRunner.cpp
+++ b/Engine/Source/Developer/TraceInsights/Private/Insights/Tests/InsightsTestRunner.cpp
@@ -11,6 +11,7 @@
 #include "IAutomationWorkerModule.h"
 #include "ISessionManager.h"
 #include "ISessionServicesModule.h"
+#include "Misc/CoreMisc.h"
 #include "Widgets/Docking/SDockTab.h"
 
 // Insights
@@ -63,6 +64,8 @@ void FInsightsTestRunner::Initialize(IUnrealInsightsModule& InsightsModule)
 		IAutomationControllerModule& AutomationControllerModule = FModuleManager::LoadModuleChecked<IAutomationControllerModule>(TEXT("AutomationController"));
 		AutomationControllerModule.Init();
 
+		// Initialize the target platform manager as it is needed by Automation Window
+		GetTargetPlatformManager();
 		FModuleManager::Get().LoadModule("AutomationWindow");
 		FModuleManager::Get().LoadModule("AutomationWorker");
 	}
