The store port argument is optional. This way the reference store client can select a default port and be free to transparently hop ports if need be.

#rnx

diff --git a/Engine/Source/Developer/TraceInsights/Private/Insights/InsightsManager.h b/Engine/Source/Developer/TraceInsights/Private/Insights/InsightsManager.h
--- a/Engine/Source/Developer/TraceInsights/Private/Insights/InsightsManager.h
+++ b/Engine/Source/Developer/TraceInsights/Private/Insights/InsightsManager.h
@@ -100,7 +100,7 @@ public:
 	void SetStoreDir(const FString& InStoreDir) { StoreDir = InStoreDir; }
 	const FString& GetStoreDir() const { return StoreDir; }
 
-	bool ConnectToStore(const TCHAR* Host, uint32 Port);
+	bool ConnectToStore(const TCHAR* Host, uint32 Port=0);
 	UE::Trace::FStoreClient* GetStoreClient() const { return StoreClient.Get(); }
 
 	/** @return an instance of the trace analysis session. */
diff --git a/Engine/Source/Developer/TraceInsights/Private/Insights/TraceInsightsModule.h b/Engine/Source/Developer/TraceInsights/Private/Insights/TraceInsightsModule.h
--- a/Engine/Source/Developer/TraceInsights/Private/Insights/TraceInsightsModule.h
+++ b/Engine/Source/Developer/TraceInsights/Private/Insights/TraceInsightsModule.h
@@ -41,7 +41,7 @@ public:
 	virtual void CreateDefaultStore() override;
 
 	virtual UE::Trace::FStoreClient* GetStoreClient() override;
-	virtual bool ConnectToStore(const TCHAR* InStoreHost, uint32 InStorePort) override;
+	virtual bool ConnectToStore(const TCHAR* InStoreHost, uint32 InStorePort=0) override;
 
 	virtual void CreateSessionBrowser(const FCreateSessionBrowserParams& Params) override;
 	virtual void CreateSessionViewer(bool bAllowDebugTools = false) override;
diff --git a/Engine/Source/Developer/TraceInsights/Public/Insights/IUnrealInsightsModule.h b/Engine/Source/Developer/TraceInsights/Public/Insights/IUnrealInsightsModule.h
--- a/Engine/Source/Developer/TraceInsights/Public/Insights/IUnrealInsightsModule.h
+++ b/Engine/Source/Developer/TraceInsights/Public/Insights/IUnrealInsightsModule.h
@@ -201,7 +201,7 @@ public:
 	 * @param InStorePort The port of the store to connect to.
 	 * @return If connected succesfully or not.
 	 */
-	virtual bool ConnectToStore(const TCHAR* InStoreHost, uint32 InStorePort) = 0;
+	virtual bool ConnectToStore(const TCHAR* InStoreHost, uint32 InStorePort=0) = 0;
 
 	//////////////////////////////////////////////////
 
