Update Unreal Insights' knowledge of where the store directory is located by ask the store where on its local filesystem.

#rnx

diff --git a/Engine/Source/Developer/TraceInsights/Private/Insights/InsightsManager.cpp b/Engine/Source/Developer/TraceInsights/Private/Insights/InsightsManager.cpp
--- a/Engine/Source/Developer/TraceInsights/Private/Insights/InsightsManager.cpp
+++ b/Engine/Source/Developer/TraceInsights/Private/Insights/InsightsManager.cpp
@@ -331,7 +331,19 @@ bool FInsightsManager::ConnectToStore(const TCHAR* Host, uint32 Port)
 	using namespace UE::Trace;
 	FStoreClient* Client = FStoreClient::Connect(Host, Port);
 	StoreClient = TUniquePtr<FStoreClient>(Client);
-	return StoreClient.IsValid();
+	if (!StoreClient.IsValid())
+	{
+		return false;
+	}
+
+	const FStoreClient::FStatus* Status = StoreClient->GetStatus();
+	FString RemoteStoreDir(Status->GetStoreDir());
+	if (RemoteStoreDir.Len() > 0)
+	{
+		SetStoreDir(RemoteStoreDir);
+	}
+
+	return true;
 }
 
 ////////////////////////////////////////////////////////////////////////////////////////////////////
