[Insights]
  - Memory: Fixed sorting by LLM Tag column in Mem Allocs Table (unset values resultin in nullptr strings).
  - Memory: Fixed sorting by Function (Callstack) column in Mem Allocs Table (column type being set as CString while values were actually of type Text).

diff --git a/Engine/Source/Developer/TraceInsights/Private/Insights/MemoryProfiler/ViewModels/MemAllocTable.cpp b/Engine/Source/Developer/TraceInsights/Private/Insights/MemoryProfiler/ViewModels/MemAllocTable.cpp
--- a/Engine/Source/Developer/TraceInsights/Private/Insights/MemoryProfiler/ViewModels/MemAllocTable.cpp
+++ b/Engine/Source/Developer/TraceInsights/Private/Insights/MemoryProfiler/ViewModels/MemAllocTable.cpp
@@ -555,7 +555,7 @@ void FMemAllocTable::AddDefaultColumns()
 		Column.SetHorizontalAlignment(HAlign_Left);
 		Column.SetInitialWidth(550.0f);
 
-		Column.SetDataType(ETableCellDataType::CString);
+		Column.SetDataType(ETableCellDataType::Text);
 
 		Column.SetIsDynamic(true);
 
@@ -633,7 +633,7 @@ void FMemAllocTable::AddDefaultColumns()
 		TSharedRef<ITableCellValueFormatter> Formatter = MakeShared<FunctionValueFormatter>();
 		Column.SetValueFormatter(Formatter);
 
-		TSharedRef<ITableCellValueSorter> Sorter = MakeShared<FSorterByCStringValue>(ColumnRef);
+		TSharedRef<ITableCellValueSorter> Sorter = MakeShared<FSorterByTextValue>(ColumnRef);
 		Column.SetValueSorter(Sorter);
 
 		AddColumn(ColumnRef);
diff --git a/Engine/Source/Developer/TraceInsights/Private/Insights/Table/ViewModels/TableCellValue.h b/Engine/Source/Developer/TraceInsights/Private/Insights/Table/ViewModels/TableCellValue.h
--- a/Engine/Source/Developer/TraceInsights/Private/Insights/Table/ViewModels/TableCellValue.h
+++ b/Engine/Source/Developer/TraceInsights/Private/Insights/Table/ViewModels/TableCellValue.h
@@ -20,8 +20,7 @@ enum class ETableCellDataType : uint32
 
 	// Custom types.
 	Custom,
-	Custom_Text, // FTextCustomTableCellValue
-	Text = Custom_Text,
+	Text, // FTextCustomTableCellValue
 
 	/** Invalid enum type, may be used as a number of enumerations. */
 	InvalidOrMax,
@@ -70,7 +69,7 @@ public:
 	explicit FTableCellValue(double Value) : DataType(ETableCellDataType::Double), Double(Value) {}
 	explicit FTableCellValue(const TCHAR* Value) : DataType(ETableCellDataType::CString), CString(Value) {}
 	explicit FTableCellValue(TSharedPtr<ICustomTableCellValue> Value) : DataType(ETableCellDataType::Custom), Custom(Value) {}
-	explicit FTableCellValue(const FText& Value) : DataType(ETableCellDataType::Custom_Text), Custom(MakeShared<FTextCustomTableCellValue>(Value)) {}
+	explicit FTableCellValue(const FText& Value) : DataType(ETableCellDataType::Text), Custom(MakeShared<FTextCustomTableCellValue>(Value)) {}
 
 	bool AsBool() const
 	{
@@ -159,7 +158,7 @@ public:
 
 	const FText& GetText() const
 	{
-		if (DataType == ETableCellDataType::Custom_Text && Custom.IsValid())
+		if (DataType == ETableCellDataType::Text && Custom.IsValid())
 		{
 			return StaticCastSharedPtr<FTextCustomTableCellValue>(Custom)->GetText();
 		}
diff --git a/Engine/Source/Developer/TraceInsights/Private/Insights/Table/ViewModels/TableCellValueSorter.cpp b/Engine/Source/Developer/TraceInsights/Private/Insights/Table/ViewModels/TableCellValueSorter.cpp
--- a/Engine/Source/Developer/TraceInsights/Private/Insights/Table/ViewModels/TableCellValueSorter.cpp
+++ b/Engine/Source/Developer/TraceInsights/Private/Insights/Table/ViewModels/TableCellValueSorter.cpp
@@ -534,17 +534,18 @@ FSorterByCStringValue::FSorterByCStringValue(TSharedRef<FTableColumn> InColumnRe
 		const TOptional<FTableCellValue> OptionalValueA = Column.GetValue(*A);
 		const TOptional<FTableCellValue> OptionalValueB = Column.GetValue(*B);
 
-		const TCHAR* ValueA = OptionalValueA.IsSet() ? OptionalValueA.GetValue().CString : nullptr;
-		const TCHAR* ValueB = OptionalValueB.IsSet() ? OptionalValueB.GetValue().CString : nullptr;
+		const TCHAR* ValueA = OptionalValueA.IsSet() ? OptionalValueA.GetValue().CString : TEXT("");
+		const TCHAR* ValueB = OptionalValueB.IsSet() ? OptionalValueB.GetValue().CString : TEXT("");
 
-		if (ValueA == ValueB)
+		const int32 Compare = FCString::Strcmp(ValueA, ValueB);
+		if (Compare == 0)
 		{
 			INSIGHTS_DEFAULT_SORTING_NODES(A, B)
 		}
 		else
 		{
 			// Sort by value (ascending).
-			return FCString::Strcmp(ValueA, ValueB) < 0;
+			return Compare < 0;
 		}
 	};
 
@@ -553,17 +554,18 @@ FSorterByCStringValue::FSorterByCStringValue(TSharedRef<FTableColumn> InColumnRe
 		const TOptional<FTableCellValue> OptionalValueA = Column.GetValue(*A);
 		const TOptional<FTableCellValue> OptionalValueB = Column.GetValue(*B);
 
-		const TCHAR* ValueA = OptionalValueA.IsSet() ? OptionalValueA.GetValue().CString : nullptr;
-		const TCHAR* ValueB = OptionalValueB.IsSet() ? OptionalValueB.GetValue().CString : nullptr;
+		const TCHAR* ValueA = OptionalValueA.IsSet() ? OptionalValueA.GetValue().CString : TEXT("");
+		const TCHAR* ValueB = OptionalValueB.IsSet() ? OptionalValueB.GetValue().CString : TEXT("");
 
-		if (ValueA == ValueB)
+		const int32 Compare = FCString::Strcmp(ValueA, ValueB);
+		if (Compare == 0)
 		{
 			INSIGHTS_DEFAULT_SORTING_NODES(A, B)
 		}
 		else
 		{
 			// Sort by value (descending).
-			return FCString::Strcmp(ValueB, ValueA) < 0;
+			return Compare > 0;
 		}
 	};
 }
@@ -625,14 +627,15 @@ FSorterByTextValue::FSorterByTextValue(TSharedRef<FTableColumn> InColumnRef)
 		const FText& ValueA = OptionalValueA.GetValue().GetText();
 		const FText& ValueB = OptionalValueB.GetValue().GetText();
 
-		if (ValueA.EqualTo(ValueB))
+		const int32 Compare = ValueA.CompareTo(ValueB);
+		if (Compare == 0)
 		{
 			INSIGHTS_DEFAULT_SORTING_NODES(A, B)
 		}
 		else
 		{
 			// Sort by value (descending).
-			return ValueB.CompareTo(ValueA) < 0;
+			return Compare > 0;
 		}
 	};
 }
