[Insights] Fix crash when a session is opened via drag and drop whilst task relations are drawed.

diff --git a/Engine/Source/Developer/TraceInsights/Private/Insights/TaskGraphProfiler/ViewModels/TaskTimingTrack.cpp b/Engine/Source/Developer/TraceInsights/Private/Insights/TaskGraphProfiler/ViewModels/TaskTimingTrack.cpp
--- a/Engine/Source/Developer/TraceInsights/Private/Insights/TaskGraphProfiler/ViewModels/TaskTimingTrack.cpp
+++ b/Engine/Source/Developer/TraceInsights/Private/Insights/TaskGraphProfiler/ViewModels/TaskTimingTrack.cpp
@@ -88,6 +88,7 @@ void FTaskTimingSharedState::OnEndSession(Insights::ITimingViewSession& InSessio
 		return;
 	}
 
+	FTaskGraphProfilerManager::Get()->ClearTaskRelations();
 	TaskTrack = nullptr;
 }
 
