[Insights] Fix warning

diff --git a/Engine/Source/Developer/TraceInsights/Private/Insights/MemoryProfiler/ViewModels/MemoryGraphTrack.cpp b/Engine/Source/Developer/TraceInsights/Private/Insights/MemoryProfiler/ViewModels/MemoryGraphTrack.cpp
--- a/Engine/Source/Developer/TraceInsights/Private/Insights/MemoryProfiler/ViewModels/MemoryGraphTrack.cpp
+++ b/Engine/Source/Developer/TraceInsights/Private/Insights/MemoryProfiler/ViewModels/MemoryGraphTrack.cpp
@@ -298,7 +298,6 @@ int32 FMemoryGraphTrack::RemoveAllMemTagSeries()
 		const TSharedPtr<FMemoryGraphSeries> Series = StaticCastSharedPtr<FMemoryGraphSeries>(GraphSeries);
 		return Series->GetTimelineType() == FMemoryGraphSeries::ETimelineType::MemTag;
 	});
-	return Count;
 }
 
 ////////////////////////////////////////////////////////////////////////////////////////////////////
