[Insights] Memory: Added "Open source file of selected callstack frame in <IDE>" in context menu of a Callstack Frame tree item (in Mem Allocs tree table).

diff --git a/Engine/Source/Developer/TraceInsights/Private/Insights/MemoryProfiler/ViewModels/MemAllocGroupingByCallstack.cpp b/Engine/Source/Developer/TraceInsights/Private/Insights/MemoryProfiler/ViewModels/MemAllocGroupingByCallstack.cpp
--- a/Engine/Source/Developer/TraceInsights/Private/Insights/MemoryProfiler/ViewModels/MemAllocGroupingByCallstack.cpp
+++ b/Engine/Source/Developer/TraceInsights/Private/Insights/MemoryProfiler/ViewModels/MemAllocGroupingByCallstack.cpp
@@ -203,6 +203,7 @@ FMemAllocGroupingByCallstack::FCallstackGroup* FMemAllocGroupingByCallstack::Cre
 
 	NewGroupPtr->Node = NodePtr.Get();
 	NewGroupPtr->Node->SetTooltip(GetGroupTooltip(InFrame));
+	NewGroupPtr->Node->SetContext((void*)InFrame);
 
 	return NewGroupPtr;
 }
diff --git a/Engine/Source/Developer/TraceInsights/Private/Insights/MemoryProfiler/Widgets/SMemAllocTableTreeView.cpp b/Engine/Source/Developer/TraceInsights/Private/Insights/MemoryProfiler/Widgets/SMemAllocTableTreeView.cpp
--- a/Engine/Source/Developer/TraceInsights/Private/Insights/MemoryProfiler/Widgets/SMemAllocTableTreeView.cpp
+++ b/Engine/Source/Developer/TraceInsights/Private/Insights/MemoryProfiler/Widgets/SMemAllocTableTreeView.cpp
@@ -3,6 +3,8 @@
 #include "SMemAllocTableTreeView.h"
 
 #include "EditorStyleSet.h"
+#include "ISourceCodeAccessModule.h"
+#include "ISourceCodeAccessor.h"
 #include "SlateOptMacros.h"
 #include "Styling/CoreStyle.h"
 #include "TraceServices/AnalysisService.h"
@@ -636,7 +638,6 @@ FReply SMemAllocTableTreeView::OnSizeViewClicked()
 
 ////////////////////////////////////////////////////////////////////////////////////////////////////
 
-
 FReply SMemAllocTableTreeView::OnMemoryPageClicked()
 {
 	ColumnBeingSorted = FTable::GetHierarchyColumnId();
@@ -975,6 +976,102 @@ ECheckBoxState SMemAllocTableTreeView::CallstackGroupingByFunction_IsChecked() c
 
 ////////////////////////////////////////////////////////////////////////////////////////////////////
 
+void SMemAllocTableTreeView::ExtendMenu(FMenuBuilder& MenuBuilder)
+{
+	MenuBuilder.BeginSection("CallstackFrame", LOCTEXT("ContextMenu_Header_CallstackFrame", "Callstack Frame"));
+	{
+		ISourceCodeAccessModule& SourceCodeAccessModule = FModuleManager::LoadModuleChecked<ISourceCodeAccessModule>("SourceCodeAccess");
+		ISourceCodeAccessor& SourceCodeAccessor = SourceCodeAccessModule.GetAccessor();
+
+		FText ItemLabel = FText::Format(LOCTEXT("ContextMenu_Header_Open", "Open in {0}"), SourceCodeAccessor.GetNameText());
+		FText FileName = GetSelectedCallstackFrameFileName();
+		FText ItemToolTip = FText::Format(LOCTEXT("ContextMenu_Header_Open_Desc", "Open source file of selected callstack frame in {0}.\n{1}"), SourceCodeAccessor.GetNameText(), FileName);
+
+		FUIAction Action_OpenIDE
+		(
+			FExecuteAction::CreateSP(this, &SMemAllocTableTreeView::OpenCallstackFrameSourceFileInIDE),
+			FCanExecuteAction::CreateSP(this, &SMemAllocTableTreeView::CanOpenCallstackFrameSourceFileInIDE)
+		);
+		MenuBuilder.AddMenuEntry
+		(
+			ItemLabel,
+			ItemToolTip,
+#if PLATFORM_WINDOWS
+			FSlateIcon(FAppStyle::GetAppStyleSetName(), "MainFrame.OpenVisualStudio"),
+#else
+			FSlateIcon(),
+#endif
+			Action_OpenIDE,
+			NAME_None,
+			EUserInterfaceActionType::Button
+		);
+	}
+	MenuBuilder.EndSection();
+}
+
+////////////////////////////////////////////////////////////////////////////////////////////////////
+
+bool SMemAllocTableTreeView::CanOpenCallstackFrameSourceFileInIDE() const
+{
+	if (TreeView->GetNumItemsSelected() != 1)
+	{
+		return false;
+	}
+
+	FTableTreeNodePtr TreeNode = TreeView->GetSelectedItems()[0];
+	return TreeNode.IsValid() && TreeNode->IsGroup() && TreeNode->GetContext();
+}
+
+////////////////////////////////////////////////////////////////////////////////////////////////////
+
+void SMemAllocTableTreeView::OpenCallstackFrameSourceFileInIDE()
+{
+	if (TreeView->GetNumItemsSelected() > 0)
+	{
+		FTableTreeNodePtr TreeNode = TreeView->GetSelectedItems()[0];
+		if (TreeNode.IsValid() && TreeNode->IsGroup() && TreeNode->GetContext())
+		{
+			const TraceServices::FStackFrame* Frame = (const TraceServices::FStackFrame*)TreeNode->GetContext();
+
+			if (Frame->Symbol && Frame->Symbol->File)
+			{
+				//if (FPlatformFileManager::Get().GetPlatformFile().FileExists(Frame->Symbol->File))
+				{
+					ISourceCodeAccessModule& SourceCodeAccessModule = FModuleManager::LoadModuleChecked<ISourceCodeAccessModule>("SourceCodeAccess");
+					ISourceCodeAccessor& SourceCodeAccessor = SourceCodeAccessModule.GetAccessor();
+					SourceCodeAccessor.OpenFileAtLine(Frame->Symbol->File, Frame->Symbol->Line);
+				}
+			}
+		}
+	}
+}
+
+////////////////////////////////////////////////////////////////////////////////////////////////////
+
+FText SMemAllocTableTreeView::GetSelectedCallstackFrameFileName() const
+{
+	if (TreeView->GetNumItemsSelected() > 0)
+	{
+		FTableTreeNodePtr TreeNode = TreeView->GetSelectedItems()[0];
+		if (TreeNode.IsValid() && TreeNode->IsGroup() && TreeNode->GetContext())
+		{
+			const TraceServices::FStackFrame* Frame = (const TraceServices::FStackFrame*)TreeNode->GetContext();
+			if (Frame->Symbol && Frame->Symbol->File)
+			{
+				FString SourceFileAndLine = FString::Printf(TEXT("%s(%d)"), Frame->Symbol->File, Frame->Symbol->Line);
+				return FText::FromString(SourceFileAndLine);
+			}
+			else
+			{
+				return LOCTEXT("NoSourceFile", "(source file not available)");
+			}
+		}
+	}
+	return LOCTEXT("NoCallstackFrame", "(only for resolved callstack frames)");
+}
+
+////////////////////////////////////////////////////////////////////////////////////////////////////
+
 } // namespace Insights
 
 #undef LOCTEXT_NAMESPACE
diff --git a/Engine/Source/Developer/TraceInsights/Private/Insights/MemoryProfiler/Widgets/SMemAllocTableTreeView.h b/Engine/Source/Developer/TraceInsights/Private/Insights/MemoryProfiler/Widgets/SMemAllocTableTreeView.h
--- a/Engine/Source/Developer/TraceInsights/Private/Insights/MemoryProfiler/Widgets/SMemAllocTableTreeView.h
+++ b/Engine/Source/Developer/TraceInsights/Private/Insights/MemoryProfiler/Widgets/SMemAllocTableTreeView.h
@@ -93,6 +93,11 @@ public:
 protected:
 	virtual void InternalCreateGroupings() override;
 
+	virtual void ExtendMenu(FMenuBuilder& Menu) override;
+	bool CanOpenCallstackFrameSourceFileInIDE() const;
+	void OpenCallstackFrameSourceFileInIDE();
+	FText GetSelectedCallstackFrameFileName() const;
+
 private:
 	void OnQueryInvalidated();
 	void StartQuery();
diff --git a/Engine/Source/Developer/TraceInsights/TraceInsights.Build.cs b/Engine/Source/Developer/TraceInsights/TraceInsights.Build.cs
--- a/Engine/Source/Developer/TraceInsights/TraceInsights.Build.cs
+++ b/Engine/Source/Developer/TraceInsights/TraceInsights.Build.cs
@@ -28,6 +28,7 @@ public class TraceInsights : ModuleRules
 				"RHI",
 				"Slate",
 				"Sockets",
+				"SourceCodeAccess",
 				"TraceAnalysis",
 				"TraceLog",
 				"TraceServices",
