[Insights] Missing file to previous commit.

diff --git a/Engine/Source/Developer/TraceInsights/Private/Insights/Table/ViewModels/BaseTreeNode.h b/Engine/Source/Developer/TraceInsights/Private/Insights/Table/ViewModels/BaseTreeNode.h
--- a/Engine/Source/Developer/TraceInsights/Private/Insights/Table/ViewModels/BaseTreeNode.h
+++ b/Engine/Source/Developer/TraceInsights/Private/Insights/Table/ViewModels/BaseTreeNode.h
@@ -32,16 +32,19 @@ class FBaseTreeNode : public TSharedFromThis<FBaseTreeNode>
 protected:
 	struct FGroupNodeData
 	{
-		/** Children of this node. */
+		/** Children of the group node. */
 		TArray<FBaseTreeNodePtr> Children;
 
-		/** Filtered children of this node. */
+		/** Filtered children of the group node. */
 		TArray<FBaseTreeNodePtr> FilteredChildren;
 
-		/** Descriptive text for this node to display in a tooltip. */
+		/** Descriptive text for the group node to display in a tooltip. */
 		FText Tooltip;
 
-		/** Whether this group node should be expanded or not. */
+		/** A generic pointer to a data context for the group node. Can be nullptr. */
+		void* Context = nullptr;
+
+		/** Whether the group node should be expanded or not. */
 		bool bIsExpanded = false;
 	};
 
@@ -98,6 +101,16 @@ public:
 	 */
 	virtual void SetTooltip(const FText& InTooltip) { if (IsGroup()) { GroupData->Tooltip = InTooltip; } }
 
+	/**
+	 * @return a pointer to a data context for this node.
+	 */
+	virtual void* GetContext() const { return IsGroup() ? GroupData->Context : nullptr; }
+
+	/**
+	 * Sets a pointer to a data context for this node.
+	 */
+	virtual void SetContext(void* InContext) { if (IsGroup()) { GroupData->Context = InContext; } }
+
 	/**
 	 * @return true, if this node is a group node.
 	 */
