[Insights] Fix warning caused by enum value name conflict

diff --git a/Engine/Source/Developer/TraceInsights/Private/Insights/MemoryProfiler/ViewModels/MemoryGraphTrack.h b/Engine/Source/Developer/TraceInsights/Private/Insights/MemoryProfiler/ViewModels/MemoryGraphTrack.h
--- a/Engine/Source/Developer/TraceInsights/Private/Insights/MemoryProfiler/ViewModels/MemoryGraphTrack.h
+++ b/Engine/Source/Developer/TraceInsights/Private/Insights/MemoryProfiler/ViewModels/MemoryGraphTrack.h
@@ -74,7 +74,7 @@ private:
 
 ////////////////////////////////////////////////////////////////////////////////////////////////////
 
-enum EMemoryTrackHeightMode
+enum class EMemoryTrackHeightMode
 {
 	Small = 0,
 	Medium,
