[Insights] Added functionality to export a TreeTableView to a file.

diff --git a/Engine/Source/Developer/TraceInsights/Private/Insights/Table/ViewModels/Table.cpp b/Engine/Source/Developer/TraceInsights/Private/Insights/Table/ViewModels/Table.cpp
--- a/Engine/Source/Developer/TraceInsights/Private/Insights/Table/ViewModels/Table.cpp
+++ b/Engine/Source/Developer/TraceInsights/Private/Insights/Table/ViewModels/Table.cpp
@@ -128,17 +128,18 @@ void FTable::GetVisibleColumns(TArray<TSharedRef<FTableColumn>>& InArray) const
 
 ////////////////////////////////////////////////////////////////////////////////////////////////////
 
-void FTable::GetVisibleColumnsData(const TArray<Insights::FBaseTreeNodePtr>& InNodes, const FName& LogListingName, FString& OutData) const
+void FTable::GetVisibleColumnsData(const TArray<Insights::FBaseTreeNodePtr>& InNodes, const FName& LogListingName, TCHAR Separator, bool bIncludeHeaders, FString& OutData) const
 {
-	constexpr TCHAR Separator = TEXT('\t');
-
 	TArray<TSharedRef<Insights::FTableColumn>> VisibleColumns;
 	GetVisibleColumns(VisibleColumns);
 
 	// Table headers
-	for (const TSharedRef<Insights::FTableColumn>& ColumnRef : VisibleColumns)
+	if (bIncludeHeaders)
 	{
-		OutData += ColumnRef->GetShortName().ToString().ReplaceCharWithEscapedChar() + Separator;
+		for (const TSharedRef<Insights::FTableColumn>& ColumnRef : VisibleColumns)
+		{
+			OutData += ColumnRef->GetShortName().ToString().ReplaceCharWithEscapedChar() + Separator;
+		}
 	}
 
 	if (OutData.Len() > 0)
diff --git a/Engine/Source/Developer/TraceInsights/Private/Insights/Table/ViewModels/Table.h b/Engine/Source/Developer/TraceInsights/Private/Insights/Table/ViewModels/Table.h
--- a/Engine/Source/Developer/TraceInsights/Private/Insights/Table/ViewModels/Table.h
+++ b/Engine/Source/Developer/TraceInsights/Private/Insights/Table/ViewModels/Table.h
@@ -51,7 +51,7 @@ public:
 
 	int32 GetColumnPositionIndex(const FName& ColumnId) const;
 
-	void GetVisibleColumnsData(const TArray<TSharedPtr<class FBaseTreeNode>>& InNodes, const FName& LogListingName, FString& OutData) const;
+	void GetVisibleColumnsData(const TArray<TSharedPtr<class FBaseTreeNode>>& InNodes, const FName& LogListingName, TCHAR Separator, bool bIncludeHeaders, FString& OutData) const;
 
 	static const FName GetHierarchyColumnId();
 
diff --git a/Engine/Source/Developer/TraceInsights/Private/Insights/Table/Widgets/STableTreeView.cpp b/Engine/Source/Developer/TraceInsights/Private/Insights/Table/Widgets/STableTreeView.cpp
--- a/Engine/Source/Developer/TraceInsights/Private/Insights/Table/Widgets/STableTreeView.cpp
+++ b/Engine/Source/Developer/TraceInsights/Private/Insights/Table/Widgets/STableTreeView.cpp
@@ -2,13 +2,16 @@
 
 #include "STableTreeView.h"
 
+#include "DesktopPlatformModule.h"
 #include "EditorStyleSet.h"
 #include "Framework/Commands/Commands.h"
 #include "Framework/Commands/UICommandList.h"
 #include "Framework/Application/SlateApplication.h"
 #include "Framework/MultiBox/MultiBoxBuilder.h"
 #include "HAL/PlatformApplicationMisc.h"
+#include "HAL/PlatformFileManager.h"
 #include "Layout/WidgetPath.h"
+#include "Logging/MessageLog.h"
 #include "SlateOptMacros.h"
 #include "TraceServices/AnalysisService.h"
 #include "Widgets/Input/SCheckBox.h"
@@ -66,6 +69,8 @@ public:
 		UI_COMMAND(Command_ExpandSubtree, "Expand Subtree", "Expand the subtree that starts from the selected group node.", EUserInterfaceActionType::Button, FInputChord(EModifierKey::None, EKeys::E));
 		UI_COMMAND(Command_ExpandCriticalPath, "Expand Critical Path", "Expand the first group child node recursively until a leaf nodes in reached.", EUserInterfaceActionType::Button, FInputChord(EModifierKey::None, EKeys::R));
 		UI_COMMAND(Command_CollapseSubtree, "Collapse Subtree", "Collapse the subtree that starts from the selected group node.", EUserInterfaceActionType::Button, FInputChord(EModifierKey::None, EKeys::C));
+		UI_COMMAND(Command_ExportToFile, "Export To File", "Export the table content to a file.", EUserInterfaceActionType::Button, FInputChord(EModifierKey::Control, EKeys::E));
+		UI_COMMAND(Command_ExportExpandedToFile, "Export Expanded To File", "Export the fully expanded table to a file.", EUserInterfaceActionType::Button, FInputChord(EModifierKey::Shift, EKeys::E));
 	}
 	PRAGMA_ENABLE_OPTIMIZATION
 
@@ -75,6 +80,8 @@ public:
 	TSharedPtr<FUICommandInfo> Command_ExpandSubtree;
 	TSharedPtr<FUICommandInfo> Command_ExpandCriticalPath;
 	TSharedPtr<FUICommandInfo> Command_CollapseSubtree;
+	TSharedPtr<FUICommandInfo> Command_ExportToFile;
+	TSharedPtr<FUICommandInfo> Command_ExportExpandedToFile;
 };
 
 const FName STableTreeView::RootNodeName(TEXT("Root"));
@@ -341,6 +348,8 @@ void STableTreeView::InitCommandList()
 	CommandList->MapAction(FTableTreeViewCommands::Get().Command_ExpandSubtree, FExecuteAction::CreateSP(this, &STableTreeView::ContextMenu_ExpandSubtree_Execute), FCanExecuteAction::CreateSP(this, &STableTreeView::ContextMenu_ExpandSubtree_CanExecute));
 	CommandList->MapAction(FTableTreeViewCommands::Get().Command_ExpandCriticalPath, FExecuteAction::CreateSP(this, &STableTreeView::ContextMenu_ExpandCriticalPath_Execute), FCanExecuteAction::CreateSP(this, &STableTreeView::ContextMenu_ExpandCriticalPath_CanExecute));
 	CommandList->MapAction(FTableTreeViewCommands::Get().Command_CollapseSubtree, FExecuteAction::CreateSP(this, &STableTreeView::ContextMenu_CollapseSubtree_Execute), FCanExecuteAction::CreateSP(this, &STableTreeView::ContextMenu_CollapseSubtree_CanExecute));
+	CommandList->MapAction(FTableTreeViewCommands::Get().Command_ExportToFile, FExecuteAction::CreateSP(this, &STableTreeView::ContextMenu_ExportToFile_Execute, true), FCanExecuteAction::CreateSP(this, &STableTreeView::ContextMenu_ExportToFile_CanExecute));
+	CommandList->MapAction(FTableTreeViewCommands::Get().Command_ExportExpandedToFile, FExecuteAction::CreateSP(this, &STableTreeView::ContextMenu_ExportToFile_Execute, false), FCanExecuteAction::CreateSP(this, &STableTreeView::ContextMenu_ExportToFile_CanExecute));
 }
 
 ////////////////////////////////////////////////////////////////////////////////////////////////////
@@ -477,6 +486,15 @@ TSharedPtr<SWidget> STableTreeView::TreeView_GetMenuContent()
 			false,
 			FSlateIcon(FEditorStyle::GetStyleSetName(), "Profiler.Misc.SortBy")
 		);
+
+		MenuBuilder.AddSubMenu
+		(
+			LOCTEXT("ContextMenu_Header_Misc_Export", "Export"),
+			LOCTEXT("ContextMenu_Header_Misc_Sort_Desc", "Export to file"),
+			FNewMenuDelegate::CreateSP(this, &STableTreeView::TreeView_BuildExportMenu),
+			false,
+			FSlateIcon(FEditorStyle::GetStyleSetName(), "Profiler.Misc.CopyToClipboard")
+		);
 	}
 	MenuBuilder.EndSection();
 
@@ -612,6 +630,34 @@ void STableTreeView::TreeView_BuildViewColumnMenu(FMenuBuilder& MenuBuilder)
 
 ////////////////////////////////////////////////////////////////////////////////////////////////////
 
+void STableTreeView::TreeView_BuildExportMenu(FMenuBuilder& MenuBuilder)
+{
+	MenuBuilder.BeginSection("Export", LOCTEXT("Export", "Export"));
+
+	MenuBuilder.AddMenuEntry
+	(
+		FTableTreeViewCommands::Get().Command_ExportToFile,
+		NAME_None,
+		TAttribute<FText>(),
+		TAttribute<FText>(),
+		FSlateIcon(FEditorStyle::GetStyleSetName(), "Profiler.Misc.CopyToClipboard")
+	);
+
+	MenuBuilder.AddMenuEntry
+	(
+		FTableTreeViewCommands::Get().Command_ExportExpandedToFile,
+		NAME_None,
+		TAttribute<FText>(),
+		TAttribute<FText>(),
+		FSlateIcon(FEditorStyle::GetStyleSetName(), "Profiler.Misc.CopyToClipboard")
+	);
+
+	MenuBuilder.EndSection();
+}
+
+////////////////////////////////////////////////////////////////////////////////////////////////////
+
+
 void STableTreeView::InitializeAndShowHeaderColumns()
 {
 	for (const TSharedRef<FTableColumn>& ColumnRef : Table->GetColumns())
@@ -2845,7 +2891,7 @@ void STableTreeView::ContextMenu_CopySelectedToClipboard_Execute()
 		CurrentSorter->Sort(SelectedNodes, ColumnSortMode == EColumnSortMode::Ascending ? Insights::ESortMode::Ascending : Insights::ESortMode::Descending);
 	}
 
-	Table->GetVisibleColumnsData(SelectedNodes, GetLogListingName(), ClipboardText);
+	Table->GetVisibleColumnsData(SelectedNodes, GetLogListingName(), TEXT('\t'), true, ClipboardText);
 
 	if (ClipboardText.Len() > 0)
 	{
@@ -3012,6 +3058,113 @@ void STableTreeView::ContextMenu_ExpandCriticalPath_Execute()
 
 ////////////////////////////////////////////////////////////////////////////////////////////////////
 
+bool STableTreeView::ContextMenu_ExportToFile_CanExecute() const
+{
+	return !FilteredGroupNodes.IsEmpty();
+}
+
+////////////////////////////////////////////////////////////////////////////////////////////////////
+
+void STableTreeView::ContextMenu_ExportToFile_Execute(bool bExportOnlyExpanded)
+{
+	TArray<FString> SaveFilenames;
+	IDesktopPlatform* DesktopPlatform = FDesktopPlatformModule::Get();
+	bool bDialogResult = false;
+	if (DesktopPlatform)
+	{
+		const FString DefaultBrowsePath = FPaths::ProjectLogDir();
+		bDialogResult = DesktopPlatform->SaveFileDialog(
+			FSlateApplication::Get().FindBestParentWindowHandleForDialogs(nullptr),
+			LOCTEXT("ExportFileTitle", "Export Table").ToString(),
+			DefaultBrowsePath,
+			TEXT(""),
+			TEXT("Tab-separated values file (*.tsv)|*.tsv|Comma-separated values file (*.csv)|*.csv"),
+			EFileDialogFlags::None,
+			SaveFilenames
+		);
+	}
+
+	if (!bDialogResult || SaveFilenames.Num() == 0)
+	{
+		return;
+	}
+
+	FString& Path = SaveFilenames[0];
+	TCHAR Separator = TEXT('\t');
+	if (Path.EndsWith(TEXT(".csv")))
+	{
+		Separator = TEXT(',');
+	}
+
+	IFileHandle* ExportFileHandle = FPlatformFileManager::Get().GetPlatformFile().OpenWrite(*Path);
+
+	if (ExportFileHandle == nullptr)
+	{
+		FMessageLog ReportMessageLog((LogListingName != NAME_None) ? LogListingName : TEXT("Other"));
+		ReportMessageLog.Error(LOCTEXT("FailedToOpenFile", "Export failed. Failed to open file for write."));
+		ReportMessageLog.Notify();
+		return;
+	}
+
+	FStopwatch Stopwatch;
+	Stopwatch.Start();
+
+	UTF16CHAR BOM = UNICODE_BOM;
+	ExportFileHandle->Write((uint8*)&BOM, sizeof(UTF16CHAR));
+
+	bool bIncludeHeaders = true;
+	auto WriteToFile = [this, &Path, &bIncludeHeaders, ExportFileHandle, Separator](TArray<Insights::FBaseTreeNodePtr>& InNodes)
+	{
+		FString Text;
+		this->GetTable()->GetVisibleColumnsData(InNodes, GetLogListingName(), Separator, bIncludeHeaders, Text);
+		// Only write the headers once, at the top of the file
+		bIncludeHeaders = false;
+
+		// Note: This is a no-op on platforms that are using a 16-bit TCHAR
+		FTCHARToUTF16 UTF16String(*Text, Text.Len());
+		ExportFileHandle->Write((uint8*)*Text, Text.Len() * sizeof(UTF16CHAR));
+
+		InNodes.Empty(InNodes.Num());
+	};
+
+	TArray<Insights::FBaseTreeNodePtr> NodesBuffer;
+	ExportToFileRec(Root, NodesBuffer, bExportOnlyExpanded, WriteToFile);
+
+	// Write the remaining lines
+	WriteToFile(NodesBuffer);
+
+	ExportFileHandle->Flush();
+	delete ExportFileHandle;
+	ExportFileHandle = nullptr;
+
+	Stopwatch.Stop();
+	const double TotalTime = Stopwatch.GetAccumulatedTime();
+	UE_LOG(TraceInsights, Log, TEXT("Exported table to file in %.3fs."), TotalTime);
+}
+
+////////////////////////////////////////////////////////////////////////////////////////////////////
+
+void STableTreeView::ExportToFileRec(const FBaseTreeNodePtr& InGroupNode, TArray<Insights::FBaseTreeNodePtr>& InNodes, bool bInExportOnlyExpanded, WriteToFileCallback Callback)
+{
+	constexpr int MaxBufferSize = 10000;
+
+	for (const FBaseTreeNodePtr& Node : InGroupNode->GetFilteredChildren())
+	{
+		if (InNodes.Num() > MaxBufferSize)
+		{
+			// Flush the buffer to the file
+			Callback(InNodes);
+		}
+		InNodes.Add(Node);
+		if (Node->IsGroup() && (!bInExportOnlyExpanded || TreeView->IsItemExpanded(StaticCastSharedPtr<FTableTreeNode>(Node))))
+		{
+			ExportToFileRec(Node, InNodes, bInExportOnlyExpanded, Callback);
+		}
+	}
+}
+
+////////////////////////////////////////////////////////////////////////////////////////////////////
+
 } // namespace Insights
 
 #undef LOCTEXT_NAMESPACE
diff --git a/Engine/Source/Developer/TraceInsights/Private/Insights/Table/Widgets/STableTreeView.h b/Engine/Source/Developer/TraceInsights/Private/Insights/Table/Widgets/STableTreeView.h
--- a/Engine/Source/Developer/TraceInsights/Private/Insights/Table/Widgets/STableTreeView.h
+++ b/Engine/Source/Developer/TraceInsights/Private/Insights/Table/Widgets/STableTreeView.h
@@ -158,6 +158,7 @@ protected:
 	TSharedPtr<SWidget> TreeView_GetMenuContent();
 	void TreeView_BuildSortByMenu(FMenuBuilder& MenuBuilder);
 	void TreeView_BuildViewColumnMenu(FMenuBuilder& MenuBuilder);
+	void TreeView_BuildExportMenu(FMenuBuilder& MenuBuilder);
 
 	bool ContextMenu_CopySelectedToClipboard_CanExecute() const;
 	void ContextMenu_CopySelectedToClipboard_Execute();
@@ -171,6 +172,8 @@ protected:
 	void ContextMenu_ExpandCriticalPath_Execute();
 	bool ContextMenu_CollapseSubtree_CanExecute() const;
 	void ContextMenu_CollapseSubtree_Execute();
+	bool ContextMenu_ExportToFile_CanExecute() const;
+	void ContextMenu_ExportToFile_Execute(bool bExportOnlyExpanded);
 
 	////////////////////////////////////////////////////////////////////////////////////////////////////
 	// Tree View - Columns' Header
@@ -391,6 +394,9 @@ protected:
 
 	virtual void ExtendMenu(FMenuBuilder& Menu) {}
 
+	typedef TFunctionRef<void(TArray<Insights::FBaseTreeNodePtr>& InNodes)> WriteToFileCallback;
+	void ExportToFileRec(const FBaseTreeNodePtr& InGroupNode, TArray<Insights::FBaseTreeNodePtr>& InNodes, bool bInExportOnlyExpanded, WriteToFileCallback Callback);
+
 protected:
 	/** Table view model. */
 	TSharedPtr<Insights::FTable> Table;
diff --git a/Engine/Source/Developer/TraceInsights/Private/Insights/Widgets/SStatsView.cpp b/Engine/Source/Developer/TraceInsights/Private/Insights/Widgets/SStatsView.cpp
--- a/Engine/Source/Developer/TraceInsights/Private/Insights/Widgets/SStatsView.cpp
+++ b/Engine/Source/Developer/TraceInsights/Private/Insights/Widgets/SStatsView.cpp
@@ -2110,7 +2110,7 @@ void SStatsView::ContextMenu_CopySelectedToClipboard_Execute()
 		CurrentSorter->Sort(SelectedNodes, ColumnSortMode == EColumnSortMode::Ascending ? Insights::ESortMode::Ascending : Insights::ESortMode::Descending);
 	}
 
-	Table->GetVisibleColumnsData(SelectedNodes, FTimingProfilerManager::Get()->GetLogListingName(), ClipboardText);
+	Table->GetVisibleColumnsData(SelectedNodes, FTimingProfilerManager::Get()->GetLogListingName(), TEXT('\t'), true, ClipboardText);
 
 	if (ClipboardText.Len() > 0)
 	{
diff --git a/Engine/Source/Developer/TraceInsights/Private/Insights/Widgets/STimerTreeView.cpp b/Engine/Source/Developer/TraceInsights/Private/Insights/Widgets/STimerTreeView.cpp
--- a/Engine/Source/Developer/TraceInsights/Private/Insights/Widgets/STimerTreeView.cpp
+++ b/Engine/Source/Developer/TraceInsights/Private/Insights/Widgets/STimerTreeView.cpp
@@ -1106,7 +1106,7 @@ void STimerTreeView::ContextMenu_CopySelectedToClipboard_Execute()
 		CurrentSorter->Sort(SelectedNodes, ColumnSortMode == EColumnSortMode::Ascending ? Insights::ESortMode::Ascending : Insights::ESortMode::Descending);
 	}
 
-	Table->GetVisibleColumnsData(SelectedNodes, FTimingProfilerManager::Get()->GetLogListingName(), ClipboardText);
+	Table->GetVisibleColumnsData(SelectedNodes, FTimingProfilerManager::Get()->GetLogListingName(), TEXT('\t'), true, ClipboardText);
 
 	if (ClipboardText.Len() > 0)
 	{
diff --git a/Engine/Source/Developer/TraceInsights/Private/Insights/Widgets/STimersView.cpp b/Engine/Source/Developer/TraceInsights/Private/Insights/Widgets/STimersView.cpp
--- a/Engine/Source/Developer/TraceInsights/Private/Insights/Widgets/STimersView.cpp
+++ b/Engine/Source/Developer/TraceInsights/Private/Insights/Widgets/STimersView.cpp
@@ -2059,7 +2059,7 @@ void STimersView::ContextMenu_CopySelectedToClipboard_Execute()
 		CurrentSorter->Sort(SelectedNodes, ColumnSortMode == EColumnSortMode::Ascending ? Insights::ESortMode::Ascending : Insights::ESortMode::Descending);
 	}
 
-	Table->GetVisibleColumnsData(SelectedNodes, FTimingProfilerManager::Get()->GetLogListingName(), ClipboardText);
+	Table->GetVisibleColumnsData(SelectedNodes, FTimingProfilerManager::Get()->GetLogListingName(), TEXT('\t'), true, ClipboardText);
 
 	if (ClipboardText.Len() > 0)
 	{
