Fix for DisableUnity compile error.

diff --git a/Engine/Source/Developer/TraceInsights/Private/Insights/MemoryProfiler/Widgets/SMemAllocTableTreeView.cpp b/Engine/Source/Developer/TraceInsights/Private/Insights/MemoryProfiler/Widgets/SMemAllocTableTreeView.cpp
--- a/Engine/Source/Developer/TraceInsights/Private/Insights/MemoryProfiler/Widgets/SMemAllocTableTreeView.cpp
+++ b/Engine/Source/Developer/TraceInsights/Private/Insights/MemoryProfiler/Widgets/SMemAllocTableTreeView.cpp
@@ -3,8 +3,10 @@
 #include "SMemAllocTableTreeView.h"
 
 #include "EditorStyleSet.h"
+#include "Framework/MultiBox/MultiBoxBuilder.h"
 #include "ISourceCodeAccessModule.h"
 #include "ISourceCodeAccessor.h"
+#include "Modules/ModuleManager.h"
 #include "SlateOptMacros.h"
 #include "Styling/CoreStyle.h"
 #include "TraceServices/AnalysisService.h"
