[Insights] Fixed how callstack top function is determined. Instead of skipping exactly two frames, it now skips only the top frames with resolved symbol starting with "FMemory::" or "FMallocWrapper::".

diff --git a/Engine/Source/Developer/TraceInsights/Private/Insights/MemoryProfiler/ViewModels/MemAllocTable.cpp b/Engine/Source/Developer/TraceInsights/Private/Insights/MemoryProfiler/ViewModels/MemAllocTable.cpp
--- a/Engine/Source/Developer/TraceInsights/Private/Insights/MemoryProfiler/ViewModels/MemAllocTable.cpp
+++ b/Engine/Source/Developer/TraceInsights/Private/Insights/MemoryProfiler/ViewModels/MemAllocTable.cpp
@@ -2,6 +2,7 @@
 
 #include "MemAllocTable.h"
 #include "CallstackFormatting.h"
+#include "Containers/StringView.h"
 
 #include "Widgets/SBoxPanel.h"
 #include "Widgets/SToolTip.h"
@@ -584,8 +585,35 @@ void FMemAllocTable::AddDefaultColumns()
 						if (Callstack)
 						{
 							check(Callstack->Num() > 0);
-							const TraceServices::FStackFrame* Frame = Callstack->Frame(FMath::Min(2u, Callstack->Num() - 1));
-							check(Frame != nullptr);
+
+							const TraceServices::FStackFrame* Frame = nullptr;
+							for (uint32 FrameIndex = 0; FrameIndex < Callstack->Num(); ++FrameIndex)
+							{
+								Frame = Callstack->Frame(FrameIndex);
+								check(Frame != nullptr);
+
+								if (!Frame->Symbol || !Frame->Symbol->Name)
+								{
+									break;
+								}
+
+								FStringView IgnoreSymbols[] = { TEXT("FMemory::"_SV), TEXT("FMallocWrapper::"_SV) };
+
+								bool bMatchIgnoreSymbol = false;
+								for (uint32 IgnoreSymbolsIndex = 0; IgnoreSymbolsIndex < UE_ARRAY_COUNT(IgnoreSymbols); ++IgnoreSymbolsIndex)
+								{
+									if (FCString::Strnicmp(Frame->Symbol->Name, IgnoreSymbols[IgnoreSymbolsIndex].GetData(), IgnoreSymbols[IgnoreSymbolsIndex].Len()) == 0)
+									{
+										bMatchIgnoreSymbol = true;
+										break;
+									}
+								}
+
+								if (!bMatchIgnoreSymbol)
+								{
+									break;
+								}
+							}
 
 							TStringBuilder<1024> Str;
 							FormatStackFrame(*Frame, Str, EStackFrameFormatFlags::Module);
