Fixed segfault if UnrealTraceServer was terminated while in use.

#rnx

diff --git a/Engine/Source/Developer/TraceInsights/Private/Insights/StoreService/StoreBrowser.cpp b/Engine/Source/Developer/TraceInsights/Private/Insights/StoreService/StoreBrowser.cpp
--- a/Engine/Source/Developer/TraceInsights/Private/Insights/StoreService/StoreBrowser.cpp
+++ b/Engine/Source/Developer/TraceInsights/Private/Insights/StoreService/StoreBrowser.cpp
@@ -101,7 +101,7 @@ void FStoreBrowser::UpdateTraces()
 	// Check if the list of trace sessions has changed.
 	{
 		const UE::Trace::FStoreClient::FStatus* Status = StoreClient->GetStatus();
-		if (StoreChangeSerial != Status->GetChangeSerial())
+		if (Status != nullptr && StoreChangeSerial != Status->GetChangeSerial())
 		{
 			StoreChangeSerial = Status->GetChangeSerial();
 
