[Insights] Initial support for context switches visualization.

diff --git a/Engine/Source/Developer/TraceInsights/Private/Insights/ContextSwitches/ContextSwitchesProfilerManager.cpp b/Engine/Source/Developer/TraceInsights/Private/Insights/ContextSwitches/ContextSwitchesProfilerManager.cpp
new file mode 100644
--- /dev/null
+++ b/Engine/Source/Developer/TraceInsights/Private/Insights/ContextSwitches/ContextSwitchesProfilerManager.cpp
@@ -0,0 +1,205 @@
+// Copyright Epic Games, Inc. All Rights Reserved.
+
+#include "ContextSwitchesProfilerManager.h"
+
+#include "Features/IModularFeatures.h"
+#include "Modules/ModuleManager.h"
+#include "TraceServices/AnalysisService.h"
+#include "TraceServices/Model/ContextSwitches.h"
+
+#include "Insights/ContextSwitches/ViewModels/ContextSwitchesTimingTrack.h"
+#include "Insights/InsightsStyle.h"
+#include "Insights/TimingProfilerManager.h"
+#include "Insights/ViewModels/ThreadTimingTrack.h"
+#include "Insights/ViewModels/ThreadTrackEvent.h"
+#include "Insights/Widgets/STimingProfilerWindow.h"
+#include "Insights/Widgets/STimingView.h"
+
+////////////////////////////////////////////////////////////////////////////////////////////////////
+
+#define LOCTEXT_NAMESPACE "ContextSwitchesProfilerManager"
+
+namespace Insights
+{
+
+////////////////////////////////////////////////////////////////////////////////////////////////////
+
+TSharedPtr<FContextSwitchesProfilerManager> FContextSwitchesProfilerManager::Instance = nullptr;
+
+////////////////////////////////////////////////////////////////////////////////////////////////////
+
+TSharedPtr<FContextSwitchesProfilerManager> FContextSwitchesProfilerManager::Get()
+{
+	return FContextSwitchesProfilerManager::Instance;
+}
+
+////////////////////////////////////////////////////////////////////////////////////////////////////
+
+TSharedPtr<FContextSwitchesProfilerManager> FContextSwitchesProfilerManager::CreateInstance()
+{
+	ensure(!FContextSwitchesProfilerManager::Instance.IsValid());
+	if (FContextSwitchesProfilerManager::Instance.IsValid())
+	{
+		FContextSwitchesProfilerManager::Instance.Reset();
+	}
+
+	FContextSwitchesProfilerManager::Instance = MakeShared<FContextSwitchesProfilerManager>(FInsightsManager::Get()->GetCommandList());
+
+	return FContextSwitchesProfilerManager::Instance;
+}
+
+////////////////////////////////////////////////////////////////////////////////////////////////////
+
+FContextSwitchesProfilerManager::FContextSwitchesProfilerManager(TSharedRef<FUICommandList> InCommandList)
+	: bIsInitialized(false)
+	, bIsAvailable(false)
+{
+}
+
+////////////////////////////////////////////////////////////////////////////////////////////////////
+
+void FContextSwitchesProfilerManager::Initialize(IUnrealInsightsModule& InsightsModule)
+{
+	ensure(!bIsInitialized);
+	if (bIsInitialized)
+	{
+		return;
+	}
+	bIsInitialized = true;
+
+	// Register tick functions.
+	OnTick = FTickerDelegate::CreateSP(this, &FContextSwitchesProfilerManager::Tick);
+	OnTickHandle = FTicker::GetCoreTicker().AddTicker(OnTick, 0.0f);
+
+	FOnRegisterMajorTabExtensions* TimingProfilerLayoutExtension = InsightsModule.FindMajorTabLayoutExtension(FInsightsManagerTabs::TimingProfilerTabId);
+	if (TimingProfilerLayoutExtension)
+	{
+		TimingProfilerLayoutExtension->AddRaw(this, &FContextSwitchesProfilerManager::RegisterTimingProfilerLayoutExtensions);
+	}
+
+	FInsightsManager::Get()->GetSessionChangedEvent().AddSP(this, &FContextSwitchesProfilerManager::OnSessionChanged);
+	OnSessionChanged();
+}
+
+////////////////////////////////////////////////////////////////////////////////////////////////////
+
+void FContextSwitchesProfilerManager::Shutdown()
+{
+	if (!bIsInitialized)
+	{
+		return;
+	}
+	bIsInitialized = false;
+
+	FInsightsManager::Get()->GetSessionChangedEvent().RemoveAll(this);
+
+	// Unregister tick function.
+	FTicker::GetCoreTicker().RemoveTicker(OnTickHandle);
+
+	FContextSwitchesProfilerManager::Instance.Reset();
+}
+
+////////////////////////////////////////////////////////////////////////////////////////////////////
+
+FContextSwitchesProfilerManager::~FContextSwitchesProfilerManager()
+{
+	ensure(!bIsInitialized);
+
+	if (ContextSwitchesSharedState.IsValid())
+	{
+		IModularFeatures::Get().UnregisterModularFeature(Insights::TimingViewExtenderFeatureName, ContextSwitchesSharedState.Get());
+	}
+}
+
+////////////////////////////////////////////////////////////////////////////////////////////////////
+
+void FContextSwitchesProfilerManager::RegisterMajorTabs(IUnrealInsightsModule& InsightsModule)
+{
+}
+
+////////////////////////////////////////////////////////////////////////////////////////////////////
+
+void FContextSwitchesProfilerManager::UnregisterMajorTabs()
+{
+}
+
+////////////////////////////////////////////////////////////////////////////////////////////////////
+
+bool FContextSwitchesProfilerManager::Tick(float DeltaTime)
+{
+	// Check if session has Memory events (to spawn the tab), but not too often.
+	if (!bIsAvailable && AvailabilityCheck.Tick())
+	{
+		bool bShouldBeAvailable = false;
+
+		TSharedPtr<const TraceServices::IAnalysisSession> Session = FInsightsManager::Get()->GetSession();
+		if (Session.IsValid())
+		{
+			TraceServices::FAnalysisSessionReadScope SessionReadScope(*Session.Get());
+
+			const TraceServices::IContextSwitchProvider* ContextSwitchesProvider = TraceServices::ReadContextSwitchProvider(*Session.Get());
+			if (ContextSwitchesProvider && ContextSwitchesProvider->HasData())
+			{
+				TSharedPtr<STimingProfilerWindow> Window = FTimingProfilerManager::Get()->GetProfilerWindow();
+				if (!Window.IsValid())
+				{
+					return true;
+				}
+
+				TSharedPtr<STimingView> TimingView = Window->GetTimingView();
+				if (!TimingView.IsValid())
+				{
+					return true;
+				}
+
+				bIsAvailable = true;
+
+				if (!ContextSwitchesSharedState.IsValid())
+				{
+					ContextSwitchesSharedState = MakeShared<FContextSwitchesSharedState>(TimingView.Get());
+					IModularFeatures::Get().RegisterModularFeature(Insights::TimingViewExtenderFeatureName, ContextSwitchesSharedState.Get());
+				}
+			}
+
+			if (Session->IsAnalysisComplete())
+			{
+				// Never check again during this session.
+				AvailabilityCheck.Disable();
+			}
+		}
+		else
+		{
+			// Do not check again until the next session changed event (see OnSessionChanged).
+			AvailabilityCheck.Disable();
+		}
+	}
+
+	return true;
+}
+
+////////////////////////////////////////////////////////////////////////////////////////////////////
+
+void FContextSwitchesProfilerManager::OnSessionChanged()
+{
+	bIsAvailable = false;
+	if (FInsightsManager::Get()->GetSession().IsValid())
+	{
+		AvailabilityCheck.Enable(0.5);
+	}
+	else
+	{
+		AvailabilityCheck.Disable();
+	}
+}
+
+////////////////////////////////////////////////////////////////////////////////////////////////////
+
+void FContextSwitchesProfilerManager::RegisterTimingProfilerLayoutExtensions(FInsightsMajorTabExtender& InOutExtender)
+{
+}
+
+////////////////////////////////////////////////////////////////////////////////////////////////////
+
+} // namespace Insights
+
+#undef LOCTEXT_NAMESPACE
diff --git a/Engine/Source/Developer/TraceInsights/Private/Insights/ContextSwitches/ContextSwitchesProfilerManager.h b/Engine/Source/Developer/TraceInsights/Private/Insights/ContextSwitches/ContextSwitchesProfilerManager.h
new file mode 100644
--- /dev/null
+++ b/Engine/Source/Developer/TraceInsights/Private/Insights/ContextSwitches/ContextSwitchesProfilerManager.h
@@ -0,0 +1,78 @@
+// Copyright Epic Games, Inc. All Rights Reserved.
+
+#pragma once
+
+#include "Containers/Ticker.h"
+#include "CoreMinimal.h"
+#include "Framework/Commands/UICommandList.h"
+
+// Insights
+#include "Insights/InsightsManager.h"
+#include "Insights/IUnrealInsightsModule.h"
+
+class FThreadTrackEvent;
+
+namespace Insights
+{
+
+class FContextSwitchesSharedState;
+
+////////////////////////////////////////////////////////////////////////////////////////////////////
+/**
+ * This class manages the Context Switches Profiler state and settings.
+ */
+class FContextSwitchesProfilerManager : public TSharedFromThis<FContextSwitchesProfilerManager>, public IInsightsComponent
+{
+public:
+	/** Creates the Context Switches Profiler manager, only one instance can exist. */
+	FContextSwitchesProfilerManager(TSharedRef<FUICommandList> InCommandList);
+
+	/** Virtual destructor. */
+	virtual ~FContextSwitchesProfilerManager();
+
+	/** Creates an instance of the Context Switches manager. */
+	static TSharedPtr<FContextSwitchesProfilerManager> CreateInstance();
+
+	static TSharedPtr<FContextSwitchesProfilerManager> Get();
+
+	//////////////////////////////////////////////////
+	// IInsightsComponent
+
+	virtual void Initialize(IUnrealInsightsModule& InsightsModule) override;
+	virtual void Shutdown() override;
+	virtual void RegisterMajorTabs(IUnrealInsightsModule& InsightsModule) override;
+	virtual void UnregisterMajorTabs() override;
+
+	//////////////////////////////////////////////////
+	bool GetIsAvailable() { return bIsAvailable; }
+
+	void OnSessionChanged();
+
+	TSharedPtr<Insights::FContextSwitchesSharedState> GetContextSwitchesSharedState() { return ContextSwitchesSharedState;	}
+
+private:
+	/** Updates this manager, done through FCoreTicker. */
+	bool Tick(float DeltaTime);
+
+	void RegisterTimingProfilerLayoutExtensions(FInsightsMajorTabExtender& InOutExtender);
+
+private:
+	bool bIsInitialized;
+	bool bIsAvailable;
+	FAvailabilityCheck AvailabilityCheck;
+
+	/** The delegate to be invoked when this manager ticks. */
+	FTickerDelegate OnTick;
+
+	/** Handle to the registered OnTick. */
+	FDelegateHandle OnTickHandle;
+
+	/** A shared pointer to the global instance of the Context Switches Profiler manager. */
+	static TSharedPtr<FContextSwitchesProfilerManager> Instance;
+
+	/** Shared state for Context Switches tracks */
+	TSharedPtr<FContextSwitchesSharedState> ContextSwitchesSharedState;
+};
+
+} // namespace Insights
+
diff --git a/Engine/Source/Developer/TraceInsights/Private/Insights/ContextSwitches/ViewModels/ContextSwitchesTimingTrack.cpp b/Engine/Source/Developer/TraceInsights/Private/Insights/ContextSwitches/ViewModels/ContextSwitchesTimingTrack.cpp
new file mode 100644
--- /dev/null
+++ b/Engine/Source/Developer/TraceInsights/Private/Insights/ContextSwitches/ViewModels/ContextSwitchesTimingTrack.cpp
@@ -0,0 +1,306 @@
+// Copyright Epic Games, Inc. All Rights Reserved.
+
+#include "ContextSwitchesTimingTrack.h"
+
+#include "EditorStyleSet.h"
+#include "Framework/Commands/Commands.h"
+#include "Framework/Commands/UICommandList.h"
+#include "Framework/MultiBox/MultiBoxBuilder.h"
+#include "TraceServices/Model/ContextSwitches.h"
+
+// Insights
+#include "Insights/Common/TimeUtils.h"
+#include "Insights/ITimingViewSession.h"
+#include "Insights/InsightsManager.h"
+#include "Insights/ViewModels/ThreadTimingTrack.h"
+#include "Insights/ViewModels/ThreadTrackEvent.h"
+#include "Insights/ViewModels/TimingTrackViewport.h"
+#include "Insights/ViewModels/TimingViewDrawHelper.h"
+#include "Insights/ViewModels/TooltipDrawState.h"
+#include "Insights/Widgets/STimingView.h"
+
+using namespace TraceServices;
+
+#define LOCTEXT_NAMESPACE "ContextSwitchesTrack"
+
+namespace Insights
+{
+////////////////////////////////////////////////////////////////////////////////////////////////////
+// FContextSwitchesStateCommands
+////////////////////////////////////////////////////////////////////////////////////////////////////
+
+class FContextSwitchesStateCommands : public TCommands<FContextSwitchesStateCommands>
+{
+public:
+	FContextSwitchesStateCommands()
+		: TCommands<FContextSwitchesStateCommands>(TEXT("FContextSwitchesStateCommands"), NSLOCTEXT("FContextSwitchesStateCommands", "Context Switches State Commands", "Context Switches Commands"), NAME_None, FEditorStyle::Get().GetStyleSetName())
+	{
+	}
+
+	virtual ~FContextSwitchesStateCommands()
+	{
+	}
+
+	// UI_COMMAND takes long for the compiler to optimize
+	PRAGMA_DISABLE_OPTIMIZATION
+	virtual void RegisterCommands() override
+	{
+	}
+	PRAGMA_ENABLE_OPTIMIZATION
+};
+
+////////////////////////////////////////////////////////////////////////////////////////////////////
+// FContextSwitchesSharedState
+////////////////////////////////////////////////////////////////////////////////////////////////////
+
+FContextSwitchesSharedState::FContextSwitchesSharedState(STimingView* InTimingView) 
+	: TimingView(InTimingView) 
+{
+}
+
+////////////////////////////////////////////////////////////////////////////////////////////////////
+
+void FContextSwitchesSharedState::OnBeginSession(Insights::ITimingViewSession& InSession)
+{
+	if (&InSession != TimingView)
+	{
+		return;
+	}
+}
+
+////////////////////////////////////////////////////////////////////////////////////////////////////
+
+void FContextSwitchesSharedState::OnEndSession(Insights::ITimingViewSession& InSession)
+{
+	if (&InSession != TimingView)
+	{
+		return;
+	}
+}
+
+////////////////////////////////////////////////////////////////////////////////////////////////////
+
+void FContextSwitchesSharedState::Tick(Insights::ITimingViewSession& InSession, const TraceServices::IAnalysisSession& InAnalysisSession)
+{
+	if (&InSession != TimingView)
+	{
+		return;
+	}
+
+	TSharedPtr<FThreadTimingSharedState> TimingSharedState = TimingView->GetThreadTimingSharedState();
+
+	if (!TimingSharedState.IsValid())
+	{
+		return;
+	}
+
+	const TMap<uint32, TSharedPtr<FCpuTimingTrack>>& CpuTracks = TimingSharedState->GetAllCpuTracks();
+
+	for (const TPair<uint32, TSharedPtr<FCpuTimingTrack>>& MapEntry : CpuTracks)
+	{
+		if (MapEntry.Value.IsValid() && !MapEntry.Value->GetChildTrack().IsValid())
+		{
+			TSharedPtr<FContextSwitchesTimingTrack> ContextSwitchesTrack = MakeShared<FContextSwitchesTimingTrack>(*this, TEXT("Context Switches"), MapEntry.Value->GetTimelineIndex(), MapEntry.Value->GetThreadId());
+			ContextSwitchesTrack->SetParentTrack(MapEntry.Value);
+			MapEntry.Value->SetChildTrack(ContextSwitchesTrack);
+		}
+	}
+}
+
+////////////////////////////////////////////////////////////////////////////////////////////////////
+
+void FContextSwitchesSharedState::ExtendFilterMenu(Insights::ITimingViewSession& InSession, FMenuBuilder& InOutMenuBuilder)
+{
+	if (&InSession != TimingView)
+	{
+		return;
+	}
+}
+
+////////////////////////////////////////////////////////////////////////////////////////////////////
+
+void FContextSwitchesSharedState::InitCommandList()
+{
+	FContextSwitchesStateCommands::Register();
+}
+
+////////////////////////////////////////////////////////////////////////////////////////////////////
+// FContextSwitchesTimingTrack
+////////////////////////////////////////////////////////////////////////////////////////////////////
+
+INSIGHTS_IMPLEMENT_RTTI(FContextSwitchesTimingTrack)
+
+////////////////////////////////////////////////////////////////////////////////////////////////////
+
+void FContextSwitchesTimingTrack::BuildDrawState(ITimingEventsTrackDrawStateBuilder& Builder, const ITimingTrackUpdateContext& Context)
+{
+	TSharedPtr<const TraceServices::IAnalysisSession> Session = FInsightsManager::Get()->GetSession();
+	if (!Session.IsValid())
+	{
+		return;
+	}
+	
+	TraceServices::FAnalysisSessionReadScope SessionReadScope(*Session.Get());
+
+	const IContextSwitchProvider* ContextSwitchesProvider = TraceServices::ReadContextSwitchProvider(*Session.Get());
+
+	if (ContextSwitchesProvider == nullptr)
+	{
+		return;
+	}
+
+	const TPagedArray<FContextSwitch>* ContextSwitches = ContextSwitchesProvider->GetContextSwitches(ThreadId);
+
+	if (ContextSwitches == nullptr)
+	{
+		return;
+	}
+
+	const FTimingTrackViewport& Viewport = Context.GetViewport();
+
+	uint64 ContextSwitchPageIndex = Algo::UpperBoundBy(*ContextSwitches, Viewport.GetStartTime(), [](const TPagedArrayPage<FContextSwitch>& ContextSwitchPage)
+		{
+			return ContextSwitchPage.Items[0].Start;
+		});
+
+	if (ContextSwitchPageIndex > 0)
+	{
+		--ContextSwitchPageIndex;
+	}
+
+	auto Iterator = ContextSwitches->GetIteratorFromPage(ContextSwitchPageIndex);
+	const FContextSwitch* CurrentContextSwitch = Iterator.NextItem();
+	while (CurrentContextSwitch && CurrentContextSwitch->Start < Viewport.GetEndTime())
+	{
+		if (CurrentContextSwitch->End > Viewport.GetStartTime())
+		{
+			Builder.AddEvent(CurrentContextSwitch->Start, CurrentContextSwitch->End, 0, *FString::Printf(TEXT("Core %d"), CurrentContextSwitch->CoreNumber), 0, 0);
+		}
+
+		CurrentContextSwitch = Iterator.NextItem();
+	}
+}
+
+////////////////////////////////////////////////////////////////////////////////////////////////////
+
+void FContextSwitchesTimingTrack::BuildFilteredDrawState(ITimingEventsTrackDrawStateBuilder& Builder, const ITimingTrackUpdateContext& Context)
+{
+	BuildDrawState(Builder, Context);
+}
+
+////////////////////////////////////////////////////////////////////////////////////////////////////
+
+void FContextSwitchesTimingTrack::PostDraw(const ITimingTrackDrawContext& Context) const
+{
+}
+
+////////////////////////////////////////////////////////////////////////////////////////////////////
+
+const TSharedPtr<const ITimingEvent> FContextSwitchesTimingTrack::GetEvent(float InPosX, float InPosY, const FTimingTrackViewport& Viewport) const
+{
+	TSharedPtr<FTimingEvent> TimingEvent;
+
+	const FTimingViewLayout& Layout = Viewport.GetLayout();
+
+	float TopLaneY = 0.0f;
+	float TrackLanesHeight = 0.0f;
+	if (ChildTrack.IsValid())
+	{
+		const float HeaderDY = InPosY - ChildTrack->GetPosY();
+		if (HeaderDY >= 0 && HeaderDY < ChildTrack->GetHeight())
+		{
+			return ChildTrack->GetEvent(InPosX, InPosY, Viewport);
+		}
+		TopLaneY = GetPosY() + 1.0f + Layout.TimelineDY + ChildTrack->GetHeight() + Layout.ChildTimelineDY;
+		TrackLanesHeight = GetHeight() - ChildTrack->GetHeight() - 1.0f - 2 * Layout.TimelineDY - Layout.ChildTimelineDY;
+	}
+	else
+	{
+		if (IsChildTrack())
+		{
+			TopLaneY = GetPosY();
+			TrackLanesHeight = GetHeight();
+		}
+		else
+		{
+			TopLaneY = GetPosY() + 1.0f + Layout.TimelineDY;
+			TrackLanesHeight = GetHeight() - 1.0f - 2 * Layout.TimelineDY;
+		}
+	}
+
+	const float DY = InPosY - TopLaneY;
+
+	// If mouse is not above first sub-track or below last sub-track...
+	if (DY >= 0 && DY < TrackLanesHeight)
+	{
+		const int32 Depth = DY / (Layout.EventH + Layout.EventDY);
+
+		const double SecondsPerPixel = 1.0 / Viewport.GetScaleX();
+
+		const double EventTime = Viewport.SlateUnitsToTime(InPosX);
+
+		TSharedPtr<const TraceServices::IAnalysisSession> Session = FInsightsManager::Get()->GetSession();
+
+		if (!Session.IsValid())
+		{
+			return TimingEvent;
+		}
+
+		TraceServices::FAnalysisSessionReadScope SessionReadScope(*Session.Get());
+
+		const IContextSwitchProvider* ContextSwitchesProvider = TraceServices::ReadContextSwitchProvider(*Session.Get());
+
+		if (ContextSwitchesProvider == nullptr)
+		{
+			return TimingEvent;
+		}
+
+		const TPagedArray<FContextSwitch>* ContextSwitches = ContextSwitchesProvider->GetContextSwitches(ThreadId);
+
+		if (ContextSwitches == nullptr)
+		{
+			return TimingEvent;
+		}
+
+		uint64 ContextSwitchPageIndex = Algo::UpperBoundBy(*ContextSwitches, EventTime + SecondsPerPixel, [](const TPagedArrayPage<FContextSwitch>& ContextSwitchPage)
+			{
+				return ContextSwitchPage.Items[0].Start;
+			});
+
+		if (ContextSwitchPageIndex > 0)
+		{
+			--ContextSwitchPageIndex;
+		}
+
+		auto Iterator = ContextSwitches->GetIteratorFromPage(ContextSwitchPageIndex);
+		const FContextSwitch* CurrentContextSwitch = Iterator.NextItem();
+		while (CurrentContextSwitch && CurrentContextSwitch->End < EventTime)
+		{
+			CurrentContextSwitch = Iterator.NextItem();
+		}
+
+		if (CurrentContextSwitch)
+		{
+			ensure(ParentTrack.IsValid());
+			if (CurrentContextSwitch->Start < EventTime && EventTime < CurrentContextSwitch->End)
+			{
+				TimingEvent = MakeShared<FTimingEvent>(SharedThis(this), CurrentContextSwitch->Start, CurrentContextSwitch->End, 0);
+			}
+		}
+	}
+
+	return TimingEvent;
+}
+
+////////////////////////////////////////////////////////////////////////////////////////////////////
+
+void FContextSwitchesTimingTrack::InitTooltip(FTooltipDrawState& InOutTooltip, const ITimingEvent& InTooltipEvent) const
+{
+	InOutTooltip.ResetContent();
+
+	InOutTooltip.UpdateLayout();
+}
+
+} // namespace Insights
+
+#undef LOCTEXT_NAMESPACE
diff --git a/Engine/Source/Developer/TraceInsights/Private/Insights/ContextSwitches/ViewModels/ContextSwitchesTimingTrack.h b/Engine/Source/Developer/TraceInsights/Private/Insights/ContextSwitches/ViewModels/ContextSwitchesTimingTrack.h
new file mode 100644
--- /dev/null
+++ b/Engine/Source/Developer/TraceInsights/Private/Insights/ContextSwitches/ViewModels/ContextSwitchesTimingTrack.h
@@ -0,0 +1,94 @@
+// Copyright Epic Games, Inc. All Rights Reserved.
+
+#pragma once
+
+#include "CoreMinimal.h"
+
+// Insights
+#include "Insights/ITimingViewExtender.h"
+#include "Insights/ViewModels/TimingEventsTrack.h"
+
+class FThreadTrackEvent;
+class STimingView;
+
+////////////////////////////////////////////////////////////////////////////////////////////////////
+
+namespace Insights
+{
+class FContextSwitchesTimingTrack;
+
+class FContextSwitchesSharedState : public Insights::ITimingViewExtender, public TSharedFromThis<FContextSwitchesSharedState>
+{
+
+public:
+	FContextSwitchesSharedState(STimingView* InTimingView);
+	virtual ~FContextSwitchesSharedState() = default;
+
+	bool IsContextSwitchesVisible() const;
+
+	//////////////////////////////////////////////////
+	// ITimingViewExtender interface
+
+	virtual void OnBeginSession(Insights::ITimingViewSession& InSession) override;
+	virtual void OnEndSession(Insights::ITimingViewSession& InSession) override;
+	virtual void Tick(Insights::ITimingViewSession& InSession, const TraceServices::IAnalysisSession& InAnalysisSession) override;
+	virtual void ExtendFilterMenu(Insights::ITimingViewSession& InSession, FMenuBuilder& InMenuBuilder) override;
+	virtual bool ExtendGlobalContextMenu(ITimingViewSession& InSession, FMenuBuilder& InMenuBuilder) override { return false; }
+
+	//////////////////////////////////////////////////
+
+	bool IsContextSwitchesToggleOn() const { return bShowHideContextSwitchesTrack; }
+	void SetContextSwitchesToggle(bool bOnOff) { bShowHideContextSwitchesTrack = bOnOff; }
+	void ShowContextSwitches() { SetContextSwitchesToggle(true); }
+	void HideContextSwitches() { SetContextSwitchesToggle(false); }
+
+	STimingView* GetTimingView() { return TimingView; }
+
+private:
+	void InitCommandList();
+
+private:
+	STimingView* TimingView;
+
+	bool bShowHideContextSwitchesTrack;
+};
+
+////////////////////////////////////////////////////////////////////////////////////////////////////
+
+class FContextSwitchesTimingTrack : public FTimingEventsTrack
+{
+	INSIGHTS_DECLARE_RTTI(FContextSwitchesTimingTrack, FTimingEventsTrack)
+
+public:
+	explicit FContextSwitchesTimingTrack(FContextSwitchesSharedState& InSharedState, const FString& InName, uint32 InTimelineIndex, uint32 ThreadId)
+		: FTimingEventsTrack(InName)
+		, TimelineIndex(InTimelineIndex)
+		, ThreadId(ThreadId)
+		, SharedState(InSharedState)
+	{
+	}
+
+	virtual ~FContextSwitchesTimingTrack() {}
+
+	uint32 GetTimelineIndex() const { return TimelineIndex; }
+
+	virtual void BuildDrawState(ITimingEventsTrackDrawStateBuilder& Builder, const ITimingTrackUpdateContext& Context) override;
+	virtual void BuildFilteredDrawState(ITimingEventsTrackDrawStateBuilder& Builder, const ITimingTrackUpdateContext& Context) override;
+
+	virtual void PostDraw(const ITimingTrackDrawContext& Context) const override;
+
+	virtual void InitTooltip(FTooltipDrawState& InOutTooltip, const ITimingEvent& InTooltipEvent) const override;
+
+	virtual const TSharedPtr<const ITimingEvent> GetEvent(float InPosX, float InPosY, const FTimingTrackViewport& Viewport) const override;
+
+
+private:
+	uint32 TimelineIndex;
+	uint32 ThreadId;
+
+	FContextSwitchesSharedState& SharedState;
+};
+
+////////////////////////////////////////////////////////////////////////////////////////////////////
+
+} //namespace Insights
diff --git a/Engine/Source/Developer/TraceInsights/Private/Insights/TraceInsightsModule.cpp b/Engine/Source/Developer/TraceInsights/Private/Insights/TraceInsightsModule.cpp
--- a/Engine/Source/Developer/TraceInsights/Private/Insights/TraceInsightsModule.cpp
+++ b/Engine/Source/Developer/TraceInsights/Private/Insights/TraceInsightsModule.cpp
@@ -13,6 +13,7 @@
 #include "TraceServices/ITraceServicesModule.h"
 
 // Insights
+#include "Insights/ContextSwitches/ContextSwitchesProfilerManager.h"
 #include "Insights/InsightsManager.h"
 #include "Insights/InsightsStyle.h"
 #include "Insights/IUnrealInsightsModule.h"
@@ -54,6 +55,7 @@ void FTraceInsightsModule::StartupModule()
 	RegisterComponent(FNetworkingProfilerManager::CreateInstance());
 	RegisterComponent(FMemoryProfilerManager::CreateInstance());
 	RegisterComponent(Insights::FTaskGraphProfilerManager::CreateInstance());
+	RegisterComponent(Insights::FContextSwitchesProfilerManager::CreateInstance());
 
 	UnrealInsightsLayoutIni = FConfigCacheIni::GetDestIniFilename(TEXT("UnrealInsightsLayout"), nullptr, *FPaths::GeneratedConfigDir());
 }
diff --git a/Engine/Source/Developer/TraceInsights/Private/Insights/ViewModels/ThreadTimingTrack.cpp b/Engine/Source/Developer/TraceInsights/Private/Insights/ViewModels/ThreadTimingTrack.cpp
--- a/Engine/Source/Developer/TraceInsights/Private/Insights/ViewModels/ThreadTimingTrack.cpp
+++ b/Engine/Source/Developer/TraceInsights/Private/Insights/ViewModels/ThreadTimingTrack.cpp
@@ -1220,11 +1220,37 @@ const TSharedPtr<const ITimingEvent> FThreadTimingTrack::GetEvent(float InPosX,
 	TSharedPtr<FThreadTrackEvent> TimingEvent;
 
 	const FTimingViewLayout& Layout = Viewport.GetLayout();
-	const float TopLaneY = GetPosY() + 1.0f + Layout.TimelineDY; // +1.0f is for horizontal line between timelines
+
+	float TopLaneY = 0.0f;
+	float TrackLanesHeight = 0.0f;
+	if (ChildTrack.IsValid())
+	{
+		const float HeaderDY = InPosY - ChildTrack->GetPosY();
+		if (HeaderDY >= 0 && HeaderDY < ChildTrack->GetHeight())
+		{
+			return ChildTrack->GetEvent(InPosX, InPosY, Viewport);
+		}
+		TopLaneY = GetPosY() + 1.0f + Layout.TimelineDY + ChildTrack->GetHeight() + Layout.ChildTimelineDY;
+		TrackLanesHeight = GetHeight() - ChildTrack->GetHeight() - 1.0f - 2 * Layout.TimelineDY - Layout.ChildTimelineDY;
+	}
+	else
+	{
+		if (IsChildTrack())
+		{
+			TopLaneY = GetPosY();
+			TrackLanesHeight = GetHeight();
+		}
+		else
+		{
+			TopLaneY = GetPosY() + 1.0f + Layout.TimelineDY;
+			TrackLanesHeight = GetHeight() - 1.0f - 2 * Layout.TimelineDY;
+		}
+	}
+
 	const float DY = InPosY - TopLaneY;
 
 	// If mouse is not above first sub-track or below last sub-track...
-	if (DY >= 0 && DY < GetHeight() - 1.0f - 2 * Layout.TimelineDY)
+	if (DY >= 0 && DY < TrackLanesHeight)
 	{
 		const int32 Depth = DY / (Layout.EventH + Layout.EventDY);
 
diff --git a/Engine/Source/Developer/TraceInsights/Private/Insights/ViewModels/ThreadTimingTrack.h b/Engine/Source/Developer/TraceInsights/Private/Insights/ViewModels/ThreadTimingTrack.h
--- a/Engine/Source/Developer/TraceInsights/Private/Insights/ViewModels/ThreadTimingTrack.h
+++ b/Engine/Source/Developer/TraceInsights/Private/Insights/ViewModels/ThreadTimingTrack.h
@@ -45,6 +45,7 @@ public:
 	TSharedPtr<FGpuTimingTrack> GetGpuTrack() { return GpuTrack; }
 	TSharedPtr<FGpuTimingTrack> GetGpu2Track() { return Gpu2Track; }
 	TSharedPtr<FCpuTimingTrack> GetCpuTrack(uint32 InThreadId);
+	const TMap<uint32, TSharedPtr<FCpuTimingTrack>> GetAllCpuTracks() { return CpuTracks; }
 
 	bool IsGpuTrackVisible() const;
 	bool IsCpuTrackVisible(uint32 InThreadId) const;
diff --git a/Engine/Source/Developer/TraceInsights/Private/Insights/ViewModels/TimingEventsTrack.cpp b/Engine/Source/Developer/TraceInsights/Private/Insights/ViewModels/TimingEventsTrack.cpp
--- a/Engine/Source/Developer/TraceInsights/Private/Insights/ViewModels/TimingEventsTrack.cpp
+++ b/Engine/Source/Developer/TraceInsights/Private/Insights/ViewModels/TimingEventsTrack.cpp
@@ -59,6 +59,11 @@ void FTimingEventsTrack::Reset()
 
 void FTimingEventsTrack::PreUpdate(const ITimingTrackUpdateContext& Context)
 {
+	if (ChildTrack.IsValid())
+	{
+		ChildTrack->PreUpdate(Context);
+	}
+
 	if (IsDirty() || Context.GetViewport().IsHorizontalViewportDirty())
 	{
 		ClearDirtyFlag();
@@ -158,7 +163,19 @@ void FTimingEventsTrack::UpdateTrackHeight(const ITimingTrackUpdateContext& Cont
 	const FTimingTrackViewport& Viewport = Context.GetViewport();
 
 	const float CurrentTrackHeight = GetHeight();
-	const float DesiredTrackHeight = Viewport.GetLayout().ComputeTrackHeight(NumLanes);
+	float DesiredTrackHeight = 0.0f;
+	if (IsChildTrack())
+	{
+		DesiredTrackHeight = Viewport.GetLayout().ComputeChildTrackHeight(NumLanes);
+	}
+	else
+	{
+		DesiredTrackHeight = Viewport.GetLayout().ComputeTrackHeight(NumLanes);
+		if (ChildTrack.IsValid())
+		{
+			DesiredTrackHeight += ChildTrack->GetHeight();
+		}
+	}
 
 	if (CurrentTrackHeight < DesiredTrackHeight)
 	{
@@ -192,6 +209,11 @@ void FTimingEventsTrack::UpdateTrackHeight(const ITimingTrackUpdateContext& Cont
 
 void FTimingEventsTrack::PostUpdate(const ITimingTrackUpdateContext& Context)
 {
+	if (ChildTrack.IsValid())
+	{
+		ChildTrack->PostUpdate(Context);
+	}
+
 	FBaseTimingTrack::PostUpdate(Context);
 
 	constexpr float HeaderWidth = 100.0f;
@@ -215,6 +237,11 @@ void FTimingEventsTrack::PostUpdate(const ITimingTrackUpdateContext& Context)
 
 void FTimingEventsTrack::Draw(const ITimingTrackDrawContext& Context) const
 {
+	if (ChildTrack.IsValid())
+	{
+		ChildTrack->Draw(Context);
+	}
+
 	DrawEvents(Context, 1.0f);
 	DrawHeader(Context);
 }
@@ -282,6 +309,12 @@ int32 FTimingEventsTrack::GetHeaderTextLayerId(const ITimingTrackDrawContext& Co
 
 void FTimingEventsTrack::DrawHeader(const ITimingTrackDrawContext& Context) const
 {
+	if (IsChildTrack())
+	{
+		// Do not draw the header for child tracks for now
+		return;
+	}
+
 	const FTimingViewDrawHelper& Helper = *static_cast<const FTimingViewDrawHelper*>(&Context.GetHelper());
 	Helper.DrawTrackHeader(*this);
 }
@@ -290,11 +323,27 @@ void FTimingEventsTrack::DrawHeader(const ITimingTrackDrawContext& Context) cons
 
 void FTimingEventsTrack::DrawEvent(const ITimingTrackDrawContext& Context, const ITimingEvent& InTimingEvent, EDrawEventMode InDrawMode) const
 {
-	if (InTimingEvent.CheckTrack(this) && InTimingEvent.Is<FTimingEvent>())
+	if (ChildTrack.IsValid() && InTimingEvent.CheckTrack(ChildTrack.Get()))
+	{
+		ChildTrack->DrawEvent(Context, InTimingEvent, InDrawMode);
+	}
+	else if (InTimingEvent.CheckTrack(this) && InTimingEvent.Is<FTimingEvent>())
 	{
 		const FTimingEvent& TrackEvent = InTimingEvent.As<FTimingEvent>();
 		const FTimingViewLayout& Layout = Context.GetViewport().GetLayout();
-		const float Y = TrackEvent.GetTrack()->GetPosY() + Layout.GetLaneY(TrackEvent.GetDepth());
+		float Y = TrackEvent.GetTrack()->GetPosY();
+		if (ChildTrack.IsValid())
+		{
+			Y += ChildTrack->GetHeight() + Layout.ChildTimelineDY;
+		}
+		if (IsChildTrack())
+		{
+			Y += Layout.GetChildLaneY(TrackEvent.GetDepth());
+		}
+		else
+		{
+			Y += Layout.GetLaneY(TrackEvent.GetDepth());
+		}
 
 		const FTimingViewDrawHelper& Helper = *static_cast<const FTimingViewDrawHelper*>(&Context.GetHelper());
 		Helper.DrawTimingEventHighlight(TrackEvent.GetStartTime(), TrackEvent.GetEndTime(), Y, InDrawMode);
@@ -307,11 +356,36 @@ const TSharedPtr<const ITimingEvent> FTimingEventsTrack::GetEvent(float InPosX,
 {
 	const FTimingViewLayout& Layout = Viewport.GetLayout();
 
-	const float TopLaneY = GetPosY() + 1.0f + Layout.TimelineDY; // +1.0f is for horizontal line between timelines
+	float TopLaneY = 0.0f;
+	float TrackLanesHeight = 0.0f;
+	if (ChildTrack.IsValid())
+	{
+		const float HeaderDY = InPosY - ChildTrack->GetPosY();
+		if (HeaderDY >= 0 && HeaderDY < ChildTrack->GetHeight())
+		{
+			return ChildTrack->GetEvent(InPosX, InPosY, Viewport);
+		}
+		TopLaneY = GetPosY() + 1.0f + Layout.TimelineDY + ChildTrack->GetHeight() + Layout.ChildTimelineDY;
+		TrackLanesHeight = GetHeight() - ChildTrack->GetHeight() - 1.0f - 2 * Layout.TimelineDY - Layout.ChildTimelineDY;
+	}
+	else
+	{
+		if (IsChildTrack())
+		{
+			TopLaneY = GetPosY();
+			TrackLanesHeight = GetHeight();
+		}
+		else
+		{
+			TopLaneY = GetPosY() + 1.0f + Layout.TimelineDY;
+			TrackLanesHeight = GetHeight() - 1.0f - 2 * Layout.TimelineDY;
+		}
+	}
+
 	const float DY = InPosY - TopLaneY;
 
 	// If mouse is not above first sub-track or below last sub-track...
-	if (DY >= 0 && DY < GetHeight() - 1.0f - 2 * Layout.TimelineDY)
+	if (DY >= 0 && DY < TrackLanesHeight)
 	{
 		const int32 Depth = DY / (Layout.EventH + Layout.EventDY);
 
diff --git a/Engine/Source/Developer/TraceInsights/Private/Insights/ViewModels/TimingViewDrawHelper.cpp b/Engine/Source/Developer/TraceInsights/Private/Insights/ViewModels/TimingViewDrawHelper.cpp
--- a/Engine/Source/Developer/TraceInsights/Private/Insights/ViewModels/TimingViewDrawHelper.cpp
+++ b/Engine/Source/Developer/TraceInsights/Private/Insights/ViewModels/TimingViewDrawHelper.cpp
@@ -312,7 +312,15 @@ void FTimingViewDrawHelper::DrawEvents(const FTimingEventsTrackDrawState& DrawSt
 		NumEvents += DrawState.GetNumEvents();
 		NumMergedBoxes += DrawState.GetNumMergedBoxes();
 
-		const float TopLaneY = TrackY + OffsetY + Layout.TimelineDY;
+		float TopLaneY = TrackY;
+		if (!Track.IsChildTrack())
+		{
+			TopLaneY += OffsetY + Layout.TimelineDY;
+		}
+		if (Track.GetChildTrack().IsValid())
+		{
+			TopLaneY += Track.ChildTrack->GetHeight() + Layout.ChildTimelineDY;
+		}
 
 		// Draw filled boxes (merged borders).
 		//if (Layout.EventH > 0.0f)
@@ -381,7 +389,15 @@ void FTimingViewDrawHelper::DrawFadedEvents(const FTimingEventsTrackDrawState& D
 		NumEvents += DrawState.GetNumEvents();
 		NumMergedBoxes += DrawState.GetNumMergedBoxes();
 
-		const float TopLaneY = TrackY + OffsetY + Layout.TimelineDY;
+		float TopLaneY = TrackY;
+		if (!Track.IsChildTrack())
+		{
+			TopLaneY += OffsetY + Layout.TimelineDY;
+		}
+		if (Track.GetChildTrack().IsValid())
+		{
+			TopLaneY += Track.ChildTrack->GetHeight() + Layout.ChildTimelineDY;
+		}
 
 		// Draw filled boxes (merged borders).
 		//if (Layout.EventH > 0.0f)
diff --git a/Engine/Source/Developer/TraceInsights/Private/Insights/ViewModels/TimingViewLayout.cpp b/Engine/Source/Developer/TraceInsights/Private/Insights/ViewModels/TimingViewLayout.cpp
--- a/Engine/Source/Developer/TraceInsights/Private/Insights/ViewModels/TimingViewLayout.cpp
+++ b/Engine/Source/Developer/TraceInsights/Private/Insights/ViewModels/TimingViewLayout.cpp
@@ -10,6 +10,7 @@ void FTimingViewLayout::ForceNormalMode()
 	EventH = NormalLayoutEventH;
 	EventDY = NormalLayoutEventDY;
 	TimelineDY = NormalLayoutTimelineDY;
+	ChildTimelineDY = NormalLayoutChildTimelineDY;
 	MinTimelineH = NormalLayoutMinTimelineH;
 	TargetMinTimelineH = NormalLayoutMinTimelineH;
 }
@@ -22,6 +23,7 @@ void FTimingViewLayout::ForceCompactMode()
 	EventH = CompactLayoutEventH;
 	EventDY = CompactLayoutEventDY;
 	TimelineDY = CompactLayoutTimelineDY;
+	ChildTimelineDY = CompactLayoutChildTimelineDY;
 	MinTimelineH = CompactLayoutMinTimelineH;
 	TargetMinTimelineH = CompactLayoutMinTimelineH;
 }
diff --git a/Engine/Source/Developer/TraceInsights/Private/Insights/Widgets/STimingView.cpp b/Engine/Source/Developer/TraceInsights/Private/Insights/Widgets/STimingView.cpp
--- a/Engine/Source/Developer/TraceInsights/Private/Insights/Widgets/STimingView.cpp
+++ b/Engine/Source/Developer/TraceInsights/Private/Insights/Widgets/STimingView.cpp
@@ -706,7 +706,7 @@ void STimingView::Tick(const FGeometry& AllottedGeometry, const double InCurrent
 	float TopOffset = 0.0f;
 	for (TSharedPtr<FBaseTimingTrack>& TrackPtr : TopDockedTracks)
 	{
-		TrackPtr->SetPosY(TopOffset);
+		SetTrackPosY(TrackPtr, TopOffset);
 		if (TrackPtr->IsVisible())
 		{
 			TopOffset += TrackPtr->GetHeight();
@@ -740,7 +740,7 @@ void STimingView::Tick(const FGeometry& AllottedGeometry, const double InCurrent
 	int32 NumVisibleBottomDockedTracks = 0;
 	for (TSharedPtr<FBaseTimingTrack>& TrackPtr : BottomDockedTracks)
 	{
-		TrackPtr->SetPosY(BottomOffsetY);
+		SetTrackPosY(TrackPtr, BottomOffsetY);
 		if (TrackPtr->IsVisible())
 		{
 			BottomOffsetY += TrackPtr->GetHeight();
@@ -1030,7 +1030,7 @@ void STimingView::UpdatePositionForScrollableTracks()
 	float ScrollableTrackPosY = Viewport.GetTopOffset() - Viewport.GetScrollPosY();
 	for (TSharedPtr<FBaseTimingTrack>& TrackPtr : ScrollableTracks)
 	{
-		TrackPtr->SetPosY(ScrollableTrackPosY); // set pos y also for the hidden tracks
+		SetTrackPosY(TrackPtr, ScrollableTrackPosY);
 		if (TrackPtr->IsVisible())
 		{
 			ScrollableTrackPosY += TrackPtr->GetHeight();
@@ -1040,6 +1040,17 @@ void STimingView::UpdatePositionForScrollableTracks()
 
 ////////////////////////////////////////////////////////////////////////////////////////////////////
 
+void STimingView::SetTrackPosY(TSharedPtr<FBaseTimingTrack>& TrackPtr, float TopOffsetY) const
+{
+	TrackPtr->SetPosY(TopOffsetY);
+	if (TrackPtr->GetChildTrack().IsValid())
+	{
+		TrackPtr->GetChildTrack()->SetPosY(TopOffsetY + this->GetViewport().GetLayout().TimelineDY + 1.0f);
+	}
+}
+
+////////////////////////////////////////////////////////////////////////////////////////////////////
+
 int32 STimingView::OnPaint(const FPaintArgs& Args, const FGeometry& AllottedGeometry, const FSlateRect& MyCullingRect, FSlateWindowElementList& OutDrawElements, int32 LayerId, const FWidgetStyle& InWidgetStyle, bool bParentEnabled) const
 {
 	FStopwatch Stopwatch;
@@ -3478,7 +3489,7 @@ void STimingView::UpdateHoveredTimingEvent(float InMousePosX, float InMousePosY)
 			Stopwatch.Start();
 
 			HoveredEvent = NewHoveredEvent;
-			ensure(HoveredTrack == HoveredEvent->GetTrack());
+			ensure(HoveredTrack == HoveredEvent->GetTrack() || HoveredTrack->GetChildTrack() == HoveredEvent->GetTrack());
 			HoveredTrack->UpdateEventStats(const_cast<ITimingEvent&>(*HoveredEvent));
 
 			Stopwatch.Update();
diff --git a/Engine/Source/Developer/TraceInsights/Private/Insights/Widgets/STimingView.h b/Engine/Source/Developer/TraceInsights/Private/Insights/Widgets/STimingView.h
--- a/Engine/Source/Developer/TraceInsights/Private/Insights/Widgets/STimingView.h
+++ b/Engine/Source/Developer/TraceInsights/Private/Insights/Widgets/STimingView.h
@@ -437,6 +437,8 @@ protected:
 	bool ContextMenu_ShowTaskDependecies_CanExecute();
 	bool ContextMenu_ShowTaskDependecies_IsChecked();
 
+	void SetTrackPosY(TSharedPtr<FBaseTimingTrack>& TrackPtr, float TopOffsetY) const;
+
 protected:
 	/** The track's viewport. Encapsulates info about position and scale. */
 	FTimingTrackViewport Viewport;
diff --git a/Engine/Source/Developer/TraceInsights/Public/Insights/ViewModels/BaseTimingTrack.h b/Engine/Source/Developer/TraceInsights/Public/Insights/ViewModels/BaseTimingTrack.h
--- a/Engine/Source/Developer/TraceInsights/Public/Insights/ViewModels/BaseTimingTrack.h
+++ b/Engine/Source/Developer/TraceInsights/Public/Insights/ViewModels/BaseTimingTrack.h
@@ -255,6 +255,13 @@ public:
 	// Called back from the timing view when an event is copied to the clipboard with Ctrl+C.
 	virtual void OnClipboardCopyEvent(const ITimingEvent& InSelectedEvent) const {}
 
+	TSharedPtr<FBaseTimingTrack> GetChildTrack() const { return ChildTrack; }
+	void SetChildTrack(TSharedPtr<FBaseTimingTrack> InTrack) { ChildTrack = InTrack; };
+
+	void SetParentTrack(TWeakPtr<FBaseTimingTrack> InTrack) { ParentTrack = InTrack; }
+	TWeakPtr<FBaseTimingTrack> GetParentTrack() const { return ParentTrack; }
+	bool IsChildTrack() const { return ParentTrack.IsValid(); }
+
 	// Returns number of text lines needed to display the debug string.
 	//TODO: virtual int GetDebugStringLineCount() const { return 0; }
 
@@ -265,6 +272,10 @@ protected:
 	void SetValidLocations(ETimingTrackLocation InValidLocations) { ValidLocations = InValidLocations; }
 	static uint64 GenerateId() { return IdGenerator++; }
 
+protected:
+	TSharedPtr<FBaseTimingTrack> ChildTrack;
+	TWeakPtr<FBaseTimingTrack> ParentTrack;
+
 private:
 	const uint64 Id;
 	FString Name;
diff --git a/Engine/Source/Developer/TraceInsights/Public/Insights/ViewModels/TimingViewLayout.h b/Engine/Source/Developer/TraceInsights/Public/Insights/ViewModels/TimingViewLayout.h
--- a/Engine/Source/Developer/TraceInsights/Public/Insights/ViewModels/TimingViewLayout.h
+++ b/Engine/Source/Developer/TraceInsights/Public/Insights/ViewModels/TimingViewLayout.h
@@ -13,11 +13,13 @@ struct TRACEINSIGHTS_API FTimingViewLayout
 	static constexpr float NormalLayoutEventH = 14.0f;
 	static constexpr float NormalLayoutEventDY = 2.0f;
 	static constexpr float NormalLayoutTimelineDY = 10.0f;
+	static constexpr float NormalLayoutChildTimelineDY = 5.0f;
 	static constexpr float NormalLayoutMinTimelineH = 0.0f;
 
 	static constexpr float CompactLayoutEventH = 4.0f;
 	static constexpr float CompactLayoutEventDY = 1.0f;
 	static constexpr float CompactLayoutTimelineDY = 1.0f;
+	static constexpr float CompactLayoutChildTimelineDY = 0.0f;
 	static constexpr float CompactLayoutMinTimelineH = 0.0f;
 
 	//////////////////////////////////////////////////
@@ -27,6 +29,7 @@ struct TRACEINSIGHTS_API FTimingViewLayout
 	float EventH; // height of a timing event, in Slate units
 	float EventDY; // vertical space between timing events in two adjacent lanes, in Slate units
 	float TimelineDY; // space at top and bottom of each track (i.e. above first lane and below last lane), in Slate units
+	float ChildTimelineDY; // space between a child track's lanes and the parent's lanes, in Slate units
 	float MinTimelineH; // current minimum height of a track, in Slate units
 	float TargetMinTimelineH; // targeted minimum height of a track (for animating min track height), in Slate units
 
@@ -38,6 +41,22 @@ struct TRACEINSIGHTS_API FTimingViewLayout
 		return 1.0f + TimelineDY + (EventDY + EventH) * Depth;
 	}
 
+	//////////////////////////////////////////////////
+
+	float GetChildLaneY(int32 Depth) const
+	{
+		return (EventDY + EventH) * Depth;
+	}
+
+	/* The layout of a track:
+	*	1.0f - Line between tracks
+	*	TimelineDY
+	*	HeaderLanes
+	*	ChildTimelineDY
+	*	TrackLanes
+	*	TimelineDY 
+	*/
+
 	float ComputeTrackHeight(int32 NumLanes) const
 	{
 		if (NumLanes <= 0)
@@ -47,7 +66,27 @@ struct TRACEINSIGHTS_API FTimingViewLayout
 		else //if (NumLanes > 0)
 		{
 			// 1.0f is for horizontal line between timelines
-			const float TrackHeight = 1.0f + TimelineDY + (EventH + EventDY) * NumLanes - EventDY + TimelineDY;
+			float TrackHeight = 1.0f + TimelineDY + (EventH + EventDY) * NumLanes - EventDY + TimelineDY;
+				
+			if (TrackHeight < FTimingViewLayout::RealMinTimelineH)
+			{
+				return FTimingViewLayout::RealMinTimelineH;
+			}
+
+			return TrackHeight;
+		}
+	}
+
+	float ComputeChildTrackHeight(int32 NumLanes) const
+	{
+		if (NumLanes <= 0)
+		{
+			return MinTimelineH;
+		}
+		else //if (NumLanes > 0)
+		{
+			// 1.0f is for horizontal line between timelines
+			const float TrackHeight = (EventH + EventDY) * NumLanes - EventDY;
 
 			if (TrackHeight < FTimingViewLayout::RealMinTimelineH)
 			{
diff --git a/Engine/Source/Developer/TraceServices/Private/Analyzers/PlatformEventTraceAnalysis.cpp b/Engine/Source/Developer/TraceServices/Private/Analyzers/PlatformEventTraceAnalysis.cpp
--- a/Engine/Source/Developer/TraceServices/Private/Analyzers/PlatformEventTraceAnalysis.cpp
+++ b/Engine/Source/Developer/TraceServices/Private/Analyzers/PlatformEventTraceAnalysis.cpp
@@ -56,4 +56,9 @@ bool FPlatformEventTraceAnalyzer::OnEvent(uint16 RouteId, EStyle Style, const FO
 	return true;
 }
 
+void FPlatformEventTraceAnalyzer::OnThreadInfo(const FThreadInfo& ThreadInfo)
+{
+	ContextSwitchProvider.AddThreadInfo(ThreadInfo.GetId(), ThreadInfo.GetSystemId());
+}
+
 } // namespace TraceServices
diff --git a/Engine/Source/Developer/TraceServices/Private/Analyzers/PlatformEventTraceAnalysis.h b/Engine/Source/Developer/TraceServices/Private/Analyzers/PlatformEventTraceAnalysis.h
--- a/Engine/Source/Developer/TraceServices/Private/Analyzers/PlatformEventTraceAnalysis.h
+++ b/Engine/Source/Developer/TraceServices/Private/Analyzers/PlatformEventTraceAnalysis.h
@@ -21,6 +21,7 @@ public:
 								FStackSampleProvider& StackSampleProvider);
 	virtual void OnAnalysisBegin(const FOnAnalysisContext& Context) override;
 	virtual bool OnEvent(uint16 RouteId, EStyle Style, const FOnEventContext& Context) override;
+	virtual void OnThreadInfo(const FThreadInfo& ThreadInfo) override;
 
 private:
 	enum : uint16
diff --git a/Engine/Source/Developer/TraceServices/Private/Model/ContextSwitches.cpp b/Engine/Source/Developer/TraceServices/Private/Model/ContextSwitches.cpp
--- a/Engine/Source/Developer/TraceServices/Private/Model/ContextSwitches.cpp
+++ b/Engine/Source/Developer/TraceServices/Private/Model/ContextSwitches.cpp
@@ -63,8 +63,13 @@ const TPagedArray<FContextSwitch>* FContextSwitchProvider::GetContextSwitches(ui
 {
 	Session.ReadAccessCheck();
 
-	auto ContextSwitches = Threads.Find(ThreadId);
-	return ContextSwitches ? *ContextSwitches : nullptr;
+	if (ThreadIdMap.Contains(ThreadId))
+	{
+		auto ContextSwitches = Threads.Find(ThreadIdMap[ThreadId]);
+		return ContextSwitches ? *ContextSwitches : nullptr;
+	}
+
+	return nullptr;
 }
 
 void FContextSwitchProvider::Add(uint32 ThreadId, double Start, double End, uint32 CoreNumber)
@@ -83,9 +88,22 @@ void FContextSwitchProvider::Add(uint32 ThreadId, double Start, double End, uint
 	ContextSwitch.CoreNumber = CoreNumber;
 }
 
-const IContextSwitchProvider& ReadContextSwitchProvider(const IAnalysisSession& Session)
+bool FContextSwitchProvider::HasData() const
+{
+	return Threads.Num() > 0;
+}
+
+void FContextSwitchProvider::AddThreadInfo(uint32 TraceThreadId, uint32 SystemThreadId)
+{
+	if (!ThreadIdMap.Contains(TraceThreadId))
+	{
+		ThreadIdMap.Add(TraceThreadId, SystemThreadId);
+	}
+}
+
+const IContextSwitchProvider* ReadContextSwitchProvider(const IAnalysisSession& Session)
 {
-	return *Session.ReadProvider<IContextSwitchProvider>(FContextSwitchProvider::ProviderName);
+	return Session.ReadProvider<IContextSwitchProvider>(FContextSwitchProvider::ProviderName);
 }
 
 } // namespace TraceServices
diff --git a/Engine/Source/Developer/TraceServices/Private/Model/ContextSwitchesPrivate.h b/Engine/Source/Developer/TraceServices/Private/Model/ContextSwitchesPrivate.h
--- a/Engine/Source/Developer/TraceServices/Private/Model/ContextSwitchesPrivate.h
+++ b/Engine/Source/Developer/TraceServices/Private/Model/ContextSwitchesPrivate.h
@@ -21,10 +21,16 @@ public:
 	const TPagedArray<FContextSwitch>* GetContextSwitches(uint32 ThreadId) const override;
 
 	void Add(uint32 ThreadId, double Start, double End, uint32 CoreNumber);
+	void AddThreadInfo(uint32 TraceThreadId, uint32 SystemThreadId);
+
+	virtual bool HasData() const override;
 
 private:
 	IAnalysisSession& Session;
+	// System Thread Id -> PagedArray
 	TMap<uint32, TPagedArray<FContextSwitch>*> Threads;
+	// Trace Thread Id -> System Thread Id
+	TMap<uint32, uint32> ThreadIdMap;
 };
 
 } // namespace TraceServices
diff --git a/Engine/Source/Developer/TraceServices/Public/TraceServices/Model/ContextSwitches.h b/Engine/Source/Developer/TraceServices/Public/TraceServices/Model/ContextSwitches.h
--- a/Engine/Source/Developer/TraceServices/Public/TraceServices/Model/ContextSwitches.h
+++ b/Engine/Source/Developer/TraceServices/Public/TraceServices/Model/ContextSwitches.h
@@ -22,8 +22,9 @@ public:
 	virtual ~IContextSwitchProvider() = default;
 	virtual int32 GetCoreNumber(uint32 ThreadId, double Time) const = 0;
 	virtual const TPagedArray<FContextSwitch>* GetContextSwitches(uint32 ThreadId) const = 0;
+	virtual bool HasData() const = 0;
 };
 
-TRACESERVICES_API const IContextSwitchProvider& ReadContextSwitchProvider(const IAnalysisSession& Session);
+TRACESERVICES_API const IContextSwitchProvider* ReadContextSwitchProvider(const IAnalysisSession& Session);
 
 } // namespace TraceServices
