Changed the default port that Trace and its server use to 1981. This is to mitigate collision and confusion as across branches.

#rnx

diff --git a/Engine/Source/Developer/TraceInsights/Private/Insights/TraceInsightsModule.cpp b/Engine/Source/Developer/TraceInsights/Private/Insights/TraceInsightsModule.cpp
--- a/Engine/Source/Developer/TraceInsights/Private/Insights/TraceInsightsModule.cpp
+++ b/Engine/Source/Developer/TraceInsights/Private/Insights/TraceInsightsModule.cpp
@@ -113,7 +113,7 @@ void FTraceInsightsModule::CreateDefaultStore()
 	// Create the Store Service.
 	UE::Trace::FStoreService::FDesc StoreServiceDesc;
 	StoreServiceDesc.StoreDir = *StoreDir;
-	StoreServiceDesc.RecorderPort = 1980;
+	StoreServiceDesc.RecorderPort = 1981;
 	StoreServiceDesc.ThreadCount = 2;
 	StoreService = TUniquePtr<UE::Trace::FStoreService>(UE::Trace::FStoreService::Create(StoreServiceDesc));
 
diff --git a/Engine/Source/Runtime/TraceLog/Private/Trace/Control.cpp b/Engine/Source/Runtime/TraceLog/Private/Trace/Control.cpp
--- a/Engine/Source/Runtime/TraceLog/Private/Trace/Control.cpp
+++ b/Engine/Source/Runtime/TraceLog/Private/Trace/Control.cpp
@@ -16,7 +16,7 @@ namespace Trace {
 namespace Private {
 
 ////////////////////////////////////////////////////////////////////////////////
-bool	Writer_SendTo(const ANSICHAR*, uint32);
+bool	Writer_SendTo(const ANSICHAR*, uint32=0);
 bool	Writer_WriteTo(const ANSICHAR*);
 bool	Writer_Stop();
 
@@ -290,7 +290,7 @@ void Writer_InitializeControl()
 		{
 			if (ArgC > 0)
 			{
-				Writer_SendTo(ArgV[0], 1980);
+				Writer_SendTo(ArgV[0]);
 			}
 		}
 	);
diff --git a/Engine/Source/Runtime/TraceLog/Private/Trace/Writer.cpp b/Engine/Source/Runtime/TraceLog/Private/Trace/Writer.cpp
--- a/Engine/Source/Runtime/TraceLog/Private/Trace/Writer.cpp
+++ b/Engine/Source/Runtime/TraceLog/Private/Trace/Writer.cpp
@@ -632,7 +632,7 @@ bool Writer_SendTo(const ANSICHAR* Host, uint32 Port)
 
 	Writer_InternalInitialize();
 
-	Port = Port ? Port : 1980;
+	Port = Port ? Port : 1981;
 	UPTRINT DataHandle = TcpSocketConnect(Host, Port);
 	if (!DataHandle)
 	{
