Fixed static analysis warning.

diff --git a/Engine/Source/Developer/TraceInsights/Private/Insights/MemoryProfiler/ViewModels/MemAllocTable.cpp b/Engine/Source/Developer/TraceInsights/Private/Insights/MemoryProfiler/ViewModels/MemAllocTable.cpp
--- a/Engine/Source/Developer/TraceInsights/Private/Insights/MemoryProfiler/ViewModels/MemAllocTable.cpp
+++ b/Engine/Source/Developer/TraceInsights/Private/Insights/MemoryProfiler/ViewModels/MemAllocTable.cpp
@@ -614,6 +614,7 @@ void FMemAllocTable::AddDefaultColumns()
 									break;
 								}
 							}
+							check(Frame != nullptr);
 
 							TStringBuilder<1024> Str;
 							FormatStackFrame(*Frame, Str, EStackFrameFormatFlags::Module);
