Added tooltips for Context Switches events
+ Other fixes related to context switches visualisation.

diff --git a/Engine/Source/Developer/TraceInsights/Private/Insights/ContextSwitches/ViewModels/ContextSwitchTimingEvent.cpp b/Engine/Source/Developer/TraceInsights/Private/Insights/ContextSwitches/ViewModels/ContextSwitchTimingEvent.cpp
new file mode 100644
--- /dev/null
+++ b/Engine/Source/Developer/TraceInsights/Private/Insights/ContextSwitches/ViewModels/ContextSwitchTimingEvent.cpp
@@ -0,0 +1,16 @@
+// Copyright Epic Games, Inc. All Rights Reserved.
+
+#include "Insights/ContextSwitches/ViewModels/ContextSwitchTimingEvent.h"
+
+////////////////////////////////////////////////////////////////////////////////////////////////////
+
+INSIGHTS_IMPLEMENT_RTTI(FContextSwitchTimingEvent)
+
+////////////////////////////////////////////////////////////////////////////////////////////////////
+
+FContextSwitchTimingEvent::FContextSwitchTimingEvent(const TSharedRef<const FBaseTimingTrack> InTrack, double InStartTime, double InEndTime, uint32 InDepth)
+	: FTimingEvent(InTrack, InStartTime, InEndTime, InDepth)
+{}
+
+////////////////////////////////////////////////////////////////////////////////////////////////////
+////////////////////////////////////////////////////////////////////////////////////////////////////
diff --git a/Engine/Source/Developer/TraceInsights/Private/Insights/ContextSwitches/ViewModels/ContextSwitchTimingEvent.h b/Engine/Source/Developer/TraceInsights/Private/Insights/ContextSwitches/ViewModels/ContextSwitchTimingEvent.h
new file mode 100644
--- /dev/null
+++ b/Engine/Source/Developer/TraceInsights/Private/Insights/ContextSwitches/ViewModels/ContextSwitchTimingEvent.h
@@ -0,0 +1,27 @@
+// Copyright Epic Games, Inc. All Rights Reserved.
+
+#pragma once
+
+#include "CoreMinimal.h"
+
+// Insights
+#include "Insights/ViewModels/TimingEvent.h"
+
+////////////////////////////////////////////////////////////////////////////////////////////////////
+
+class FContextSwitchTimingEvent : public FTimingEvent
+{
+	INSIGHTS_DECLARE_RTTI(FContextSwitchTimingEvent, FTimingEvent)
+
+public:
+	FContextSwitchTimingEvent(const TSharedRef<const FBaseTimingTrack> InTrack, double InStartTime, double InEndTime, uint32 InDepth);
+	virtual ~FContextSwitchTimingEvent() {}
+
+	void SetCoreNumber(uint32 InCoreNumber) { CoreNumber = InCoreNumber; }
+	uint32 GetCoreNumber() const { return CoreNumber; }
+
+private:
+	uint32 CoreNumber = 0;
+};
+
+////////////////////////////////////////////////////////////////////////////////////////////////////
diff --git a/Engine/Source/Developer/TraceInsights/Private/Insights/ContextSwitches/ViewModels/ContextSwitchesTimingTrack.cpp b/Engine/Source/Developer/TraceInsights/Private/Insights/ContextSwitches/ViewModels/ContextSwitchesTimingTrack.cpp
--- a/Engine/Source/Developer/TraceInsights/Private/Insights/ContextSwitches/ViewModels/ContextSwitchesTimingTrack.cpp
+++ b/Engine/Source/Developer/TraceInsights/Private/Insights/ContextSwitches/ViewModels/ContextSwitchesTimingTrack.cpp
@@ -10,6 +10,7 @@
 
 // Insights
 #include "Insights/Common/TimeUtils.h"
+#include "Insights/ContextSwitches/ViewModels/ContextSwitchTimingEvent.h"
 #include "Insights/ITimingViewSession.h"
 #include "Insights/InsightsManager.h"
 #include "Insights/ViewModels/ThreadTimingTrack.h"
@@ -149,36 +150,13 @@ void FContextSwitchesTimingTrack::BuildDrawState(ITimingEventsTrackDrawStateBuil
 		return;
 	}
 
-	const TPagedArray<FContextSwitch>* ContextSwitches = ContextSwitchesProvider->GetContextSwitches(ThreadId);
-
-	if (ContextSwitches == nullptr)
-	{
-		return;
-	}
-
 	const FTimingTrackViewport& Viewport = Context.GetViewport();
 
-	uint64 ContextSwitchPageIndex = Algo::UpperBoundBy(*ContextSwitches, Viewport.GetStartTime(), [](const TPagedArrayPage<FContextSwitch>& ContextSwitchPage)
+	ContextSwitchesProvider->EnumerateContextSwitches(ThreadId, Viewport.GetStartTime(), Viewport.GetEndTime(), [&Builder](const FContextSwitch& ContextSwitch)
 		{
-			return ContextSwitchPage.Items[0].Start;
+			Builder.AddEvent(ContextSwitch.Start, ContextSwitch.End, 0, *FString::Printf(TEXT("Core %d"), ContextSwitch.CoreNumber), 0, 0);
+			return EContextSwitchEnumerationResult::Continue;
 		});
-
-	if (ContextSwitchPageIndex > 0)
-	{
-		--ContextSwitchPageIndex;
-	}
-
-	auto Iterator = ContextSwitches->GetIteratorFromPage(ContextSwitchPageIndex);
-	const FContextSwitch* CurrentContextSwitch = Iterator.NextItem();
-	while (CurrentContextSwitch && CurrentContextSwitch->Start < Viewport.GetEndTime())
-	{
-		if (CurrentContextSwitch->End > Viewport.GetStartTime())
-		{
-			Builder.AddEvent(CurrentContextSwitch->Start, CurrentContextSwitch->End, 0, *FString::Printf(TEXT("Core %d"), CurrentContextSwitch->CoreNumber), 0, 0);
-		}
-
-		CurrentContextSwitch = Iterator.NextItem();
-	}
 }
 
 ////////////////////////////////////////////////////////////////////////////////////////////////////
@@ -198,7 +176,7 @@ void FContextSwitchesTimingTrack::PostDraw(const ITimingTrackDrawContext& Contex
 
 const TSharedPtr<const ITimingEvent> FContextSwitchesTimingTrack::GetEvent(float InPosX, float InPosY, const FTimingTrackViewport& Viewport) const
 {
-	TSharedPtr<FTimingEvent> TimingEvent;
+	TSharedPtr<FContextSwitchTimingEvent> TimingEvent;
 
 	const FTimingViewLayout& Layout = Viewport.GetLayout();
 
@@ -255,37 +233,44 @@ const TSharedPtr<const ITimingEvent> FContextSwitchesTimingTrack::GetEvent(float
 			return TimingEvent;
 		}
 
-		const TPagedArray<FContextSwitch>* ContextSwitches = ContextSwitchesProvider->GetContextSwitches(ThreadId);
+		FContextSwitch BestMatchContextSwitch;
+		double Delta = 2 * SecondsPerPixel;
 
-		if (ContextSwitches == nullptr)
-		{
-			return TimingEvent;
-		}
-
-		uint64 ContextSwitchPageIndex = Algo::UpperBoundBy(*ContextSwitches, EventTime + SecondsPerPixel, [](const TPagedArrayPage<FContextSwitch>& ContextSwitchPage)
+		ContextSwitchesProvider->EnumerateContextSwitches(ThreadId, EventTime - 2 * SecondsPerPixel, EventTime + 2 * SecondsPerPixel,
+			[EventTime, &BestMatchContextSwitch, &Delta](const FContextSwitch& ContextSwitch)
 			{
-				return ContextSwitchPage.Items[0].Start;
+				if (ContextSwitch.Start <= EventTime && ContextSwitch.End >= EventTime)
+				{
+					BestMatchContextSwitch = ContextSwitch;
+					Delta = 0.0f;
+					return EContextSwitchEnumerationResult::Stop;
+				}
+
+				if (ContextSwitch.End <= EventTime)
+				{
+					if (EventTime - ContextSwitch.End < Delta)
+					{
+						Delta = EventTime - ContextSwitch.End;
+						BestMatchContextSwitch = ContextSwitch;
+					}
+				}
+
+				if (ContextSwitch.Start >= EventTime)
+				{
+					if (ContextSwitch.Start - EventTime < Delta)
+					{
+						Delta = ContextSwitch.Start - EventTime;
+						BestMatchContextSwitch = ContextSwitch;
+					}
+				}
+
+				return EContextSwitchEnumerationResult::Continue;
 			});
 
-		if (ContextSwitchPageIndex > 0)
-		{
-			--ContextSwitchPageIndex;
-		}
-
-		auto Iterator = ContextSwitches->GetIteratorFromPage(ContextSwitchPageIndex);
-		const FContextSwitch* CurrentContextSwitch = Iterator.NextItem();
-		while (CurrentContextSwitch && CurrentContextSwitch->End < EventTime)
+		if (Delta < 2 * SecondsPerPixel)
 		{
-			CurrentContextSwitch = Iterator.NextItem();
-		}
-
-		if (CurrentContextSwitch)
-		{
-			ensure(ParentTrack.IsValid());
-			if (CurrentContextSwitch->Start < EventTime && EventTime < CurrentContextSwitch->End)
-			{
-				TimingEvent = MakeShared<FTimingEvent>(SharedThis(this), CurrentContextSwitch->Start, CurrentContextSwitch->End, 0);
-			}
+			TimingEvent = MakeShared<FContextSwitchTimingEvent>(SharedThis(this), BestMatchContextSwitch.Start, BestMatchContextSwitch.End, 0);
+			TimingEvent->SetCoreNumber(BestMatchContextSwitch.CoreNumber);
 		}
 	}
 
@@ -298,9 +283,23 @@ void FContextSwitchesTimingTrack::InitTooltip(FTooltipDrawState& InOutTooltip, c
 {
 	InOutTooltip.ResetContent();
 
+	if (!InTooltipEvent.CheckTrack(this) || !InTooltipEvent.Is<FContextSwitchTimingEvent>())
+	{
+		return;
+	}
+
+	const FContextSwitchTimingEvent& ContextSwitchEvent = InTooltipEvent.As<FContextSwitchTimingEvent>();
+	InOutTooltip.AddTitle(FString::Printf(TEXT("Core %d"), ContextSwitchEvent.GetCoreNumber()));
+
+	InOutTooltip.AddNameValueTextLine(TEXT("Start Time:"), TimeUtils::FormatTimeAuto(InTooltipEvent.GetStartTime(), 6));
+	InOutTooltip.AddNameValueTextLine(TEXT("End Time:"), TimeUtils::FormatTimeAuto(InTooltipEvent.GetEndTime(), 6));
+	InOutTooltip.AddNameValueTextLine(TEXT("Duration:"), TimeUtils::FormatTimeAuto(InTooltipEvent.GetDuration(), 6));
+
 	InOutTooltip.UpdateLayout();
 }
 
+////////////////////////////////////////////////////////////////////////////////////////////////////
+
 } // namespace Insights
 
 #undef LOCTEXT_NAMESPACE
diff --git a/Engine/Source/Developer/TraceInsights/Private/Insights/ViewModels/ThreadTimingTrack.cpp b/Engine/Source/Developer/TraceInsights/Private/Insights/ViewModels/ThreadTimingTrack.cpp
--- a/Engine/Source/Developer/TraceInsights/Private/Insights/ViewModels/ThreadTimingTrack.cpp
+++ b/Engine/Source/Developer/TraceInsights/Private/Insights/ViewModels/ThreadTimingTrack.cpp
@@ -1173,6 +1173,10 @@ void FThreadTimingTrack::InitTooltip(FTooltipDrawState& InOutTooltip, const ITim
 			} while (false);
 		}
 	}
+	else if (ChildTrack.IsValid())
+	{
+		ChildTrack->InitTooltip(InOutTooltip, InTooltipEvent);
+	}
 
 	InOutTooltip.UpdateLayout();
 }
diff --git a/Engine/Source/Developer/TraceInsights/Private/Insights/ViewModels/TimingEventsTrack.cpp b/Engine/Source/Developer/TraceInsights/Private/Insights/ViewModels/TimingEventsTrack.cpp
--- a/Engine/Source/Developer/TraceInsights/Private/Insights/ViewModels/TimingEventsTrack.cpp
+++ b/Engine/Source/Developer/TraceInsights/Private/Insights/ViewModels/TimingEventsTrack.cpp
@@ -171,9 +171,13 @@ void FTimingEventsTrack::UpdateTrackHeight(const ITimingTrackUpdateContext& Cont
 	else
 	{
 		DesiredTrackHeight = Viewport.GetLayout().ComputeTrackHeight(NumLanes);
-		if (ChildTrack.IsValid())
+		if (ChildTrack.IsValid() && ChildTrack->GetHeight() > Viewport.GetLayout().MinTimelineH)
 		{
-			DesiredTrackHeight += ChildTrack->GetHeight();
+			if (DesiredTrackHeight == Viewport.GetLayout().MinTimelineH)
+			{
+				DesiredTrackHeight += 1.0f + 2.0 * Viewport.GetLayout().TimelineDY;
+			}
+			DesiredTrackHeight += ChildTrack->GetHeight() + Viewport.GetLayout().ChildTimelineDY;
 		}
 	}
 
diff --git a/Engine/Source/Developer/TraceInsights/Private/Insights/ViewModels/TimingViewLayout.cpp b/Engine/Source/Developer/TraceInsights/Private/Insights/ViewModels/TimingViewLayout.cpp
--- a/Engine/Source/Developer/TraceInsights/Private/Insights/ViewModels/TimingViewLayout.cpp
+++ b/Engine/Source/Developer/TraceInsights/Private/Insights/ViewModels/TimingViewLayout.cpp
@@ -51,6 +51,11 @@ bool FTimingViewLayout::Update()
 		{
 			TimelineDY -= LayoutTransitionSpeed;
 			bHasChanged = true;
+		}		
+		if (ChildTimelineDY > CompactLayoutChildTimelineDY)
+		{
+			ChildTimelineDY -= LayoutTransitionSpeed / 2;
+			bHasChanged = true;
 		}
 	}
 	else
@@ -69,6 +74,11 @@ bool FTimingViewLayout::Update()
 		{
 			TimelineDY += LayoutTransitionSpeed;
 			bHasChanged = true;
+		}		
+		if (ChildTimelineDY < NormalLayoutChildTimelineDY)
+		{
+			ChildTimelineDY += LayoutTransitionSpeed / 2;
+			bHasChanged = true;
 		}
 	}
 
diff --git a/Engine/Source/Developer/TraceServices/Private/Model/ContextSwitches.cpp b/Engine/Source/Developer/TraceServices/Private/Model/ContextSwitches.cpp
--- a/Engine/Source/Developer/TraceServices/Private/Model/ContextSwitches.cpp
+++ b/Engine/Source/Developer/TraceServices/Private/Model/ContextSwitches.cpp
@@ -72,6 +72,43 @@ const TPagedArray<FContextSwitch>* FContextSwitchProvider::GetContextSwitches(ui
 	return nullptr;
 }
 
+const void FContextSwitchProvider::EnumerateContextSwitches(uint32 ThreadId, double StartTime, double EndTime, ContextSwitchCallback Callback) const
+{
+	Session.ReadAccessCheck();
+
+	const TPagedArray<FContextSwitch>* PagedArray = GetContextSwitches(ThreadId);
+
+	if (PagedArray == nullptr)
+	{
+		return;
+	}
+
+	uint64 ContextSwitchPageIndex = Algo::UpperBoundBy(*PagedArray, StartTime, [](const TPagedArrayPage<FContextSwitch>& ContextSwitchPage)
+		{
+			return ContextSwitchPage.Items[0].Start;
+		});
+
+	if (ContextSwitchPageIndex > 0)
+	{
+		--ContextSwitchPageIndex;
+	}
+
+	auto Iterator = PagedArray->GetIteratorFromPage(ContextSwitchPageIndex);
+	const FContextSwitch* CurrentContextSwitch = Iterator.NextItem();
+	while (CurrentContextSwitch && CurrentContextSwitch->Start < EndTime)
+	{
+		if (CurrentContextSwitch->End > StartTime)
+		{
+			if (Callback(*CurrentContextSwitch) == EContextSwitchEnumerationResult::Stop)
+			{
+				break;
+			}
+		}
+
+		CurrentContextSwitch = Iterator.NextItem();
+	}
+}
+
 void FContextSwitchProvider::Add(uint32 ThreadId, double Start, double End, uint32 CoreNumber)
 {
 	Session.WriteAccessCheck();
diff --git a/Engine/Source/Developer/TraceServices/Private/Model/ContextSwitchesPrivate.h b/Engine/Source/Developer/TraceServices/Private/Model/ContextSwitchesPrivate.h
--- a/Engine/Source/Developer/TraceServices/Private/Model/ContextSwitchesPrivate.h
+++ b/Engine/Source/Developer/TraceServices/Private/Model/ContextSwitchesPrivate.h
@@ -18,7 +18,7 @@ public:
 	~FContextSwitchProvider();
 
 	int32 GetCoreNumber(uint32 ThreadId, double Time) const override;
-	const TPagedArray<FContextSwitch>* GetContextSwitches(uint32 ThreadId) const override;
+	const void EnumerateContextSwitches(uint32 ThreadId, double StartTime, double EndTime, ContextSwitchCallback Callback) const override;
 
 	void Add(uint32 ThreadId, double Start, double End, uint32 CoreNumber);
 	void AddThreadInfo(uint32 TraceThreadId, uint32 SystemThreadId);
@@ -26,6 +26,8 @@ public:
 	virtual bool HasData() const override;
 
 private:
+	const TPagedArray<FContextSwitch>* GetContextSwitches(uint32 ThreadId) const;
+
 	IAnalysisSession& Session;
 	// System Thread Id -> PagedArray
 	TMap<uint32, TPagedArray<FContextSwitch>*> Threads;
diff --git a/Engine/Source/Developer/TraceServices/Public/TraceServices/Model/ContextSwitches.h b/Engine/Source/Developer/TraceServices/Public/TraceServices/Model/ContextSwitches.h
--- a/Engine/Source/Developer/TraceServices/Public/TraceServices/Model/ContextSwitches.h
+++ b/Engine/Source/Developer/TraceServices/Public/TraceServices/Model/ContextSwitches.h
@@ -15,13 +15,21 @@ struct FContextSwitch
 	uint32 CoreNumber;
 };
 
+enum class EContextSwitchEnumerationResult
+{
+	Continue,
+	Stop,
+};
+
 class IContextSwitchProvider
 	: public IProvider
 {
 public:
+	typedef TFunctionRef<EContextSwitchEnumerationResult(const FContextSwitch&)> ContextSwitchCallback;
+
 	virtual ~IContextSwitchProvider() = default;
 	virtual int32 GetCoreNumber(uint32 ThreadId, double Time) const = 0;
-	virtual const TPagedArray<FContextSwitch>* GetContextSwitches(uint32 ThreadId) const = 0;
+	virtual const void EnumerateContextSwitches(uint32 ThreadId, double StartTime, double EndTime, ContextSwitchCallback Callback) const = 0;
 	virtual bool HasData() const = 0;
 };
 
