[Insights] Renamed table export fo file menu options and added a new option to export only non leaf nodes

diff --git a/Engine/Source/Developer/TraceInsights/Private/Insights/Table/Widgets/STableTreeView.cpp b/Engine/Source/Developer/TraceInsights/Private/Insights/Table/Widgets/STableTreeView.cpp
--- a/Engine/Source/Developer/TraceInsights/Private/Insights/Table/Widgets/STableTreeView.cpp
+++ b/Engine/Source/Developer/TraceInsights/Private/Insights/Table/Widgets/STableTreeView.cpp
@@ -69,8 +69,9 @@ public:
 		UI_COMMAND(Command_ExpandSubtree, "Expand Subtree", "Expand the subtree that starts from the selected group node.", EUserInterfaceActionType::Button, FInputChord(EModifierKey::None, EKeys::E));
 		UI_COMMAND(Command_ExpandCriticalPath, "Expand Critical Path", "Expand the first group child node recursively until a leaf nodes in reached.", EUserInterfaceActionType::Button, FInputChord(EModifierKey::None, EKeys::R));
 		UI_COMMAND(Command_CollapseSubtree, "Collapse Subtree", "Collapse the subtree that starts from the selected group node.", EUserInterfaceActionType::Button, FInputChord(EModifierKey::None, EKeys::C));
-		UI_COMMAND(Command_ExportToFile, "Export To File", "Export the table content to a file.", EUserInterfaceActionType::Button, FInputChord(EModifierKey::Control, EKeys::E));
-		UI_COMMAND(Command_ExportExpandedToFile, "Export Expanded To File", "Export the fully expanded table to a file.", EUserInterfaceActionType::Button, FInputChord(EModifierKey::Shift, EKeys::E));
+		UI_COMMAND(Command_ExportToFile, "Export Visible Tree to File...", "Exports the tree/table content to a file. It exports only the tree nodes currently expanded in the tree, including leaf nodes.", EUserInterfaceActionType::Button, FInputChord(EModifierKey::Control, EKeys::E));
+		UI_COMMAND(Command_ExportEntireTreeToFile, "Export Entire Tree to File...", "Exports the entire tree/table content to a file. It exports also the collapsed tree nodes, including the leaf nodes. Filtered out nodes are not exported.", EUserInterfaceActionType::Button, FInputChord(EModifierKey::Shift, EKeys::E));
+		UI_COMMAND(Command_ExportEntireTreeToFileNoLeafs, "Export Entire Tree (No Leafs) to File...", "Exports the entire tree/table content to a file, but not the leaf nodes. It exports the collapsed tree nodes. Filtered out nodes are not exported.", EUserInterfaceActionType::Button, FInputChord(EModifierKey::Alt, EKeys::E));
 	}
 	PRAGMA_ENABLE_OPTIMIZATION
 
@@ -81,7 +82,8 @@ public:
 	TSharedPtr<FUICommandInfo> Command_ExpandCriticalPath;
 	TSharedPtr<FUICommandInfo> Command_CollapseSubtree;
 	TSharedPtr<FUICommandInfo> Command_ExportToFile;
-	TSharedPtr<FUICommandInfo> Command_ExportExpandedToFile;
+	TSharedPtr<FUICommandInfo> Command_ExportEntireTreeToFile;
+	TSharedPtr<FUICommandInfo> Command_ExportEntireTreeToFileNoLeafs;
 };
 
 const FName STableTreeView::RootNodeName(TEXT("Root"));
@@ -348,8 +350,9 @@ void STableTreeView::InitCommandList()
 	CommandList->MapAction(FTableTreeViewCommands::Get().Command_ExpandSubtree, FExecuteAction::CreateSP(this, &STableTreeView::ContextMenu_ExpandSubtree_Execute), FCanExecuteAction::CreateSP(this, &STableTreeView::ContextMenu_ExpandSubtree_CanExecute));
 	CommandList->MapAction(FTableTreeViewCommands::Get().Command_ExpandCriticalPath, FExecuteAction::CreateSP(this, &STableTreeView::ContextMenu_ExpandCriticalPath_Execute), FCanExecuteAction::CreateSP(this, &STableTreeView::ContextMenu_ExpandCriticalPath_CanExecute));
 	CommandList->MapAction(FTableTreeViewCommands::Get().Command_CollapseSubtree, FExecuteAction::CreateSP(this, &STableTreeView::ContextMenu_CollapseSubtree_Execute), FCanExecuteAction::CreateSP(this, &STableTreeView::ContextMenu_CollapseSubtree_CanExecute));
-	CommandList->MapAction(FTableTreeViewCommands::Get().Command_ExportToFile, FExecuteAction::CreateSP(this, &STableTreeView::ContextMenu_ExportToFile_Execute, true), FCanExecuteAction::CreateSP(this, &STableTreeView::ContextMenu_ExportToFile_CanExecute));
-	CommandList->MapAction(FTableTreeViewCommands::Get().Command_ExportExpandedToFile, FExecuteAction::CreateSP(this, &STableTreeView::ContextMenu_ExportToFile_Execute, false), FCanExecuteAction::CreateSP(this, &STableTreeView::ContextMenu_ExportToFile_CanExecute));
+	CommandList->MapAction(FTableTreeViewCommands::Get().Command_ExportToFile, FExecuteAction::CreateSP(this, &STableTreeView::ContextMenu_ExportToFile_Execute, false, true), FCanExecuteAction::CreateSP(this, &STableTreeView::ContextMenu_ExportToFile_CanExecute));
+	CommandList->MapAction(FTableTreeViewCommands::Get().Command_ExportEntireTreeToFile, FExecuteAction::CreateSP(this, &STableTreeView::ContextMenu_ExportToFile_Execute, true, true), FCanExecuteAction::CreateSP(this, &STableTreeView::ContextMenu_ExportToFile_CanExecute));
+	CommandList->MapAction(FTableTreeViewCommands::Get().Command_ExportEntireTreeToFileNoLeafs, FExecuteAction::CreateSP(this, &STableTreeView::ContextMenu_ExportToFile_Execute, true, false), FCanExecuteAction::CreateSP(this, &STableTreeView::ContextMenu_ExportToFile_CanExecute));
 }
 
 ////////////////////////////////////////////////////////////////////////////////////////////////////
@@ -493,7 +496,7 @@ TSharedPtr<SWidget> STableTreeView::TreeView_GetMenuContent()
 			LOCTEXT("ContextMenu_Header_Misc_Sort_Desc", "Export to file"),
 			FNewMenuDelegate::CreateSP(this, &STableTreeView::TreeView_BuildExportMenu),
 			false,
-			FSlateIcon(FEditorStyle::GetStyleSetName(), "Profiler.Misc.CopyToClipboard")
+			FSlateIcon(FEditorStyle::GetStyleSetName(), "ProfilerCommand.ProfilerManager_Save")
 		);
 	}
 	MenuBuilder.EndSection();
@@ -640,16 +643,25 @@ void STableTreeView::TreeView_BuildExportMenu(FMenuBuilder& MenuBuilder)
 		NAME_None,
 		TAttribute<FText>(),
 		TAttribute<FText>(),
-		FSlateIcon(FEditorStyle::GetStyleSetName(), "Profiler.Misc.CopyToClipboard")
+		FSlateIcon(FEditorStyle::GetStyleSetName(), "ProfilerCommand.ProfilerManager_Save")
 	);
 
 	MenuBuilder.AddMenuEntry
 	(
-		FTableTreeViewCommands::Get().Command_ExportExpandedToFile,
+		FTableTreeViewCommands::Get().Command_ExportEntireTreeToFile,
 		NAME_None,
 		TAttribute<FText>(),
 		TAttribute<FText>(),
-		FSlateIcon(FEditorStyle::GetStyleSetName(), "Profiler.Misc.CopyToClipboard")
+		FSlateIcon(FEditorStyle::GetStyleSetName(), "ProfilerCommand.ProfilerManager_Save")
+	);
+
+	MenuBuilder.AddMenuEntry
+	(
+		FTableTreeViewCommands::Get().Command_ExportEntireTreeToFileNoLeafs,
+		NAME_None,
+		TAttribute<FText>(),
+		TAttribute<FText>(),
+		FSlateIcon(FEditorStyle::GetStyleSetName(), "ProfilerCommand.ProfilerManager_Save")
 	);
 
 	MenuBuilder.EndSection();
@@ -657,7 +669,6 @@ void STableTreeView::TreeView_BuildExportMenu(FMenuBuilder& MenuBuilder)
 
 ////////////////////////////////////////////////////////////////////////////////////////////////////
 
-
 void STableTreeView::InitializeAndShowHeaderColumns()
 {
 	for (const TSharedRef<FTableColumn>& ColumnRef : Table->GetColumns())
@@ -3065,7 +3076,7 @@ bool STableTreeView::ContextMenu_ExportToFile_CanExecute() const
 
 ////////////////////////////////////////////////////////////////////////////////////////////////////
 
-void STableTreeView::ContextMenu_ExportToFile_Execute(bool bExportOnlyExpanded)
+void STableTreeView::ContextMenu_ExportToFile_Execute(bool bInExportCollapsed, bool InExportLeafs)
 {
 	TArray<FString> SaveFilenames;
 	IDesktopPlatform* DesktopPlatform = FDesktopPlatformModule::Get();
@@ -3128,7 +3139,7 @@ void STableTreeView::ContextMenu_ExportToFile_Execute(bool bExportOnlyExpanded)
 	};
 
 	TArray<Insights::FBaseTreeNodePtr> NodesBuffer;
-	ExportToFileRec(Root, NodesBuffer, bExportOnlyExpanded, WriteToFile);
+	ExportToFileRec(Root, NodesBuffer, bInExportCollapsed, InExportLeafs, WriteToFile);
 
 	// Write the remaining lines
 	WriteToFile(NodesBuffer);
@@ -3144,7 +3155,7 @@ void STableTreeView::ContextMenu_ExportToFile_Execute(bool bExportOnlyExpanded)
 
 ////////////////////////////////////////////////////////////////////////////////////////////////////
 
-void STableTreeView::ExportToFileRec(const FBaseTreeNodePtr& InGroupNode, TArray<Insights::FBaseTreeNodePtr>& InNodes, bool bInExportOnlyExpanded, WriteToFileCallback Callback)
+void STableTreeView::ExportToFileRec(const FBaseTreeNodePtr& InGroupNode, TArray<Insights::FBaseTreeNodePtr>& InNodes, bool bInExportCollapsed, bool InExportLeafs, WriteToFileCallback Callback)
 {
 	constexpr int MaxBufferSize = 10000;
 
@@ -3155,10 +3166,15 @@ void STableTreeView::ExportToFileRec(const FBaseTreeNodePtr& InGroupNode, TArray
 			// Flush the buffer to the file
 			Callback(InNodes);
 		}
-		InNodes.Add(Node);
-		if (Node->IsGroup() && (!bInExportOnlyExpanded || TreeView->IsItemExpanded(StaticCastSharedPtr<FTableTreeNode>(Node))))
+
+		if (Node->IsGroup() || InExportLeafs)
+		{
+			InNodes.Add(Node);
+		}
+
+		if (Node->IsGroup() && (bInExportCollapsed || TreeView->IsItemExpanded(StaticCastSharedPtr<FTableTreeNode>(Node))))
 		{
-			ExportToFileRec(Node, InNodes, bInExportOnlyExpanded, Callback);
+			ExportToFileRec(Node, InNodes, bInExportCollapsed, InExportLeafs, Callback);
 		}
 	}
 }
diff --git a/Engine/Source/Developer/TraceInsights/Private/Insights/Table/Widgets/STableTreeView.h b/Engine/Source/Developer/TraceInsights/Private/Insights/Table/Widgets/STableTreeView.h
--- a/Engine/Source/Developer/TraceInsights/Private/Insights/Table/Widgets/STableTreeView.h
+++ b/Engine/Source/Developer/TraceInsights/Private/Insights/Table/Widgets/STableTreeView.h
@@ -173,7 +173,7 @@ protected:
 	bool ContextMenu_CollapseSubtree_CanExecute() const;
 	void ContextMenu_CollapseSubtree_Execute();
 	bool ContextMenu_ExportToFile_CanExecute() const;
-	void ContextMenu_ExportToFile_Execute(bool bExportOnlyExpanded);
+	void ContextMenu_ExportToFile_Execute(bool bInExportCollapsed, bool InExportLeafs);
 
 	////////////////////////////////////////////////////////////////////////////////////////////////////
 	// Tree View - Columns' Header
@@ -395,7 +395,7 @@ protected:
 	virtual void ExtendMenu(FMenuBuilder& Menu) {}
 
 	typedef TFunctionRef<void(TArray<Insights::FBaseTreeNodePtr>& InNodes)> WriteToFileCallback;
-	void ExportToFileRec(const FBaseTreeNodePtr& InGroupNode, TArray<Insights::FBaseTreeNodePtr>& InNodes, bool bInExportOnlyExpanded, WriteToFileCallback Callback);
+	void ExportToFileRec(const FBaseTreeNodePtr& InGroupNode, TArray<Insights::FBaseTreeNodePtr>& InNodes, bool bInExportCollapsed, bool InExportLeafs, WriteToFileCallback Callback);
 
 protected:
 	/** Table view model. */
