[Insights] Change Task Timing Track order so it is displayed as the first top docked track.

diff --git a/Engine/Source/Developer/TraceInsights/Public/Insights/ViewModels/BaseTimingTrack.h b/Engine/Source/Developer/TraceInsights/Public/Insights/ViewModels/BaseTimingTrack.h
--- a/Engine/Source/Developer/TraceInsights/Public/Insights/ViewModels/BaseTimingTrack.h
+++ b/Engine/Source/Developer/TraceInsights/Public/Insights/ViewModels/BaseTimingTrack.h
@@ -45,10 +45,10 @@ struct TRACEINSIGHTS_API FTimingTrackOrder
 	static constexpr int32 TimeRuler  = -2 * GroupRange;
 	static constexpr int32 Markers    = -1 * GroupRange;
 	static constexpr int32 First      = 0;
-	static constexpr int32 Memory     = 1 * GroupRange;
-	static constexpr int32 Gpu        = 2 * GroupRange;
-	static constexpr int32 Cpu        = 3 * GroupRange;
-	static constexpr int32 Task       = 4 * GroupRange;
+	static constexpr int32 Task       = 1 * GroupRange;
+	static constexpr int32 Memory     = 2 * GroupRange;
+	static constexpr int32 Gpu        = 3 * GroupRange;
+	static constexpr int32 Cpu        = 4 * GroupRange;
 	static constexpr int32 Last       = 5 * GroupRange;
 };
 
