[Insights] Fix event relations when header tracks are displayed.
+ Context Switches now update the session duration

diff --git a/Engine/Source/Developer/TraceInsights/Private/Insights/ViewModels/TaskGraphRelation.cpp b/Engine/Source/Developer/TraceInsights/Private/Insights/ViewModels/TaskGraphRelation.cpp
--- a/Engine/Source/Developer/TraceInsights/Private/Insights/ViewModels/TaskGraphRelation.cpp
+++ b/Engine/Source/Developer/TraceInsights/Private/Insights/ViewModels/TaskGraphRelation.cpp
@@ -23,9 +23,17 @@ void FTaskGraphRelation::Draw(const FDrawContext& DrawContext, const FTimingTrac
 {
 	int32 LayerId = Helper.GetRelationLayerId();
 
+	TSharedPtr<const FBaseTimingTrack> SourceTrackShared = SourceTrack.Pin();
+	TSharedPtr<const FBaseTimingTrack> TargetTrackShared = TargetTrack.Pin();
+
+	if (!SourceTrackShared.IsValid() || !TargetTrackShared.IsValid())
+	{
+		return;
+	}
+
 	if (Filter == ITimingEventRelation::EDrawFilter::BetweenScrollableTracks)
 	{
-		if (SourceTrack.Pin()->GetLocation() > ETimingTrackLocation::Scrollable || TargetTrack.Pin()->GetLocation() > ETimingTrackLocation::Scrollable)
+		if (SourceTrackShared->GetLocation() > ETimingTrackLocation::Scrollable || TargetTrackShared->GetLocation() > ETimingTrackLocation::Scrollable)
 		{
 			return;
 		}
@@ -33,7 +41,7 @@ void FTaskGraphRelation::Draw(const FDrawContext& DrawContext, const FTimingTrac
 
 	if (Filter == ITimingEventRelation::EDrawFilter::BetweenDockedTracks)
 	{
-		if (SourceTrack.Pin()->GetLocation() <= ETimingTrackLocation::Scrollable && TargetTrack.Pin()->GetLocation() <= ETimingTrackLocation::Scrollable)
+		if (SourceTrackShared->GetLocation() <= ETimingTrackLocation::Scrollable && TargetTrackShared->GetLocation() <= ETimingTrackLocation::Scrollable)
 		{
 			return;
 		}
@@ -42,12 +50,20 @@ void FTaskGraphRelation::Draw(const FDrawContext& DrawContext, const FTimingTrac
 	}
 
 	float X1 = Viewport.TimeToSlateUnitsRounded(SourceTime);
-	float Y1 = SourceTrack.Pin()->GetPosY();
+	float Y1 = SourceTrackShared->GetPosY();
 	Y1 += Viewport.GetLayout().GetLaneY(SourceDepth) + Viewport.GetLayout().EventH / 2.0f;
+	if (SourceTrackShared->GetChildTrack())
+	{
+		Y1 += SourceTrackShared->GetChildTrack()->GetHeight() + Viewport.GetLayout().ChildTimelineDY;
+	}
 
 	float X2 = Viewport.TimeToSlateUnitsRounded(TargetTime);
-	float Y2 = TargetTrack.Pin()->GetPosY();
+	float Y2 = TargetTrackShared->GetPosY();
 	Y2 += Viewport.GetLayout().GetLaneY(TargetDepth) + Viewport.GetLayout().EventH / 2.0f;
+	if (TargetTrackShared->GetChildTrack())
+	{
+		Y2 += TargetTrackShared->GetChildTrack()->GetHeight() + Viewport.GetLayout().ChildTimelineDY;
+	}
 
 	FVector2D StartPoint = FVector2D(X1, Y1);
 	FVector2D EndPoint = FVector2D(X2, Y2);
@@ -64,7 +80,7 @@ void FTaskGraphRelation::Draw(const FDrawContext& DrawContext, const FTimingTrac
 	constexpr float ArrowRotationAngle = 20.0f;
 	FVector2D ArrowDirection(-ArrowDirectionLen, 0.0f);
 
-	if (Distance > LineLenghtAtEnds)
+	if (Distance > LineLenghtAtEnds && !FMath::IsNearlyEqual(StartPoint.Y, EndPoint.Y))
 	{
 		FVector2D SplineStart(StartPoint.X + LineLenghtAtEnds, StartPoint.Y);
 		FVector2D SplineEnd(EndPoint.X - LineLenghtAtEnds, EndPoint.Y);
diff --git a/Engine/Source/Developer/TraceServices/Private/Analyzers/PlatformEventTraceAnalysis.cpp b/Engine/Source/Developer/TraceServices/Private/Analyzers/PlatformEventTraceAnalysis.cpp
--- a/Engine/Source/Developer/TraceServices/Private/Analyzers/PlatformEventTraceAnalysis.cpp
+++ b/Engine/Source/Developer/TraceServices/Private/Analyzers/PlatformEventTraceAnalysis.cpp
@@ -39,6 +39,9 @@ bool FPlatformEventTraceAnalyzer::OnEvent(uint16 RouteId, EStyle Style, const FO
 		uint32 ThreadId = EventData.GetValue<uint32>("ThreadId");
 		uint32 CoreNumber = EventData.GetValue<uint8>("CoreNumber");
 		ContextSwitchProvider.Add(ThreadId, Start, End, CoreNumber);
+
+		Session.UpdateDurationSeconds(End);
+
 		break;
 	}
 
@@ -48,6 +51,9 @@ bool FPlatformEventTraceAnalyzer::OnEvent(uint16 RouteId, EStyle Style, const FO
 		uint32 ThreadId = EventData.GetValue<uint32>("ThreadId");
 		const TArrayReader<uint64>& Addresses = EventData.GetArray<uint64>("Addresses");
 		StackSampleProvider.Add(ThreadId, Time, Addresses.Num(), Addresses.GetData());
+
+		Session.UpdateDurationSeconds(Time);
+
 		break;
 	}
 
@@ -58,6 +64,7 @@ bool FPlatformEventTraceAnalyzer::OnEvent(uint16 RouteId, EStyle Style, const FO
 
 void FPlatformEventTraceAnalyzer::OnThreadInfo(const FThreadInfo& ThreadInfo)
 {
+	FAnalysisSessionEditScope _(Session);
 	ContextSwitchProvider.AddThreadInfo(ThreadInfo.GetId(), ThreadInfo.GetSystemId());
 }
 
