[Insights]
  - Improved visualization of context switches. The "Core X" events are draw as 3px horizontal segments with text displayed on top of the segment. Both segment and text are using the same color (based on cpu core index).
  - Added hashed area for time regions between context switch events. This extends the edges of the context switch segments and fades out the outside regions.

diff --git a/Engine/Source/Developer/TraceInsights/Private/Insights/ContextSwitches/ViewModels/ContextSwitchesTimingTrack.cpp b/Engine/Source/Developer/TraceInsights/Private/Insights/ContextSwitches/ViewModels/ContextSwitchesTimingTrack.cpp
--- a/Engine/Source/Developer/TraceInsights/Private/Insights/ContextSwitches/ViewModels/ContextSwitchesTimingTrack.cpp
+++ b/Engine/Source/Developer/TraceInsights/Private/Insights/ContextSwitches/ViewModels/ContextSwitchesTimingTrack.cpp
@@ -168,8 +168,110 @@ void FContextSwitchesTimingTrack::BuildFilteredDrawState(ITimingEventsTrackDrawS
 
 ////////////////////////////////////////////////////////////////////////////////////////////////////
 
+void FContextSwitchesTimingTrack::Draw(const ITimingTrackDrawContext& Context) const
+{
+	DrawLineEvents(Context, 1.0f);
+}
+
+////////////////////////////////////////////////////////////////////////////////////////////////////
+
+void FContextSwitchesTimingTrack::DrawLineEvents(const ITimingTrackDrawContext& Context, const float OffsetY) const
+{
+	const FTimingViewDrawHelper& Helper = *static_cast<const FTimingViewDrawHelper*>(&Context.GetHelper());
+
+	if (Context.GetEventFilter().IsValid() || HasCustomFilter())
+	{
+		Helper.DrawFadedLineEvents(GetDrawState(), *this, OffsetY, 0.1f);
+
+		if (UpdateFilteredDrawStateOpacity())
+		{
+			Helper.DrawLineEvents(GetFilteredDrawState(), *this, OffsetY);
+		}
+		else
+		{
+			Helper.DrawFadedLineEvents(GetFilteredDrawState(), *this, OffsetY, GetFilteredDrawStateOpacity());
+		}
+	}
+	else
+	{
+		Helper.DrawLineEvents(GetDrawState(), *this, OffsetY);
+	}
+}
+
+////////////////////////////////////////////////////////////////////////////////////////////////////
+
 void FContextSwitchesTimingTrack::PostDraw(const ITimingTrackDrawContext& Context) const
 {
+	float LineY1 = 0.0f;
+	float LineY2 = 0.0f;
+	ETimingTrackLocation LocalLocation = ETimingTrackLocation::None;
+	TSharedPtr<FBaseTimingTrack> LocalParentTrack = GetParentTrack().Pin();
+	if (LocalParentTrack)
+	{
+		LineY1 = LocalParentTrack->GetPosY();
+		LineY2 = LineY1 + LocalParentTrack->GetHeight();
+		LocalLocation = LocalParentTrack->GetLocation();
+	}
+	else
+	{
+		LineY1 = GetPosY();
+		LineY2 = LineY1 + GetHeight();
+		LocalLocation = GetLocation();
+	}
+
+	const FTimingTrackViewport& Viewport = Context.GetViewport();
+	switch (LocalLocation)
+	{
+		case ETimingTrackLocation::Scrollable:
+		{
+			const float TopY = Viewport.GetTopOffset();
+			if (LineY1 < TopY)
+			{
+				LineY1 = TopY;
+			}
+			const float BottomY = Viewport.GetHeight() - Viewport.GetBottomOffset();
+			if (LineY2 > BottomY)
+			{
+				LineY2 = BottomY;
+			}
+			break;
+		}
+		case ETimingTrackLocation::TopDocked:
+		{
+			const float TopY = 0.0f;
+			if (LineY1 < TopY)
+			{
+				LineY1 = TopY;
+			}
+			const float BottomY = Viewport.GetTopOffset();
+			if (LineY2 > BottomY)
+			{
+				LineY2 = BottomY;
+			}
+			break;
+		}
+		case ETimingTrackLocation::BottomDocked:
+		{
+			const float TopY = Viewport.GetHeight() - Viewport.GetBottomOffset();
+			if (LineY1 < TopY)
+			{
+				LineY1 = TopY;
+			}
+			const float BottomY = Viewport.GetHeight();
+			if (LineY2 > BottomY)
+			{
+				LineY2 = BottomY;
+			}
+			break;
+		}
+	}
+
+	const float LineH = LineY2 - LineY1;
+	if (LineH > 0.0f)
+	{
+		const FTimingViewDrawHelper& Helper = *static_cast<const FTimingViewDrawHelper*>(&Context.GetHelper());
+		Helper.DrawContextSwitchMarkers(GetDrawState(), LineY1, LineH, 0.25f);
+	}
 }
 
 ////////////////////////////////////////////////////////////////////////////////////////////////////
diff --git a/Engine/Source/Developer/TraceInsights/Private/Insights/ContextSwitches/ViewModels/ContextSwitchesTimingTrack.h b/Engine/Source/Developer/TraceInsights/Private/Insights/ContextSwitches/ViewModels/ContextSwitchesTimingTrack.h
--- a/Engine/Source/Developer/TraceInsights/Private/Insights/ContextSwitches/ViewModels/ContextSwitchesTimingTrack.h
+++ b/Engine/Source/Developer/TraceInsights/Private/Insights/ContextSwitches/ViewModels/ContextSwitchesTimingTrack.h
@@ -75,12 +75,15 @@ public:
 	virtual void BuildDrawState(ITimingEventsTrackDrawStateBuilder& Builder, const ITimingTrackUpdateContext& Context) override;
 	virtual void BuildFilteredDrawState(ITimingEventsTrackDrawStateBuilder& Builder, const ITimingTrackUpdateContext& Context) override;
 
+	virtual void Draw(const ITimingTrackDrawContext& Context) const override;
 	virtual void PostDraw(const ITimingTrackDrawContext& Context) const override;
 
 	virtual void InitTooltip(FTooltipDrawState& InOutTooltip, const ITimingEvent& InTooltipEvent) const override;
 
 	virtual const TSharedPtr<const ITimingEvent> GetEvent(float InPosX, float InPosY, const FTimingTrackViewport& Viewport) const override;
 
+protected:
+	void DrawLineEvents(const ITimingTrackDrawContext& Context, const float OffsetY = 1.0f) const;
 
 private:
 	uint32 TimelineIndex;
diff --git a/Engine/Source/Developer/TraceInsights/Private/Insights/InsightsStyle.cpp b/Engine/Source/Developer/TraceInsights/Private/Insights/InsightsStyle.cpp
--- a/Engine/Source/Developer/TraceInsights/Private/Insights/InsightsStyle.cpp
+++ b/Engine/Source/Developer/TraceInsights/Private/Insights/InsightsStyle.cpp
@@ -59,6 +59,7 @@ TSharedRef<FSlateStyleSet> FInsightsStyle::Create()
 	Style.Set("WhiteBrush", new FSlateColorBrush(FLinearColor::White));
 	Style.Set("SingleBorder", new FSlateBorderBrush(NAME_None, FMargin(1.0f)));
 	Style.Set("DoubleBorder", new FSlateBorderBrush(NAME_None, FMargin(2.0f)));
+	Style.Set("DashedBrush", new IMAGE_BRUSH("Old/Window/WindowBackground", FVector2D(74, 74), FLinearColor::White, ESlateBrushTileType::Both));
 
 	Style.Set("EventBorder", new FSlateBorderBrush(NAME_None, FMargin(1.0f)));
 	Style.Set("HoveredEventBorder", new FSlateBorderBrush(NAME_None, FMargin(2.0f)));
diff --git a/Engine/Source/Developer/TraceInsights/Private/Insights/ViewModels/ThreadTimingTrack.cpp b/Engine/Source/Developer/TraceInsights/Private/Insights/ViewModels/ThreadTimingTrack.cpp
--- a/Engine/Source/Developer/TraceInsights/Private/Insights/ViewModels/ThreadTimingTrack.cpp
+++ b/Engine/Source/Developer/TraceInsights/Private/Insights/ViewModels/ThreadTimingTrack.cpp
@@ -935,6 +935,11 @@ void FThreadTimingTrack::BuildFilteredDrawState(ITimingEventsTrackDrawStateBuild
 
 void FThreadTimingTrack::PostDraw(const ITimingTrackDrawContext& Context) const
 {
+	if (ChildTrack.IsValid())
+	{
+		ChildTrack->PostDraw(Context);
+	}
+
 	const TSharedPtr<const ITimingEvent> SelectedEventPtr = Context.GetSelectedEvent();
 	if (SelectedEventPtr.IsValid() &&
 		SelectedEventPtr->CheckTrack(this) &&
diff --git a/Engine/Source/Developer/TraceInsights/Private/Insights/ViewModels/TimingEventsTrack.cpp b/Engine/Source/Developer/TraceInsights/Private/Insights/ViewModels/TimingEventsTrack.cpp
--- a/Engine/Source/Developer/TraceInsights/Private/Insights/ViewModels/TimingEventsTrack.cpp
+++ b/Engine/Source/Developer/TraceInsights/Private/Insights/ViewModels/TimingEventsTrack.cpp
@@ -247,7 +247,21 @@ void FTimingEventsTrack::Draw(const ITimingTrackDrawContext& Context) const
 	}
 
 	DrawEvents(Context, 1.0f);
-	DrawHeader(Context);
+
+	if (!IsChildTrack())
+	{
+		DrawHeader(Context);
+	}
+}
+
+////////////////////////////////////////////////////////////////////////////////////////////////////
+
+void FTimingEventsTrack::PostDraw(const ITimingTrackDrawContext& Context) const
+{
+	if (ChildTrack.IsValid())
+	{
+		ChildTrack->PostDraw(Context);
+	}
 }
 
 ////////////////////////////////////////////////////////////////////////////////////////////////////
@@ -258,21 +272,20 @@ void FTimingEventsTrack::DrawEvents(const ITimingTrackDrawContext& Context, cons
 
 	if (Context.GetEventFilter().IsValid() || HasCustomFilter())
 	{
-		Helper.DrawFadedEvents(*DrawState, *this, OffsetY, 0.1f);
+		Helper.DrawFadedEvents(GetDrawState(), *this, OffsetY, 0.1f);
 
-		if (FilteredDrawStateInfo.Opacity == 1.0f)
+		if (UpdateFilteredDrawStateOpacity())
 		{
-			Helper.DrawEvents(*FilteredDrawState, *this, OffsetY);
+			Helper.DrawEvents(GetFilteredDrawState(), *this, OffsetY);
 		}
 		else
 		{
-			FilteredDrawStateInfo.Opacity = FMath::Min(1.0f, FilteredDrawStateInfo.Opacity + 0.05f);
-			Helper.DrawFadedEvents(*FilteredDrawState, *this, OffsetY, FilteredDrawStateInfo.Opacity);
+			Helper.DrawFadedEvents(GetFilteredDrawState(), *this, OffsetY, GetFilteredDrawStateOpacity());
 		}
 	}
 	else
 	{
-		Helper.DrawEvents(*DrawState, *this, OffsetY);
+		Helper.DrawEvents(GetDrawState(), *this, OffsetY);
 	}
 }
 
@@ -284,12 +297,12 @@ void FTimingEventsTrack::DrawMarkers(const ITimingTrackDrawContext& Context, flo
 
 	if (Context.GetEventFilter().IsValid())
 	{
-		Helper.DrawMarkers(*DrawState, LineY, LineH, 0.2f);
-		Helper.DrawMarkers(*FilteredDrawState, LineY, LineH, 0.75f * FilteredDrawStateInfo.Opacity);
+		Helper.DrawMarkers(GetDrawState(), LineY, LineH, 0.2f);
+		Helper.DrawMarkers(GetFilteredDrawState(), LineY, LineH, 0.75f * GetFilteredDrawStateOpacity());
 	}
 	else
 	{
-		Helper.DrawMarkers(*DrawState, LineY, LineH, 0.2f);
+		Helper.DrawMarkers(GetDrawState(), LineY, LineH, 0.2f);
 	}
 }
 
@@ -313,12 +326,6 @@ int32 FTimingEventsTrack::GetHeaderTextLayerId(const ITimingTrackDrawContext& Co
 
 void FTimingEventsTrack::DrawHeader(const ITimingTrackDrawContext& Context) const
 {
-	if (IsChildTrack())
-	{
-		// Do not draw the header for child tracks for now
-		return;
-	}
-
 	const FTimingViewDrawHelper& Helper = *static_cast<const FTimingViewDrawHelper*>(&Context.GetHelper());
 	Helper.DrawTrackHeader(*this);
 }
diff --git a/Engine/Source/Developer/TraceInsights/Private/Insights/ViewModels/TimingViewDrawHelper.cpp b/Engine/Source/Developer/TraceInsights/Private/Insights/ViewModels/TimingViewDrawHelper.cpp
--- a/Engine/Source/Developer/TraceInsights/Private/Insights/ViewModels/TimingViewDrawHelper.cpp
+++ b/Engine/Source/Developer/TraceInsights/Private/Insights/ViewModels/TimingViewDrawHelper.cpp
@@ -226,6 +226,7 @@ void FTimingEventsTrackDrawStateBuilder::AddEvent(double EventStartTime, double
 			DrawText.X = EventX1 + 2.0f;
 			DrawText.Text = Name.Left(LastWholeCharacterIndex + 1);
 			DrawText.bWhite = bIsDarkColor;
+			DrawText.Color = FLinearColor(EventColorFill.R * BorderColorFactor, EventColorFill.G * BorderColorFactor, EventColorFill.B * BorderColorFactor, EventColorFill.A);
 		}
 	}
 }
@@ -269,6 +270,7 @@ FTimingViewDrawHelper::FTimingViewDrawHelper(const FDrawContext& InDrawContext,
 	, HoveredEventBorderBrush(FInsightsStyle::Get().GetBrush("HoveredEventBorder"))
 	, SelectedEventBorderBrush(FInsightsStyle::Get().GetBrush("SelectedEventBorder"))
 	, BackgroundAreaBrush(WhiteBrush)
+	, IdleAreaBrush(FInsightsStyle::Get().GetBrush("DashedBrush"))
 	, ValidAreaColor(0.07f, 0.07f, 0.07f, 1.0f)
 	, InvalidAreaColor(0.1f, 0.07f, 0.07f, 1.0f)
 	, EdgeColor(0.05f, 0.05f, 0.05f, 1.0f)
@@ -302,9 +304,7 @@ void FTimingViewDrawHelper::DrawBackground() const
 
 void FTimingViewDrawHelper::DrawEvents(const FTimingEventsTrackDrawState& DrawState, const FTimingEventsTrack& Track, const float OffsetY) const
 {
-	const float TrackY = Track.GetPosY();
 	const float TrackH = Track.GetHeight();
-
 	if (TrackH > 0.0f)
 	{
 		const FTimingViewLayout& Layout = Viewport.GetLayout();
@@ -312,7 +312,7 @@ void FTimingViewDrawHelper::DrawEvents(const FTimingEventsTrackDrawState& DrawSt
 		NumEvents += DrawState.GetNumEvents();
 		NumMergedBoxes += DrawState.GetNumMergedBoxes();
 
-		float TopLaneY = TrackY;
+		float TopLaneY = Track.GetPosY();
 		if (!Track.IsChildTrack())
 		{
 			TopLaneY += OffsetY + Layout.TimelineDY;
@@ -379,9 +379,7 @@ void FTimingViewDrawHelper::DrawEvents(const FTimingEventsTrackDrawState& DrawSt
 
 void FTimingViewDrawHelper::DrawFadedEvents(const FTimingEventsTrackDrawState& DrawState, const FTimingEventsTrack& Track, const float OffsetY, const float Opacity) const
 {
-	const float TrackY = Track.GetPosY();
 	const float TrackH = Track.GetHeight();
-
 	if (TrackH > 0.0f)
 	{
 		const FTimingViewLayout& Layout = Viewport.GetLayout();
@@ -389,7 +387,7 @@ void FTimingViewDrawHelper::DrawFadedEvents(const FTimingEventsTrackDrawState& D
 		NumEvents += DrawState.GetNumEvents();
 		NumMergedBoxes += DrawState.GetNumMergedBoxes();
 
-		float TopLaneY = TrackY;
+		float TopLaneY = Track.GetPosY();
 		if (!Track.IsChildTrack())
 		{
 			TopLaneY += OffsetY + Layout.TimelineDY;
@@ -457,6 +455,156 @@ void FTimingViewDrawHelper::DrawFadedEvents(const FTimingEventsTrackDrawState& D
 
 ////////////////////////////////////////////////////////////////////////////////////////////////////
 
+void FTimingViewDrawHelper::DrawLineEvents(const FTimingEventsTrackDrawState& DrawState, const FTimingEventsTrack& Track, const float OffsetY) const
+{
+	const float TrackH = Track.GetHeight();
+
+	if (TrackH > 0.0f)
+	{
+		const FTimingViewLayout& Layout = Viewport.GetLayout();
+
+		NumEvents += DrawState.GetNumEvents();
+		NumMergedBoxes += DrawState.GetNumMergedBoxes();
+
+		float TopLaneY = Track.GetPosY();
+		if (!Track.IsChildTrack())
+		{
+			TopLaneY += OffsetY + Layout.TimelineDY;
+		}
+		if (Track.GetChildTrack().IsValid())
+		{
+			TopLaneY += Track.ChildTrack->GetHeight() + Layout.ChildTimelineDY;
+		}
+
+		constexpr float LineH = 3.0f;
+		constexpr float LineEndsDH = 1.0f;
+		constexpr float LineEndsH = LineH + 2 * LineEndsDH;
+
+		// Draw filled boxes (merged borders).
+		//if (Layout.EventH > 0.0f)
+		{
+			const int32 EventFillLayerId = ReservedLayerId + ToInt32(EDrawLayer::EventFill);
+			const float Y = TopLaneY + Layout.EventH - LineEndsH + LineEndsDH + 1.0f;
+			for (const FTimingEventsTrackDrawState::FBoxPrimitive& Box : DrawState.Boxes)
+			{
+				DrawContext.DrawBox(EventFillLayerId, Box.X, Y, Box.W, LineEndsH, WhiteBrush, Box.Color);
+			}
+			NumDrawBoxes += DrawState.Boxes.Num();
+		}
+
+		// Draw filled boxes (event inside area).
+		if (Layout.EventH > 2.0f)
+		{
+			const int32 EventFillLayerId = ReservedLayerId + ToInt32(EDrawLayer::EventFill);
+			const float Y = TopLaneY + Layout.EventH - LineEndsH + LineEndsDH + 2.0f;
+			for (const FTimingEventsTrackDrawState::FBoxPrimitive& Box : DrawState.InsideBoxes)
+			{
+				DrawContext.DrawBox(EventFillLayerId, Box.X, Y, Box.W, LineH, WhiteBrush, Box.Color);
+			}
+			NumDrawBoxes += DrawState.InsideBoxes.Num();
+		}
+
+		// Draw borders (just two vertical lines at ends).
+		//if (Layout.EventH > 0.0f)
+		{
+			const int32 EventBorderLayerId = ReservedLayerId + ToInt32(EDrawLayer::EventBorder);
+			const float Y = TopLaneY + Layout.EventH - LineEndsH + LineEndsDH + 1.0f;
+			for (const FTimingEventsTrackDrawState::FBoxPrimitive& Box : DrawState.Borders)
+			{
+				DrawContext.DrawBox(EventBorderLayerId, Box.X, Y, 1.0f, LineEndsH, WhiteBrush, Box.Color);
+				DrawContext.DrawBox(EventBorderLayerId, Box.X + Box.W - 1.0f, Y, 1.0f, LineEndsH, WhiteBrush, Box.Color);
+			}
+			NumDrawBorders += DrawState.Borders.Num();
+		}
+
+		// Draw texts.
+		if (Layout.EventH > 10.0f)
+		{
+			const int32 EventTextLayerId = ReservedLayerId + ToInt32(EDrawLayer::EventText);
+			const float Y = TopLaneY;
+			for (const FTimingEventsTrackDrawState::FTextPrimitive& Text : DrawState.Texts)
+			{
+				DrawContext.DrawText(EventTextLayerId, Text.X, Y, Text.Text, EventFont, Text.Color);
+			}
+			NumDrawTexts += DrawState.Texts.Num();
+		}
+	}
+}
+
+////////////////////////////////////////////////////////////////////////////////////////////////////
+
+void FTimingViewDrawHelper::DrawFadedLineEvents(const FTimingEventsTrackDrawState& DrawState, const FTimingEventsTrack& Track, const float OffsetY, const float Opacity) const
+{
+	// TODO
+	DrawLineEvents(DrawState, Track, OffsetY);
+}
+
+////////////////////////////////////////////////////////////////////////////////////////////////////
+
+void FTimingViewDrawHelper::DrawContextSwitchMarkers(const FTimingEventsTrackDrawState& DrawState, float LineY, float LineH, float Opacity) const
+{
+	if (LineH > 0.0f)
+	{
+		const int32 LayerId = ReservedLayerId + ToInt32(EDrawLayer::EventHighlight);
+
+		// Draw vertical lines (merged into large boxes).
+		for (const FTimingEventsTrackDrawState::FBoxPrimitive& Box : DrawState.Boxes)
+		{
+			const FLinearColor Color = Box.Color.CopyWithNewOpacity(Opacity);
+			DrawContext.DrawBox(LayerId, Box.X, LineY, Box.W, LineH, WhiteBrush, Color);
+		}
+
+		// Draw vertical lines (edges of larger events).
+		for (const FTimingEventsTrackDrawState::FBoxPrimitive& Box : DrawState.Borders)
+		{
+			const FLinearColor Color = Box.Color.CopyWithNewOpacity(Opacity);
+			DrawContext.DrawBox(LayerId, Box.X, LineY, 1.0f, LineH, WhiteBrush, Color);
+			DrawContext.DrawBox(LayerId, Box.X + Box.W - 1.0f, LineY, 1.0f, LineH, WhiteBrush, Color);
+		}
+
+		// Draw overlay for idle areas.
+		const FSlateBrush* IdleBrush = IdleAreaBrush;
+		const FLinearColor IdleColor(1.0f, 1.0f, 1.0f, 0.5f);
+		const int32 Count1 = DrawState.Boxes.Num();
+		const int32 Count2 = DrawState.Borders.Num();
+		int32 Index1 = 0;
+		int32 Index2 = 0;
+		const float MinX = Viewport.TimeToSlateUnitsRounded(Viewport.GetMinValidTime());
+		float CurrentX = MinX;
+		while (Index1 < Count1 || Index2 < Count2)
+		{
+			const float X1 = (Index1 < Count1) ? DrawState.Boxes[Index1].X : FLT_MAX;
+			const float X2 = (Index2 < Count2) ? DrawState.Borders[Index2].X : FLT_MAX;
+			if (X1 < X2)
+			{
+				if (X1 - CurrentX > 0.0f)
+				{
+					DrawContext.DrawBox(LayerId, CurrentX, LineY, X1 - CurrentX, LineH, IdleBrush, IdleColor);
+				}
+				CurrentX = FMath::Max(CurrentX, X1 + DrawState.Boxes[Index1].W);
+				++Index1;
+			}
+			else
+			{
+				if (X2 - CurrentX > 0.0f)
+				{
+					DrawContext.DrawBox(LayerId, CurrentX, LineY, X2 - CurrentX, LineH, IdleBrush, IdleColor);
+				}
+				CurrentX = FMath::Max(CurrentX, X2 + DrawState.Borders[Index2].W);
+				++Index2;
+			}
+		}
+		const float MaxX = Viewport.TimeToSlateUnitsRounded(Viewport.GetMaxValidTime());
+		const float LastAreaW = FMath::Min(MaxX, Viewport.GetWidth()) - CurrentX;
+		if (LastAreaW > 0.0f)
+		{
+			DrawContext.DrawBox(LayerId, CurrentX, LineY, LastAreaW, LineH, IdleBrush, IdleColor);
+		}
+	}
+}
+
+////////////////////////////////////////////////////////////////////////////////////////////////////
+
 void FTimingViewDrawHelper::DrawMarkers(const FTimingEventsTrackDrawState& DrawState, float LineY, float LineH, float Opacity) const
 {
 	if (LineH > 0.0f)
diff --git a/Engine/Source/Developer/TraceInsights/Private/Insights/ViewModels/TimingViewDrawHelper.h b/Engine/Source/Developer/TraceInsights/Private/Insights/ViewModels/TimingViewDrawHelper.h
--- a/Engine/Source/Developer/TraceInsights/Private/Insights/ViewModels/TimingViewDrawHelper.h
+++ b/Engine/Source/Developer/TraceInsights/Private/Insights/ViewModels/TimingViewDrawHelper.h
@@ -42,6 +42,7 @@ struct FTimingEventsTrackDrawState
 		float X;
 		FString Text;
 		bool bWhite;
+		FLinearColor Color;
 	};
 
 	FTimingEventsTrackDrawState()
@@ -186,9 +187,12 @@ public:
 
 	// OffsetY = 1.0f is for the top horizontal line (which separates the timelines) added by DrawTrackHeader.
 	void DrawEvents(const FTimingEventsTrackDrawState& DrawState, const FTimingEventsTrack& Track, const float OffsetY = 1.0f) const;
-
 	void DrawFadedEvents(const FTimingEventsTrackDrawState& DrawState, const FTimingEventsTrack& Track, const float OffsetY = 1.0f, const float Opacity = 0.1f) const;
 
+	void DrawLineEvents(const FTimingEventsTrackDrawState& DrawState, const FTimingEventsTrack& Track, const float OffsetY = 1.0f) const;
+	void DrawFadedLineEvents(const FTimingEventsTrackDrawState& DrawState, const FTimingEventsTrack& Track, const float OffsetY = 1.0f, const float Opacity = 0.1f) const;
+	void DrawContextSwitchMarkers(const FTimingEventsTrackDrawState& DrawState, float LineY, float LineH, float Opacity) const;
+
 	void DrawMarkers(const FTimingEventsTrackDrawState& DrawState, float LineY, float LineH, float Opacity) const;
 
 	void DrawTrackHeader(const FBaseTimingTrack& Track) const;
@@ -217,6 +221,7 @@ private:
 	const FSlateBrush* HoveredEventBorderBrush;
 	const FSlateBrush* SelectedEventBorderBrush;
 	const FSlateBrush* BackgroundAreaBrush;
+	const FSlateBrush* IdleAreaBrush;
 	const FLinearColor ValidAreaColor;
 	const FLinearColor InvalidAreaColor;
 	const FLinearColor EdgeColor;
diff --git a/Engine/Source/Developer/TraceInsights/Public/Insights/ViewModels/TimingEventsTrack.h b/Engine/Source/Developer/TraceInsights/Public/Insights/ViewModels/TimingEventsTrack.h
--- a/Engine/Source/Developer/TraceInsights/Public/Insights/ViewModels/TimingEventsTrack.h
+++ b/Engine/Source/Developer/TraceInsights/Public/Insights/ViewModels/TimingEventsTrack.h
@@ -41,6 +41,7 @@ public:
 	virtual void PostUpdate(const ITimingTrackUpdateContext& Context) override;
 
 	virtual void Draw(const ITimingTrackDrawContext& Context) const override;
+	virtual void PostDraw(const ITimingTrackDrawContext& Context) const override;
 	virtual void DrawEvent(const ITimingTrackDrawContext& Context, const ITimingEvent& InTimingEvent, EDrawEventMode InDrawMode) const override;
 
 	virtual const TSharedPtr<const ITimingEvent> GetEvent(float InPosX, float InPosY, const FTimingTrackViewport& Viewport) const override;
@@ -56,6 +57,23 @@ protected:
 	int32 GetNumLanes() const { return NumLanes; }
 	void SetNumLanes(int32 InNumLanes) { NumLanes = InNumLanes; }
 
+	const struct FTimingEventsTrackDrawState& GetDrawState() const { return *DrawState; }
+	const struct FTimingEventsTrackDrawState& GetFilteredDrawState() const { return *FilteredDrawState; }
+
+	float GetFilteredDrawStateOpacity() const { return FilteredDrawStateInfo.Opacity; }
+	bool UpdateFilteredDrawStateOpacity() const
+	{
+		if (FilteredDrawStateInfo.Opacity == 1.0f)
+		{
+			return true;
+		}
+		else
+		{
+			FilteredDrawStateInfo.Opacity = FMath::Min(1.0f, FilteredDrawStateInfo.Opacity + 0.05f);
+			return false;
+		}
+	}
+
 	void UpdateTrackHeight(const ITimingTrackUpdateContext& Context);
 
 	void DrawEvents(const ITimingTrackDrawContext& Context, const float OffsetY = 1.0f) const;
