[Insights] Fixed issue with display of context switches highlight area.

diff --git a/Engine/Source/Developer/TraceInsights/Private/Insights/ViewModels/TimingViewDrawHelper.cpp b/Engine/Source/Developer/TraceInsights/Private/Insights/ViewModels/TimingViewDrawHelper.cpp
--- a/Engine/Source/Developer/TraceInsights/Private/Insights/ViewModels/TimingViewDrawHelper.cpp
+++ b/Engine/Source/Developer/TraceInsights/Private/Insights/ViewModels/TimingViewDrawHelper.cpp
@@ -570,7 +570,7 @@ void FTimingViewDrawHelper::DrawContextSwitchMarkers(const FTimingEventsTrackDra
 		int32 Index1 = 0;
 		int32 Index2 = 0;
 		const float MinX = Viewport.TimeToSlateUnitsRounded(Viewport.GetMinValidTime());
-		float CurrentX = MinX;
+		float CurrentX = FMath::Max(MinX, 0.0f);
 		while (Index1 < Count1 || Index2 < Count2)
 		{
 			const float X1 = (Index1 < Count1) ? DrawState.Boxes[Index1].X : FLT_MAX;
