[Insights] Menu option to show/hide context switches header tracks

diff --git a/Engine/Source/Developer/TraceInsights/Private/Insights/ContextSwitches/ContextSwitchesProfilerManager.cpp b/Engine/Source/Developer/TraceInsights/Private/Insights/ContextSwitches/ContextSwitchesProfilerManager.cpp
--- a/Engine/Source/Developer/TraceInsights/Private/Insights/ContextSwitches/ContextSwitchesProfilerManager.cpp
+++ b/Engine/Source/Developer/TraceInsights/Private/Insights/ContextSwitches/ContextSwitchesProfilerManager.cpp
@@ -157,6 +157,7 @@ bool FContextSwitchesProfilerManager::Tick(float DeltaTime)
 				if (!ContextSwitchesSharedState.IsValid())
 				{
 					ContextSwitchesSharedState = MakeShared<FContextSwitchesSharedState>(TimingView.Get());
+					ContextSwitchesSharedState->AddCommands();
 					IModularFeatures::Get().RegisterModularFeature(Insights::TimingViewExtenderFeatureName, ContextSwitchesSharedState.Get());
 				}
 			}
diff --git a/Engine/Source/Developer/TraceInsights/Private/Insights/ContextSwitches/ViewModels/ContextSwitchesTimingTrack.cpp b/Engine/Source/Developer/TraceInsights/Private/Insights/ContextSwitches/ViewModels/ContextSwitchesTimingTrack.cpp
--- a/Engine/Source/Developer/TraceInsights/Private/Insights/ContextSwitches/ViewModels/ContextSwitchesTimingTrack.cpp
+++ b/Engine/Source/Developer/TraceInsights/Private/Insights/ContextSwitches/ViewModels/ContextSwitchesTimingTrack.cpp
@@ -46,8 +46,11 @@ public:
 	PRAGMA_DISABLE_OPTIMIZATION
 	virtual void RegisterCommands() override
 	{
+		UI_COMMAND(Command_ShowContextSwitches, "Show Context Switches ", "Show/hide context switches events as header tracks", EUserInterfaceActionType::ToggleButton, FInputChord(EModifierKey::Shift, EKeys::C));
 	}
 	PRAGMA_ENABLE_OPTIMIZATION
+
+	TSharedPtr<FUICommandInfo> Command_ShowContextSwitches;
 };
 
 ////////////////////////////////////////////////////////////////////////////////////////////////////
@@ -77,6 +80,8 @@ void FContextSwitchesSharedState::OnEndSession(Insights::ITimingViewSession& InS
 	{
 		return;
 	}
+
+	AddContextSwitchesChildTracks();
 }
 
 ////////////////////////////////////////////////////////////////////////////////////////////////////
@@ -88,6 +93,58 @@ void FContextSwitchesSharedState::Tick(Insights::ITimingViewSession& InSession,
 		return;
 	}
 
+	if (!FInsightsManager::Get()->IsAnalysisComplete() && bShowContextSwitchesTrack)
+	{
+		AddContextSwitchesChildTracks();
+	}
+}
+
+////////////////////////////////////////////////////////////////////////////////////////////////////
+
+void FContextSwitchesSharedState::ExtendFilterMenu(Insights::ITimingViewSession& InSession, FMenuBuilder& InOutMenuBuilder)
+{
+	if (&InSession != TimingView)
+	{
+		return;
+	}
+}
+
+////////////////////////////////////////////////////////////////////////////////////////////////////
+
+bool FContextSwitchesSharedState::ExtendGlobalContextMenu(ITimingViewSession& InSession, FMenuBuilder& InMenuBuilder)
+{
+	InMenuBuilder.AddMenuEntry
+	(
+		FContextSwitchesStateCommands::Get().Command_ShowContextSwitches,
+		NAME_None,
+		TAttribute<FText>(),
+		TAttribute<FText>(),
+		FSlateIcon()
+	);
+
+	return true;
+}
+
+////////////////////////////////////////////////////////////////////////////////////////////////////
+
+void FContextSwitchesSharedState::AddCommands()
+{
+	FContextSwitchesStateCommands::Register();
+
+	TSharedPtr<FUICommandList> CommandList = TimingView->GetCommandList();
+	ensure(CommandList.IsValid());
+
+	CommandList->MapAction(
+		FContextSwitchesStateCommands::Get().Command_ShowContextSwitches,
+		FExecuteAction::CreateSP(this, &FContextSwitchesSharedState::ContextMenu_ShowContextSwitches_Execute),
+		FCanExecuteAction::CreateSP(this, &FContextSwitchesSharedState::ContextMenu_ShowContextSwitches_CanExecute),
+		FIsActionChecked::CreateSP(this, &FContextSwitchesSharedState::ContextMenu_ShowContextSwitches_IsChecked));
+}
+
+////////////////////////////////////////////////////////////////////////////////////////////////////
+
+void FContextSwitchesSharedState::AddContextSwitchesChildTracks()
+{
 	TSharedPtr<FThreadTimingSharedState> TimingSharedState = TimingView->GetThreadTimingSharedState();
 
 	if (!TimingSharedState.IsValid())
@@ -110,19 +167,54 @@ void FContextSwitchesSharedState::Tick(Insights::ITimingViewSession& InSession,
 
 ////////////////////////////////////////////////////////////////////////////////////////////////////
 
-void FContextSwitchesSharedState::ExtendFilterMenu(Insights::ITimingViewSession& InSession, FMenuBuilder& InOutMenuBuilder)
+void FContextSwitchesSharedState::RemoveContextSwitchesChildTracks()
 {
-	if (&InSession != TimingView)
+	TSharedPtr<FThreadTimingSharedState> TimingSharedState = TimingView->GetThreadTimingSharedState();
+
+	if (!TimingSharedState.IsValid())
 	{
 		return;
 	}
+
+	const TMap<uint32, TSharedPtr<FCpuTimingTrack>>& CpuTracks = TimingSharedState->GetAllCpuTracks();
+
+	for (const TPair<uint32, TSharedPtr<FCpuTimingTrack>>& MapEntry : CpuTracks)
+	{
+		if (MapEntry.Value.IsValid() && MapEntry.Value->GetChildTrack().IsValid() && MapEntry.Value->GetChildTrack()->Is<FContextSwitchesTimingTrack>())
+		{
+			MapEntry.Value->SetChildTrack(nullptr);
+		}
+	}
 }
 
 ////////////////////////////////////////////////////////////////////////////////////////////////////
 
-void FContextSwitchesSharedState::InitCommandList()
+void FContextSwitchesSharedState::ContextMenu_ShowContextSwitches_Execute()
 {
-	FContextSwitchesStateCommands::Register();
+	SetContextSwitchesToggle(!bShowContextSwitchesTrack);
+
+	if (bShowContextSwitchesTrack)
+	{
+		AddContextSwitchesChildTracks();
+	}
+	else
+	{
+		RemoveContextSwitchesChildTracks();
+	}
+}
+
+////////////////////////////////////////////////////////////////////////////////////////////////////
+
+bool FContextSwitchesSharedState::ContextMenu_ShowContextSwitches_CanExecute()
+{
+	return true;
+}
+
+////////////////////////////////////////////////////////////////////////////////////////////////////
+
+bool FContextSwitchesSharedState::ContextMenu_ShowContextSwitches_IsChecked()
+{
+	return IsContextSwitchesToggleOn();
 }
 
 ////////////////////////////////////////////////////////////////////////////////////////////////////
diff --git a/Engine/Source/Developer/TraceInsights/Private/Insights/ContextSwitches/ViewModels/ContextSwitchesTimingTrack.h b/Engine/Source/Developer/TraceInsights/Private/Insights/ContextSwitches/ViewModels/ContextSwitchesTimingTrack.h
--- a/Engine/Source/Developer/TraceInsights/Private/Insights/ContextSwitches/ViewModels/ContextSwitchesTimingTrack.h
+++ b/Engine/Source/Developer/TraceInsights/Private/Insights/ContextSwitches/ViewModels/ContextSwitchesTimingTrack.h
@@ -33,24 +33,31 @@ public:
 	virtual void OnEndSession(Insights::ITimingViewSession& InSession) override;
 	virtual void Tick(Insights::ITimingViewSession& InSession, const TraceServices::IAnalysisSession& InAnalysisSession) override;
 	virtual void ExtendFilterMenu(Insights::ITimingViewSession& InSession, FMenuBuilder& InMenuBuilder) override;
-	virtual bool ExtendGlobalContextMenu(ITimingViewSession& InSession, FMenuBuilder& InMenuBuilder) override { return false; }
+	virtual bool ExtendGlobalContextMenu(ITimingViewSession& InSession, FMenuBuilder& InMenuBuilder) override;
 
 	//////////////////////////////////////////////////
 
-	bool IsContextSwitchesToggleOn() const { return bShowHideContextSwitchesTrack; }
-	void SetContextSwitchesToggle(bool bOnOff) { bShowHideContextSwitchesTrack = bOnOff; }
+	bool IsContextSwitchesToggleOn() const { return bShowContextSwitchesTrack; }
+	void SetContextSwitchesToggle(bool bOnOff) { bShowContextSwitchesTrack = bOnOff; }
 	void ShowContextSwitches() { SetContextSwitchesToggle(true); }
 	void HideContextSwitches() { SetContextSwitchesToggle(false); }
 
 	STimingView* GetTimingView() { return TimingView; }
+	void AddCommands();
 
 private:
-	void InitCommandList();
+	void ContextMenu_ShowContextSwitches_Execute();
+	bool ContextMenu_ShowContextSwitches_CanExecute();
+	bool ContextMenu_ShowContextSwitches_IsChecked();
+
+	void AddContextSwitchesChildTracks();
+	void RemoveContextSwitchesChildTracks();
+
 
 private:
 	STimingView* TimingView;
 
-	bool bShowHideContextSwitchesTrack;
+	bool bShowContextSwitchesTrack = true;
 };
 
 ////////////////////////////////////////////////////////////////////////////////////////////////////
diff --git a/Engine/Source/Developer/TraceInsights/Private/Insights/TaskGraphProfiler/ViewModels/TaskTimingTrack.cpp b/Engine/Source/Developer/TraceInsights/Private/Insights/TaskGraphProfiler/ViewModels/TaskTimingTrack.cpp
--- a/Engine/Source/Developer/TraceInsights/Private/Insights/TaskGraphProfiler/ViewModels/TaskTimingTrack.cpp
+++ b/Engine/Source/Developer/TraceInsights/Private/Insights/TaskGraphProfiler/ViewModels/TaskTimingTrack.cpp
@@ -150,26 +150,13 @@ void FTaskTimingSharedState::SetTaskId(uint32 InTaskId)
 
 bool FTaskTimingSharedState::ExtendGlobalContextMenu(ITimingViewSession& InSession, FMenuBuilder& InMenuBuilder)
 {
-	if (!CommandList.IsValid())
-	{
-		return false;
-	}
-
 	if (&InSession != TimingView)
 	{
 		return false;
 	}
 
-	InMenuBuilder.PushCommandList(CommandList.ToSharedRef());
-
-	InMenuBuilder.BeginSection(TEXT("Event"), LOCTEXT("Event", "Event"));
-
 	InMenuBuilder.AddSubMenu(LOCTEXT("Tasks", "Tasks"), LOCTEXT("Task", "Task Graph Insights settings"), FNewMenuDelegate::CreateSP(this, &FTaskTimingSharedState::BuildTasksSubMenu));
 
-	InMenuBuilder.EndSection();
-
-	InMenuBuilder.PopCommandList();
-
 	return true;
 }
 
@@ -220,15 +207,11 @@ void FTaskTimingSharedState::BuildTasksSubMenu(FMenuBuilder& MenuBuilder)
 
 void FTaskTimingSharedState::InitCommandList()
 {
-	if (CommandList.IsValid())
-	{
-		return;
-	}
+	TSharedPtr<FUICommandList> CommandList = TimingView->GetCommandList();
+	ensure(CommandList);
 
 	FTaskTimingStateCommands::Register();
 
-	CommandList = MakeShared<FUICommandList>();
-
 	CommandList->MapAction(
 		FTaskTimingStateCommands::Get().Command_ShowTaskDependencies,
 		FExecuteAction::CreateSP(this, &FTaskTimingSharedState::ContextMenu_ShowTaskDependecies_Execute),
diff --git a/Engine/Source/Developer/TraceInsights/Private/Insights/TaskGraphProfiler/ViewModels/TaskTimingTrack.h b/Engine/Source/Developer/TraceInsights/Private/Insights/TaskGraphProfiler/ViewModels/TaskTimingTrack.h
--- a/Engine/Source/Developer/TraceInsights/Private/Insights/TaskGraphProfiler/ViewModels/TaskTimingTrack.h
+++ b/Engine/Source/Developer/TraceInsights/Private/Insights/TaskGraphProfiler/ViewModels/TaskTimingTrack.h
@@ -82,7 +82,6 @@ private:
 	bool bResetOnNextTick = false;
 
 	TSharedPtr<FTaskTimingTrack> TaskTrack;
-	TSharedPtr<FUICommandList> CommandList;
 };
 
 ////////////////////////////////////////////////////////////////////////////////////////////////////
diff --git a/Engine/Source/Developer/TraceInsights/Private/Insights/Widgets/STimingView.h b/Engine/Source/Developer/TraceInsights/Private/Insights/Widgets/STimingView.h
--- a/Engine/Source/Developer/TraceInsights/Private/Insights/Widgets/STimingView.h
+++ b/Engine/Source/Developer/TraceInsights/Private/Insights/Widgets/STimingView.h
@@ -349,6 +349,8 @@ public:
 	void AddRelation(TUniquePtr<ITimingEventRelation>& Relation) { CurrentRelations.Add(MoveTemp(Relation)); }
 	void ClearRelations();
 
+	TSharedPtr<FUICommandList> GetCommandList() { return CommandList; }
+
 protected:
 	virtual FVector2D ComputeDesiredSize(float) const override
 	{
