Fixed misspelled variable name.

diff --git a/Engine/Source/Developer/TraceInsights/Private/Insights/ViewModels/TaskGraphRelation.cpp b/Engine/Source/Developer/TraceInsights/Private/Insights/ViewModels/TaskGraphRelation.cpp
--- a/Engine/Source/Developer/TraceInsights/Private/Insights/ViewModels/TaskGraphRelation.cpp
+++ b/Engine/Source/Developer/TraceInsights/Private/Insights/ViewModels/TaskGraphRelation.cpp
@@ -70,7 +70,7 @@ void FTaskGraphRelation::Draw(const FDrawContext& DrawContext, const FTimingTrac
 	float Distance = FVector2D::Distance(StartPoint, EndPoint);
 
 	constexpr float DirectionFactor = 2.0f;
-	constexpr float LineLenghtAtEnds = 20.0f;
+	constexpr float LineLengthAtEnds = 20.0f;
 	FVector2D StartDir((X2 - X1) / DirectionFactor, 0.0f);
 
 	FLinearColor Color = FTaskGraphProfilerManager::Get()->GetColorForTaskEvent(Type);
@@ -80,10 +80,10 @@ void FTaskGraphRelation::Draw(const FDrawContext& DrawContext, const FTimingTrac
 	constexpr float ArrowRotationAngle = 20.0f;
 	FVector2D ArrowDirection(-ArrowDirectionLen, 0.0f);
 
-	if (Distance > LineLenghtAtEnds && !FMath::IsNearlyEqual(StartPoint.Y, EndPoint.Y))
+	if (Distance > LineLengthAtEnds && !FMath::IsNearlyEqual(StartPoint.Y, EndPoint.Y))
 	{
-		FVector2D SplineStart(StartPoint.X + LineLenghtAtEnds, StartPoint.Y);
-		FVector2D SplineEnd(EndPoint.X - LineLenghtAtEnds, EndPoint.Y);
+		FVector2D SplineStart(StartPoint.X + LineLengthAtEnds, StartPoint.Y);
+		FVector2D SplineEnd(EndPoint.X - LineLengthAtEnds, EndPoint.Y);
 
 		DrawContext.DrawSpline(LayerId, 0.0f, 0.0f, SplineStart, StartDir, SplineEnd, StartDir, /*Thickness=*/ 2.0f, Color);
 
