[Insights]
  - Enabled "HighDPI" support for Unreal Insights app.
  - Fixed size and location of app windows (for both browser and viewer modes) when high dpi is enabled and DPIScaleFactor is not 1.
  - Added name of trace in the application's title.

diff --git a/Engine/Source/Developer/TraceInsights/Private/Insights/TraceInsightsModule.cpp b/Engine/Source/Developer/TraceInsights/Private/Insights/TraceInsightsModule.cpp
--- a/Engine/Source/Developer/TraceInsights/Private/Insights/TraceInsightsModule.cpp
+++ b/Engine/Source/Developer/TraceInsights/Private/Insights/TraceInsightsModule.cpp
@@ -168,15 +168,42 @@ void FTraceInsightsModule::CreateSessionBrowser(const FCreateSessionBrowserParam
 {
 	RegisterTabSpawners();
 
-	const float DPIScaleFactor = FPlatformApplicationMisc::GetDPIScaleFactorAtPoint(10.0f, 10.0f);
+	//////////////////////////////////////////////////
+	// Create the main window.
 
-	TSharedRef<FTabManager::FLayout> DefaultLayout = FTabManager::NewLayout("TraceSessionBrowserLayout_v1.0");
-	float WindowWidth = 0.0f;
-	float WindowHeight = 0.0f;
+	const bool bEmbedTitleAreaContent = false;
+
+	// Get desktop metrics. It also ensures the correct metrics will be used later in SWindow.
+	FDisplayMetrics DisplayMetrics;
+	FSlateApplication::Get().GetDisplayMetrics(DisplayMetrics);
+	const float DPIScaleFactor = FPlatformApplicationMisc::GetDPIScaleFactorAtPoint(DisplayMetrics.PrimaryDisplayWorkAreaRect.Left, DisplayMetrics.PrimaryDisplayWorkAreaRect.Top);
+
+	const FVector2D ClientSize(960.0f * DPIScaleFactor, 640.0f * DPIScaleFactor);
+
+	TSharedRef<SWindow> RootWindow = SNew(SWindow)
+		.Title(NSLOCTEXT("TraceInsightsModule", "UnrealInsightsBrowserAppName", "Unreal Insights Session Browser"))
+		.CreateTitleBar(!bEmbedTitleAreaContent)
+		.SupportsMaximize(true)
+		.SupportsMinimize(true)
+		.IsInitiallyMaximized(false)
+		.IsInitiallyMinimized(false)
+		.SizingRule(ESizingRule::UserSized)
+		.AutoCenter(EAutoCenter::PreferredWorkArea)
+		.ClientSize(ClientSize)
+		.AdjustInitialSizeAndPositionForDPIScale(false);
+
+	const bool bShowRootWindowImmediately = false;
+	FSlateApplication::Get().AddWindow(RootWindow, bShowRootWindowImmediately);
+
+	FGlobalTabmanager::Get()->SetRootWindow(RootWindow);
+	FGlobalTabmanager::Get()->SetAllowWindowMenuBar(true);
+
+	FSlateNotificationManager::Get().SetRootWindow(RootWindow);
 
-	WindowWidth = 920.0f;
-	WindowHeight = 665.0f;
+	//////////////////////////////////////////////////
+	// Setup the window's content.
 
+	TSharedRef<FTabManager::FLayout> DefaultLayout = FTabManager::NewLayout("TraceSessionBrowserLayout_v1.0");
 	DefaultLayout->AddArea
 	(
 		FTabManager::NewPrimaryArea()
@@ -186,19 +213,6 @@ void FTraceInsightsModule::CreateSessionBrowser(const FCreateSessionBrowserParam
 			->AddTab(FInsightsManagerTabs::StartPageTabId, ETabState::OpenedTab)
 		)
 	);
-	
-	TSharedRef<SWindow> RootWindow = SNew(SWindow)
-		.AutoCenter(EAutoCenter::PreferredWorkArea)
-		.Title(NSLOCTEXT("TraceInsightsModule", "UnrealInsightsBrowserAppName", "Unreal Insights Session Browser"))
-		.IsInitiallyMaximized(false)
-		.ClientSize(FVector2D(WindowWidth * DPIScaleFactor, WindowHeight * DPIScaleFactor))
-		.SupportsMaximize(true)
-		.SupportsMinimize(true);
-
-	FSlateApplication::Get().AddWindow(RootWindow, true);
-
-	FGlobalTabmanager::Get()->SetRootWindow(RootWindow);
-	FSlateNotificationManager::Get().SetRootWindow(RootWindow);
 
 	AddAreaForWidgetReflector(DefaultLayout, Params.bAllowDebugTools);
 
@@ -206,12 +220,20 @@ void FTraceInsightsModule::CreateSessionBrowser(const FCreateSessionBrowserParam
 	PersistentLayout = FLayoutSaveRestore::LoadFromConfig(UnrealInsightsLayoutIni, DefaultLayout);
 
 	// Restore application layout.
-	const bool bEmbedTitleAreaContent = false;
 	const EOutputCanBeNullptr OutputCanBeNullptr = EOutputCanBeNullptr::Never;
 	TSharedPtr<SWidget> Content = FGlobalTabmanager::Get()->RestoreFrom(PersistentLayout.ToSharedRef(), RootWindow, bEmbedTitleAreaContent, OutputCanBeNullptr);
 	RootWindow->SetContent(Content.ToSharedRef());
 
+	//////////////////////////////////////////////////
+	// Show the window.
+
+	RootWindow->ShowWindow();
+	const bool bForceWindowToFront = true;
+	RootWindow->BringToFront(bForceWindowToFront);
+
+	//////////////////////////////////////////////////
 	// Set up command line parameter forwarding.
+
 	TSharedPtr<class SStartPageWindow> StartPageWnd = FInsightsManager::Get()->GetStartPageWindow();
 	if (StartPageWnd.IsValid())
 	{
@@ -228,21 +250,46 @@ void FTraceInsightsModule::CreateSessionViewer(bool bAllowDebugTools)
 	RegisterTabSpawners();
 
 #if !WITH_EDITOR
-	const float DPIScaleFactor = FPlatformApplicationMisc::GetDPIScaleFactorAtPoint(10.0f, 10.0f);
+
+	//////////////////////////////////////////////////
+	// Create the main window.
+
+#if PLATFORM_MAC
+	const bool bEmbedTitleAreaContent = true;
+#else
+	const bool bEmbedTitleAreaContent = false;
+#endif
+
+	// Get desktop metrics. It also ensures the correct metrics will be used later in SWindow.
+	FDisplayMetrics DisplayMetrics;
+	FSlateApplication::Get().GetDisplayMetrics(DisplayMetrics);
+	const float DPIScaleFactor = FPlatformApplicationMisc::GetDPIScaleFactorAtPoint(DisplayMetrics.PrimaryDisplayWorkAreaRect.Left, DisplayMetrics.PrimaryDisplayWorkAreaRect.Top);
+	
+	const FVector2D ClientSize(1280.f * DPIScaleFactor, 720.0f * DPIScaleFactor);
 
 	TSharedRef<SWindow> RootWindow = SNew(SWindow)
-		.AutoCenter(EAutoCenter::PreferredWorkArea)
 		.Title(NSLOCTEXT("TraceInsightsModule", "UnrealInsightsAppName", "Unreal Insights"))
-		.IsInitiallyMaximized(false)
-		.ClientSize(FVector2D(1280.f * DPIScaleFactor, 720.0f * DPIScaleFactor))
+		.CreateTitleBar(!bEmbedTitleAreaContent)
 		.SupportsMaximize(true)
-		.SupportsMinimize(true);
+		.SupportsMinimize(true)
+		.IsInitiallyMaximized(false)
+		.IsInitiallyMinimized(false)
+		.SizingRule(ESizingRule::UserSized)
+		.AutoCenter(EAutoCenter::PreferredWorkArea)
+		.ClientSize(ClientSize)
+		.AdjustInitialSizeAndPositionForDPIScale(false);
 
-	FSlateApplication::Get().AddWindow(RootWindow, true);
+	const bool bShowRootWindowImmediately = false;
+	FSlateApplication::Get().AddWindow(RootWindow, bShowRootWindowImmediately);
 
 	FGlobalTabmanager::Get()->SetRootWindow(RootWindow);
+	FGlobalTabmanager::Get()->SetAllowWindowMenuBar(true);
+
 	FSlateNotificationManager::Get().SetRootWindow(RootWindow);
 
+	//////////////////////////////////////////////////
+	// Setup the window's content.
+
 	TSharedRef<FTabManager::FLayout> DefaultLayout = FTabManager::NewLayout("UnrealInsightsLayout_v1.0");
 
 	AddAreaForSessionViewer(DefaultLayout);
@@ -253,16 +300,19 @@ void FTraceInsightsModule::CreateSessionViewer(bool bAllowDebugTools)
 	PersistentLayout = FLayoutSaveRestore::LoadFromConfig(UnrealInsightsLayoutIni, DefaultLayout);
 
 	// Restore application layout.
-#if PLATFORM_MAC
-	const bool bEmbedTitleAreaContent = true;
-#else
-	const bool bEmbedTitleAreaContent = false;
-#endif
 	const EOutputCanBeNullptr OutputCanBeNullptr = EOutputCanBeNullptr::Never;
 	TSharedPtr<SWidget> Content = FGlobalTabmanager::Get()->RestoreFrom(PersistentLayout.ToSharedRef(), RootWindow, bEmbedTitleAreaContent, OutputCanBeNullptr);
 
 	RootWindow->SetContent(Content.ToSharedRef());
 	RootWindow->GetOnWindowClosedEvent().AddRaw(this, &FTraceInsightsModule::OnWindowClosedEvent);
+
+	//////////////////////////////////////////////////
+	// Show the window.
+
+	RootWindow->ShowWindow();
+	const bool bForceWindowToFront = true;
+	RootWindow->BringToFront(bForceWindowToFront);
+
 #endif
 }
 
@@ -407,6 +457,7 @@ void FTraceInsightsModule::StartAnalysisForTrace(uint32 InTraceId, bool InAutoQu
 	{
 		FInsightsManager::Get()->LoadTrace(InTraceId, InAutoQuit);
 	}
+	UpdateAppTitle();
 }
 
 ////////////////////////////////////////////////////////////////////////////////////////////////////
@@ -414,6 +465,7 @@ void FTraceInsightsModule::StartAnalysisForTrace(uint32 InTraceId, bool InAutoQu
 void FTraceInsightsModule::StartAnalysisForLastLiveSession()
 {
 	FInsightsManager::Get()->LoadLastLiveSession();
+	UpdateAppTitle();
 }
 
 ////////////////////////////////////////////////////////////////////////////////////////////////////
@@ -424,6 +476,7 @@ void FTraceInsightsModule::StartAnalysisForTraceFile(const TCHAR* InTraceFile, b
 	{
 		FInsightsManager::Get()->LoadTraceFile(FString(InTraceFile), InAutoQuit);
 	}
+	UpdateAppTitle();
 }
 
 ////////////////////////////////////////////////////////////////////////////////////////////////////
@@ -462,3 +515,27 @@ void FTraceInsightsModule::InitializeTesting(bool InInitAutomationModules, bool
 }
 
 ////////////////////////////////////////////////////////////////////////////////////////////////////
+
+void FTraceInsightsModule::UpdateAppTitle()
+{
+#if !WITH_EDITOR
+	TSharedPtr<SWindow> RootWindow = FGlobalTabmanager::Get()->GetRootWindow();
+	if (RootWindow)
+	{
+		const FString TraceFilename = FInsightsManager::Get()->GetTraceFilename();
+		if (TraceFilename.IsEmpty())
+		{
+			const FText AppTitle = NSLOCTEXT("TraceInsightsModule", "UnrealInsightsAppName", "Unreal Insights");
+			RootWindow->SetTitle(AppTitle);
+		}
+		else
+		{
+			const FString SessionName = FPaths::GetBaseFilename(TraceFilename);
+			const FText AppTitle = FText::Format(NSLOCTEXT("TraceInsightsModule", "UnrealInsightsAppNameFmt", "{0} - Unreal Insights"), FText::FromString(SessionName));
+			RootWindow->SetTitle(AppTitle);
+		}
+	}
+#endif
+}
+
+////////////////////////////////////////////////////////////////////////////////////////////////////
diff --git a/Engine/Source/Developer/TraceInsights/Private/Insights/TraceInsightsModule.h b/Engine/Source/Developer/TraceInsights/Private/Insights/TraceInsightsModule.h
--- a/Engine/Source/Developer/TraceInsights/Private/Insights/TraceInsightsModule.h
+++ b/Engine/Source/Developer/TraceInsights/Private/Insights/TraceInsightsModule.h
@@ -98,6 +98,8 @@ protected:
 
 	void OnWindowClosedEvent(const TSharedRef<SWindow>&);
 
+	void UpdateAppTitle();
+
 protected:
 #if WITH_TRACE_STORE
 	TUniquePtr<UE::Trace::FStoreService> StoreService;
diff --git a/Engine/Source/Programs/UnrealInsights/Private/UserInterfaceCommand.cpp b/Engine/Source/Programs/UnrealInsights/Private/UserInterfaceCommand.cpp
--- a/Engine/Source/Programs/UnrealInsights/Private/UserInterfaceCommand.cpp
+++ b/Engine/Source/Programs/UnrealInsights/Private/UserInterfaceCommand.cpp
@@ -243,14 +243,13 @@ void FUserInterfaceCommand::Run()
 
 void FUserInterfaceCommand::InitializeSlateApplication(bool bOpenTraceFile, const TCHAR* TraceFile)
 {
-	//TODO: FSlateApplication::InitHighDPI(true);
+	FSlateApplication::InitHighDPI(true);
 
 	FModuleManager::Get().LoadModuleChecked("EditorStyle");
 
 	//const FSlateBrush* AppIcon = new FSlateImageBrush(FPaths::EngineContentDir() / "Editor/Slate/Icons/Insights/AppIcon_24x.png", FVector2D(24.0f, 24.0f));
 	//FSlateApplication::Get().SetAppIcon(AppIcon);
 
-
 	// Set the application name.
 	const FText ApplicationTitle = FText::Format(NSLOCTEXT("UnrealInsights", "AppTitle", "Unreal Insights {0}"), FText::FromString(TEXT(UNREAL_INSIGHTS_VERSION_STRING_EX)));
 	FGlobalTabmanager::Get()->SetApplicationTitle(ApplicationTitle);
