[Insights]
  - Improved struct packing for FTraceViewModel, FStoreBrowserTraceInfo and FLogMessageRecord.
  - Log View: Limited the height of "Category Filter" menu to height of the Log View panel.

diff --git a/Engine/Source/Developer/TraceInsights/Private/Insights/StoreService/StoreBrowser.h b/Engine/Source/Developer/TraceInsights/Private/Insights/StoreService/StoreBrowser.h
--- a/Engine/Source/Developer/TraceInsights/Private/Insights/StoreService/StoreBrowser.h
+++ b/Engine/Source/Developer/TraceInsights/Private/Insights/StoreService/StoreBrowser.h
@@ -26,9 +26,9 @@ namespace Insights
 struct FStoreBrowserTraceInfo
 {
 	uint32 TraceId = 0;
+	int32 TraceIndex = -1;
 
 	uint64 ChangeSerial = 0;
-	int32 TraceIndex = -1;
 
 	FString Name;
 	//FString Uri;
@@ -36,10 +36,6 @@ struct FStoreBrowserTraceInfo
 	FDateTime Timestamp = 0;
 	uint64 Size = 0;
 
-	bool bIsLive = false;
-	uint32 IpAddress = 0;
-
-	bool bIsMetadataUpdated = false;
 	FString Platform;
 	FString AppName;
 	FString CommandLine;
@@ -49,6 +45,11 @@ struct FStoreBrowserTraceInfo
 	EBuildConfiguration ConfigurationType = EBuildConfiguration::Unknown;
 	EBuildTargetType TargetType = EBuildTargetType::Unknown;
 
+	bool bIsMetadataUpdated = false;
+	bool bIsLive = false;
+
+	uint32 IpAddress = 0;
+
 	FStoreBrowserTraceInfo() = default;
 
 	static FDateTime ConvertTimestamp(uint64 InTimestamp)
diff --git a/Engine/Source/Developer/TraceInsights/Private/Insights/ViewModels/LogFilter.cpp b/Engine/Source/Developer/TraceInsights/Private/Insights/ViewModels/LogFilter.cpp
--- a/Engine/Source/Developer/TraceInsights/Private/Insights/ViewModels/LogFilter.cpp
+++ b/Engine/Source/Developer/TraceInsights/Private/Insights/ViewModels/LogFilter.cpp
@@ -20,7 +20,7 @@ public:
 	/** Test the given value against the strings extracted from the current item */
 	virtual bool TestBasicStringExpression(const FTextFilterString& InValue, const ETextFilterTextComparisonMode InTextComparisonMode) const override
 	{
-		return TextFilterUtils::TestBasicStringExpression(Message->Message.ToString(), InValue, InTextComparisonMode);
+		return TextFilterUtils::TestBasicStringExpression(Message->GetMessageAsString(), InValue, InTextComparisonMode);
 	}
 
 	/**
@@ -107,7 +107,7 @@ bool FLogFilter::IsMessageAllowed(const FLogMessageRecord& Message)
 	// Filter by Verbosity
 	if (bIsFilterSetByVerbosity)
 	{
-		if (Message.Verbosity > VerbosityThreshold)
+		if (Message.GetVerbosity() > VerbosityThreshold)
 		{
 			return false;
 		}
@@ -116,7 +116,7 @@ bool FLogFilter::IsMessageAllowed(const FLogMessageRecord& Message)
 	// Filter by Category
 	if (bIsFilterSetByCategory)
 	{
-		if (!EnabledLogCategories.Contains(FName(*Message.Category.ToString())))
+		if (!EnabledLogCategories.Contains(FName(Message.GetCategory())))
 		{
 			return false;
 		}
diff --git a/Engine/Source/Developer/TraceInsights/Private/Insights/ViewModels/LogMessage.cpp b/Engine/Source/Developer/TraceInsights/Private/Insights/ViewModels/LogMessage.cpp
--- a/Engine/Source/Developer/TraceInsights/Private/Insights/ViewModels/LogMessage.cpp
+++ b/Engine/Source/Developer/TraceInsights/Private/Insights/ViewModels/LogMessage.cpp
@@ -13,27 +13,14 @@
 // FLogMessageRecord
 ////////////////////////////////////////////////////////////////////////////////////////////////////
 
-FLogMessageRecord::FLogMessageRecord()
-	: Index(0)
-	, Time(0)
-	, Verbosity(ELogVerbosity::Type::NoLogging)
-	, Category()
-	, Message()
-	, File()
-	, Line(0)
-{
-}
-
-////////////////////////////////////////////////////////////////////////////////////////////////////
-
 FLogMessageRecord::FLogMessageRecord(const TraceServices::FLogMessageInfo& TraceLogMessage)
-	: Index(static_cast<int32>(TraceLogMessage.Index))
-	, Time(TraceLogMessage.Time)
-	, Verbosity(TraceLogMessage.Verbosity)
+	: Time(TraceLogMessage.Time)
 	//, Category(FText::FromString(TraceLogMessage.Category))
 	//, Message(FText::FromString(TraceLogMessage.Message))
 	, File(FText::FromString(TraceLogMessage.File))
 	, Line(TraceLogMessage.Line)
+	, Index(static_cast<uint32>(TraceLogMessage.Index))
+	, Verbosity(TraceLogMessage.Verbosity)
 {
 	// Strip the "Log" prefix.
 	FString CategoryStr(TraceLogMessage.Category->Name);
@@ -41,12 +28,12 @@ FLogMessageRecord::FLogMessageRecord(const TraceServices::FLogMessageInfo& Trace
 	{
 		CategoryStr.RightChopInline(3, false);
 	}
-	Category = FText::FromString(CategoryStr);
+	Category = FText::FromString(MoveTemp(CategoryStr));
 
 	// Strip the trailing whitespaces (ex. some messages ends with "\n" and we do not want the LogView rows to have an unnecessary increased height).
 	FString MessageStr(TraceLogMessage.Message);
 	MessageStr.TrimEndInline();
-	Message = FText::FromString(MessageStr);
+	Message = FText::FromString(MoveTemp(MessageStr));
 }
 
 ////////////////////////////////////////////////////////////////////////////////////////////////////
@@ -72,27 +59,6 @@ FText FLogMessageRecord::GetVerbosityAsText() const
 
 ////////////////////////////////////////////////////////////////////////////////////////////////////
 
-FText FLogMessageRecord::GetCategoryAsText() const
-{
-	return Category;
-}
-
-////////////////////////////////////////////////////////////////////////////////////////////////////
-
-FText FLogMessageRecord::GetMessageAsText() const
-{
-	return Message;
-}
-
-////////////////////////////////////////////////////////////////////////////////////////////////////
-
-FText FLogMessageRecord::GetFileAsText() const
-{
-	return File;
-}
-
-////////////////////////////////////////////////////////////////////////////////////////////////////
-
 FText FLogMessageRecord::GetLineAsText() const
 {
 	return FText::AsNumber(Line);
diff --git a/Engine/Source/Developer/TraceInsights/Private/Insights/ViewModels/LogMessage.h b/Engine/Source/Developer/TraceInsights/Private/Insights/ViewModels/LogMessage.h
--- a/Engine/Source/Developer/TraceInsights/Private/Insights/ViewModels/LogMessage.h
+++ b/Engine/Source/Developer/TraceInsights/Private/Insights/ViewModels/LogMessage.h
@@ -17,27 +17,38 @@ namespace TraceServices
 class FLogMessageRecord
 {
 public:
-	FLogMessageRecord();
-	FLogMessageRecord(const TraceServices::FLogMessageInfo& Message);
+	FLogMessageRecord() = default;
+	explicit FLogMessageRecord(const TraceServices::FLogMessageInfo& Message);
+
+	uint32 GetIndex() const { return Index; }
+	double GetTime() const { return Time; }
+	ELogVerbosity::Type GetVerbosity() const { return Verbosity; }
+	const TCHAR* GetCategory() const { return *Category.ToString(); }
+	const TCHAR* GetMessage() const { return *Message.ToString(); }
+	const TCHAR* GetFile() const { return *File.ToString(); }
+	uint32 GetLine() const { return Line; }
+
+	FString GetCategoryAsString() const { return Category.ToString(); }
+	FString GetMessageAsString() const { return Message.ToString(); }
 
 	FText GetIndexAsText() const;
 	FText GetTimeAsText() const;
 	FText GetVerbosityAsText() const;
-	FText GetCategoryAsText() const;
-	FText GetMessageAsText() const;
-	FText GetFileAsText() const;
+	FText GetCategoryAsText() const { return Category; }
+	FText GetMessageAsText() const { return Message; }
+	FText GetFileAsText() const { return File; }
 	FText GetLineAsText() const;
 
 	FText ToDisplayString() const;
 
-public:
-	int32 Index;
-	double Time;
-	ELogVerbosity::Type Verbosity;
+private:
+	double Time = 0.0;
 	FText Category;
 	FText Message;
 	FText File;
-	uint32 Line;
+	uint32 Line = 0;
+	uint32 Index = 0;
+	ELogVerbosity::Type Verbosity = ELogVerbosity::Type::NoLogging;
 };
 
 ////////////////////////////////////////////////////////////////////////////////////////////////////
diff --git a/Engine/Source/Developer/TraceInsights/Private/Insights/Widgets/SLogView.cpp b/Engine/Source/Developer/TraceInsights/Private/Insights/Widgets/SLogView.cpp
--- a/Engine/Source/Developer/TraceInsights/Private/Insights/Widgets/SLogView.cpp
+++ b/Engine/Source/Developer/TraceInsights/Private/Insights/Widgets/SLogView.cpp
@@ -152,7 +152,7 @@ public:
 			if (!SelectedLogMessage || SelectedLogMessage->GetIndex() != LogMessagePin->GetIndex()) // if row is not selected
 			{
 				FLogMessageRecord& CacheEntry = ParentWidgetPin->GetCache().Get(LogMessagePin->GetIndex());
-				const double Time = CacheEntry.Time;
+				const double Time = CacheEntry.GetTime();
 
 				TSharedPtr<STimingProfilerWindow> Window = FTimingProfilerManager::Get()->GetProfilerWindow();
 				if (Window)
@@ -187,7 +187,7 @@ public:
 			}
 
 			FLogMessageRecord& CacheEntry = ParentWidgetPin->GetCache().Get(LogMessagePin->GetIndex());
-			const double Time = CacheEntry.Time;
+			const double Time = CacheEntry.GetTime();
 
 			TSharedPtr<STimingProfilerWindow> Window = FTimingProfilerManager::Get()->GetProfilerWindow();
 			if (Window)
@@ -277,7 +277,7 @@ public:
 		if (ParentWidgetPin.IsValid() && LogMessagePin.IsValid())
 		{
 			FLogMessageRecord& CacheEntry = ParentWidgetPin->GetCache().Get(LogMessagePin->GetIndex());
-			return FSlateColor(FTimeMarkerTrackBuilder::GetColorByVerbosity(CacheEntry.Verbosity));
+			return FSlateColor(FTimeMarkerTrackBuilder::GetColorByVerbosity(CacheEntry.GetVerbosity()));
 		}
 		else
 		{
@@ -310,7 +310,7 @@ public:
 		if (ParentWidgetPin.IsValid() && LogMessagePin.IsValid())
 		{
 			FLogMessageRecord& CacheEntry = ParentWidgetPin->GetCache().Get(LogMessagePin->GetIndex());
-			return FSlateColor(FTimeMarkerTrackBuilder::GetColorByCategory(*CacheEntry.Category.ToString()));
+			return FSlateColor(FTimeMarkerTrackBuilder::GetColorByCategory(CacheEntry.GetCategory()));
 		}
 		else
 		{
@@ -932,7 +932,7 @@ void SLogView::SelectLogMessage(TSharedPtr<FLogMessage> LogMessage)
 			TSharedPtr<STimingView> TimingView = Window->GetTimingView();
 			if (TimingView)
 			{
-				const double Time = Cache.Get(LogMessage->GetIndex()).Time;
+				const double Time = Cache.Get(LogMessage->GetIndex()).GetTime();
 
 				if (FSlateApplication::Get().GetModifierKeys().IsShiftDown())
 				{
@@ -1056,11 +1056,11 @@ TSharedPtr<SWidget> SLogView::ListView_GetContextMenu()
 		if (SelectedLogMessage.IsValid())
 		{
 			FLogMessageRecord& Record = Cache.Get(SelectedLogMessage->GetIndex());
-			FName CategoryName(*Record.Category.ToString());
+			FName CategoryName(Record.GetCategory());
 
 			MenuBuilder.AddMenuEntry(
-				FText::Format(LOCTEXT("HideCategory", "Hide \"{0}\" Category"), Record.Category),
-				FText::Format(LOCTEXT("HideCategory_Tooltip", "Hide the \"{0}\" log category."), Record.Category),
+				FText::Format(LOCTEXT("HideCategory", "Hide \"{0}\" Category"), Record.GetCategoryAsText()),
+				FText::Format(LOCTEXT("HideCategory_Tooltip", "Hide the \"{0}\" log category."), Record.GetCategoryAsText()),
 				FSlateIcon(),
 				FUIAction(FExecuteAction::CreateSP(this, &SLogView::ToggleCategory_Execute, CategoryName)),
 				NAME_None,
@@ -1068,8 +1068,8 @@ TSharedPtr<SWidget> SLogView::ListView_GetContextMenu()
 			);
 
 			MenuBuilder.AddMenuEntry(
-				FText::Format(LOCTEXT("ShowOnlyCategory", "Show Only \"{0}\" Category"), Record.Category),
-				FText::Format(LOCTEXT("ShowOnlyCategory_Tooltip", "Show only the \"{0}\" log category (hide all other log categories)."), Record.Category),
+				FText::Format(LOCTEXT("ShowOnlyCategory", "Show Only \"{0}\" Category"), Record.GetCategoryAsText()),
+				FText::Format(LOCTEXT("ShowOnlyCategory_Tooltip", "Show only the \"{0}\" log category (hide all other log categories)."), Record.GetCategoryAsText()),
 				FSlateIcon(),
 				FUIAction(FExecuteAction::CreateSP(this, &SLogView::ShowOnlyCategory_Execute, CategoryName)),
 				NAME_None,
@@ -1187,7 +1187,8 @@ TSharedRef<SWidget> SLogView::MakeCategoryFilterMenu()
 	CreateCategoriesFilterMenuSection(MenuBuilder);
 	MenuBuilder.EndSection();
 
-	return MenuBuilder.MakeWidget();
+	const float MaxMenuHeight = FMath::Clamp(this->GetCachedGeometry().GetLocalSize().Y - 40.0f, 100.0f, 500.0f);
+	return MenuBuilder.MakeWidget(nullptr, MaxMenuHeight);
 }
 
 ////////////////////////////////////////////////////////////////////////////////////////////////////
diff --git a/Engine/Source/Developer/TraceInsights/Private/Insights/Widgets/SStartPageWindow.h b/Engine/Source/Developer/TraceInsights/Private/Insights/Widgets/SStartPageWindow.h
--- a/Engine/Source/Developer/TraceInsights/Private/Insights/Widgets/SStartPageWindow.h
+++ b/Engine/Source/Developer/TraceInsights/Private/Insights/Widgets/SStartPageWindow.h
@@ -49,9 +49,9 @@ typedef TWeakPtr<class SNotificationItem> SNotificationItemWeak;
 struct FTraceViewModel
 {
 	uint32 TraceId = 0;
+	int32 TraceIndex = -1; // debug
 
 	uint64 ChangeSerial = 0;
-	int32 TraceIndex = -1; // debug
 
 	FText Name;
 	FText Uri;
@@ -59,19 +59,19 @@ struct FTraceViewModel
 	FDateTime Timestamp = 0;
 	uint64 Size = 0;
 
-	bool bIsLive = false;
-	uint32 IpAddress = 0;
-
-	bool bIsMetadataUpdated = false;
 	FText Platform;
 	FText AppName;
 	FText CommandLine;
 	FText Branch;
 	FText BuildVersion;
-	uint32 Changelist;
+	uint32 Changelist = 0;
 	EBuildConfiguration ConfigurationType = EBuildConfiguration::Unknown;
 	EBuildTargetType TargetType = EBuildTargetType::Unknown;
 
+	bool bIsMetadataUpdated = false;
+	bool bIsLive = false;
+	uint32 IpAddress = 0;
+
 	FTraceViewModel() = default;
 
 	static FDateTime ConvertTimestamp(uint64 InTimestamp)
