[Insights] Disabled striped brush for context switches.

diff --git a/Engine/Source/Developer/TraceInsights/Private/Insights/ViewModels/TimingViewDrawHelper.cpp b/Engine/Source/Developer/TraceInsights/Private/Insights/ViewModels/TimingViewDrawHelper.cpp
--- a/Engine/Source/Developer/TraceInsights/Private/Insights/ViewModels/TimingViewDrawHelper.cpp
+++ b/Engine/Source/Developer/TraceInsights/Private/Insights/ViewModels/TimingViewDrawHelper.cpp
@@ -270,7 +270,7 @@ FTimingViewDrawHelper::FTimingViewDrawHelper(const FDrawContext& InDrawContext,
 	, HoveredEventBorderBrush(FInsightsStyle::Get().GetBrush("HoveredEventBorder"))
 	, SelectedEventBorderBrush(FInsightsStyle::Get().GetBrush("SelectedEventBorder"))
 	, BackgroundAreaBrush(WhiteBrush)
-	, IdleAreaBrush(FInsightsStyle::Get().GetBrush("DashedBrush"))
+	, IdleAreaBrush(FInsightsStyle::Get().GetBrush("IdleAreaBrush"))
 	, ValidAreaColor(0.07f, 0.07f, 0.07f, 1.0f)
 	, InvalidAreaColor(0.1f, 0.07f, 0.07f, 1.0f)
 	, EdgeColor(0.05f, 0.05f, 0.05f, 1.0f)
@@ -563,8 +563,10 @@ void FTimingViewDrawHelper::DrawContextSwitchMarkers(const FTimingEventsTrackDra
 		}
 
 		// Draw overlay for idle areas.
-		const FSlateBrush* IdleBrush = IdleAreaBrush;
-		const FLinearColor IdleColor(1.0f, 1.0f, 1.0f, 0.5f);
+		//const FSlateBrush* IdleBrush = IdleAreaBrush;
+		//const FLinearColor IdleColor(1.0f, 1.0f, 1.0f, 0.75f);
+		const FSlateBrush* IdleBrush = WhiteBrush;
+		const FLinearColor IdleColor(0.0f, 0.0f, 0.0f, 0.3f);
 		const int32 Count1 = DrawState.Boxes.Num();
 		const int32 Count2 = DrawState.Borders.Num();
 		int32 Index1 = 0;
