[Insights] Removed unused IdleAreaBrush member and the DashedBrush style.

diff --git a/Engine/Source/Developer/TraceInsights/Private/Insights/InsightsStyle.cpp b/Engine/Source/Developer/TraceInsights/Private/Insights/InsightsStyle.cpp
--- a/Engine/Source/Developer/TraceInsights/Private/Insights/InsightsStyle.cpp
+++ b/Engine/Source/Developer/TraceInsights/Private/Insights/InsightsStyle.cpp
@@ -59,7 +59,6 @@ TSharedRef<FSlateStyleSet> FInsightsStyle::Create()
 	Style.Set("WhiteBrush", new FSlateColorBrush(FLinearColor::White));
 	Style.Set("SingleBorder", new FSlateBorderBrush(NAME_None, FMargin(1.0f)));
 	Style.Set("DoubleBorder", new FSlateBorderBrush(NAME_None, FMargin(2.0f)));
-	Style.Set("DashedBrush", new IMAGE_BRUSH("Old/Window/WindowBackground", FVector2D(74, 74), FLinearColor::White, ESlateBrushTileType::Both));
 
 	Style.Set("EventBorder", new FSlateBorderBrush(NAME_None, FMargin(1.0f)));
 	Style.Set("HoveredEventBorder", new FSlateBorderBrush(NAME_None, FMargin(2.0f)));
diff --git a/Engine/Source/Developer/TraceInsights/Private/Insights/ViewModels/TimingViewDrawHelper.cpp b/Engine/Source/Developer/TraceInsights/Private/Insights/ViewModels/TimingViewDrawHelper.cpp
--- a/Engine/Source/Developer/TraceInsights/Private/Insights/ViewModels/TimingViewDrawHelper.cpp
+++ b/Engine/Source/Developer/TraceInsights/Private/Insights/ViewModels/TimingViewDrawHelper.cpp
@@ -270,7 +270,6 @@ FTimingViewDrawHelper::FTimingViewDrawHelper(const FDrawContext& InDrawContext,
 	, HoveredEventBorderBrush(FInsightsStyle::Get().GetBrush("HoveredEventBorder"))
 	, SelectedEventBorderBrush(FInsightsStyle::Get().GetBrush("SelectedEventBorder"))
 	, BackgroundAreaBrush(WhiteBrush)
-	, IdleAreaBrush(FInsightsStyle::Get().GetBrush("IdleAreaBrush"))
 	, ValidAreaColor(0.07f, 0.07f, 0.07f, 1.0f)
 	, InvalidAreaColor(0.1f, 0.07f, 0.07f, 1.0f)
 	, EdgeColor(0.05f, 0.05f, 0.05f, 1.0f)
@@ -563,8 +562,6 @@ void FTimingViewDrawHelper::DrawContextSwitchMarkers(const FTimingEventsTrackDra
 		}
 
 		// Draw overlay for idle areas.
-		//const FSlateBrush* IdleBrush = IdleAreaBrush;
-		//const FLinearColor IdleColor(1.0f, 1.0f, 1.0f, 0.75f);
 		const FSlateBrush* IdleBrush = WhiteBrush;
 		const FLinearColor IdleColor(0.0f, 0.0f, 0.0f, 0.3f);
 		const int32 Count1 = DrawState.Boxes.Num();
diff --git a/Engine/Source/Developer/TraceInsights/Private/Insights/ViewModels/TimingViewDrawHelper.h b/Engine/Source/Developer/TraceInsights/Private/Insights/ViewModels/TimingViewDrawHelper.h
--- a/Engine/Source/Developer/TraceInsights/Private/Insights/ViewModels/TimingViewDrawHelper.h
+++ b/Engine/Source/Developer/TraceInsights/Private/Insights/ViewModels/TimingViewDrawHelper.h
@@ -221,7 +221,6 @@ private:
 	const FSlateBrush* HoveredEventBorderBrush;
 	const FSlateBrush* SelectedEventBorderBrush;
 	const FSlateBrush* BackgroundAreaBrush;
-	const FSlateBrush* IdleAreaBrush;
 	const FLinearColor ValidAreaColor;
 	const FLinearColor InvalidAreaColor;
 	const FLinearColor EdgeColor;
