[Insights]
  - Timing Insights: Added more switches related to context switches:
      * Overlays - Extend the visualisation of context switches over the cpu timing tracks. (Shift+O)
      * Extended Lines - Show/hide the extended vertical lines at edges of each context switch event. (Shift+L)
  - Moved all options into a "Context Switches" sub-menu in the Tracks menu.
  - Added placeholder code for the Cpu Core tracks.

diff --git a/Engine/Source/Developer/TraceInsights/Private/Insights/ContextSwitches/ContextSwitchesProfilerManager.h b/Engine/Source/Developer/TraceInsights/Private/Insights/ContextSwitches/ContextSwitchesProfilerManager.h
--- a/Engine/Source/Developer/TraceInsights/Private/Insights/ContextSwitches/ContextSwitchesProfilerManager.h
+++ b/Engine/Source/Developer/TraceInsights/Private/Insights/ContextSwitches/ContextSwitchesProfilerManager.h
@@ -44,7 +44,8 @@ public:
 	virtual void UnregisterMajorTabs() override;
 
 	//////////////////////////////////////////////////
-	bool GetIsAvailable() { return bIsAvailable; }
+
+	bool IsAvailable() const { return bIsAvailable; }
 
 	void OnSessionChanged();
 
@@ -75,4 +76,3 @@ private:
 };
 
 } // namespace Insights
-
diff --git a/Engine/Source/Developer/TraceInsights/Private/Insights/ContextSwitches/ViewModels/ContextSwitchesTimingTrack.cpp b/Engine/Source/Developer/TraceInsights/Private/Insights/ContextSwitches/ViewModels/ContextSwitchesTimingTrack.cpp
--- a/Engine/Source/Developer/TraceInsights/Private/Insights/ContextSwitches/ViewModels/ContextSwitchesTimingTrack.cpp
+++ b/Engine/Source/Developer/TraceInsights/Private/Insights/ContextSwitches/ViewModels/ContextSwitchesTimingTrack.cpp
@@ -10,6 +10,7 @@
 
 // Insights
 #include "Insights/Common/TimeUtils.h"
+#include "Insights/ContextSwitches/ContextSwitchesProfilerManager.h"
 #include "Insights/ContextSwitches/ViewModels/ContextSwitchTimingEvent.h"
 #include "Insights/ITimingViewSession.h"
 #include "Insights/InsightsManager.h"
@@ -46,11 +47,17 @@ public:
 	PRAGMA_DISABLE_OPTIMIZATION
 	virtual void RegisterCommands() override
 	{
-		UI_COMMAND(Command_ShowContextSwitches, "Show Context Switches ", "Show/hide context switches events as header tracks", EUserInterfaceActionType::ToggleButton, FInputChord(EModifierKey::Shift, EKeys::C));
+		UI_COMMAND(Command_ShowCoreTracks, "Core Tracks", "Show/hide the Cpu Core tracks.", EUserInterfaceActionType::ToggleButton, FInputChord());
+		UI_COMMAND(Command_ShowContextSwitches, "Context Switches", "Show/hide context switches on top of cpu timing tracks.", EUserInterfaceActionType::ToggleButton, FInputChord(EModifierKey::Shift, EKeys::C));
+		UI_COMMAND(Command_ShowOverlays, "Overlays", "Extend the visualisation of context switches over the cpu timing tracks.", EUserInterfaceActionType::ToggleButton, FInputChord(EModifierKey::Shift, EKeys::O));
+		UI_COMMAND(Command_ShowExtendedLines, "Extended Lines", "Show/hide the extended vertical lines at edges of each context switch event.", EUserInterfaceActionType::ToggleButton, FInputChord(EModifierKey::Shift, EKeys::L));
 	}
 	PRAGMA_ENABLE_OPTIMIZATION
 
+	TSharedPtr<FUICommandInfo> Command_ShowCoreTracks;
 	TSharedPtr<FUICommandInfo> Command_ShowContextSwitches;
+	TSharedPtr<FUICommandInfo> Command_ShowOverlays;
+	TSharedPtr<FUICommandInfo> Command_ShowExtendedLines;
 };
 
 ////////////////////////////////////////////////////////////////////////////////////////////////////
@@ -58,7 +65,12 @@ public:
 ////////////////////////////////////////////////////////////////////////////////////////////////////
 
 FContextSwitchesSharedState::FContextSwitchesSharedState(STimingView* InTimingView) 
-	: TimingView(InTimingView) 
+	: TimingView(InTimingView)
+	, bAreContextSwitchesVisible(true)
+	, bAreCoreTracksVisible(true)
+	, bAreOverlaysVisible(true)
+	, bAreExtendedLinesVisible(true)
+	, bSyncWithProviders(true)
 {
 }
 
@@ -70,6 +82,13 @@ void FContextSwitchesSharedState::OnBeginSession(Insights::ITimingViewSession& I
 	{
 		return;
 	}
+
+	bAreCoreTracksVisible = true;
+	bAreContextSwitchesVisible = true;
+	bAreOverlaysVisible = true;
+	bAreExtendedLinesVisible = true;
+
+	bSyncWithProviders = true;
 }
 
 ////////////////////////////////////////////////////////////////////////////////////////////////////
@@ -81,7 +100,12 @@ void FContextSwitchesSharedState::OnEndSession(Insights::ITimingViewSession& InS
 		return;
 	}
 
-	AddContextSwitchesChildTracks();
+	bAreCoreTracksVisible = true;
+	bAreContextSwitchesVisible = true;
+	bAreOverlaysVisible = true;
+	bAreExtendedLinesVisible = true;
+
+	bSyncWithProviders = false;
 }
 
 ////////////////////////////////////////////////////////////////////////////////////////////////////
@@ -93,36 +117,60 @@ void FContextSwitchesSharedState::Tick(Insights::ITimingViewSession& InSession,
 		return;
 	}
 
-	if (!FInsightsManager::Get()->IsAnalysisComplete() && bShowContextSwitchesTrack)
+	if (bSyncWithProviders && AreContextSwitchesAvailable())
 	{
-		AddContextSwitchesChildTracks();
+		if (bAreCoreTracksVisible)
+		{
+			AddCoreTracks();
+		}
+
+		if (bAreContextSwitchesVisible)
+		{
+			AddContextSwitchesChildTracks();
+		}
+
+		if (FInsightsManager::Get()->IsAnalysisComplete())
+		{
+			// No need to sync anymore when analysis is completed.
+			bSyncWithProviders = false;
+		}
 	}
 }
 
 ////////////////////////////////////////////////////////////////////////////////////////////////////
 
-void FContextSwitchesSharedState::ExtendFilterMenu(Insights::ITimingViewSession& InSession, FMenuBuilder& InOutMenuBuilder)
+void FContextSwitchesSharedState::ExtendFilterMenu(Insights::ITimingViewSession& InSession, FMenuBuilder& InMenuBuilder)
 {
 	if (&InSession != TimingView)
 	{
 		return;
 	}
+
+	InMenuBuilder.BeginSection("ContextSwitches");
+	{
+		InMenuBuilder.AddSubMenu(
+			LOCTEXT("ContextSwitches_SubMenu", "Context Switches"),
+			LOCTEXT("ContextSwitches_SubMenu_Desc", "Context Switch track options"),
+			FNewMenuDelegate::CreateSP(this, &FContextSwitchesSharedState::BuildSubMenu),
+			false,
+			FSlateIcon()
+		);
+	}
+	InMenuBuilder.EndSection();
 }
 
 ////////////////////////////////////////////////////////////////////////////////////////////////////
 
-bool FContextSwitchesSharedState::ExtendGlobalContextMenu(ITimingViewSession& InSession, FMenuBuilder& InMenuBuilder)
+void FContextSwitchesSharedState::BuildSubMenu(FMenuBuilder& InMenuBuilder)
 {
-	InMenuBuilder.AddMenuEntry
-	(
-		FContextSwitchesStateCommands::Get().Command_ShowContextSwitches,
-		NAME_None,
-		TAttribute<FText>(),
-		TAttribute<FText>(),
-		FSlateIcon()
-	);
-
-	return true;
+	InMenuBuilder.BeginSection("ContextSwitches", LOCTEXT("ContextSwitchesHeading", "Context Switches"));
+	{
+		//TODO: InMenuBuilder.AddMenuEntry(FContextSwitchesStateCommands::Get().Command_ShowCoreTracks);
+		InMenuBuilder.AddMenuEntry(FContextSwitchesStateCommands::Get().Command_ShowContextSwitches);
+		InMenuBuilder.AddMenuEntry(FContextSwitchesStateCommands::Get().Command_ShowOverlays);
+		InMenuBuilder.AddMenuEntry(FContextSwitchesStateCommands::Get().Command_ShowExtendedLines);
+	}
+	InMenuBuilder.EndSection();
 }
 
 ////////////////////////////////////////////////////////////////////////////////////////////////////
@@ -134,11 +182,50 @@ void FContextSwitchesSharedState::AddCommands()
 	TSharedPtr<FUICommandList> CommandList = TimingView->GetCommandList();
 	ensure(CommandList.IsValid());
 
+	CommandList->MapAction(
+		FContextSwitchesStateCommands::Get().Command_ShowCoreTracks,
+		FExecuteAction::CreateSP(this, &FContextSwitchesSharedState::ContextMenu_ShowCoreTracks_Execute),
+		FCanExecuteAction::CreateSP(this, &FContextSwitchesSharedState::ContextMenu_ShowCoreTracks_CanExecute),
+		FIsActionChecked::CreateSP(this, &FContextSwitchesSharedState::ContextMenu_ShowCoreTracks_IsChecked));
+
 	CommandList->MapAction(
 		FContextSwitchesStateCommands::Get().Command_ShowContextSwitches,
 		FExecuteAction::CreateSP(this, &FContextSwitchesSharedState::ContextMenu_ShowContextSwitches_Execute),
 		FCanExecuteAction::CreateSP(this, &FContextSwitchesSharedState::ContextMenu_ShowContextSwitches_CanExecute),
 		FIsActionChecked::CreateSP(this, &FContextSwitchesSharedState::ContextMenu_ShowContextSwitches_IsChecked));
+
+	CommandList->MapAction(
+		FContextSwitchesStateCommands::Get().Command_ShowOverlays,
+		FExecuteAction::CreateSP(this, &FContextSwitchesSharedState::ContextMenu_ShowOverlays_Execute),
+		FCanExecuteAction::CreateSP(this, &FContextSwitchesSharedState::ContextMenu_ShowOverlays_CanExecute),
+		FIsActionChecked::CreateSP(this, &FContextSwitchesSharedState::ContextMenu_ShowOverlays_IsChecked));
+
+	CommandList->MapAction(
+		FContextSwitchesStateCommands::Get().Command_ShowExtendedLines,
+		FExecuteAction::CreateSP(this, &FContextSwitchesSharedState::ContextMenu_ShowExtendedLines_Execute),
+		FCanExecuteAction::CreateSP(this, &FContextSwitchesSharedState::ContextMenu_ShowExtendedLines_CanExecute),
+		FIsActionChecked::CreateSP(this, &FContextSwitchesSharedState::ContextMenu_ShowExtendedLines_IsChecked));
+}
+
+////////////////////////////////////////////////////////////////////////////////////////////////////
+
+bool FContextSwitchesSharedState::AreContextSwitchesAvailable() const
+{
+	return FContextSwitchesProfilerManager::Get()->IsAvailable();
+}
+
+////////////////////////////////////////////////////////////////////////////////////////////////////
+
+void FContextSwitchesSharedState::AddCoreTracks()
+{
+	//TODO
+}
+
+////////////////////////////////////////////////////////////////////////////////////////////////////
+
+void FContextSwitchesSharedState::RemoveCoreTracks()
+{
+	//TODO
 }
 
 ////////////////////////////////////////////////////////////////////////////////////////////////////
@@ -156,11 +243,12 @@ void FContextSwitchesSharedState::AddContextSwitchesChildTracks()
 
 	for (const TPair<uint32, TSharedPtr<FCpuTimingTrack>>& MapEntry : CpuTracks)
 	{
-		if (MapEntry.Value.IsValid() && !MapEntry.Value->GetChildTrack().IsValid())
+		const TSharedPtr<FCpuTimingTrack>& CpuTrack = MapEntry.Value;
+		if (CpuTrack.IsValid() && !CpuTrack->GetChildTrack().IsValid())
 		{
-			TSharedPtr<FContextSwitchesTimingTrack> ContextSwitchesTrack = MakeShared<FContextSwitchesTimingTrack>(*this, TEXT("Context Switches"), MapEntry.Value->GetTimelineIndex(), MapEntry.Value->GetThreadId());
-			ContextSwitchesTrack->SetParentTrack(MapEntry.Value);
-			MapEntry.Value->SetChildTrack(ContextSwitchesTrack);
+			TSharedPtr<FContextSwitchesTimingTrack> ContextSwitchesTrack = MakeShared<FContextSwitchesTimingTrack>(*this, TEXT("Context Switches"), CpuTrack->GetTimelineIndex(), CpuTrack->GetThreadId());
+			ContextSwitchesTrack->SetParentTrack(CpuTrack);
+			CpuTrack->SetChildTrack(ContextSwitchesTrack);
 		}
 	}
 }
@@ -180,41 +268,64 @@ void FContextSwitchesSharedState::RemoveContextSwitchesChildTracks()
 
 	for (const TPair<uint32, TSharedPtr<FCpuTimingTrack>>& MapEntry : CpuTracks)
 	{
-		if (MapEntry.Value.IsValid() && MapEntry.Value->GetChildTrack().IsValid() && MapEntry.Value->GetChildTrack()->Is<FContextSwitchesTimingTrack>())
+		const TSharedPtr<FCpuTimingTrack>& CpuTrack = MapEntry.Value;
+		if (CpuTrack.IsValid() && CpuTrack->GetChildTrack().IsValid() && CpuTrack->GetChildTrack()->Is<FContextSwitchesTimingTrack>())
 		{
-			MapEntry.Value->SetChildTrack(nullptr);
+			CpuTrack->SetChildTrack(nullptr);
 		}
 	}
 }
 
 ////////////////////////////////////////////////////////////////////////////////////////////////////
 
-void FContextSwitchesSharedState::ContextMenu_ShowContextSwitches_Execute()
+void FContextSwitchesSharedState::SetCoreTracksVisible(bool bOnOff)
 {
-	SetContextSwitchesToggle(!bShowContextSwitchesTrack);
-
-	if (bShowContextSwitchesTrack)
+	if (bAreCoreTracksVisible != bOnOff)
 	{
-		AddContextSwitchesChildTracks();
+		bAreCoreTracksVisible = bOnOff;
+
+		if (bAreCoreTracksVisible)
+		{
+			AddCoreTracks();
+		}
+		else
+		{
+			RemoveCoreTracks();
+		}
 	}
-	else
+}
+
+////////////////////////////////////////////////////////////////////////////////////////////////////
+
+void FContextSwitchesSharedState::SetContextSwitchesVisible(bool bOnOff)
+{
+	if (bAreContextSwitchesVisible != bOnOff)
 	{
-		RemoveContextSwitchesChildTracks();
+		bAreContextSwitchesVisible = bOnOff;
+
+		if (bAreContextSwitchesVisible)
+		{
+			AddContextSwitchesChildTracks();
+		}
+		else
+		{
+			RemoveContextSwitchesChildTracks();
+		}
 	}
 }
 
 ////////////////////////////////////////////////////////////////////////////////////////////////////
 
-bool FContextSwitchesSharedState::ContextMenu_ShowContextSwitches_CanExecute()
+void FContextSwitchesSharedState::SetOverlaysVisible(bool bOnOff)
 {
-	return true;
+	bAreOverlaysVisible = bOnOff;
 }
 
 ////////////////////////////////////////////////////////////////////////////////////////////////////
 
-bool FContextSwitchesSharedState::ContextMenu_ShowContextSwitches_IsChecked()
+void FContextSwitchesSharedState::SetExtendedLinesVisible(bool bOnOff)
 {
-	return IsContextSwitchesToggleOn();
+	bAreExtendedLinesVisible = bOnOff;
 }
 
 ////////////////////////////////////////////////////////////////////////////////////////////////////
@@ -294,75 +405,78 @@ void FContextSwitchesTimingTrack::DrawLineEvents(const ITimingTrackDrawContext&
 
 void FContextSwitchesTimingTrack::PostDraw(const ITimingTrackDrawContext& Context) const
 {
-	float LineY1 = 0.0f;
-	float LineY2 = 0.0f;
-	ETimingTrackLocation LocalLocation = ETimingTrackLocation::None;
-	TSharedPtr<FBaseTimingTrack> LocalParentTrack = GetParentTrack().Pin();
-	if (LocalParentTrack)
+	if (SharedState.AreOverlaysVisible())
 	{
-		LineY1 = LocalParentTrack->GetPosY();
-		LineY2 = LineY1 + LocalParentTrack->GetHeight();
-		LocalLocation = LocalParentTrack->GetLocation();
-	}
-	else
-	{
-		LineY1 = GetPosY();
-		LineY2 = LineY1 + GetHeight();
-		LocalLocation = GetLocation();
-	}
-
-	const FTimingTrackViewport& Viewport = Context.GetViewport();
-	switch (LocalLocation)
-	{
-		case ETimingTrackLocation::Scrollable:
+		float LineY1 = 0.0f;
+		float LineY2 = 0.0f;
+		ETimingTrackLocation LocalLocation = ETimingTrackLocation::None;
+		TSharedPtr<FBaseTimingTrack> LocalParentTrack = GetParentTrack().Pin();
+		if (LocalParentTrack)
 		{
-			const float TopY = Viewport.GetTopOffset();
-			if (LineY1 < TopY)
-			{
-				LineY1 = TopY;
-			}
-			const float BottomY = Viewport.GetHeight() - Viewport.GetBottomOffset();
-			if (LineY2 > BottomY)
-			{
-				LineY2 = BottomY;
-			}
-			break;
+			LineY1 = LocalParentTrack->GetPosY();
+			LineY2 = LineY1 + LocalParentTrack->GetHeight();
+			LocalLocation = LocalParentTrack->GetLocation();
+		}
+		else
+		{
+			LineY1 = GetPosY();
+			LineY2 = LineY1 + GetHeight();
+			LocalLocation = GetLocation();
 		}
-		case ETimingTrackLocation::TopDocked:
+
+		const FTimingTrackViewport& Viewport = Context.GetViewport();
+		switch (LocalLocation)
 		{
-			const float TopY = 0.0f;
-			if (LineY1 < TopY)
+			case ETimingTrackLocation::Scrollable:
 			{
-				LineY1 = TopY;
-			}
-			const float BottomY = Viewport.GetTopOffset();
-			if (LineY2 > BottomY)
-			{
-				LineY2 = BottomY;
+				const float TopY = Viewport.GetTopOffset();
+				if (LineY1 < TopY)
+				{
+					LineY1 = TopY;
+				}
+				const float BottomY = Viewport.GetHeight() - Viewport.GetBottomOffset();
+				if (LineY2 > BottomY)
+				{
+					LineY2 = BottomY;
+				}
+				break;
 			}
-			break;
-		}
-		case ETimingTrackLocation::BottomDocked:
-		{
-			const float TopY = Viewport.GetHeight() - Viewport.GetBottomOffset();
-			if (LineY1 < TopY)
+			case ETimingTrackLocation::TopDocked:
 			{
-				LineY1 = TopY;
+				const float TopY = 0.0f;
+				if (LineY1 < TopY)
+				{
+					LineY1 = TopY;
+				}
+				const float BottomY = Viewport.GetTopOffset();
+				if (LineY2 > BottomY)
+				{
+					LineY2 = BottomY;
+				}
+				break;
 			}
-			const float BottomY = Viewport.GetHeight();
-			if (LineY2 > BottomY)
+			case ETimingTrackLocation::BottomDocked:
 			{
-				LineY2 = BottomY;
+				const float TopY = Viewport.GetHeight() - Viewport.GetBottomOffset();
+				if (LineY1 < TopY)
+				{
+					LineY1 = TopY;
+				}
+				const float BottomY = Viewport.GetHeight();
+				if (LineY2 > BottomY)
+				{
+					LineY2 = BottomY;
+				}
+				break;
 			}
-			break;
 		}
-	}
 
-	const float LineH = LineY2 - LineY1;
-	if (LineH > 0.0f)
-	{
-		const FTimingViewDrawHelper& Helper = *static_cast<const FTimingViewDrawHelper*>(&Context.GetHelper());
-		Helper.DrawContextSwitchMarkers(GetDrawState(), LineY1, LineH, 0.25f);
+		const float LineH = LineY2 - LineY1;
+		if (LineH > 0.0f)
+		{
+			const FTimingViewDrawHelper& Helper = *static_cast<const FTimingViewDrawHelper*>(&Context.GetHelper());
+			Helper.DrawContextSwitchMarkers(GetDrawState(), LineY1, LineH, 0.25f, SharedState.AreExtendedLinesVisible());
+		}
 	}
 }
 
diff --git a/Engine/Source/Developer/TraceInsights/Private/Insights/ContextSwitches/ViewModels/ContextSwitchesTimingTrack.h b/Engine/Source/Developer/TraceInsights/Private/Insights/ContextSwitches/ViewModels/ContextSwitchesTimingTrack.h
--- a/Engine/Source/Developer/TraceInsights/Private/Insights/ContextSwitches/ViewModels/ContextSwitchesTimingTrack.h
+++ b/Engine/Source/Developer/TraceInsights/Private/Insights/ContextSwitches/ViewModels/ContextSwitchesTimingTrack.h
@@ -24,8 +24,6 @@ public:
 	FContextSwitchesSharedState(STimingView* InTimingView);
 	virtual ~FContextSwitchesSharedState() = default;
 
-	bool IsContextSwitchesVisible() const;
-
 	//////////////////////////////////////////////////
 	// ITimingViewExtender interface
 
@@ -33,31 +31,73 @@ public:
 	virtual void OnEndSession(Insights::ITimingViewSession& InSession) override;
 	virtual void Tick(Insights::ITimingViewSession& InSession, const TraceServices::IAnalysisSession& InAnalysisSession) override;
 	virtual void ExtendFilterMenu(Insights::ITimingViewSession& InSession, FMenuBuilder& InMenuBuilder) override;
-	virtual bool ExtendGlobalContextMenu(ITimingViewSession& InSession, FMenuBuilder& InMenuBuilder) override;
 
 	//////////////////////////////////////////////////
 
-	bool IsContextSwitchesToggleOn() const { return bShowContextSwitchesTrack; }
-	void SetContextSwitchesToggle(bool bOnOff) { bShowContextSwitchesTrack = bOnOff; }
-	void ShowContextSwitches() { SetContextSwitchesToggle(true); }
-	void HideContextSwitches() { SetContextSwitchesToggle(false); }
-
 	STimingView* GetTimingView() { return TimingView; }
+
+	bool AreContextSwitchesAvailable() const;
+
+	bool AreCoreTracksVisible() const { return bAreCoreTracksVisible; }
+	void ShowCoreTracks() { SetCoreTracksVisible(true); }
+	void HideCoreTracks() { SetCoreTracksVisible(false); }
+	void ToggleCoreTracks() { SetCoreTracksVisible(!bAreCoreTracksVisible); }
+	void SetCoreTracksVisible(bool bOnOff);
+
+	bool AreContextSwitchesVisible() const { return bAreContextSwitchesVisible; }
+	void ShowContextSwitches() { SetContextSwitchesVisible(true); }
+	void HideContextSwitches() { SetContextSwitchesVisible(false); }
+	void ToggleContextSwitches() { SetContextSwitchesVisible(!bAreContextSwitchesVisible); }
+	void SetContextSwitchesVisible(bool bOnOff);
+
+	bool AreOverlaysVisible() const { return bAreOverlaysVisible; }
+	void ShowOverlays() { SetOverlaysVisible(true); }
+	void HideOverlays() { SetOverlaysVisible(false); }
+	void ToggleOverlays() { SetOverlaysVisible(!bAreOverlaysVisible); }
+	void SetOverlaysVisible(bool bOnOff);
+
+	bool AreExtendedLinesVisible() const { return bAreExtendedLinesVisible; }
+	void ShowExtendedLines() { SetExtendedLinesVisible(true); }
+	void HideExtendedLines() { SetExtendedLinesVisible(false); }
+	void ToggleExtendedLines() { SetExtendedLinesVisible(!bAreExtendedLinesVisible); }
+	void SetExtendedLinesVisible(bool bOnOff);
+
 	void AddCommands();
 
 private:
-	void ContextMenu_ShowContextSwitches_Execute();
-	bool ContextMenu_ShowContextSwitches_CanExecute();
-	bool ContextMenu_ShowContextSwitches_IsChecked();
+	void AddCoreTracks();
+	void RemoveCoreTracks();
 
 	void AddContextSwitchesChildTracks();
 	void RemoveContextSwitchesChildTracks();
 
+	void BuildSubMenu(FMenuBuilder& InMenuBuilder);
+
+	void ContextMenu_ShowCoreTracks_Execute() { ToggleCoreTracks(); }
+	bool ContextMenu_ShowCoreTracks_CanExecute() const { return AreContextSwitchesAvailable(); }
+	bool ContextMenu_ShowCoreTracks_IsChecked() const { return AreCoreTracksVisible(); }
+
+	void ContextMenu_ShowContextSwitches_Execute() { ToggleContextSwitches(); }
+	bool ContextMenu_ShowContextSwitches_CanExecute() const { return AreContextSwitchesAvailable(); }
+	bool ContextMenu_ShowContextSwitches_IsChecked() const { return AreContextSwitchesVisible(); }
+
+	void ContextMenu_ShowOverlays_Execute() { ToggleOverlays(); }
+	bool ContextMenu_ShowOverlays_CanExecute() const { return AreContextSwitchesAvailable() && AreContextSwitchesVisible(); }
+	bool ContextMenu_ShowOverlays_IsChecked() const { return AreOverlaysVisible(); }
+
+	void ContextMenu_ShowExtendedLines_Execute() { ToggleExtendedLines(); }
+	bool ContextMenu_ShowExtendedLines_CanExecute() const { return AreContextSwitchesAvailable() && AreContextSwitchesVisible() && AreOverlaysVisible(); }
+	bool ContextMenu_ShowExtendedLines_IsChecked() const { return AreExtendedLinesVisible(); }
 
 private:
 	STimingView* TimingView;
 
-	bool bShowContextSwitchesTrack = true;
+	bool bAreContextSwitchesVisible;
+	bool bAreCoreTracksVisible;
+	bool bAreOverlaysVisible;
+	bool bAreExtendedLinesVisible;
+
+	bool bSyncWithProviders;
 };
 
 ////////////////////////////////////////////////////////////////////////////////////////////////////
diff --git a/Engine/Source/Developer/TraceInsights/Private/Insights/ViewModels/TimingViewDrawHelper.cpp b/Engine/Source/Developer/TraceInsights/Private/Insights/ViewModels/TimingViewDrawHelper.cpp
--- a/Engine/Source/Developer/TraceInsights/Private/Insights/ViewModels/TimingViewDrawHelper.cpp
+++ b/Engine/Source/Developer/TraceInsights/Private/Insights/ViewModels/TimingViewDrawHelper.cpp
@@ -540,25 +540,28 @@ void FTimingViewDrawHelper::DrawFadedLineEvents(const FTimingEventsTrackDrawStat
 
 ////////////////////////////////////////////////////////////////////////////////////////////////////
 
-void FTimingViewDrawHelper::DrawContextSwitchMarkers(const FTimingEventsTrackDrawState& DrawState, float LineY, float LineH, float Opacity) const
+void FTimingViewDrawHelper::DrawContextSwitchMarkers(const FTimingEventsTrackDrawState& DrawState, float LineY, float LineH, float Opacity, bool bShowVerticalLines) const
 {
 	if (LineH > 0.0f)
 	{
 		const int32 LayerId = ReservedLayerId + ToInt32(EDrawLayer::EventHighlight);
 
-		// Draw vertical lines (merged into large boxes).
-		for (const FTimingEventsTrackDrawState::FBoxPrimitive& Box : DrawState.Boxes)
+		if (bShowVerticalLines)
 		{
-			const FLinearColor Color = Box.Color.CopyWithNewOpacity(Opacity);
-			DrawContext.DrawBox(LayerId, Box.X, LineY, Box.W, LineH, WhiteBrush, Color);
-		}
+			// Draw vertical lines (merged into large boxes).
+			for (const FTimingEventsTrackDrawState::FBoxPrimitive& Box : DrawState.Boxes)
+			{
+				const FLinearColor Color = Box.Color.CopyWithNewOpacity(Opacity);
+				DrawContext.DrawBox(LayerId, Box.X, LineY, Box.W, LineH, WhiteBrush, Color);
+			}
 
-		// Draw vertical lines (edges of larger events).
-		for (const FTimingEventsTrackDrawState::FBoxPrimitive& Box : DrawState.Borders)
-		{
-			const FLinearColor Color = Box.Color.CopyWithNewOpacity(Opacity);
-			DrawContext.DrawBox(LayerId, Box.X, LineY, 1.0f, LineH, WhiteBrush, Color);
-			DrawContext.DrawBox(LayerId, Box.X + Box.W - 1.0f, LineY, 1.0f, LineH, WhiteBrush, Color);
+			// Draw vertical lines (edges of larger events).
+			for (const FTimingEventsTrackDrawState::FBoxPrimitive& Box : DrawState.Borders)
+			{
+				const FLinearColor Color = Box.Color.CopyWithNewOpacity(Opacity);
+				DrawContext.DrawBox(LayerId, Box.X, LineY, 1.0f, LineH, WhiteBrush, Color);
+				DrawContext.DrawBox(LayerId, Box.X + Box.W - 1.0f, LineY, 1.0f, LineH, WhiteBrush, Color);
+			}
 		}
 
 		// Draw overlay for idle areas.
diff --git a/Engine/Source/Developer/TraceInsights/Private/Insights/ViewModels/TimingViewDrawHelper.h b/Engine/Source/Developer/TraceInsights/Private/Insights/ViewModels/TimingViewDrawHelper.h
--- a/Engine/Source/Developer/TraceInsights/Private/Insights/ViewModels/TimingViewDrawHelper.h
+++ b/Engine/Source/Developer/TraceInsights/Private/Insights/ViewModels/TimingViewDrawHelper.h
@@ -191,7 +191,7 @@ public:
 
 	void DrawLineEvents(const FTimingEventsTrackDrawState& DrawState, const FTimingEventsTrack& Track, const float OffsetY = 1.0f) const;
 	void DrawFadedLineEvents(const FTimingEventsTrackDrawState& DrawState, const FTimingEventsTrack& Track, const float OffsetY = 1.0f, const float Opacity = 0.1f) const;
-	void DrawContextSwitchMarkers(const FTimingEventsTrackDrawState& DrawState, float LineY, float LineH, float Opacity) const;
+	void DrawContextSwitchMarkers(const FTimingEventsTrackDrawState& DrawState, float LineY, float LineH, float Opacity, bool bShowVerticalLines) const;
 
 	void DrawMarkers(const FTimingEventsTrackDrawState& DrawState, float LineY, float LineH, float Opacity) const;
 
diff --git a/Engine/Source/Developer/TraceInsights/Private/Insights/Widgets/STimingView.cpp b/Engine/Source/Developer/TraceInsights/Private/Insights/Widgets/STimingView.cpp
--- a/Engine/Source/Developer/TraceInsights/Private/Insights/Widgets/STimingView.cpp
+++ b/Engine/Source/Developer/TraceInsights/Private/Insights/Widgets/STimingView.cpp
@@ -2467,51 +2467,60 @@ FReply STimingView::OnKeyDown(const FGeometry& MyGeometry, const FKeyEvent& InKe
 {
 	if (InKeyEvent.GetKey() == EKeys::B)
 	{
-		// Toggle Bookmarks.
-		if (MarkersTrack->IsVisible())
+		if (!InKeyEvent.GetModifierKeys().IsControlDown() && !InKeyEvent.GetModifierKeys().IsShiftDown())
 		{
-			if (!MarkersTrack->IsBookmarksTrack())
+			// Toggle Bookmarks.
+			if (MarkersTrack->IsVisible())
 			{
-				SetDrawOnlyBookmarks(true);
+				if (!MarkersTrack->IsBookmarksTrack())
+				{
+					SetDrawOnlyBookmarks(true);
+				}
+				else
+				{
+					SetTimeMarkersVisible(false);
+				}
 			}
 			else
 			{
-				SetTimeMarkersVisible(false);
+				SetTimeMarkersVisible(true);
+				SetDrawOnlyBookmarks(true);
 			}
+			return FReply::Handled();
 		}
-		else
-		{
-			SetTimeMarkersVisible(true);
-			SetDrawOnlyBookmarks(true);
-		}
-		return FReply::Handled();
 	}
 	else if (InKeyEvent.GetKey() == EKeys::M)
 	{
-		// Toggle Time Markers.
-		if (MarkersTrack->IsVisible())
+		if (!InKeyEvent.GetModifierKeys().IsControlDown() && !InKeyEvent.GetModifierKeys().IsShiftDown())
 		{
-			if (MarkersTrack->IsBookmarksTrack())
+			// Toggle Time Markers.
+			if (MarkersTrack->IsVisible())
 			{
-				SetDrawOnlyBookmarks(false);
+				if (MarkersTrack->IsBookmarksTrack())
+				{
+					SetDrawOnlyBookmarks(false);
+				}
+				else
+				{
+					SetTimeMarkersVisible(false);
+				}
 			}
 			else
 			{
-				SetTimeMarkersVisible(false);
+				SetTimeMarkersVisible(true);
+				SetDrawOnlyBookmarks(false);
 			}
-		}
-		else
-		{
-			SetTimeMarkersVisible(true);
-			SetDrawOnlyBookmarks(false);
-		}
 
-		return FReply::Handled();
+			return FReply::Handled();
+		}
 	}
 	else if (InKeyEvent.GetKey() == EKeys::F)
 	{
-		FrameSelection();
-		return FReply::Handled();
+		if (!InKeyEvent.GetModifierKeys().IsControlDown() && !InKeyEvent.GetModifierKeys().IsShiftDown())
+		{
+			FrameSelection();
+			return FReply::Handled();
+		}
 	}
 	else if (InKeyEvent.GetKey() == EKeys::C)
 	{
@@ -2619,39 +2628,60 @@ FReply STimingView::OnKeyDown(const FGeometry& MyGeometry, const FKeyEvent& InKe
 	}
 	else if (InKeyEvent.GetKey() == EKeys::D) // debug: toggles down-sampling on/off
 	{
-		FTimingEventsTrack::bUseDownSampling = !FTimingEventsTrack::bUseDownSampling;
-		Viewport.AddDirtyFlags(ETimingTrackViewportDirtyFlags::HInvalidated);
-		return FReply::Handled();
+		if (!InKeyEvent.GetModifierKeys().IsControlDown() && !InKeyEvent.GetModifierKeys().IsShiftDown())
+		{
+			FTimingEventsTrack::bUseDownSampling = !FTimingEventsTrack::bUseDownSampling;
+			Viewport.AddDirtyFlags(ETimingTrackViewportDirtyFlags::HInvalidated);
+			return FReply::Handled();
+		}
 	}
 	else if (InKeyEvent.GetKey() == EKeys::R)
 	{
-		FrameSharedState->ShowHideAllFrameTracks();
-		return FReply::Handled();
+		if (!InKeyEvent.GetModifierKeys().IsControlDown() && !InKeyEvent.GetModifierKeys().IsShiftDown())
+		{
+			FrameSharedState->ShowHideAllFrameTracks();
+			return FReply::Handled();
+		}
 	}
 	else if (InKeyEvent.GetKey() == EKeys::Y)
 	{
-		ThreadTimingSharedState->ShowHideAllGpuTracks();
-		return FReply::Handled();
+		if (!InKeyEvent.GetModifierKeys().IsControlDown() && !InKeyEvent.GetModifierKeys().IsShiftDown())
+		{
+			ThreadTimingSharedState->ShowHideAllGpuTracks();
+			return FReply::Handled();
+		}
 	}
 	else if (InKeyEvent.GetKey() == EKeys::U)
 	{
-		ThreadTimingSharedState->ShowHideAllCpuTracks();
-		return FReply::Handled();
+		if (!InKeyEvent.GetModifierKeys().IsControlDown() && !InKeyEvent.GetModifierKeys().IsShiftDown())
+		{
+			ThreadTimingSharedState->ShowHideAllCpuTracks();
+			return FReply::Handled();
+		}
 	}
 	else if (InKeyEvent.GetKey() == EKeys::L)
 	{
-		LoadingSharedState->ShowHideAllLoadingTracks();
-		return FReply::Handled();
+		if (!InKeyEvent.GetModifierKeys().IsControlDown() && !InKeyEvent.GetModifierKeys().IsShiftDown())
+		{
+			LoadingSharedState->ShowHideAllLoadingTracks();
+			return FReply::Handled();
+		}
 	}
 	else if (InKeyEvent.GetKey() == EKeys::I)
 	{
-		FileActivitySharedState->ShowHideAllIoTracks();
-		return FReply::Handled();
+		if (!InKeyEvent.GetModifierKeys().IsControlDown() && !InKeyEvent.GetModifierKeys().IsShiftDown())
+		{
+			FileActivitySharedState->ShowHideAllIoTracks();
+			return FReply::Handled();
+		}
 	}
 	else if (InKeyEvent.GetKey() == EKeys::O)
 	{
-		FileActivitySharedState->ToggleBackgroundEvents();
-		return FReply::Handled();
+		if (!InKeyEvent.GetModifierKeys().IsControlDown() && !InKeyEvent.GetModifierKeys().IsShiftDown())
+		{
+			FileActivitySharedState->ToggleBackgroundEvents();
+			return FReply::Handled();
+		}
 	}
 	else if (InKeyEvent.GetKey() == EKeys::A)
 	{
@@ -2696,7 +2726,18 @@ FReply STimingView::OnKeyDown(const FGeometry& MyGeometry, const FKeyEvent& InKe
 	else if (InKeyEvent.GetKey() == EKeys::X)
 	{
 		uint32 DepthLimit = FTimingProfilerManager::Get()->GetEventDepthLimit();
-		DepthLimit = (DepthLimit != 4) ? 4 : 1000;
+		if (DepthLimit == 1)
+		{
+			DepthLimit = 4;
+		}
+		else if (DepthLimit == 4)
+		{
+			DepthLimit = 1000;
+		}
+		else
+		{
+			DepthLimit = 1;
+		}
 		FTimingProfilerManager::Get()->SetEventDepthLimit(DepthLimit);
 		Viewport.AddDirtyFlags(ETimingTrackViewportDirtyFlags::HInvalidated);
 		return FReply::Handled();
